Initial merge of the new Sequencer Movie Render system. This provides some basic ways to output movies but is mostly scripting based without UI still.

diff --git a/Engine/Plugins/MovieScene/MovieRenderPipeline/MovieRenderPipeline.uplugin b/Engine/Plugins/MovieScene/MovieRenderPipeline/MovieRenderPipeline.uplugin
new file mode 100644
--- /dev/null
+++ b/Engine/Plugins/MovieScene/MovieRenderPipeline/MovieRenderPipeline.uplugin
@@ -0,0 +1,41 @@
+{
+	"FileVersion": 3,
+	"Version": 1,
+	"VersionName": "1.0",
+	"FriendlyName": "Movie Render Pipeline",
+	"Description": "Advanced movie rendering pipeline for use in creating rendered cinematics or other multi-media creation.",
+	"Category": "Scripting",
+	"CreatedBy": "Epic Games Inc.",
+	"CreatedByURL": "",
+	"DocsURL": "",
+	"MarketplaceURL": "",
+	"SupportURL": "",
+	"CanContainContent": true,
+	"IsBetaVersion": true,
+	"Installed": false,
+	"EnabledByDefault": true,
+	"Plugins": [
+	],
+	"Modules": [
+		{
+			"Name": "MovieRenderPipelineCore",
+			"Type": "DeveloperTool",
+			"LoadingPhase": "Default"
+		},
+		{
+			"Name": "MovieRenderPipelineSettings",
+			"Type": "DeveloperTool",
+			"LoadingPhase": "Default"
+		}, 
+		{
+			"Name" : "MovieRenderPipelineRenderPasses",
+			"Type": "DeveloperTool",
+			"LoadingPhase": "Default"
+		},
+		{
+			"Name": "MovieRenderPipelineEditor",
+			"Type": "Editor",
+			"LoadingPhase": "Default"
+		}
+	]
+}
\ No newline at end of file
diff --git a/Engine/Plugins/MovieScene/MovieRenderPipeline/Source/MovieRenderPipelineCore/Private/MoviePipeline.cpp b/Engine/Plugins/MovieScene/MovieRenderPipeline/Source/MovieRenderPipelineCore/Private/MoviePipeline.cpp
new file mode 100644
--- /dev/null
+++ b/Engine/Plugins/MovieScene/MovieRenderPipeline/Source/MovieRenderPipelineCore/Private/MoviePipeline.cpp
@@ -0,0 +1,1033 @@
+// Copyright 1998-2019 Epic Games, Inc. All Rights Reserved.
+#include "MoviePipeline.h"
+#include "Editor.h"
+#include "MovieRenderPipelineCoreModule.h"
+#include "LevelSequence.h"
+#include "Tracks/MovieSceneSubTrack.h"
+#include "Tracks/MovieSceneCameraCutTrack.h"
+#include "LevelSequenceActor.h"
+#include "EngineUtils.h"
+#include "MovieSceneSection.h"
+#include "Sections/MovieSceneCameraCutSection.h"
+#include "Sections/MovieSceneCinematicShotSection.h"
+#include "Tracks/MovieSceneCinematicShotTrack.h"
+#include "CanvasTypes.h"
+#include "AudioDeviceManager.h"
+#include "MovieSceneTimeHelpers.h"
+#include "Engine/World.h"
+#include "GameFramework/PlayerController.h"
+#include "MovieRenderDebugWidget.h"
+#include "MovieRenderPipelineConfig.h"
+#include "MoviePipelineShotConfig.h"
+#include "MovieRenderPipelineDataTypes.h"
+#include "MoviePipelineRenderPass.h"
+#include "MoviePipelineOutput.h"
+#include "ShaderCompiler.h"
+#include "ImageWriteStream.h"
+#include "MoviePipelineAccumulationSetting.h"
+#include "MoviePipelineOutputBuilder.h"
+#include "DistanceFieldAtlas.h"
+
+static TAutoConsoleVariable<int32> CVarMovieRenderPipelineFrameStepper(
+	TEXT("MovieRenderPipeline.FrameStepDebug"),
+	-1,
+	TEXT("How many frames should the Movie Render Pipeline produce before pausing. Set to zero on launch to stall at the first frame. Debug tool.\n")
+	TEXT("-1: Don't pause after each frame (default)\n")
+	TEXT("0: Process engine ticks but don't progress in the movie rendering pipeline.\n")
+	TEXT("1+: Run this many loops of the movie rendering pipeline before pausing again.\n"),
+	ECVF_Default);
+
+
+UMoviePipeline::UMoviePipeline()
+	: CustomTimeStep(nullptr)
+	, CachedCustomTimeStep(nullptr)
+	, TargetSequence(nullptr)
+	, LevelSequenceActor(nullptr)
+	, PipelineState(EMovieRenderPipelineState::Uninitialized)
+	, CurrentShotIndex(-1)
+	, bHasRunBeginFrameOnce(false)
+{
+	CustomTimeStep = CreateDefaultSubobject<UMoviePipelineCustomTimeStep>("MoviePipelineCustomTimeStep");
+	CustomSequenceTimeController = MakeShared<FMoviePipelineTimeController>();
+}
+
+ULevelSequence* PostProcessSequence(ULevelSequence* InSequence)
+{
+	ULevelSequence* DestinationSequence = InSequence;
+
+	return DestinationSequence;
+}
+
+void ValidateSequence(ULevelSequence* InSequence)
+{
+
+}
+
+void UMoviePipeline::Initialize(UMovieRenderPipelineConfig* InConfig)
+{
+	// This function is called after the PIE world has finished initializing, but before
+	// the PIE world is ticked for the first time. We'll end up waiting for the next tick
+	// for FCoreDelegateS::OnBeginFrame to get called to actually start processing.
+	UE_LOG(LogMovieRenderPipeline, Log, TEXT("[%d] Initializing overall Movie Pipeline"), GFrameCounter);
+	
+	check(InConfig);
+	check(PipelineState == EMovieRenderPipelineState::Uninitialized);
+
+	// Ensure this object has the World as part of its Outer (so that it has context to spawn things)
+	check(GetWorld());
+
+	Config = InConfig;
+
+	// Add the Config to root to ensure it doesn't go out of scope until we've shut down.
+	Config->AddToRoot();
+
+	// When start up we want to override the engine's Custom Timestep with our own.
+	// This gives us the ability to completely control the engine tick/delta time before the frame
+	// is started so that we don't have to always be thinking of delta times one frame ahead.
+	CachedCustomTimeStep = GEngine->GetCustomTimeStep();
+	GEngine->SetCustomTimeStep(CustomTimeStep);
+
+	// Duplicate the target level sequence. This reduces complexity by preventing the need to undo all
+	// changes applied to the sequence when running from the Editor. This is a medium-depth copy, we
+	// copy all sub-sequences and pointed to sequences but don't copy any actual assets.
+	ULevelSequence* SequenceAsset = LoadObject<ULevelSequence>(this, *Config->Sequence.GetAssetPathString());
+	if (!ensureAlwaysMsgf(SequenceAsset, TEXT("Failed to load Sequence Asset from specified path, aborting movie render! Path: %s"), *Config->Sequence.GetAssetPathString()))
+	{
+		Shutdown();
+		// ToDo: Promote this to a error that gets broadcast out.
+		return;
+	}
+
+	TargetSequence = CreateCopyOfSequence(SequenceAsset);
+
+	// ToDo: Allow extensions to modify the sequence
+	{
+
+	}
+
+	// Once the sequence has been modified by any extensions (which may have added new sub-sequences, added
+	// or removed tracks, etc.) we want to do some validation on the sequence such as adding Camera Cut tracks
+	// or other needed things we can do to smooth over the user experience. This can modify the root sequence
+	// such as adding a Cinematic Shot Section, so we re-assign our Target Sequence in case it does.
+	TargetSequence = PostProcessSequence(TargetSequence);
+
+	// Now that we've post-processed it, we're going to run a validation pass on it. This will produce warnings
+	// for anything we can't fix that might be an issue - extending sections, etc. This should be const as this
+	// validation should re-use what was used in the UI.
+	ValidateSequence(TargetSequence);
+
+	// Now that we've fixed up the sequence and validated it, we're going to build a list of shots that we need
+	// to produce in a simplified data structure. The simplified structure makes the flow/debugging easier.
+	ShotList = BuildShotListFromSequence(TargetSequence);
+
+	// Finally, we're going to create a Level Sequence Actor in the world that has its settings configured by us.
+	// Because this callback is at the end of PIE startup (and before tick) we should be able to spawn the actor
+	// and give it a chance to tick once (where it should do nothing) before we start manually manipulating it.
+	InitializeLevelSequenceActor(SequenceAsset, TargetSequence);
+
+	// Register any additional engine callbacks needed.
+	{
+		// Called before the Custom Timestep is updated. This gives us a chance to calculate
+		// what we want the frame to look like and then cache that information so that the
+		// Custom Timestep doesn't have to perform its own logic.
+		FCoreDelegates::OnBeginFrame.AddUObject(this, &UMoviePipeline::OnEngineTickBeginFrame);
+		// Called at the end of the frame after everything has been ticked and rendered for the frame.
+		FCoreDelegates::OnEndFrame.AddUObject(this, &UMoviePipeline::OnEngineTickEndFrame);
+	}
+
+	// Construct a debug UI and bind it to this instance.
+	LoadDebugWidget();
+
+	OutputBuilder = MakeShared<FMoviePipelineOutputMerger, ESPMode::ThreadSafe>();
+
+	OutputPipe = MakeShared<FImagePixelPipe, ESPMode::ThreadSafe>();
+
+	// Loop through our Output Containers and let them do any first-frame initialization. It's important that
+	// this happens on the first frame (before we start rendering movies).
+	for (UMoviePipelineOutput* OutputContainer : Config->OutputContainers)
+	{
+		OutputContainer->OnInitializedForPipeline(this);
+	}
+
+	// Initialization is complete. This engine frame is a wash (because the tick started with a 
+	// delta time not generated by us) so we'll wait until the next engine frame to start rendering.
+	UE_LOG(LogMovieRenderPipeline, Log, TEXT("[%d] Initialized Movie Pipeline with [%d Shots]. Expected total output is %d frames."), GFrameCounter, ShotList.Num(), 0);
+
+	// We'll get started on the next frame
+	CurrentShotIndex = 0;
+	PipelineState = EMovieRenderPipelineState::ProducingFrames;
+	InitializationTime = FDateTime::UtcNow();
+}
+
+void UMoviePipeline::Shutdown()
+{
+	// Pipeline should not be shut down if already shut down. You probably have a delegate
+	// loop where we shutdown because finish and you responded by requesting we shutdown.
+	ensure(PipelineState != EMovieRenderPipelineState::Shutdown);
+
+	// Restore any custom Time Step that may have been set before.
+	GEngine->SetCustomTimeStep(CachedCustomTimeStep);
+
+	// Ensure our delegates don't get called anymore as we're going to become null soon.
+	FCoreDelegates::OnBeginFrame.RemoveAll(this);
+	FCoreDelegates::OnEndFrame.RemoveAll(this);
+
+	// If we're not shutting down naturally (ie: we got canceled) then force the output
+	// containers to flush so that all in-flight write processes get finished. If we
+	// did shut down naturally, we should be in the export state.
+	if (PipelineState != EMovieRenderPipelineState::Export)
+	{
+		const bool bForceFinish = true;
+		TickFinalizeOutputContainers(bForceFinish);
+	}
+
+	if (DebugWidget)
+	{
+		DebugWidget->RemoveFromParent();
+		DebugWidget = nullptr;
+	}
+
+	for (UMoviePipelineOutput* Setting : Config->OutputContainers)
+	{
+		Setting->OnPipelineFinished();
+	}
+
+	UE_LOG(LogMovieRenderPipeline, Log, TEXT("Movie Pipeline completed. Duration: %s"), *(FDateTime::UtcNow() - InitializationTime).ToString());
+	PipelineState = EMovieRenderPipelineState::Shutdown;
+}
+
+void UMoviePipeline::OnEngineTickBeginFrame()
+{
+	UE_LOG(LogMovieRenderPipeline, VeryVerbose, TEXT("OnEngineTickBeginFrame (Start) Engine Frame: [%d]"), GFrameCounter);
+	
+	// We should have a custom timestep set up by now.
+	check(CustomTimeStep);
+
+	switch (PipelineState)
+	{
+	case EMovieRenderPipelineState::Uninitialized:
+		// We shouldn't register this delegate until we're initialized.
+		check(false);
+		break;
+	case EMovieRenderPipelineState::ProducingFrames:
+		TickProducingFrames();
+		break;
+	case EMovieRenderPipelineState::Finalize:
+		TickFinalizeOutputContainers(false);
+		break;
+	case EMovieRenderPipelineState::Export:
+		// We support ticking export state as well to allow async long export steps.
+		TickPostFinalizeExport();
+		break;
+	}
+
+	bHasRunBeginFrameOnce = true;
+	UE_LOG(LogMovieRenderPipeline, VeryVerbose, TEXT("OnEngineTickBeginFrame (End) Engine Frame: [%d]"), GFrameCounter);
+}
+
+void UMoviePipeline::OnSequenceEvaluated(const UMovieSceneSequencePlayer& Player, FFrameTime CurrentTime, FFrameTime PreviousTime)
+{
+	// This callback exists for logging purposes. DO NOT HINGE LOGIC ON THIS CALLBACK
+	// because this may get called multiple times per frame and may be the result of
+	// a seek operation which is reverted before a frame is even rendered.
+	UE_LOG(LogMovieRenderPipeline, VeryVerbose, TEXT("[GFrameCounter: %d] Sequence Evaluated. CurrentTime: %s PreviousTime: %s"), GFrameCounter, *LexToString(CurrentTime), *LexToString(PreviousTime));
+}
+
+void UMoviePipeline::OnEngineTickEndFrame()
+{
+	// Unfortunately, since we can't control when our Initialization function is called
+	// we can end up in a situation where this callback is registered but the matching
+	// OnEngineTickBeginFrame() hasn't been called for that given engine tick. Instead of
+	// changing this registration to hang off of the end of the first OnEngineTickBeginFrame()
+	// we instead just early out here if that hasn't actually been called once. This decision
+	// is designed to minimize places where callbacks are registered and where flow changes.
+	if (!bHasRunBeginFrameOnce)
+	{
+		return;
+	}
+
+	// Early out if we're idling as we don't want to process a frame. This prevents us from
+	// overwriting render state when the engine is processing ticks but we're not trying to
+	// change the evaluation. 
+	if (IsDebugFrameStepIdling())
+	{
+		return;
+	}
+
+	// It is important that there is no early out that skips hitting this
+	// (Otherwise we don't pause on the frame we transition from step -> idle
+	// and the world plays even though the state is frozen).
+	if (bPauseAtEndOfFrame)
+	{
+		GetWorld()->GetFirstPlayerController()->SetPause(true);
+		bPauseAtEndOfFrame = false;
+	}
+
+	UE_LOG(LogMovieRenderPipeline, VeryVerbose, TEXT("OnEngineTickEndFrame (Start) Engine Frame: [%d]"), GFrameCounter);
+
+	RenderFrame();
+
+	UE_LOG(LogMovieRenderPipeline, VeryVerbose, TEXT("OnEngineTickEndFrame (End) Engine Frame: [%d]"), GFrameCounter);
+}
+
+void UMoviePipeline::BeginFinalize()
+{
+	// Notify all of our output containers that we have finished producing and
+	// submitting all frames to them and that they should start any async flushes.
+	for (UMoviePipelineOutput* Container : Config->OutputContainers)
+	{
+		Container->BeginFinalize();
+	}
+}
+
+void UMoviePipeline::TickFinalizeOutputContainers(const bool bInForceFinish)
+{
+	// Tick all containers until they all report that they have finalized.
+	bool bAllContainsFinishedProcessing;
+
+	do
+	{
+		bAllContainsFinishedProcessing = true;
+
+		// Ask the containers if they're all done processing.
+		for (UMoviePipelineOutput* Container : Config->OutputContainers)
+		{
+			bAllContainsFinishedProcessing &= Container->HasFinishedProcessing();
+		}
+
+		// If we aren't forcing a finish, early out after one loop to keep
+		// the editor/ui responsive.
+		if (!bInForceFinish || bAllContainsFinishedProcessing)
+		{
+			break;
+		}
+
+		// If they've reached here, they're forcing them to finish so we'll sleep for a touch to give
+		// everyone a chance to actually do work before asking them if they're done.
+		FPlatformProcess::Sleep(1.f);
+
+	} while (true);
+
+	// If an output container is still working, we'll early out to keep the UI responsive.
+	// If they've forced a finish this will have to be true before we can reach this block.
+	if (!bAllContainsFinishedProcessing)
+	{
+		return;
+	}
+
+	for (UMoviePipelineOutput* Container : Config->OutputContainers)
+	{
+		// All containers have finished processing, final shutdown.
+		Container->Finalize();
+	}
+
+	// Move onto the Export stage and call the Export function directly. Most export steps are
+	// instantaneous and will complete here, but on the next tick we will check them all anyways
+	// to allow running long-running export processes (such as encoding).
+	PipelineState = EMovieRenderPipelineState::Export;
+	PostFinalizeExport();
+}
+
+void UMoviePipeline::PostFinalizeExport()
+{
+	check(PipelineState == EMovieRenderPipelineState::Export);
+	
+	// This is called once
+}
+
+void UMoviePipeline::TickPostFinalizeExport()
+{
+	// This step assumes you have produced data and filled the data structures.
+	check(PipelineState == EMovieRenderPipelineState::Export);
+	UE_LOG(LogMovieRenderPipeline, Log, TEXT("[%d] PostFinalize Export (Start)."), GFrameCounter);
+
+	// ToDo: Loop through any extensions (such as XML export) and let them export using all of the
+	// data that was generated during this run such as containers, output names and lengths.
+
+
+
+	bool bAllExportsFinishedProcessing = true;
+
+	// Ask the containers if they're all done processing.
+	//for (UMoviePipelineOutput* Container : Config->OutputContainers)
+	//{
+	//	bAllContainsFinishedProcessing &= Container->HasFinishedProcessing();
+	//}
+
+	if (!bAllExportsFinishedProcessing)
+	{
+		return;
+	}
+
+
+	UE_LOG(LogMovieRenderPipeline, Log, TEXT("[%d] PostFinalize Export (End)."), GFrameCounter);
+
+	// Now that we have done the last step of the process we're going to call Shutdown on ourself
+	// to unregister all delegates, and then we'll broadcast to anyone listening that we're done.
+	// This allows an external controller to decide what to do, ie: start another one, shut down
+	// PIE or exit the game.
+	Shutdown();
+
+	UE_LOG(LogMovieRenderPipeline, Verbose, TEXT("[%d] Broadcasting Post-Finalize Finished event."), GFrameCounter);
+	OnMoviePipelineFinished().Broadcast(this);
+}
+
+bool UMoviePipelineCustomTimeStep::UpdateTimeStep(UEngine* /*InEngine*/)
+{
+	if (ensureMsgf(!FMath::IsNearlyZero(TimeCache.DeltaTime), TEXT("An incorrect or uninitialized time step was used! Delta Time of 0 isn't allowed.")))
+	{
+		FApp::UpdateLastTime();
+		FApp::SetDeltaTime(TimeCache.DeltaTime);
+		FApp::SetCurrentTime(FApp::GetCurrentTime() + FApp::GetDeltaTime());
+	}
+
+	// Clear our cached time to ensure we're always explicitly telling this what to do and never relying on the last set value.
+	// (This will cause the ensure above to check on the next tick if someone didn't reset our value.)
+	TimeCache = MoviePipeline::FFrameTimeStepCache();
+
+	// Return false so the engine doesn't run its own logic to overwrite FApp timings.
+	return false;
+}
+
+ULevelSequence* UMoviePipeline::CreateCopyOfSequence(ULevelSequence* InSequence)
+{
+	// We perform a medium-depth copy here. The standard Duplicate Object
+	// function only performs a shallow copy which will duplicate all of the
+	// tracks and sections that belong to this sequence and do reference fix
+	// up. We don't want a deep copy (because we don't want to copy all assets)
+	// but because we are modifying the Playback Start/End on sub-sequences
+	// we need to manually copy those and fix up the references ourself.
+
+	// Duplicate the sequence and change the outer to be ourself.
+	FObjectDuplicationParameters DuplicationParams(InSequence, this);
+	DuplicationParams.DestName = FName(*(InSequence->GetName() + TEXT("_Instanced")));
+	ULevelSequence* DuplicatedSequence = (ULevelSequence*)StaticDuplicateObjectEx(DuplicationParams);
+
+	// // Now that we've copied this sequence, we'll look for any camera cut sections and duplicate those.
+	// TArray<UMovieSceneCameraCutTrack*> SubTracks = DuplicatedSequence->GetMovieScene()->GetMasterTracks<UMovieSceneCameraCutTrack>();
+	// 
+	// for (const UMovieSceneCameraCutTrack* CameraCutTrack : SubTracks)
+	// {
+	// 	for (const UMovieSceneSection* Section : CameraCutTrack->GetAllSections())
+	// 	{
+	// 		UMovieSceneCameraCutSection* CameraCutSection = CastChecked<UMovieSceneCameraCutSection>(Section);
+	// 		CameraCutSection
+	// 	}
+	// }
+
+	return DuplicatedSequence;
+}
+
+void UMoviePipeline::InitializeLevelSequenceActor(ULevelSequence* OriginalLevelSequence, ULevelSequence* InSequenceToApply)
+{
+	// There is a reasonable chance that there exists a Level Sequence Actor in the world
+	// already set up to play this sequence. We don't want that however as it may have unknown
+	// settings on it.
+	for (auto It = TActorIterator<ALevelSequenceActor>(GetWorld()); It; ++It)
+	{
+		// Iterate through all of them in the event someone has multiple copies in the world on accident.
+		if (It->LevelSequence == OriginalLevelSequence)
+		{
+			// Found it!
+			ALevelSequenceActor* ExistingActor = *It;
+
+			//  Clear the sequence ref instead of killing actor to reduce likelyhood of game code throwing a NRE
+			// that may have expected this actor to exist.
+			if (ExistingActor->GetSequencePlayer())
+			{
+				ExistingActor->GetSequencePlayer()->Stop();
+			}
+
+			// Sequence must have been stopped before we can switch sequences.
+			ExistingActor->SetSequence(nullptr);
+		}
+	}
+
+	// Now we can construct a nice new actor which has the settings we need.
+	LevelSequenceActor = GetWorld()->SpawnActor<ALevelSequenceActor>();
+	check(LevelSequenceActor);
+
+	// Use our duplicated sequence
+	LevelSequenceActor->SetSequence(InSequenceToApply);
+
+	// Enforce settings.
+	LevelSequenceActor->PlaybackSettings.LoopCount.Value = 0;
+	LevelSequenceActor->PlaybackSettings.bAutoPlay = false;
+	LevelSequenceActor->PlaybackSettings.bPauseAtEnd = true;
+	LevelSequenceActor->GetSequencePlayer()->SetTimeController(CustomSequenceTimeController);
+
+	LevelSequenceActor->GetSequencePlayer()->OnSequenceUpdated().AddUObject(this, &UMoviePipeline::OnSequenceEvaluated);
+}
+
+
+FMoviePipelineShotCache CreateShotFromMovieScene(const UMovieScene* InMovieScene, const TRange<FFrameNumber>& InIntersectionRange)
+{
+	check(InMovieScene);
+
+	// We will generate a range for each sub-section within the shot we wish to render. These will be clipped by
+	// the overall range of the shot to represent the final desired render duration.
+	struct FCameraCutRange
+	{
+		TRange<FFrameNumber> Range;
+		UMovieSceneCameraCutSection* Section;
+	};
+
+	TArray<FCameraCutRange> IntersectedRanges;
+
+	// We're going to search for Camera Cut tracks within this shot. If none are found, we'll use the whole range of the shot.
+	UMovieSceneCameraCutTrack* CameraCutTrack = InMovieScene->FindMasterTrack<UMovieSceneCameraCutTrack>();
+	if (CameraCutTrack)
+	{
+		for (UMovieSceneSection* Section : CameraCutTrack->GetAllSections())
+		{
+			UMovieSceneCameraCutSection* CameraCutSection = CastChecked<UMovieSceneCameraCutSection>(Section);
+
+			// ToDo: Inner vs. Outer resolution differences.
+			// Intersect this cut with the outer range in the likely event that the section goes past the bounds.
+			TRange<FFrameNumber> IntersectingRange = TRange<FFrameNumber>::Intersection(Section->GetRange(), InIntersectionRange);
+
+			FCameraCutRange& NewRange = IntersectedRanges.AddDefaulted_GetRef();
+			NewRange.Range = IntersectingRange;
+			NewRange.Section = CameraCutSection;
+		}
+	}
+	else
+	{
+		// No camera cut track was found inside. We'll treat the whole shot as the desired range.
+		FCameraCutRange& NewRange = IntersectedRanges.AddDefaulted_GetRef();
+		NewRange.Range = InIntersectionRange;
+		NewRange.Section = nullptr;
+	}
+
+	FMoviePipelineShotCache NewShot;
+	NewShot.OriginalRange = InIntersectionRange;
+	NewShot.TotalOutputRange = NewShot.OriginalRange;
+
+	for (const FCameraCutRange& Range : IntersectedRanges)
+	{
+		// Generate a CameraCut for each range.
+		FMoviePipelineShotCutCache& CameraCut = NewShot.CameraCuts.AddDefaulted_GetRef();
+
+		CameraCut.CameraCutSection = Range.Section; // May be nullptr.
+		CameraCut.OriginalRange = Range.Range;
+		CameraCut.TotalOutputRange = CameraCut.OriginalRange;
+	}
+
+	return NewShot;
+}
+
+TArray<FMoviePipelineShotCache> UMoviePipeline::BuildShotListFromSequence(const ULevelSequence* InSequence)
+{
+	TArray<FMoviePipelineShotCache> NewShotList;
+	
+	// Shot Tracks take precedent over camera cuts, as settings can only be applied as granular as a shot.
+	UMovieSceneCinematicShotTrack* CinematicShotTrack = InSequence->GetMovieScene()->FindMasterTrack<UMovieSceneCinematicShotTrack>();
+	
+	if (CinematicShotTrack)
+	{
+		ensureMsgf(CinematicShotTrack->GetAllSections().Num() > 0, TEXT("Cinematic Shot track must have at least one section to be rendered."));
+
+		for (UMovieSceneSection* Section : CinematicShotTrack->GetAllSections())
+		{
+			UMovieSceneCinematicShotSection* ShotSection = CastChecked<UMovieSceneCinematicShotSection>(Section);
+
+			// If the user has manually marked a section as inactive we don't produce a shot for it.
+			if (!Section->IsActive())
+			{
+				UE_LOG(LogMovieRenderPipeline, Verbose, TEXT("Skipped adding Shot %s to Shot List due to being inactive."), *ShotSection->GetShotDisplayName());
+				continue;
+			}
+
+			if (!ShotSection->GetSequence())
+			{
+				UE_LOG(LogMovieRenderPipeline, Warning, TEXT("Skipped adding Shot %s to Shot List due to no inner sequence."), *ShotSection->GetShotDisplayName());
+				continue;
+			}
+
+			// The Shot Section may extend past our Sequence's Playback Bounds. We intersect the two bounds to ensure that
+			// the Playback Start/Playback End of the overall sequence is respected.
+			TRange<FFrameNumber> CinematicShotSectionRange = TRange<FFrameNumber>::Intersection(ShotSection->GetRange(), InSequence->GetMovieScene()->GetPlaybackRange());
+
+			FMoviePipelineShotCache NewShot = CreateShotFromMovieScene(ShotSection->GetSequence()->GetMovieScene(), CinematicShotSectionRange);
+
+			// The first thing we do is find the appropriate configuration from the settings. Each shot can have its own config
+			// or they fall back to a default one specified for the whole pipeline.
+			NewShot.ShotConfig = Config->GetConfigForShot(ShotSection->GetShotDisplayName());
+			NewShot.CinematicShotSection = ShotSection;
+
+			// There should always be a shot config as the Pipeline default is returned in the event they didn't customize.
+			check(NewShot.ShotConfig);
+
+			NewShotList.Add(MoveTemp(NewShot));
+		}
+	}
+	else
+	{
+		// They don't have a cinematic shot track. We'll slice them up by camera cuts instead.
+		FMoviePipelineShotCache NewShot = CreateShotFromMovieScene(TargetSequence->GetMovieScene(), TargetSequence->GetMovieScene()->GetPlaybackRange());
+		NewShot.ShotConfig = Config->DefaultShotConfig;
+		
+		NewShotList.Add(MoveTemp(NewShot));
+	}
+
+	// If they don't have a cinematic shot track, or a camera cut track then they want to control the camera
+	// through their own logic. We'll just use the duration of the Sequence as the render, plus warn them.
+	if (NewShotList.Num() == 0)
+	{
+		UE_LOG(LogMovieRenderPipeline, Warning, TEXT("No Cinematic Shot Tracks found, and no Camera Cut Tracks found. Playback Range will be used but camera will render from Pawns perspective."));
+		FMoviePipelineShotCache NewShot;
+		NewShot.ShotConfig = Config->DefaultShotConfig;
+
+		FMoviePipelineShotCutCache& CameraCut = NewShot.CameraCuts.AddDefaulted_GetRef();
+		CameraCut.OriginalRange = TargetSequence->GetMovieScene()->GetPlaybackRange();
+	}
+
+	// Now that we've gathered at least one or more shots with one or more cuts, we can apply settings. It's easier to
+	// debug when all of the shots are calculated up front and debug info is printed in a block instead of as it is reached.
+	int32 ShotIndex = 1;
+	for (FMoviePipelineShotCache& Shot : NewShotList)
+	{
+		UE_LOG(LogMovieRenderPipeline, Log, TEXT("Shot %d/%d has %d Camera Cuts."), ShotIndex, NewShotList.Num(), Shot.CameraCuts.Num());
+		ShotIndex++;
+
+		Shot.NumHandleFrames = 0; // ToDo: Pull this from a config.
+		UMoviePipelineAccumulationSetting* AccumulationSettings = Shot.ShotConfig->FindOrAddSetting<UMoviePipelineAccumulationSetting>();
+
+		// Expand the shot to encompass handle frames. This will modify our Camera Cuts bounds.
+		ExpandShot(Shot);
+
+		FString ShotName;
+		if (Shot.CinematicShotSection.IsValid())
+		{
+			ShotName = Shot.CinematicShotSection->GetShotDisplayName();
+		}
+
+		for (FMoviePipelineShotCutCache& CameraCut : Shot.CameraCuts)
+		{
+			CameraCut.NumWarmUpFramesRemaining = CameraCut.NumWarmUpFrames = 4; // ToDo: pull from config
+			CameraCut.bAccurateFirstFrameHistory = true; // ToDo: pull from config
+			CameraCut.NumTemporalSamples = AccumulationSettings->TemporalSampleCount;
+			CameraCut.NumSpatialSamples = AccumulationSettings->SpatialSampleCount;
+			CameraCut.CachedFrameRate = Config->GetEffectiveFrameRate();
+			CameraCut.CachedTickResolution = TargetSequence->GetMovieScene()->GetTickResolution();
+			
+			FString CameraName;
+			if (CameraCut.CameraCutSection.IsValid())
+			{
+				CameraName = TargetSequence->GetMovieScene()->FindBinding(CameraCut.CameraCutSection->GetCameraBindingID().GetGuid())->GetName();
+			}
+
+			CameraCut.DisplayName = FString::Printf(TEXT("Shot: %s Cam: %s"), *ShotName, *CameraName);
+			
+			// When we expanded the shot above, it pushed the first/last camera cuts ranges to account for Handle Frames.
+			// We want to start rendering from the first handle frame. Shutter Timing is a fixed offset from this number.
+			CameraCut.CurrentTick = CameraCut.TotalOutputRange.GetLowerBoundValue();
+		}
+	}
+
+	return NewShotList;
+}
+
+void UMoviePipeline::SetSoloCameraCut(FMoviePipelineShotCutCache& InCameraCut)
+{
+
+}
+
+void UMoviePipeline::InitializeShot(FMoviePipelineShotCache& InShot)
+{
+	// We handle both tearing down the previous shot (if there is one) and
+	// standing up the specified shot here. This gives us a good chance to
+	// notify output containers/extensions that we're switching contexts
+	// and they can decide what to do with it (ie: split into multiple files)
+	TOptional<FMoviePipelineShotCache> PrevShot;
+	int32 NewShotIndex = ShotList.IndexOfByKey(InShot);
+	if (NewShotIndex > 0)
+	{
+		PrevShot = ShotList[NewShotIndex - 1];
+	}
+
+	// Set the new shot as the active shot. This enables the specified shot section and disables all other shot sections.
+	SetSoloShot(InShot);
+	
+	// ToDo
+	for (UMoviePipelineSetting* Setting : InShot.ShotConfig->GetSettings())
+	{
+		Setting->OnInitializedForPipeline(this);
+	}
+
+	for (UMoviePipelineSetting* Setting : InShot.ShotConfig->InputBuffers)
+	{
+		Setting->OnInitializedForPipeline(this);
+	}
+
+
+	// Setup required rendering architecture for all passes in this shot.
+	SetupRenderingPipelineForShot(InShot);
+}
+
+void UMoviePipeline::SetSoloShot(const FMoviePipelineShotCache& InShot)
+{
+	// Iterate through the shot list and ensure all shots have been set to inactive.
+	for (FMoviePipelineShotCache& Shot : ShotList)
+	{
+		UMovieSceneCinematicShotSection* Section = Shot.CinematicShotSection.Get();
+
+		// Skip the one we will be modifying below to avoid extra MarkAsChanged calls.
+		if (Section && Section != InShot.CinematicShotSection.Get())
+		{
+			Section->SetIsActive(false);
+			Section->MarkAsChanged();
+		}
+	}
+
+	// Now that we've set them all to inactive we'll ensure that our passed in shot is active.
+	if (UMovieSceneCinematicShotSection* CurrentShotSection = InShot.CinematicShotSection.Get())
+	{
+		UE_LOG(LogMovieRenderPipeline, Verbose, TEXT("Disabled all shot tracks and re-enabling %s for solo."), *CurrentShotSection->GetShotDisplayName());
+		CurrentShotSection->SetIsActive(true);
+		CurrentShotSection->MarkAsChanged();
+	}
+	else
+	{
+		UE_LOG(LogMovieRenderPipeline, Verbose, TEXT("Skipped enabling a shot track due to not shot section associated with the provided shot."));
+	}
+}
+
+FFrameNumber UMoviePipeline::GetMotionBlurDuration(int32 InShotIndex) const
+{
+	const FMoviePipelineShotCache& CurrentShot = ShotList[InShotIndex];
+	UMoviePipelineAccumulationSetting* AccumulationSettings = CurrentShot.ShotConfig->FindOrAddSetting<UMoviePipelineAccumulationSetting>();
+
+	const MoviePipeline::FFrameConstantMetrics FrameMetrics = CalculateShotFrameMetrics(CurrentShot);
+	
+	// We don't have data from before the first frame, at least from the common user perspective. To solve this,
+	// we can take the sample from the next frame, and then evaluate the first frame again and we will get a
+	// motion vector which is a close-enough approximation of the real motion (had there been any). While only
+	// an approximation, it is worth the trade off to be significantly more user friendly in regards to data prep.
+	return FrameMetrics.TicksPerSample.FloorToFrame();
+}
+
+FFrameNumber UMoviePipeline::CalculateHandleFrameDuration(const FMoviePipelineShotCache& InShot) const
+{
+	// Convert to Tick Resolution
+	FFrameNumber HandleFrameDuration = FFrameRate::TransformTime(FFrameTime(FFrameNumber(InShot.NumHandleFrames)),
+		TargetSequence->GetMovieScene()->GetDisplayRate(),
+		TargetSequence->GetMovieScene()->GetTickResolution()).FloorToFrame();
+
+	return HandleFrameDuration;
+}
+
+void UMoviePipeline::ExpandShot(FMoviePipelineShotCache& InShot)
+{
+	const MoviePipeline::FFrameConstantMetrics FrameMetrics = CalculateShotFrameMetrics(InShot);
+
+	// Handle Frames will be added onto our original shot size. We track the ranges separately for counting purposes
+	// later - the actual rendering code is unaware of the handle frames. Handle Frames only apply to the shot
+	// and expand the first/last inner cut to cover this area.
+	FFrameNumber HandleFrameTicks = CalculateHandleFrameDuration(InShot);
+	InShot.TotalOutputRange = MovieScene::DilateRange(InShot.OriginalRange, -HandleFrameTicks, HandleFrameTicks);
+
+	// We'll expand the overall sequence by at least this much to ensure we don't get clamped. This is +1 frame for shutter timing overlap.
+	TRange<FFrameNumber> TotalPlaybackRange = MovieScene::DilateRange(InShot.TotalOutputRange, -FrameMetrics.TicksPerOutputFrame.CeilToFrame(), FrameMetrics.TicksPerOutputFrame.CeilToFrame());
+
+	// We can now take the difference between the two ranges to isolate just the opening/closing range (so we can use the range for counting later)
+	TArray<TRange<FFrameNumber>> HandleFrameRanges = TRange<FFrameNumber>::Difference(InShot.TotalOutputRange, InShot.OriginalRange);
+	if (HandleFrameRanges.Num() == 2)
+	{
+		InShot.HandleFrameRangeStart = HandleFrameRanges[0];
+		InShot.HandleFrameRangeEnd = HandleFrameRanges[1];
+	}
+
+	// Expand the first and last cut inside this shot to cover the handle frame distance.
+	if (InShot.CameraCuts.Num() > 0)
+	{
+		// These may be the same or they may be different.
+		FMoviePipelineShotCutCache& FirstCut = InShot.CameraCuts[0];
+		FMoviePipelineShotCutCache& LastCut = InShot.CameraCuts[InShot.CameraCuts.Num() - 1];
+
+		FirstCut.TotalOutputRange = MovieScene::DilateRange(FirstCut.TotalOutputRange, -HandleFrameTicks, FFrameNumber(0));
+		LastCut.TotalOutputRange = MovieScene::DilateRange(LastCut.TotalOutputRange, FFrameNumber(0), HandleFrameTicks);
+	}
+
+	// When we expand a shot section, we also need to expand the inner sequence's range by the same amount.
+	UMovieSceneCinematicShotSection* ShotSection = InShot.CinematicShotSection.Get();
+	if(ShotSection)
+	{
+		UMovieScene* ShotMovieScene = ShotSection->GetSequence() ? ShotSection->GetSequence()->GetMovieScene() : nullptr;
+		if (!ShotMovieScene)
+		{
+			UE_LOG(LogMovieRenderPipeline, Warning, TEXT("ShotSection (%s) with no inner sequence, skipping expansion!"), *ShotSection->GetShotDisplayName());
+			return;
+		}
+
+		// Expand the inner shot section range by the handle size, multiplied by the difference between the outer and inner tick resolutions (and factoring in the time scale)
+		const bool bRateMatches = TargetSequence->GetMovieScene()->GetTickResolution() == ShotMovieScene->GetTickResolution();
+		const double OuterToInnerRateDilation = bRateMatches ? 1.0 : (ShotMovieScene->GetTickResolution() / TargetSequence->GetMovieScene()->GetTickResolution()).AsDecimal();
+		const double OuterToInnerScale = OuterToInnerRateDilation * (double)ShotSection->Parameters.TimeScale;
+
+		// We push the playback bounds one frame further (in both directions) to handle the case where shutter timings offset them +/- the actual edge.
+		// These bounds don't affect the period of time which is rendered, so it's okay that this pushes beyond the actual evaluated range.
+
+		const FFrameNumber StartOffset = HandleFrameTicks + FrameMetrics.TicksPerOutputFrame.FloorToFrame();
+		const FFrameNumber EndOffset = HandleFrameTicks + FrameMetrics.TicksPerOutputFrame.FloorToFrame();
+
+		const FFrameNumber DilatedStartOffset = FFrameNumber(FMath::CeilToInt(StartOffset.Value * OuterToInnerScale));
+		const FFrameNumber DilatedEndOffset = FFrameNumber(FMath::CeilToInt(EndOffset.Value * OuterToInnerScale));
+
+		TRange<FFrameNumber> OldPlaybackRange = ShotMovieScene->GetPlaybackRange();
+		TRange<FFrameNumber> NewPlaybackRange = MovieScene::DilateRange(ShotMovieScene->GetPlaybackRange(), -DilatedStartOffset, DilatedEndOffset);
+
+		FFrameNumber TotalExpansionSizeInFrames = FFrameRate::TransformTime(FFrameTime(HandleFrameTicks) + FrameMetrics.TicksPerOutputFrame, FrameMetrics.TickResolution, FrameMetrics.FrameRate).CeilToFrame();
+		// ToDo: ShotList isn't assigned yet so IndexOfByKey fails.
+		UE_LOG(LogMovieRenderPipeline, Log, TEXT("ShotSection ([%d] %s) Expanding Shot by [-%d, %d] Ticks (+/-%d Frames). Old Inner Playback Range: %s New Inner Playback Range: %s"),
+			ShotList.IndexOfByKey(InShot), *ShotSection->GetShotDisplayName(), 
+			StartOffset.Value, EndOffset.Value, TotalExpansionSizeInFrames.Value,
+			*LexToString(OldPlaybackRange), *LexToString(NewPlaybackRange));
+
+		// Expand the inner scene
+		ShotMovieScene->SetPlaybackRange(NewPlaybackRange, false);
+
+		// Expand the outer owning section
+		ShotSection->SetRange(MovieScene::DilateRange(ShotSection->GetRange(), -StartOffset, EndOffset));
+
+	}
+
+	// Ensure the overall Movie Scene Playback Range is large enough. This will clamp evaluation if we don't expand it. We hull the existing range
+	// with the new range.
+	TRange<FFrameNumber> EncompassingPlaybackRange = TRange<FFrameNumber>::Hull(TotalPlaybackRange, TargetSequence->MovieScene->GetPlaybackRange());
+	TargetSequence->GetMovieScene()->SetPlaybackRange(EncompassingPlaybackRange);
+
+}
+
+bool UMoviePipeline::IsDebugFrameStepIdling() const
+{
+	// We're only idling when we're at zero, otherwise there's more frames to process.
+	// Caveat is that this will be zero on the last frame we want to render, so we
+	// take into account whether or not we've queued up a pause at the end of the frame
+	// which is indicator that we want to process the current frame.
+	int32 DebugFrameStepValue = CVarMovieRenderPipelineFrameStepper.GetValueOnGameThread();
+	return DebugFrameStepValue == 0 && !bPauseAtEndOfFrame;
+}
+
+bool UMoviePipeline::DebugFrameStepPreTick()
+{
+	int32 DebugFrameStepValue = CVarMovieRenderPipelineFrameStepper.GetValueOnGameThread();
+	if (DebugFrameStepValue == 0)
+	{
+		// A value of 0 means that they are using the frame stepper and that we have stepped
+		// the specified number of frames. We will create a DeltaTime for the engine
+		// and not process anything below, which prevents us from trying to produce an
+		// output frame later.
+		CustomTimeStep->SetCachedFrameTiming(MoviePipeline::FFrameTimeStepCache(1 / 24.0));
+		return true;
+	}
+	else if (DebugFrameStepValue > 0)
+	{
+		// They want to process at least one frame, deincrement, then we
+		// process the frame. We pause the game here to preserve render state.
+		CVarMovieRenderPipelineFrameStepper->Set(DebugFrameStepValue - 1, ECVF_SetByConsole);
+
+		// We want to run this one frame and then pause again at the end.
+		bPauseAtEndOfFrame = true;
+	}
+
+	return false;
+}
+
+FMoviePipelineShotCache UMoviePipeline::GetCurrentShotSnapshot() const
+{
+	if (CurrentShotIndex >= 0 && CurrentShotIndex < ShotList.Num())
+	{
+		return ShotList[CurrentShotIndex];
+	}
+
+	return FMoviePipelineShotCache();
+}
+
+FMoviePipelineShotCutCache UMoviePipeline::GetCurrentCameraCutSnapshot() const
+{
+	FMoviePipelineShotCache CurrentShot = GetCurrentShotSnapshot();
+	if (CurrentShot.CurrentCameraCutIndex >= 0 && CurrentShot.CurrentCameraCutIndex < CurrentShot.CameraCuts.Num())
+	{
+		return CurrentShot.GetCurrentCameraCut();
+	}
+
+	return FMoviePipelineShotCutCache();
+}
+
+
+FMoviePipelineFrameOutputState UMoviePipeline::GetOutputStateSnapshot() const
+{
+	return CachedOutputState;
+}
+
+void UMoviePipeline::LoadDebugWidget()
+{
+	FSoftClassPath DebugWidgetClassRef(TEXT("/MovieRenderPipeline/Blueprints/UI_MovieRenderPipelineScreenOverlay.UI_MovieRenderPipelineScreenOverlay_C"));
+	if (UClass* DebugWidgetClass = DebugWidgetClassRef.TryLoadClass<UMovieRenderDebugWidget>())
+	{
+		DebugWidget = CreateWidget<UMovieRenderDebugWidget>(GetWorld(), DebugWidgetClass);
+		if (DebugWidget)
+		{
+			DebugWidget->OnInitializedForPipeline(this);
+			DebugWidget->AddToViewport();
+		}
+		else
+		{
+			UE_LOG(LogMovieRenderPipeline, Warning, TEXT("Failed to create Debug Screen UMG Widget. No debug overlay available."));
+		}
+	}
+	else
+	{
+		UE_LOG(LogMovieRenderPipeline, Warning, TEXT("Failed to find Debug Screen UMG Widget class. No debug overlay available."));
+	}
+}
+
+
+void UMoviePipeline::FlushAsyncSystems()
+{
+	// Flush Level Streaming. This solves the problem where levels that are not controlled
+	// by the Sequencer Level Visibility track are marked for Async Load by a gameplay system.
+	// This will register any new actors/components that were spawned during this frame. This needs 
+	// to be done before the shader compiler is flushed so that we compile shaders for any newly
+	// spawned component materials.
+	if (GetWorld())
+	{
+		GetWorld()->FlushLevelStreaming(EFlushLevelStreamingType::Full);
+	}
+
+	// Now we can flush the shader compiler. ToDo: This should probably happen right before SendAllEndOfFrameUpdates() is normally called
+	if (GShaderCompilingManager && GShaderCompilingManager->GetNumRemainingJobs() > 0)
+	{
+		UE_LOG(LogMovieRenderPipeline, Log, TEXT("[%d] Waiting for %d shaders to finish compiling..."), GFrameCounter, GShaderCompilingManager->GetNumRemainingJobs());
+		GShaderCompilingManager->FinishAllCompilation();
+		UE_LOG(LogMovieRenderPipeline, Log, TEXT("[%d] Done waiting for shaders to finish."), GFrameCounter);
+	}
+
+	// Flush the Mesh Distance Field builder as well.
+	if (GDistanceFieldAsyncQueue && GDistanceFieldAsyncQueue->GetNumOutstandingTasks() > 0)
+	{
+		UE_LOG(LogMovieRenderPipeline, Log, TEXT("[%d] Waiting for %d Mesh Distance Fields to finish building..."), GFrameCounter, GDistanceFieldAsyncQueue->GetNumOutstandingTasks());
+		GDistanceFieldAsyncQueue->BlockUntilAllBuildsComplete();
+		UE_LOG(LogMovieRenderPipeline, Log, TEXT("[%d] Done waiting for Mesh Distance Fields to build."), GFrameCounter);
+	}
+}
+
+FFrameTime FMoviePipelineTimeController::OnRequestCurrentTime(const FQualifiedFrameTime& InCurrentTime, float InPlayRate)
+{
+	FFrameTime RequestTime = FFrameRate::TransformTime(TimeCache.Time, TimeCache.Rate, InCurrentTime.Rate);
+	UE_LOG(LogMovieRenderPipeline, VeryVerbose, TEXT("[%d] OnRequestCurrentTime: %d %f"), GFrameCounter, RequestTime.FloorToFrame().Value, RequestTime.GetSubFrame());
+
+	return RequestTime;
+}
+
+MoviePipeline::FFrameConstantMetrics UMoviePipeline::CalculateShotFrameMetrics(const FMoviePipelineShotCache& InShot) const
+{
+	MoviePipeline::FFrameConstantMetrics Output;
+	Output.TickResolution = TargetSequence->GetMovieScene()->GetTickResolution();
+	Output.FrameRate = Config->GetEffectiveFrameRate();
+	Output.TicksPerOutputFrame = FFrameRate::TransformTime(FFrameTime(FFrameNumber(1)), Output.FrameRate, Output.TickResolution);
+
+	UMoviePipelineAccumulationSetting* AccumulationSettings = InShot.ShotConfig->FindOrAddSetting<UMoviePipelineAccumulationSetting>();
+
+	// (CameraShutterAngle/360) gives us the fraction-of-the-output-frame the accumulation frames should cover.
+	Output.ShutterAnglePercentage = AccumulationSettings->CameraShutterAngle / 360.0;
+
+	{
+		/*
+		* Calculate out how many ticks a normal sub-frame occupies.
+		* (TickRes/FrameRate) gives you ticks-per-second, and then divide that by the percentage of time the
+		* shutter is open. Finally, divide the percentage of time the shutter is open by the number of frames
+		* we're accumulating.
+		*
+		* It is common that there is potential to have some leftover here. ie:
+		* 24000 Ticks / 24fps = 1000 ticks per second. At 180 degree shutter angle that gives you 500 ticks
+		* spread out amongst n=3 sub-frames. 500/3 = 166.66 ticks. We'll floor that when we use it, and ensure
+		* we accumulate the sub-tick and choose when to apply it.
+		*/
+
+		// Now we take the amount of time the shutter is open.
+		Output.TicksWhileShutterOpen = Output.TicksPerOutputFrame * Output.ShutterAnglePercentage;
+
+		// Divide that amongst all of our accumulation sample frames.
+		Output.TicksPerSample = Output.TicksWhileShutterOpen / AccumulationSettings->TemporalSampleCount;
+	}
+
+	Output.ShutterClosedFraction = (360 - AccumulationSettings->CameraShutterAngle) / 360.0;
+	Output.TicksWhileShutterClosed = Output.TicksPerOutputFrame * Output.ShutterClosedFraction;
+
+	// Shutter Offset
+	switch (AccumulationSettings->ShutterTiming)
+	{
+		// Subtract the entire time the shutter is open.
+	case EMoviePipelineShutterTiming::FrameClose:
+		Output.ShutterOffsetTicks = -Output.TicksWhileShutterOpen;
+		break;
+		// Only subtract half the time the shutter is open.
+	case EMoviePipelineShutterTiming::FrameCenter:
+		Output.ShutterOffsetTicks = -Output.TicksWhileShutterOpen / 2.0;
+		break;
+		// No offset needed
+	case EMoviePipelineShutterTiming::FrameOpen:
+		break;
+	}
+
+	// Then, calculate our motion blur offset. Motion Blur in the engine is always
+	// centered around the object so we offset our time sampling by half of the
+	// motion blur distance so that the distance blurred represents that time.
+	Output.MotionBlurCenteringOffsetTicks = Output.TicksPerSample / 2.0;
+
+	return Output;
+}
+
+FFrameNumber UMoviePipeline::GetTotalOutputFrameCountEstimate() const
+{
+	FFrameNumber EstimatedFrameCount = FFrameNumber(0);
+	
+	for (const FMoviePipelineShotCache& Shot : ShotList)
+	{
+		for (const FMoviePipelineShotCutCache& CameraCut : Shot.CameraCuts)
+		{
+			EstimatedFrameCount += CameraCut.GetOutputFrameCountEstimate();
+		}
+	}
+
+	return EstimatedFrameCount;
+}
+
+bool UMoviePipeline::GetRemainingTimeEstimate(FTimespan& OutTimespan) const
+{
+	// If they haven't produced a single frame yet, we can't give an estimate.
+	if (CachedOutputState.TotalSamplesRendered <= 0)
+	{
+		OutTimespan = FTimespan();
+		return false;
+	}
+
+	// Look at how many total samples we expect across all shots. This includes
+	// samples produced for warm-ups, motion blur fixes, and temporal/spatial samples.
+	FFrameNumber TotalExpectedSamples = FFrameNumber(0);
+
+	for (const FMoviePipelineShotCache& Shot : ShotList)
+	{
+		for (const FMoviePipelineShotCutCache& CameraCut : Shot.CameraCuts)
+		{
+			TotalExpectedSamples += CameraCut.GetSampleCountEstimate(true, true);
+		}
+	}
+
+	// Check to see how many frames we've rendered vs. our estimate.
+	int32 RenderedFrames = CachedOutputState.TotalSamplesRendered;
+	int32 TotalFrames = TotalExpectedSamples.Value;
+
+	double CompletionPercentage = FMath::Clamp(RenderedFrames / (double)TotalFrames, 0.0, 1.0);
+	FTimespan CurrentDuration = FDateTime::UtcNow() - InitializationTime;
+
+	// If it has taken us CurrentDuration to process CompletionPercentage frames, then we can get a total duration
+	// estimate by taking (CurrentDuration/CompletionPercentage) and then take that total estimate minus elapsed
+	// to get remaining. 
+	FTimespan EstimatedTotalDuration = CurrentDuration / CompletionPercentage;
+	OutTimespan = EstimatedTotalDuration - CurrentDuration;
+
+	return true;
+}
\ No newline at end of file
diff --git a/Engine/Plugins/MovieScene/MovieRenderPipeline/Source/MovieRenderPipelineCore/Private/MoviePipelineBlueprintLibrary.cpp b/Engine/Plugins/MovieScene/MovieRenderPipeline/Source/MovieRenderPipelineCore/Private/MoviePipelineBlueprintLibrary.cpp
new file mode 100644
--- /dev/null
+++ b/Engine/Plugins/MovieScene/MovieRenderPipeline/Source/MovieRenderPipelineCore/Private/MoviePipelineBlueprintLibrary.cpp
@@ -0,0 +1,49 @@
+// Copyright 1998-2019 Epic Games, Inc. All Rights Reserved.
+
+#include "MoviePipelineBlueprintLibrary.h"
+#include "MoviePipeline.h"
+#include "MovieRenderPipelineDataTypes.h"
+
+void UMoviePipelineBlueprintLibrary::GetCameraCutCounts(const UMoviePipeline* InPipeline, int32& OutTotalCuts, int32& OutCurrentCutIndex)
+{
+	OutTotalCuts = 0;
+	OutCurrentCutIndex = 0;
+	
+	if(!InPipeline)
+	{
+		return;
+	}
+
+	const TArray<FMoviePipelineShotCache>&  ShotList = InPipeline->GetShotList();
+	for (int32 ShotIdx = 0; ShotIdx < ShotList.Num(); ShotIdx++)
+	{
+		OutTotalCuts += ShotList[ShotIdx].CameraCuts.Num();
+
+		if (ShotIdx <= InPipeline->GetCurrentShotIndex())
+		{
+			// We don't add an additional one here because we're currently processing that
+			// index.
+			OutCurrentCutIndex += ShotList[ShotIdx].CurrentCameraCutIndex;
+		}
+	}
+}
+
+int32 UMoviePipelineBlueprintLibrary::GetOutputFrameCountEstimate(const FMoviePipelineShotCutCache& InCameraCut)
+{
+	return InCameraCut.GetOutputFrameCountEstimate().Value;
+}
+
+int32 UMoviePipelineBlueprintLibrary::GetTemporalFrameCountEstimate(const FMoviePipelineShotCutCache& InCameraCut)
+{
+	return InCameraCut.GetTemporalFrameCountEstimate().Value;
+}
+
+int32 UMoviePipelineBlueprintLibrary::GetUtilityFrameCountEstimate(const FMoviePipelineShotCutCache& InCameraCut)
+{
+	return InCameraCut.GetUtilityFrameCountEstimate().Value;
+}
+
+int32 UMoviePipelineBlueprintLibrary::GetSampleCountEstimate(const FMoviePipelineShotCutCache& InCameraCut, const bool bIncludeWarmup, const bool bIncludeMotionBlur)
+{
+	return InCameraCut.GetSampleCountEstimate(bIncludeWarmup, bIncludeMotionBlur).Value;
+}
diff --git a/Engine/Plugins/MovieScene/MovieRenderPipeline/Source/MovieRenderPipelineCore/Private/MoviePipelineLinearExecutor.cpp b/Engine/Plugins/MovieScene/MovieRenderPipeline/Source/MovieRenderPipelineCore/Private/MoviePipelineLinearExecutor.cpp
new file mode 100644
--- /dev/null
+++ b/Engine/Plugins/MovieScene/MovieRenderPipeline/Source/MovieRenderPipelineCore/Private/MoviePipelineLinearExecutor.cpp
@@ -0,0 +1,62 @@
+// Copyright 1998-2019 Epic Games, Inc. All Rights Reserved.
+#include "MoviePipelineLinearExecutor.h"
+#include "MovieRenderPipelineCoreModule.h"
+
+#define LOCTEXT_NAMESPACE "MoviePipelineLinearExecutorBase"
+
+void UMoviePipelineLinearExecutorBase::ExecuteImpl(TArray<UMovieRenderPipelineConfig*>& InPipelines)
+{
+	if (InPipelines.Num() == 0)
+	{
+		UE_LOG(LogMovieRenderPipeline, Warning, TEXT("Executor asked to execute on empty list of pipelines."));
+		OnExecutorFinishedImpl(false);
+		return;
+	}
+
+	// We'll process them in linear fashion and wait until each one is canceled or finishes on its own
+	// before moving onto the next one. This may be parallelizable in the future (either multiple PIE
+	// sessions, or multiple external processes) but ideally one render would maximize resource usage anyways...
+	Pipelines = InPipelines;
+
+	InitializationTime = FDateTime::UtcNow();
+
+	StartPipelineByIndex(0);
+}
+
+void UMoviePipelineLinearExecutorBase::StartPipelineByIndex(int32 InPipelineIndex)
+{
+	check(InPipelineIndex >= 0 && InPipelineIndex < Pipelines.Num());
+	
+	CurrentPipelineIndex = InPipelineIndex;
+
+	if (!Pipelines[CurrentPipelineIndex])
+	{
+		UE_LOG(LogMovieRenderPipeline, Warning, TEXT("Found null config in list of configs to render. Aborting the pipeline processing!"));
+		OnExecutorFinishedImpl(false);
+		return;
+	}
+	
+	Start(Pipelines[CurrentPipelineIndex], CurrentPipelineIndex, Pipelines.Num());
+}
+
+void UMoviePipelineLinearExecutorBase::OnIndividualPipelineFinished(UMoviePipeline* /* FinishedPipeline */)
+{
+	if (CurrentPipelineIndex == Pipelines.Num() - 1)
+	{
+		OnExecutorFinishedImpl(true);
+	}
+	else
+	{
+		// Onto the next one!
+		StartPipelineByIndex(CurrentPipelineIndex + 1);
+	}
+}
+
+void UMoviePipelineLinearExecutorBase::OnExecutorFinishedImpl(const bool bInSuccess)
+{
+	UE_LOG(LogMovieRenderPipeline, Log, TEXT("MoviePipelineLinearExecutorBase produced %d Pipelines in %s."), Pipelines.Num(), *(FDateTime::UtcNow() - InitializationTime).ToString());
+
+	Super::OnExecutorFinishedImpl(bInSuccess);
+}
+
+#undef LOCTEXT_NAMESPACE // "MoviePipelineLinearExecutorBase"
diff --git a/Engine/Plugins/MovieScene/MovieRenderPipeline/Source/MovieRenderPipelineCore/Private/MoviePipelineOutput.cpp b/Engine/Plugins/MovieScene/MovieRenderPipeline/Source/MovieRenderPipelineCore/Private/MoviePipelineOutput.cpp
new file mode 100644
--- /dev/null
+++ b/Engine/Plugins/MovieScene/MovieRenderPipeline/Source/MovieRenderPipelineCore/Private/MoviePipelineOutput.cpp
@@ -0,0 +1,3 @@
+// Copyright 1998-2019 Epic Games, Inc. All Rights Reserved.
+
+#include "MoviePipelineOutput.h"
diff --git a/Engine/Plugins/MovieScene/MovieRenderPipeline/Source/MovieRenderPipelineCore/Private/MoviePipelineOutputBuilder.h b/Engine/Plugins/MovieScene/MovieRenderPipeline/Source/MovieRenderPipelineCore/Private/MoviePipelineOutputBuilder.h
new file mode 100644
--- /dev/null
+++ b/Engine/Plugins/MovieScene/MovieRenderPipeline/Source/MovieRenderPipelineCore/Private/MoviePipelineOutputBuilder.h
@@ -0,0 +1,33 @@
+// Copyright 1998-2019 Epic Games, Inc. All Rights Reserved.
+
+#pragma once
+#include "MovieRenderPipelineDataTypes.h"
+
+// Forward Declares
+struct FImagePixelData;
+
+struct FMoviePipelineMergerOutputFrame
+{
+	FMoviePipelineFrameOutputState FrameOutput;
+	TArray<FMoviePipelinePassIdentifier> ExpectedRenderPasses;
+};
+
+class FMoviePipelineOutputMerger
+{
+public:
+	FMoviePipelineOutputMerger()
+	{
+	}
+
+public:
+	FMoviePipelineMergerOutputFrame& QueueOutputFrame_GameThread(const FMoviePipelineFrameOutputState& CachedOutputState);
+	void OnDataAvailable_AnyThread(TUniquePtr<FImagePixelData> InData);
+
+private:
+	TMap<FMoviePipelineFrameOutputState, FMoviePipelineMergerOutputFrame> ActiveData;
+
+	/** Mutex that protects adding/updating/removing from ActiveData */
+	FCriticalSection ActiveDataMutex;
+};
+
+
diff --git a/Engine/Plugins/MovieScene/MovieRenderPipeline/Source/MovieRenderPipelineCore/Private/MoviePipelineRendering.cpp b/Engine/Plugins/MovieScene/MovieRenderPipeline/Source/MovieRenderPipelineCore/Private/MoviePipelineRendering.cpp
new file mode 100644
--- /dev/null
+++ b/Engine/Plugins/MovieScene/MovieRenderPipeline/Source/MovieRenderPipelineCore/Private/MoviePipelineRendering.cpp
@@ -0,0 +1,349 @@
+// Copyright 1998-2019 Epic Games, Inc. All Rights Reserved.
+
+#include "MoviePipeline.h"
+#include "MovieRenderPipelineCoreModule.h"
+#include "Misc/FrameRate.h"
+#include "MoviePipelineOutput.h"
+#include "MovieRenderPipelineConfig.h"
+#include "MoviePipelineAccumulationSetting.h"
+#include "MoviePipelineShotConfig.h"
+#include "MoviePipelineRenderPass.h"
+#include "MoviePipelineOutputBuilder.h"
+
+void UMoviePipeline::SetupRenderingPipelineForShot(FMoviePipelineShotCache& Shot)
+{
+	/*
+	* To support tiled rendering we take the final effective resolution and divide
+	* it by the number of tiles to find the resolution of each render target. To 
+	* handle non-evenly divisible numbers/resolutions we may oversize the targets
+	* by a few pixels and then take the center of the resulting image when interlacing
+	* to produce the final image at the right resolution. For example:
+	*
+	* 1920x1080 in 7x7 tiles gives you 274.29x154.29. We ceiling this to set the resolution
+	* of the render pass to 275x155 which will give us a final interleaved image size of
+	* 1925x1085. To ensure that the image matches a non-scaled one we take the center out.
+	* LeftOffset = floor((1925-1920)/2) = 2
+	* RightOffset = (1925-1920-LeftOffset)
+	*/
+	UMoviePipelineAccumulationSetting* AccumulationSettings = Shot.ShotConfig->FindOrAddSetting<UMoviePipelineAccumulationSetting>();
+
+	FMoviePipelineRenderPassInitSettings RenderPassInitSettings;
+	RenderPassInitSettings.TileResolution = FIntPoint(
+		FMath::CeilToInt(Config->OutputResolution.X / AccumulationSettings->TileCount),
+		FMath::CeilToInt(Config->OutputResolution.Y / AccumulationSettings->TileCount));
+	RenderPassInitSettings.TileCount = AccumulationSettings->TileCount;
+	RenderPassInitSettings.ShotConfig = Shot.ShotConfig;
+
+	if (AccumulationSettings->bIsUsingOverlappedTiles)
+	{
+		// this code is duplicated in a few places, we should clean this up
+
+		int32 OverlappedPadX = int32(float(RenderPassInitSettings.TileResolution.X) * AccumulationSettings->PadRatioX);
+		int32 OverlappedPadY = int32(float(RenderPassInitSettings.TileResolution.Y) * AccumulationSettings->PadRatioY);
+		RenderPassInitSettings.TileResolution.X += 2 * OverlappedPadX;
+		RenderPassInitSettings.TileResolution.Y += 2 * OverlappedPadY;
+	}
+
+
+
+	// Code expects at least a 1x1 tile.
+	ensure(RenderPassInitSettings.TileCount > 0);
+
+	for (UMoviePipelineRenderPass* Input : Shot.ShotConfig->InputBuffers)
+	{
+		Input->Setup(RenderPassInitSettings, OutputPipe.ToSharedRef());
+	}
+}
+
+void UMoviePipeline::TeardownRenderingPipelineForShot(FMoviePipelineShotCache& Shot)
+{
+	for (UMoviePipelineRenderPass* Input : Shot.ShotConfig->InputBuffers)
+	{
+		Input->Teardown();
+	}
+}
+
+float Halton(int32 Index, int32 Base)
+{
+	float Result = 0.0f;
+	float InvBase = 1.0f / Base;
+	float Fraction = InvBase;
+	while (Index > 0)
+	{
+		Result += (Index % Base) * Fraction;
+		Index /= Base;
+		Fraction *= InvBase;
+	}
+	return Result;
+}
+
+void UMoviePipeline::RenderFrame()
+{
+	// Flush built in systems before we render anything. This maximizes the likelyhood that the data is prepared for when
+	// the render thread uses it.
+	FlushAsyncSystems();
+
+	FMoviePipelineShotCache& CurrentShot = ShotList[CurrentShotIndex];
+	FMoviePipelineShotCutCache& CurrentCameraCut = CurrentShot.GetCurrentCameraCut();
+
+	// We render a frame during motion blur fixes (but don't try to submit it for output) to allow seeding
+	// the view histories with data. We want to do this all in one tick so that the motion is correct.
+	const bool bIsHistoryOnlyFrame = CurrentCameraCut.State == EMovieRenderShotState::MotionBlur;
+	const bool bIsRenderFrame = CurrentCameraCut.State == EMovieRenderShotState::Rendering;
+
+	if(!(bIsHistoryOnlyFrame || bIsRenderFrame))
+	{
+		UE_LOG(LogMovieRenderPipeline, VeryVerbose, TEXT("[%d] Skipping RenderFrame() call due to state being ineligable for rendering. State: %d"), GFrameCounter, CurrentCameraCut.State);
+		return;
+	}
+
+	// Allow our containers to think
+	// ToDo: Is this necessaary?
+	for (UMoviePipelineOutput* Container : Config->OutputContainers)
+	{
+		Container->OnPostTick();
+	}
+	
+	// To produce a frame from the movie pipeline we may render many frames over a period of time, additively collecting the results
+	// together before submitting it for writing on the last result - this is referred to as an "output frame". The 1 (or more) samples
+	// that make up each output frame are referred to as "sample frames". Within each sample frame, we may need to render the scene many
+	// times. In order to support ultra-high-resolution rendering (>16k) movie pipelines support building an output frame out of 'tiles'. 
+	// Each tile renders the entire viewport with a small offset which causes different samples to be picked for each final pixel. These
+	// 'tiles' are then interleaved together (on the CPU) to produce a higher resolution result. For each tile, we can render a number
+	// of jitters that get added together to produce a higher quality single frame. This is useful for cases where you may not want any 
+	// motion (such as trees fluttering in the wind) but you do want high quality anti-aliasing on the edges of the pixels. Finally,
+	// the outermost loop (which is not represented here) is accumulation over time which happens over multiple engine ticks.
+	// 
+	// In short, for each output frame, for each accumulation frame, for each tile X/Y, for each jitter, we render a pass. This setup is
+	// designed to maximize the likely hood of deterministic rendering and that different passes line up with each other.
+	UMoviePipelineAccumulationSetting* AccumulationSettings = CurrentShot.ShotConfig->FindOrAddSetting<UMoviePipelineAccumulationSetting>();
+	const int32 NumTilesX = AccumulationSettings->TileCount;
+	const int32 NumTilesY = NumTilesX;
+	const bool bIsUsingOverlappedTiles = AccumulationSettings->bIsUsingOverlappedTiles;
+
+	int32 NumSpatialSamples = AccumulationSettings->SpatialSampleCount;
+	ensure(NumTilesX > 0 && NumTilesY > 0 && NumSpatialSamples > 0);
+	bool bHasCountedSamples = false;
+
+	TArray<UMoviePipelineRenderPass*> InputBuffers = CurrentShot.ShotConfig->InputBuffers;
+
+	// If this is the first sample for a new frame, we want to notify the output builder that it should expect data to accumulate for this frame.
+	if (CachedOutputState.TemporalSampleIndex == 0)
+	{
+		// This happens before any data is queued for this frame.
+		FMoviePipelineMergerOutputFrame& OutputFrame = OutputBuilder->QueueOutputFrame_GameThread(CachedOutputState);
+
+		// Now we need to go through all passes and get any identifiers from them of what this output frame should expect.
+		for (UMoviePipelineRenderPass* RenderPass : InputBuffers)
+		{
+			RenderPass->GatherOutputPasses(OutputFrame.ExpectedRenderPasses);
+		}
+	}
+
+	// Each render has been initialized with ResX/NumTiles, ResY/NumTiles already to know how big each resulting render target should be.
+	for (int32 TileY = 0; TileY < NumTilesY; TileY++)
+	{
+		for (int32 TileX = 0; TileX < NumTilesX; TileX++)
+		{
+			/*
+			* Two different features need to move the viewport and need to work in combination with each other.
+			* Ultra-high resolution rendering and single-frame anti-aliasing both work by moving the viewport
+			* by a small amount to create new samples for each render pass which generates more usable data.
+			*
+			* High resolution rendering is referred to as "tiles" but it is not actually a tile based render.
+			* Splitting the view into many smaller rectangles and using an off-center camera projection on them
+			* wouldn't work due to post-processing effects (such as vignette) as you would end up with many copies
+			* of the effect on your final image. Instead, it renders the entire screen for every 'tile', but shifts
+			* the camera position a portion of a pixel to cause new, unique data to be captured in the next tile.
+			*
+			* Single-frame anti-aliasing works in a similar way, by shifting the view by less than a pixel for each
+			* sample in the anti-aliasing. While the high-resolution rendering is a specific grid pattern, the
+			* anti-aliasing uses a more nuanced technique for picking offsets. It uses a deterministic Halton
+			* sequence to pick the offset, weighted by a guassian curve to mostly focus on samples around the
+			* center, but occasionally picking an offset which would be a sample outside of it's current sub-pixel.
+			*
+			* In the example below we're looking at one render target pixel, split into a 2x2 tiling. To gain
+			* 2x resolution, we shift the viewport so that the center of the pixel is no longer 'o', but now
+			* 'v'. Since this shift is less than a whole pixel we can interlace the resulting images and they
+			* won't be duplicates-but-shifted-by-one of each other. This is how we end up with new unique data.
+			*
+			* |-------------|-------------|
+			* |        a    |             |
+			* |    a        |             |
+			* |      v      |      v      |
+			* |   a         |             |
+			* |a       a    |             |
+			* |-------------o-------------|
+			* |             |             |
+			* |             |             |
+			* |      v      |      v      |
+			* |             |             |
+			* |             |             |
+			* |-------------|-------------|
+			*
+			* Now that a given render target pixel |v| has been shifted to it's correct location, we need to randomly
+			* sample within that box (perhaps with a slight overshoot) for the anti-aliasing samples. These samples
+			* ('a') are added together to create one |v| pixel and then only the |v| pixels are interleaved.
+			* Interleaving is performed on the CPU to support > 16k output sizes which is the maximum texture resolution
+			* for most GPUs.
+			*/
+
+			// Calculate the sub-pixel shift to get this pass into the right portion of the main pixel. Assuming the
+			// pixel coordinates are [0-1], the size of a sub-pixel (for tiles) is 1/NumTiles. To get to the right
+			// 'quadrant' we need to take our current tile index and multiply it by the size of each tile, so
+			// for the top right quadrant (in above example) size=(1/2), offset=(size*1). Finally we add half
+			// of the size to put us in the center of that pixel.
+			double TilePixelSizeFractionX = 1.0 / NumTilesX;
+			double TilePixelSizeFractionY = 1.0 / NumTilesY;
+			double TileShiftX = (TilePixelSizeFractionX * TileX) + (TilePixelSizeFractionX / 2.0);
+			double TileShiftY = (TilePixelSizeFractionY * TileY) + (TilePixelSizeFractionY / 2.0);
+
+			if (bIsUsingOverlappedTiles)
+			{
+				TileShiftX = 0.5f;
+				TileShiftY = 0.5f;
+
+				TilePixelSizeFractionX = 1.0f;
+				TilePixelSizeFractionY = 1.0f;
+			}
+
+			// Tell each pass to reset their spatial sampling accumulation.
+			for (UMoviePipelineRenderPass* Input : InputBuffers)
+			{
+				//Input->ResetSpatialSamplingAccumulation();
+			}
+
+			// If they're trying to seed the histories, then we'll just override the NumSpatialSamples since they're contributing towards
+			// our final target buffer which needs the history.
+			if (bIsHistoryOnlyFrame)
+			{
+				NumSpatialSamples = 1;
+			}
+
+			// Now we want to render a user-configured number of jitters to come up with the final output for this tile. These are considered
+			// spatial jitters which are accumulated together before contributing to the temporal accumulation.
+			for (int32 SpatialSample = 0; SpatialSample < NumSpatialSamples; SpatialSample++)
+			{
+				// We only want to count samples from the first tile
+				if (!bHasCountedSamples)
+				{
+					CurrentCameraCut.NumSamplesRendered++;
+					CachedOutputState.TotalSamplesRendered++;
+				}
+
+				float OffsetX = 0.5f;
+				float OffsetY = 0.5f; 
+
+				// CachedOutputState.TemporalSampleIndex is -1 on the first frame for motion blur, so ignore jitter on this frame.
+				// Also, we want jitter if either time samples or spatial samples are > 1, otherwise not.
+				if (CachedOutputState.TemporalSampleIndex >= 0 &&
+					(AccumulationSettings->SpatialSampleCount > 1 || AccumulationSettings->TemporalSampleCount > 1))
+				{
+					OffsetX = Halton((SpatialSample + 1) + (CachedOutputState.TemporalSampleIndex * NumSpatialSamples), 2);
+					OffsetY = Halton((SpatialSample + 1) + (CachedOutputState.TemporalSampleIndex * NumSpatialSamples), 3);
+				}
+
+				// Our spatial sample can only move within a +/- range within the tile we're in.
+				double SpatialShiftX = FMath::GetMappedRangeValueClamped(FVector2D(0, 1), FVector2D(-TilePixelSizeFractionX / 2.f, TilePixelSizeFractionX / 2.f), OffsetX);
+				double SpatialShiftY = FMath::GetMappedRangeValueClamped(FVector2D(0, 1), FVector2D(-TilePixelSizeFractionY / 2.f, TilePixelSizeFractionY / 2.f), OffsetY);
+
+				double FinalSubPixelShiftX = - (TileShiftX + SpatialShiftX - 0.5);
+				double FinalSubPixelShiftY = (TileShiftY + SpatialShiftY - 0.5);
+
+				int32 TileSizeX = FMath::CeilToInt(Config->OutputResolution.X / NumTilesX);
+				int32 TileSizeY = FMath::CeilToInt(Config->OutputResolution.Y / NumTilesY);
+
+				// in the case of overlapped rendering, the render target size will be larger than tile size because of overlap
+				int32 PadX = bIsUsingOverlappedTiles ? int32(float(TileSizeX) * AccumulationSettings->PadRatioX) : 0;
+				int32 PadY = bIsUsingOverlappedTiles ? int32(float(TileSizeY) * AccumulationSettings->PadRatioY) : 0;
+
+				int TargetSizeX = TileSizeX + 2 * PadX;
+				int TargetSizeY = TileSizeY + 2 * PadY;
+
+				FMoviePipelineRenderPassMetrics FrameMetrics;
+				FrameMetrics.OutputState = CachedOutputState;
+				FrameMetrics.TileIndex = (TileY * NumTilesY) + TileX;
+				FrameMetrics.SpatialShift = FVector2D((float)( FinalSubPixelShiftX) * 2.0f / TargetSizeX, (float)FinalSubPixelShiftY * 2.0f / TargetSizeX);
+				FrameMetrics.TileIndexX = TileX;
+				FrameMetrics.TileIndexY = TileY;
+				FrameMetrics.NumTilesX = NumTilesX;
+				FrameMetrics.NumTilesY = NumTilesY;
+				FrameMetrics.bIsHistoryOnlyFrame = bIsHistoryOnlyFrame;
+				FrameMetrics.SpatialJitterIndex = SpatialSample;
+				FrameMetrics.NumSpatialJitters = NumSpatialSamples;
+				FrameMetrics.TemporalJitterIndex = CachedOutputState.TemporalSampleIndex;
+				FrameMetrics.NumTemporalJitters = AccumulationSettings->TemporalSampleCount;
+				FrameMetrics.AccumulationGamma = AccumulationSettings->AccumulationGamma;
+				FrameMetrics.JitterOffsetX = 1.0f - OffsetX; // reversing to match FinalSubPixelShiftX
+				FrameMetrics.JitterOffsetY = OffsetY;
+				FrameMetrics.bIsUsingOverlappedTiles = bIsUsingOverlappedTiles;
+
+				if (bIsUsingOverlappedTiles)
+				{
+					FrameMetrics.OverlappedPadX = int32(float(TileSizeX) * AccumulationSettings->PadRatioX);
+					FrameMetrics.OverlappedPadY = int32(float(TileSizeY) * AccumulationSettings->PadRatioY);
+					FrameMetrics.OverlappedOffsetX = TileX * TileSizeX - FrameMetrics.OverlappedPadX;
+					FrameMetrics.OverlappedOffsetY = TileY * TileSizeY - FrameMetrics.OverlappedPadY;
+					FrameMetrics.OverlappedSizeX = TileSizeX;
+					FrameMetrics.OverlappedSizeY = TileSizeY;
+					FrameMetrics.OverlappedSubpixelShift = FVector2D(1.0f - OffsetX,OffsetY);
+				}
+				else
+				{
+					FrameMetrics.OverlappedOffsetX = 0;
+					FrameMetrics.OverlappedOffsetY = 0;
+					FrameMetrics.OverlappedSizeX = 0;
+					FrameMetrics.OverlappedSizeY = 0;
+					FrameMetrics.OverlappedPadX = 0;
+					FrameMetrics.OverlappedPadY = 0;
+					FrameMetrics.OverlappedSubpixelShift = FVector2D(0.0f,0.0f);
+				}
+				// Finally, we can ask each pass to render. Passes have already been initialized knowing how many tiles they're broken into
+				for (UMoviePipelineRenderPass* Input : InputBuffers)
+				{
+					// Our Cached Output state has the expected output for this frame.
+					Input->CaptureFrame(FrameMetrics);
+
+					// If we needed to see each individual jittered image, we would need to
+					// enqueue a copy of the pass's render target here. ToDo:
+
+					// We can calculate the influence of the results of this pass before accumulating
+					// If we are accumulating to a dedicated target for sub-frame jitter results, then
+					// the influence of this pass is simply AccumTarget += ThisResult * (1/NumSpatialSamples);
+					// However, to minimize memory usage we want to use a single accumulation target for
+					// each pass on each tile. To simulate a camera shutter, different temporal sample frames
+					// may have different weights (ie: the first and last sample may be slightly darker), and
+					// in this case the influence of this pass is:
+					// AccumTarget += (ThisResult * (1/NumSpatialSamples)) * GetWeightForTemporalSample(CachedOutputState.SubFrameIndex);
+				}
+
+			}
+
+			// We only want to count samples from the first tile
+			bHasCountedSamples = true;
+		}
+	}
+	
+	// 	// Get our list of buffers that are supposed to be producing this frame
+
+	// 
+	// 	// Once all inputs have had a chance to render, we're going to flush the rendering and then
+	// 	// transfer all of the created data to the output containers in one go.
+	FlushRenderingCommands();
+	// 
+	// 	// Store all of the data generated this frame
+	TArray<MoviePipeline::FOutputFrameData> AllPassFrameData;
+	// 
+	for (UMoviePipelineRenderPass* Input : InputBuffers)
+	{
+		MoviePipeline::FOutputFrameData& OutFrameData = AllPassFrameData.Add_GetRef(MoviePipeline::FOutputFrameData());
+		Input->GetFrameData(OutFrameData);
+	}
+	// 	
+	TArray<UMoviePipelineOutput*> OutputContainers = Config->OutputContainers;
+	for (UMoviePipelineOutput* Output : OutputContainers)
+	{
+		Output->ProcessFrame(AllPassFrameData, CachedOutputState, Config->OutputDirectory);
+	}
+	// }
+}
\ No newline at end of file
diff --git a/Engine/Plugins/MovieScene/MovieRenderPipeline/Source/MovieRenderPipelineCore/Private/MoviePipelineSetting.cpp b/Engine/Plugins/MovieScene/MovieRenderPipeline/Source/MovieRenderPipelineCore/Private/MoviePipelineSetting.cpp
new file mode 100644
--- /dev/null
+++ b/Engine/Plugins/MovieScene/MovieRenderPipeline/Source/MovieRenderPipelineCore/Private/MoviePipelineSetting.cpp
@@ -0,0 +1,24 @@
+// Copyright 1998-2019 Epic Games, Inc. All Rights Reserved.
+
+#include "MoviePipelineSetting.h"
+#include "MoviePipeline.h"
+
+void UMoviePipelineSetting::OnInitializedForPipeline(UMoviePipeline* InPipeline)
+{
+	CachedPipeline = MakeWeakObjectPtr(InPipeline);
+	OnInitializedForPipelineImpl(InPipeline);
+}
+
+UMoviePipeline* UMoviePipelineSetting::GetPipeline() const
+{
+	// If this check trips then life cycles aren't as expected.
+	UMoviePipeline* OutPipeline = CachedPipeline.Get();
+	check(OutPipeline);
+
+	return OutPipeline;
+}
+
+UWorld* UMoviePipelineSetting::GetWorld() const
+{
+	return GetPipeline()->GetWorld();
+}
\ No newline at end of file
diff --git a/Engine/Plugins/MovieScene/MovieRenderPipeline/Source/MovieRenderPipelineCore/Private/MoviePipelineShotConfig.cpp b/Engine/Plugins/MovieScene/MovieRenderPipeline/Source/MovieRenderPipelineCore/Private/MoviePipelineShotConfig.cpp
new file mode 100644
--- /dev/null
+++ b/Engine/Plugins/MovieScene/MovieRenderPipeline/Source/MovieRenderPipelineCore/Private/MoviePipelineShotConfig.cpp
@@ -0,0 +1,83 @@
+// Copyright 1998-2019 Epic Games, Inc. All Rights Reserved.
+
+#include "MoviePipelineShotConfig.h"
+#include "MoviePipelineSetting.h"
+#include "MoviePipelineRenderPass.h"
+#include "MoviePipelineOutput.h"
+
+#define LOCTEXT_NAMESPACE "MovieRenderShotConfig"
+
+bool UMoviePipelineShotConfig::ValidateConfig(TArray<FText>& OutValidationErrors) const
+{
+	bool bValidPipeline = true;
+	OutValidationErrors.Reset();
+	
+	// Check that we have a sequence to render
+	if(!Sequence.IsValid())
+	{
+		OutValidationErrors.Add(LOCTEXT("InvalidInputSequence", "No sequence to render has been specified."));
+		bValidPipeline = false;
+	}
+	
+	// Check that we have a valid output resolution
+	if(OutputResolution.X == 0 || OutputResolution.Y == 0)
+	{
+		OutValidationErrors.Add(LOCTEXT("InvalidOutputResolution", "Output Resolution must be greater than zero."));
+		bValidPipeline = false;
+	}
+	
+	// Check for a non-zero framerate
+	if(OutputFrameRate.AsDecimal() == 0.0)
+	{
+		OutValidationErrors.Add(LOCTEXT("InvalidOutputFrameRate", "Output FrameRate must be greater than zero."));
+		bValidPipeline = false;
+	}
+	
+	// Output Directory can probably be ignored. The folder path you output to may not exist yet anyways.
+	
+	// Give each setting a chance to validate the pipeline
+	for(const UMoviePipelineSetting* Setting : Settings)
+	{
+		bValidPipeline &= Setting->ValidatePipeline(OutValidationErrors);
+	}
+	
+	// ToDo-MoviePipeline: Reflect over all settings classes and make sure all required classes have been added.
+	// ToDo-MoviePipeline: Scan for any shots/sections that don't start on whole frames issue warning.
+	
+	// Give each Input Buffer a chance to validate the pipeline
+	for(const UMoviePipelineRenderPass* InputBuffer : InputBuffers)
+	{
+		bValidPipeline &= InputBuffer->ValidatePipeline(OutValidationErrors);
+	}
+	
+	// Give each output container a chance to validate the pipeline
+	for(const UMoviePipelineOutput* OutputContainer : OutputContainers)
+	{
+		bValidPipeline &= OutputContainer->ValidatePipeline(OutValidationErrors);
+	}
+
+	// Ensure we have at least one input buffer and one output buffer.
+	if (InputBuffers.Num() == 0)
+	{
+		OutValidationErrors.Add(LOCTEXT("NoInputBuffer", "Must specify at least one Input Buffer to capture."));
+		bValidPipeline = false;
+	}
+
+	if (OutputContainers.Num() == 0)
+	{
+		OutValidationErrors.Add(LOCTEXT("NoOutputContainer", "Must specify at least one Output Container to write to."));
+		bValidPipeline = false;
+	}
+	
+	return bValidPipeline;
+}
+
+
+void UMoviePipelineShotConfig::RemoveSetting(UMoviePipelineSetting* InSource)
+{
+	Settings.Remove(InSource);
+
+	// Update our cached serial number so the UI rebuilds the tree
+	++SettingsSerialNumber;
+}
+#undef LOCTEXT_NAMESPACE // "MovieRenderShotConfig"
\ No newline at end of file
