Initial merge of the new Sequencer Movie Render system. This provides some basic ways to output movies but is mostly scripting based without UI still.

diff --git a/Engine/Plugins/MovieScene/MovieRenderPipeline/Source/MovieRenderPipelineCore/Private/MoviePipelineTiming.cpp b/Engine/Plugins/MovieScene/MovieRenderPipeline/Source/MovieRenderPipelineCore/Private/MoviePipelineTiming.cpp
new file mode 100644
--- /dev/null
+++ b/Engine/Plugins/MovieScene/MovieRenderPipeline/Source/MovieRenderPipelineCore/Private/MoviePipelineTiming.cpp
@@ -0,0 +1,561 @@
+// Copyright 1998-2019 Epic Games, Inc. All Rights Reserved.
+
+#include "MoviePipeline.h"
+#include "MovieRenderPipelineDataTypes.h"
+#include "MoviePipelineAccumulationSetting.h"
+#include "MoviePipelineShotConfig.h"
+#include "MoviePipelineOutput.h"
+#include "MovieRenderPipelineConfig.h"
+#include "CoreGlobals.h"
+#include "Containers/Array.h"
+#include "MovieRenderPipelineCoreModule.h"
+#include "MovieScene.h"
+#include "LevelSequence.h"
+#include "LevelSequenceActor.h"
+
+void UMoviePipeline::TickProducingFrames()
+{
+	// The callback for this function does not get registered until Initialization has been called, which sets
+	// the state to Render. If it's not, we have a initialization order/flow issue!
+	check(PipelineState == EMovieRenderPipelineState::ProducingFrames);
+
+	// We should not be calling this once we have completed all the shots.
+	check(CurrentShotIndex >= 0 && CurrentShotIndex < ShotList.Num());
+
+	// If we're debug stepping one frame at a time, this will return true
+	// and we early out so that we don't end up advancing the pipeline at all.
+	// This handles setting the delta time for us to a fixed number.
+	if (DebugFrameStepPreTick())
+	{
+		return;
+	}
+	
+	// Make sure we're unpaused for the duration of this frame. We could have been paused by the frame stepper
+	// on the last Tick.
+	GetWorld()->GetFirstPlayerController()->SetPause(false);
+
+	FMoviePipelineShotCache& CurrentShot = ShotList[CurrentShotIndex];
+	FMoviePipelineShotCutCache& CurrentCameraCut = CurrentShot.GetCurrentCameraCut();
+
+	// Check to see if we need to initialize a new shot. This potentially changes a lot of state.
+	if (CurrentCameraCut.State == EMovieRenderShotState::Uninitialized)
+	{
+		// If this is the first camera cut, and it was uninitialized, then this is the first time the shot is being set up.
+		if (CurrentShot.CurrentCameraCutIndex == 0)
+		{
+			// Solo our shot and create rendering resources.
+			InitializeShot(CurrentShot);
+		}
+
+		UE_LOG(LogMovieRenderPipeline, Log, TEXT("[%d] Initializing Camera Cut [%d/%d] in Shot (%s)."), GFrameCounter, 
+			CurrentShot.CurrentCameraCutIndex + 1, CurrentShot.CameraCuts.Num(), *CurrentShot.GetDisplayName());
+		
+		// Ensure our Camera Cut is the only one active so it will evaluate the right camera.
+		SetSoloCameraCut(CurrentCameraCut);
+
+		CurrentCameraCut.SetNextState(EMovieRenderShotState::Uninitialized);
+
+		// The shot has been initialized and the new state chosen. Jump the Sequence to the starting time.
+		// This means that the camera will be in the correct location for warm-up frames to allow any systems
+		// dependent on camera position to properly warm up. If motion blur fixes are enabled, then we'll jump 
+		// again after the warm-up frames.
+		FFrameTime TimeInPlayRate = FFrameRate::TransformTime(CurrentCameraCut.CurrentTick, TargetSequence->GetMovieScene()->GetTickResolution(), TargetSequence->GetMovieScene()->GetDisplayRate());
+		
+		// We tell it to jump so that things responding to different scrub-types work correctly.
+		LevelSequenceActor->GetSequencePlayer()->JumpToFrame(TimeInPlayRate);
+		CustomSequenceTimeController->SetCachedFrameTiming(FQualifiedFrameTime(CurrentCameraCut.CurrentTick, TargetSequence->GetMovieScene()->GetTickResolution()));
+
+		// Ensure we don't try to evaluate as we want to sit and wait during warm up and motion blur frames.
+		LevelSequenceActor->GetSequencePlayer()->Pause();
+
+		// We can safely fall through to the below states as they're OK to process the same frame we set up.
+		UE_LOG(LogMovieRenderPipeline, Log, TEXT("[%d] Finished initializing Camera Cut [%d] in Shot [%d] %s."), GFrameCounter, 
+			CurrentShot.CurrentCameraCutIndex + 1, CurrentShotIndex, *CurrentShot.GetDisplayName());
+	}
+	
+	// Do the WarmingUp State exit condition before the WarmingUp loop, to allow the correct WarmingUp state to
+	// persist through to the post-frame render before switching to the next state.
+	if (CurrentCameraCut.State == EMovieRenderShotState::WarmingUp && CurrentCameraCut.NumWarmUpFramesRemaining == 0)
+	{
+		// We can either jump to the MotionBlur state or to the Rendering state, same as with coming from Uninitialized stage. 
+		// We just let the shot decide to move on if it's appropriate (to avoid duplicating the code from Uninitialized)
+		CurrentCameraCut.SetNextState(EMovieRenderShotState::WarmingUp);
+	}
+
+	// If the shot is warming up then we will just pass a reasonable amount of time this frame (no output data will be produced).
+	if (CurrentCameraCut.State == EMovieRenderShotState::WarmingUp)
+	{
+		UE_LOG(LogMovieRenderPipeline, Verbose, TEXT("[%d] Evaluating WarmUp frame %d. Remaining: %d"), GFrameCounter, CurrentCameraCut.NumWarmUpFramesRemaining, CurrentCameraCut.NumWarmUpFramesRemaining - 1);
+
+		// This if block will get called the same frame that we transition into the Warm Up state. We should not be in the
+		// warm up state unless we have at least one frame. 
+		CurrentCameraCut.NumWarmUpFramesRemaining--;
+
+		// Render the next frame at a reasonable frame rate to pass time in the world.
+		double FrameDeltaTime = Config->GetEffectiveFrameRate().AsInterval();
+
+		UE_LOG(LogMovieRenderPipeline, VeryVerbose, TEXT("[%d] Shot WarmUp set engine DeltaTime to %f seconds."), GFrameCounter, FrameDeltaTime);
+		CustomTimeStep->SetCachedFrameTiming(MoviePipeline::FFrameTimeStepCache(FrameDeltaTime));
+		
+		CachedOutputState.FrameDeltaTime = FrameDeltaTime;
+		CachedOutputState.WorldSeconds = CachedOutputState.WorldSeconds + FrameDeltaTime;
+		// We don't want to execute the other possible states until at least the next frame.
+		// This ensures that we actually tick the world once for the WarmUp state.
+		return;
+	}
+
+	// At this point, from now on, the Sequence should be in the Playing state. MotionBlur fixes only
+	// last for one frame, and rendering expects it to be playing anyways. 
+	if (!LevelSequenceActor->GetSequencePlayer()->IsPlaying())
+	{
+		UE_LOG(LogMovieRenderPipeline, Verbose, TEXT("[%d] Setting Sequence Player to Play due to being paused when motion started."), GFrameCounter);
+		LevelSequenceActor->GetSequencePlayer()->Play();
+	}
+
+	// Do the MotionBlur State exit condition before the MotionBlur loop, to allow the correct MotionBlur state to
+	// persist through to the post-frame render before switching to the next state.
+	if (CurrentCameraCut.State == EMovieRenderShotState::MotionBlur && CurrentCameraCut.bHasEvaluatedMotionBlurFrame)
+	{
+		// Motion Blur will evaluate further on in the frame than the first output frame. This is so that when
+		// we jump to the correct evaluation time for the first frame, the motion vectors (which are bi-directional)
+		// are correct and more likely to have valid data than if we started before the frame.
+		const MoviePipeline::FFrameConstantMetrics FrameMetrics = CalculateShotFrameMetrics(CurrentShot);
+
+		// The first frame of rendering code always assumes that we're coming from the previous frame. To avoid
+		// complicating that code, we'll jump back to where that frame would have been rendered, had it existed.
+		FFrameTime TicksToEndOfPreviousFrame;
+		float WorldTimeDilation = GetWorld()->GetWorldSettings()->GetEffectiveTimeDilation();
+
+		UMoviePipelineAccumulationSetting* AccumulationSettings = CurrentShot.ShotConfig->FindOrAddSetting<UMoviePipelineAccumulationSetting>();
+		if (AccumulationSettings->TemporalSampleCount == 1)
+		{
+			TicksToEndOfPreviousFrame = FrameMetrics.TicksPerOutputFrame;
+		}
+		else
+		{
+			// The first sub-frame accumulation is going to try to go forward
+			// by the amount the shutter is closed, plus the duration of one sample
+			// because that is the logic we use for every other frame.
+			TicksToEndOfPreviousFrame = FrameMetrics.TicksPerSample + FrameMetrics.TicksWhileShutterClosed;
+		}
+
+		if (!FMath::IsNearlyEqual(WorldTimeDilation, 1.f))
+		{
+			TicksToEndOfPreviousFrame = TicksToEndOfPreviousFrame * WorldTimeDilation;
+		}
+		AccumulatedTickSubFrameDeltas -= TicksToEndOfPreviousFrame.GetSubFrame();
+		CurrentCameraCut.CurrentTick = CurrentCameraCut.CurrentTick - TicksToEndOfPreviousFrame.FloorToFrame();
+
+		// Skip to rendering which will skip the next block.
+		CurrentCameraCut.State = EMovieRenderShotState::Rendering;
+	}
+
+	// This block is optional, as they may not want motion blur fixes.
+	if (CurrentCameraCut.State == EMovieRenderShotState::MotionBlur)
+	{
+		const MoviePipeline::FFrameConstantMetrics FrameMetrics = CalculateShotFrameMetrics(CurrentShot);
+		CachedOutputState.MotionBlurFraction = FrameMetrics.ShutterAnglePercentage;
+		
+		// For the Motion Blur frame, we evaluate the Sequence past the starting frame so that on the next frame
+		// we go backwards, and that gives us an approximation of motion blur for data we don't have. The world moves
+		// forward in both cases which is a little odd conceptually but we can't use negative delta times.
+		double FrameDeltaTime = FrameMetrics.TickResolution.AsSeconds(FrameMetrics.TicksPerSample);
+		CachedOutputState.FrameDeltaTime = FrameDeltaTime;
+		CachedOutputState.WorldSeconds = CachedOutputState.WorldSeconds + FrameDeltaTime;
+
+		// We subtract a single tick here so that if you're trying to render one frame, and it is the last frame of the
+		// camera cut, if we jump forward a whole frame the evaluation will fail (since end frames are exclusive). To
+		// work around, we evaluate one tick before, which is indistinguishable in positions but neatly avoids the eval problem.
+		FFrameNumber MotionBlurDurationTicks = FrameMetrics.TicksPerOutputFrame.FloorToFrame() - FFrameNumber(1);
+
+		// CurrentShot.CurrentTick should be at the first frame they want to evaluate, without the shutter timing/motion blur
+		// offsets taken into account. When we evaluate for this frame, take those into account.
+		FFrameTime FinalEvalTime = FrameMetrics.GetFinalEvalTime(CurrentCameraCut.CurrentTick + MotionBlurDurationTicks);
+
+		CustomTimeStep->SetCachedFrameTiming(MoviePipeline::FFrameTimeStepCache(FrameDeltaTime));
+		CustomSequenceTimeController->SetCachedFrameTiming(FQualifiedFrameTime(FinalEvalTime, FrameMetrics.TickResolution));
+		
+		// We early out on this loop so that at least one tick passes for motion blur. We need to leave our state in
+		// MotionBlur because the render tick needs to know that it is processing motion blur (it's a special case
+		// with many renders for a frame to fill history buffers).
+		CurrentCameraCut.bHasEvaluatedMotionBlurFrame = true;
+		UE_LOG(LogMovieRenderPipeline, Verbose, TEXT("[%d] Shot MotionBlur set engine DeltaTime to %f seconds."), GFrameCounter, FrameDeltaTime);
+		return;
+	}
+
+	// We notify our output containers that we're starting to produce content again. This is useful for things like
+	// sound where we may want to skip a period of sound (during warm-ups, etc.) We don't do this as part of another
+	// block because you could be getting to this point from any of three states (Initializing, Warm Up, MotionBlur)
+	// so it's easier to just track it separately and do it once here.
+	if (!CurrentCameraCut.bHasNotifiedFrameProductionStarted)
+	{
+		CurrentCameraCut.bHasNotifiedFrameProductionStarted = true;
+		for (UMoviePipelineOutput* Container : Config->OutputContainers)
+		{
+			// ToDo: We should name this OnRenderingStart/OnRenderingStop
+			Container->OnFrameProductionStart();
+		}
+	}
+
+	// Alright we're in the point where we finally want to start rendering out frames, we've finished warming up,
+	// we've evaluated the correct frame to give us motion blur vector sources, etc. This will be called each time
+	// we advance the world, even when sub-stepping, etc.
+	if (CurrentCameraCut.State == EMovieRenderShotState::Rendering)
+	{
+		const MoviePipeline::FFrameConstantMetrics FrameMetrics = CalculateShotFrameMetrics(CurrentShot);
+		float WorldTimeDilation = GetWorld()->GetWorldSettings()->GetEffectiveTimeDilation();
+
+		// This delta time may get modified to respect slowmo tracks and will be converted to seconds for
+		// engine delta time at the end.
+		FFrameTime DeltaFrameTime = FFrameTime();
+
+		UMoviePipelineAccumulationSetting* AccumulationSettings = CurrentShot.ShotConfig->FindOrAddSetting<UMoviePipelineAccumulationSetting>();
+		check(AccumulationSettings);
+
+		// If we've rendered the last sample and wrapped around, then we're going to be
+		// working on the next output frame for this frame. Since OutputFrameNumber gets
+		// initialized to -1, this is ok on the very first frame of rendering as it pushes
+		// us to be working on outputting frame zero.
+		if (CachedOutputState.TemporalSampleIndex >= AccumulationSettings->TemporalSampleCount - 1)
+		{
+			CachedOutputState.TemporalSampleIndex = -1;
+		}
+
+		CachedOutputState.TemporalSampleIndex++;
+		CurrentCameraCut.CurrentFrameSamples = (CachedOutputState.TemporalSampleIndex + 1) * CurrentCameraCut.NumSpatialSamples;
+
+		if (AccumulationSettings->TemporalSampleCount == 1)
+		{
+			// It should always be zero when using no sub-sampling.
+			check(CachedOutputState.TemporalSampleIndex == 0);
+
+			// When not using any sub-frame accumulation, the delta time is just
+			// 1/FrameRate * World Timescale (Slo-mo), which is applied afterwards.
+			DeltaFrameTime = FrameMetrics.TicksPerOutputFrame;
+		}
+		else
+		{
+			/*
+			* Sub-frame accumulation is surprisingly nuanced and complicated. Unreal's real-time motion blur is
+			* is an approximation that derives a velocity (per-vertex) by taking the difference in positions per 
+			* frame and then simulates a 'centered' motion blur, blurring pixels on either side of the object.
+			* The amount of pixels that are blurred is determined by the Shutter Angle ("Motion Blur Amount").
+			* Because we are simulating a higher quality version of the existing motion blur, we override the
+			* Shutter Angle to 360 during accumulation frames, as they are designed to represent one long blur
+			* and we emulate the shutter angle ourself over multiple frames.
+			*
+			* Different industries have different needs for when data is collected relative to an output frame.
+			* Since a "frame" is a span of time, and a non-360 shutter angle determines that no photons would
+			* accumulate while it is closed, there are three options for what period of time makes up the frame.
+			* In addition to the sampling being shifted, there may also be a brightness curve that gets applied
+			* to each shutter type during the accumulation frame, as choosing to give a particular frame more
+			* weight during the accumulation will make it brighter and thus appear to have more influence. This
+			* is handled during the accumulation/pixel data instead of shifting samples around as that would make
+			* the timing on this very complicated.
+			*
+			* In the diagram below, 'f' marks an output frame, 's' marks the shutter-angle (180 in this case)
+			* and it is assumed there are four motion samples ('-')
+			*
+			* f         f         f         f
+			* |    s    |    s    |    s    |
+			* |    |    |    |    |    |    |
+			*       ----     Frame Close
+			*         -- --   Frame Center
+			*		     ---- Frame Open
+			* Frame Closed - In this scenario, a given frame 'n' is when the shutter should close, meaning that
+			*                the accumulation is the motion that happened before the frame.
+			* Frame Center - In this scenario, a given frame 'n' is split evenly with half of the motion coming
+			*                from before that frame, and half afterwards. In this diagram, the space in the
+			*                middle is irrelevant (-- --), and is just there to make the diagram line up evenly.
+			* Frame Opened - In this scenario, a given frame 'n' is when the shutter should open, meaning that
+			*                the accumulation is the motion that happens after the frame.
+			*
+			* Depending on the shutter angle the spacing between the frames will change, ie: 270 shutter angle
+			* means that the accumulation will happen across 75% of the time. In the example below, this assumes
+			* that you are using Frame Close, meaning that the 75% of time before each frame.
+			* 
+			* f         f         f         f
+			* |  s      |  s      |  s      |
+			* |  |      |  |      |  |      |
+			*     ------    ------    ------ Frame Close accumulation.
+			*
+			* The above offset (Frame Open/Closed/Centered) is the first timing offset that we apply, which only
+			* changes what data is sampled to be considered a given frame n. We need to apply a second offset to
+			* our sampling, because motion blur is centered around the object. For our motion blur to accurately
+			* represent the period of time denoted on the chart, we need to offset the sample by half of the length
+			* of a sub-frame so that when the blur is applied on either side of it, no pixels beyond where the frame
+			* lies are affected.
+			*
+			* These offsets are calculated at the start of the shot and counter-adjusted for when calculating out which
+			* frame number the output frames should end up with.
+			*
+			* In the diagram below, 'f' marks an output frame, 's' marks the shutter angle time, and '-' marks where
+			* we actually sample time to generate a sub-frame (for which there are 4 in this example).
+			*
+			* f               f               f
+			* |       s       |       s       |
+			* |       |       |       |       |
+			* |       |-|-|-|-|               |
+			*
+			* Finally, the tricky part. Because time in Unreal only goes forward (due to world tick) and renders may
+			* be very expensive, we need to minimize any extra, unnecessary renders. To emulate Shutter Angle, there is
+			* a period of time where the shutter is closed. We still need this time to pass in the world, but we don't
+			* want our camera motion blur to stretch to cover this gap. One solution would be to render an extra frame
+			* (with a duration of |-|) right before the first accumulation frame but this would be unnecessarily expensive.
+			* 
+			* Instead, what we do is we jump from the last accumulation frame to the start of the next frame (keeping in
+			* mind that we have a timing offset meaning that we're in the center of the frame), but we override the motion
+			* blur for this frame and say that the motion blur only has an effective length matching |-|, even though 
+			* |-----| time has passed. Our motion blur amount is a 0-1 fraction, so to calculate the desired motion blur
+			* for the first frame we take the length of a sub-frame and divide it by the total number of time that has
+			* passed. Because the motion blur is centered and we've landed in the middle of the first accumulation frame
+			* this ends up being the correct length and centered at the correct period of time. Yay!
+			*
+			* f               f               f
+			* |       s       |       s       |
+			* |       |       |       |       |
+			* |       |-|-|-|-|       |-|-|-|-|
+			*                ^         ^
+			*        Jump From         Jump To
+			*
+			* The only remaining thing to watch out for is accumulation drift. The Movie Pipeline does not rely on 
+			* accumulating delta times via engine-tick (even though we control engine tick) to avoid any floating point
+			* issues. Instead, we calculate the exact number of ticks to move each time and store the accumulated time as 
+			* ticks. Integers don't drift, but certain tick rates/accumulation sub-frame counts can end up not moving time
+			* forward by a whole number.
+			*
+			* Given a default tick rate of 24,000 ticks per second and a output frame rate of 24fps, we have 1,000 ticks to
+			* represent each frame. This is a great deal of precision, but it doesn't divide evenly by the number of samples
+			* they may choose, such as three. 1000/3 = 333.3333 ticks per sub-frame. Since ticks are our finest resolution
+			* we have to either advance time by 333 or 334. To account for this, we need a running tally of how much we're
+			* drifting behind and then add it once it crosses the threshold.
+			*
+			* The reason we derive the Sequence time from our own time, and we derive our own time from the series of multiple
+			* samples, is to handle slo-mo tracks. Slo-mo is not linear and can change each frame (which is why we multiply our
+			* final delta amount by it) but to ensure the newly generated frames have the most correct timecode we move time
+			* on our own, and find the nearest frame number and timecode from the sequence. When slo-mo tracks are not in play
+			* the numbers should always add up evenly and it will be an exact 1-to-1 match.
+			*
+			* Easy!
+			*/
+
+			
+			// Shouldn't be in this else branch if they don't have an Accumulation Setting for this shot.
+			check(AccumulationSettings);
+			check(AccumulationSettings->TemporalSampleCount > 1);
+				
+			const bool bFirstFrame = CachedOutputState.TemporalSampleIndex == 0;
+			if (bFirstFrame)
+			{
+				/* The first sub-frame has its timing calculated differently than the rest, because we're trying to skip
+				* any unnecessary ticks so our world/engine delta time is long (to cover the gap where the shutter is 
+				* closed) but the motion blur shutter angle is overridden to be short to simulate only the portion of time
+				* with the shutter open that we care about.
+				* f               f               f
+				* |       s       |       s       |
+				* |       |       |       |       |
+				* |       |-|-|-|-|       |-|-|-|-|
+				*                ^         ^
+				*        Jump From         Jump To
+				*
+				* The duration between the last accumulation frame and the start of the next accumulation frame is
+				* (360-ShutterAngle)/360 * Ticks Per Frame + Ticks Per Subframe.
+				* The (360-ShutterAngle)/360 gives us the fraction of the output frame that the shutter was closed,
+				* so in a 270 shutter angle (open 75% of the time) at 24000 ticks/24fps (1000 ticks per frame) the
+				* shutter is closed for 1000 * 0.25 = 250 ticks. We need to add the duration of one accumulation frame
+				* so that we end up in the middle of the first frame, and not before the first frame.
+				*
+				* We don't actually end up adding half offsets or shutter biases because the entire starting index has
+				* been adjusted for those and they are constant throughout the shot.
+				*/
+
+
+				// Finally, we take the time it is closed plus the sample duration to put us into the new sample.
+				DeltaFrameTime = FrameMetrics.TicksWhileShutterClosed + FrameMetrics.TicksPerSample;
+
+				while (AccumulatedTickSubFrameDeltas >= 1.f)
+				{
+					// We add the deltas on the start of the first sample, since it's the one where 
+					// we jump by a non-consistent amount of time anyways.
+					DeltaFrameTime += FFrameTime(FFrameNumber(1));
+					AccumulatedTickSubFrameDeltas -= 1.f;
+				}
+			}
+			else
+			{
+				// We're moving from one sub-frame to the next. Our calculation is simpler.
+				DeltaFrameTime = FrameMetrics.TicksPerSample;
+			}
+		}
+
+		
+
+		// Now that we've calculated our delta ticks, we need to multiply it by
+		// time dilation so that we advance through the sequence as slow as we
+		// advance through the world.
+		bool bWasAffectedByTimeDilation = false;
+		if (!FMath::IsNearlyEqual(WorldTimeDilation, 1.f))
+		{
+			UE_LOG(LogMovieRenderPipeline, VeryVerbose, TEXT("[%d] Modified FrameDeltaTime by a factor of %f to account for World Time Dilation."), GFrameCounter, WorldTimeDilation);
+
+			// This is technically always a frame behind as this function executes before Sequencer evaluates slow-mo tracks
+			// and sets the World Time Dilation, but it is close enough for now.
+			DeltaFrameTime = DeltaFrameTime * WorldTimeDilation;
+			bWasAffectedByTimeDilation = true;
+		}
+
+		/*
+		* Now that we know how much we want to advance the sequence by (and the world)
+		* We can calculate out the desired frame for the shot to evaluate. This is slightly
+		* complicated, as it's not just a simple increment. The actual evaluation position
+		* is modified by two things:
+		*
+		* - Shutter Timing: To affect whether the motion before, during, or after
+		*                   the frame count as this frame, we add this offset to
+		*                   the sampled range.
+		*
+		* - Motion Blur:    Unreal Motion blur is centered on the object. To
+		*                   do our best to ensure that the written pixels match
+		*					the motion that would have occurred during a frame,
+		*                   we offset by half of a sample to ensure our samples
+		*                   are in the middle and stretch half way to each side.
+		*
+		* When the shutter timing is Frame Open (meaning the motion all comes from
+		* after the shutter opens) we have to add half a sample's duration so that
+		* our motion blur is centered.
+		* When the shutter timing is Frame Center, we add half a sample's duration
+		* which counter-acts the half-sample removed by the shutter timing offset
+		* which ends up centering our sample exactly on the frame.
+		* When the shutter timing is Frame Close, we subtract an entire sample's
+		* duration and add half of it back, meaning we only sample information
+		* before the current frame, but stretch the motion to cover up to that.
+		*
+		* For example, example we can look at the following situations when using
+		* a 24fps sequence, 24000 tick resolution (1000 tick per frame), and
+		* 270* shutter angle.
+		*
+		* 270* Shutter Angle means the shutter is open for 750 ticks (1000*0.75)
+		* since we want to stretch from the 0 to 750, and our motion blur is bi
+		* directional, we offset 350 ticks (750/2) to center it.
+		*/
+
+		// Convert any sub-frame time we have into a rolling accumulation
+		AccumulatedTickSubFrameDeltas += DeltaFrameTime.GetSubFrame();
+
+		// Increment where we should evaluate on the current camera cut.
+		CurrentCameraCut.CurrentTick += DeltaFrameTime.GetFrame();
+
+		FFrameTime FinalEvalTime = FrameMetrics.GetFinalEvalTime(CurrentCameraCut.CurrentTick);
+		UE_LOG(LogMovieRenderPipeline, VeryVerbose, TEXT("[%d] FinalEvalTime: %s (Tick: %d SOffset: %d MOffset: %d)"), 
+			GFrameCounter, *LexToString(FinalEvalTime.FloorToFrame()),
+			CurrentCameraCut.CurrentTick.Value, FrameMetrics.ShutterOffsetTicks.GetFrame().Value, FrameMetrics.MotionBlurCenteringOffsetTicks.GetFrame().Value);
+
+		
+		// If the user doesn't want to render every frame we need to increase the delta time so that enough
+		// time passes in the world. The framerate of the output media will be adjusted to match.
+		// FrameDeltaTime = FrameDeltaTime * 1.0 / (double)Config->OutputFrameStep.Value;
+
+		double FrameDeltaTime = FrameMetrics.TickResolution.AsSeconds(FFrameTime(DeltaFrameTime.GetFrame()));
+		CachedOutputState.FrameDeltaTime = FrameDeltaTime;
+		CachedOutputState.WorldSeconds = CachedOutputState.WorldSeconds + FrameDeltaTime;
+		CachedOutputState.bWasAffectedByTimeDilation = bWasAffectedByTimeDilation;
+
+		// The combination of shutter angle percentage, non-uniform render frame delta times and dividing by sample
+		// count produce the correct length for motion blur in all cases.
+		CachedOutputState.MotionBlurFraction = FrameMetrics.ShutterAnglePercentage / AccumulationSettings->TemporalSampleCount;
+		CalculateFrameNumbersForOutputState(FrameMetrics, CurrentCameraCut, CachedOutputState);
+
+		// Now that we know the delta time for the upcoming frame, we can see if this time would extend past the end of our
+		// camera cut. If it does, we don't want to produce this frame as that would produce an extra frame at the end!
+		FFrameNumber EndOfCameraCut = CurrentCameraCut.TotalOutputRange.GetUpperBoundValue();
+		if (CurrentCameraCut.CurrentTick >= EndOfCameraCut)
+		{
+			UE_LOG(LogMovieRenderPipeline, Log, TEXT("[%d] Finished processing Camera Cut [%d/%d] on Shot [%d] (%s) with stats [...]"), GFrameCounter, 
+				CurrentShot.CurrentCameraCutIndex + 1, CurrentShot.CameraCuts.Num(), 
+				CurrentShotIndex + 1, *CurrentShot.GetDisplayName());
+
+			CurrentCameraCut.State = EMovieRenderShotState::Finished;
+			
+			// Compare our expected vs. actual results for logging.
+			{
+				FFrameNumber NumExpectedFrames = CurrentCameraCut.GetOutputFrameCountEstimate();
+				FFrameNumber NumExpectedSamples = CurrentCameraCut.GetSampleCountEstimate(/*bIncludeWarmup*/ true, /*bIncludeMotionBlur*/ true);
+
+				FFrameNumber NumFrames = CurrentCameraCut.NumFramesProduced;
+				FFrameNumber NumSamples = CurrentCameraCut.NumSamplesRendered;
+
+				UE_LOG(LogMovieRenderPipeline, Log, TEXT("\t Frames Produced: %d (Expected: %d), Samples Rendered: %d (Expected: %d)"),
+					NumFrames.Value, NumExpectedFrames.Value, NumSamples.Value, NumExpectedSamples.Value);
+			}
+
+			// We pause at the end too, just so that frames during finalize don't continue to trigger Sequence Eval messages.
+			LevelSequenceActor->GetSequencePlayer()->Pause();
+			
+			bool bWasLastCameraCut = CurrentShot.SetNextShotActive();
+			if (bWasLastCameraCut)
+			{
+				TeardownRenderingPipelineForShot(CurrentShot);
+
+				CurrentShotIndex++;
+
+				// Notify our containers that the current shot has ended.
+				for (UMoviePipelineOutput* Container : Config->OutputContainers)
+				{
+					Container->OnShotFinished(CurrentShot);
+				}
+			}
+
+			// Check to see if this was the last shot in the Pipeline, otherwise on the next
+			// tick the new shot will be initialized and processed.
+			if (CurrentShotIndex >= ShotList.Num())
+			{
+				UE_LOG(LogMovieRenderPipeline, Log, TEXT("[%d] Finished rendering last shot. Moving to Finalize to finish writing items to disk."), GFrameCounter);
+				PipelineState = EMovieRenderPipelineState::Finalize;
+
+				// Unregister our OnEngineTickEndFrame delegate. We don't unregister BeginFrame as we need
+				// to continue to call it to allow ticking the Finalization stage.
+				FCoreDelegates::OnEndFrame.RemoveAll(this);
+
+				// Reset the Custom Timestep because we don't care how long the engine takes now
+				GEngine->SetCustomTimeStep(CachedCustomTimeStep);
+
+				BeginFinalize();
+				return;
+			}
+			else
+			{
+				// If this isn't the last shot, we'll immediately call this function again to just
+				// determine a new setup/seek/etc. on the same engine tick. Otherwise we would have
+				// used a random delta time that didn't actually correlate to anything.
+				TickProducingFrames();
+				return;
+			}
+		}
+
+		// If we've made it this far, we're going to produce a frame so we can check to see if we want to increment the output
+		// number. If we do it before the above block, then the last frame of each shot increments this global state incorrectly.
+		if (CachedOutputState.TemporalSampleIndex == 0)
+		{
+			CachedOutputState.OutputFrameNumber++;
+			CurrentCameraCut.NumFramesProduced++;
+		}
+		
+		// Set our time step for the next frame
+		CustomTimeStep->SetCachedFrameTiming(MoviePipeline::FFrameTimeStepCache(FrameDeltaTime));
+		CustomSequenceTimeController->SetCachedFrameTiming(FQualifiedFrameTime(FinalEvalTime, FrameMetrics.TickResolution));
+		return;
+	}
+
+	// This function should not be called when the shot isn't in one of the above states.
+	check(false);
+}
+
+void UMoviePipeline::CalculateFrameNumbersForOutputState(const MoviePipeline::FFrameConstantMetrics& InFrameMetrics, const FMoviePipelineShotCutCache& InCameraCut, FMoviePipelineFrameOutputState& InOutOutputState) const
+{
+	// Find the closest number on the source Sequence. This may produce duplicates when using slow-mo tracks.
+	FFrameRate SourceDisplayRate = TargetSequence->GetMovieScene()->GetDisplayRate();
+	InOutOutputState.SourceFrameNumber = FFrameRate::TransformTime(InCameraCut.CurrentTick, InFrameMetrics.TickResolution, SourceDisplayRate).RoundToFrame();
+	InOutOutputState.SourceTimeCode = FTimecode::FromFrameNumber(CachedOutputState.SourceFrameNumber, SourceDisplayRate, false);
+	InOutOutputState.EffectiveFrameNumber = FFrameRate::TransformTime(InCameraCut.CurrentTick, InFrameMetrics.TickResolution, InFrameMetrics.FrameRate).RoundToFrame();
+	InOutOutputState.EffectiveTimeCode = FTimecode::FromFrameNumber(CachedOutputState.EffectiveFrameNumber, InFrameMetrics.FrameRate, false);
+}
\ No newline at end of file
diff --git a/Engine/Plugins/MovieScene/MovieRenderPipeline/Source/MovieRenderPipelineCore/Private/MovieRenderOverlappedImage.cpp b/Engine/Plugins/MovieScene/MovieRenderPipeline/Source/MovieRenderPipelineCore/Private/MovieRenderOverlappedImage.cpp
new file mode 100644
--- /dev/null
+++ b/Engine/Plugins/MovieScene/MovieRenderPipeline/Source/MovieRenderPipelineCore/Private/MovieRenderOverlappedImage.cpp
@@ -0,0 +1,407 @@
+// Copyright 1998-2019 Epic Games, Inc. All Rights Reserved.
+
+#include "MovieRenderOverlappedImage.h"
+
+#include "Math/Vector.h"
+#include "Math/VectorRegister.h"
+
+#include "MovieRenderPipelineCoreModule.h"
+#include "Math/Float16.h"
+
+void FImageOverlappedPlane::Init(int32 InSizeX, int32 InSizeY)
+{
+	SizeX = InSizeX;
+	SizeY = InSizeY;
+
+	// leaves the data as is
+	ChannelData.SetNumUninitialized(SizeX * SizeY);
+}
+
+void FImageOverlappedPlane::ZeroPlane()
+{
+	int32 Num = SizeX * SizeY;
+	for (int32 Index = 0; Index < Num; Index++)
+	{
+		ChannelData[Index] = 0.0f;
+	}
+}
+
+void FImageOverlappedPlane::Reset()
+{
+	SizeX = 0;
+	SizeY = 0;
+	ChannelData.Empty();
+}
+
+// a subpixel offset of (0.5,0.5) means that the raw data is exactly centered on destination pixels.
+void FImageOverlappedPlane::AccumulateSinglePlane(const TArray64<float>& InRawData, const TArray64<float>& InWeightData, int32 InSizeX, int32 InSizeY, int32 InOffsetX, int32 InOffsetY,
+								float SubpixelOffsetX, float SubpixelOffsetY)
+{
+	check(InRawData.Num() == InSizeX * InSizeY);
+	check(InWeightData.Num() == InSizeX * InSizeY);
+
+	check(SubpixelOffsetX >= 0.0f);
+	check(SubpixelOffsetX <= 1.0f);
+	check(SubpixelOffsetY >= 0.0f);
+	check(SubpixelOffsetY <= 1.0f);
+
+	// if the subpixel offset is less than 0.5, go back one pixel
+	int32 StartX = (SubpixelOffsetX >= 0.5 ? InOffsetX : InOffsetX - 1);
+	int32 StartY = (SubpixelOffsetY >= 0.5 ? InOffsetY : InOffsetY - 1);
+
+	float PixelWeight[2][2];
+	{
+		// make sure that the equal sign is correct, if the subpixel offset is 0.5, the weightX is 0 and it starts on the center pixel
+		float WeightX = FMath::Frac(SubpixelOffsetX + 0.5f);
+		float WeightY = FMath::Frac(SubpixelOffsetY + 0.5f);
+
+		// row, column
+		PixelWeight[0][0] = (1.0f - WeightX) * (1.0f - WeightY);
+		PixelWeight[0][1] = (       WeightX) * (1.0f - WeightY);
+		PixelWeight[1][0] = (1.0f - WeightX) * (       WeightY);
+		PixelWeight[1][1] = (       WeightX) * (       WeightY);
+	}
+
+	// Slow, reference version. Maybe optimize later.
+	for (int CurrY = 0; CurrY < InSizeY; CurrY++)
+	{
+		for (int CurrX = 0; CurrX < InSizeX; CurrX++)
+		{
+			float Val = InRawData[CurrY * InSizeX + CurrX];
+			float BaseWeight = InWeightData[CurrY * InSizeX + CurrX];
+
+			for (int OffsetY = 0; OffsetY < 2; OffsetY++)
+			{
+				for (int OffsetX = 0; OffsetX < 2; OffsetX++)
+				{
+					int DstY = StartY + CurrY + OffsetY;
+					int DstX = StartX + CurrX + OffsetX;
+
+					float Weight = BaseWeight * PixelWeight[OffsetY][OffsetX];
+
+					if (DstX >= 0 && DstY >= 0 &&
+						DstX < SizeX && DstY < SizeY)
+					{
+						ChannelData[DstY * SizeX + DstX] += Weight * Val;
+					}
+				}
+			}
+		}
+	}
+
+}
+
+void FImageOverlappedAccumulator::InitMemory(int InPlaneSizeX, int InPlaneSizeY, int InNumChannels)
+{
+	PlaneSizeX = InPlaneSizeX;
+	PlaneSizeY = InPlaneSizeY;
+	NumChannels = InNumChannels;
+
+	ChannelPlanes.SetNum(NumChannels);
+
+	for (int Channel = 0; Channel < NumChannels; Channel++)
+	{
+		ChannelPlanes[Channel].Init(PlaneSizeX, PlaneSizeY);
+	}
+
+	WeightPlane.Init(PlaneSizeX, PlaneSizeY);
+}
+
+void FImageOverlappedAccumulator::ZeroPlanes()
+{
+	check(ChannelPlanes.Num() == NumChannels);
+
+	for (int Channel = 0; Channel < NumChannels; Channel++)
+	{
+		ChannelPlanes[Channel].ZeroPlane();
+	}
+
+	WeightPlane.ZeroPlane();
+}
+
+void FImageOverlappedAccumulator::Reset()
+{
+	PlaneSizeX = 0;
+	PlaneSizeY = 0;
+	NumChannels = 0;
+
+	// Let the desctructor clean up
+	ChannelPlanes.Empty();
+}
+
+
+void FImageOverlappedAccumulator::GenerateTileWeight(TArray64<float>& Weights, int32 SizeX, int32 SizeY)
+{
+	Weights.SetNum(SizeX * SizeY);
+
+	// we'll use a simple triangle filter, which goes from 1.0 at the center to 0.0 at 3/4 of the way to the tile edge.
+
+	float ScaleX = 1.0f / (float(SizeX / 2) * .75f);
+	float ScaleY = 1.0f / (float(SizeY / 2) * .75f);
+
+	for (int PixY = 0; PixY < SizeY; PixY++)
+	{
+		for (int PixX = 0; PixX < SizeX; PixX++)
+		{
+			float Y = float(PixY) + .5f;
+			float X = float(PixX) + .5f;
+
+			float DistX = FMath::Abs(SizeX / 2 - X);
+			float DistY = FMath::Abs(SizeY / 2 - Y);
+
+			float WeightX = FMath::Clamp(1.0f - DistX * ScaleX, 0.0f, 1.0f);
+			float WeightY = FMath::Clamp(1.0f - DistY * ScaleY, 0.0f, 1.0f);
+
+			Weights[PixY * SizeX + PixX] = WeightX * WeightY;
+		}
+	}
+}
+
+void FImageOverlappedAccumulator::AccumulatePixelData(const FImagePixelData& InPixelData, int32 InTileOffsetX, int32 InTileOffsetY, FVector2D InSubpixelOffset)
+{
+	EImagePixelType Fmt = InPixelData.GetType();
+
+	int32 RawNumChan = InPixelData.GetNumChannels();
+	int32 RawBitDepth = InPixelData.GetBitDepth();
+	FIntPoint RawSize = InPixelData.GetSize();
+
+	int32 SizeInBytes = 0;
+	const void* SrcRawDataPtr = nullptr;
+
+	bool IsFetchOk = InPixelData.GetRawData(SrcRawDataPtr, SizeInBytes);
+	check(IsFetchOk); // keeping the if below for now, but this really should always succeed
+
+	if (IsFetchOk)
+	{
+		// hardcode to 4 channels (RGBA), even if we are only saving fewer channels
+		TArray64<float> RawData[4];
+
+		check(NumChannels >= 1);
+		check(NumChannels <= 4);
+
+		for (int32 ChanIter = 0; ChanIter < 4; ChanIter++)
+		{
+			RawData[ChanIter].SetNumUninitialized(RawSize.X * RawSize.Y);
+		}
+
+		const double AccumulateBeginTime = FPlatformTime::Seconds();
+			
+		// for now, only handle rgba8
+		if (Fmt == EImagePixelType::Color && RawNumChan == 4 && RawBitDepth == 8)
+		{
+			const uint8* RawDataPtr = static_cast<const uint8*>(SrcRawDataPtr);
+
+			static bool bIsReferenceUnpack = false;
+			if (bIsReferenceUnpack)
+			{
+				// simple, slow unpack, takes about 10-15ms on a 1080p image
+				for (int32 Y = 0; Y < RawSize.Y; Y++)
+				{
+					for (int32 X = 0; X < RawSize.X; X++)
+					{
+						int32 ChanReorder[4] = { 2, 1, 0, 3 };
+						for (int32 ChanIter = 0; ChanIter < 4; ChanIter++)
+						{
+							int32 RawValue = RawDataPtr[(Y*RawSize.X + X)*RawNumChan + ChanIter];
+							float Value = float(RawValue) / 255.0f;
+							int32 Reorder = ChanReorder[ChanIter];
+							RawData[Reorder][Y*RawSize.X + X] = Value;
+						}
+					}
+				}
+			}
+			else
+			{
+				// slightly optimized, takes about 3-7ms on a 1080p image
+				for (int32 Y = 0; Y < RawSize.Y; Y++)
+				{
+					const uint8* SrcRowDataPtr = &RawDataPtr[Y*RawSize.X*RawNumChan];
+
+					float* DstRowDataR = &RawData[0][Y*RawSize.X];
+					float* DstRowDataG = &RawData[1][Y*RawSize.X];
+					float* DstRowDataB = &RawData[2][Y*RawSize.X];
+					float* DstRowDataA = &RawData[3][Y*RawSize.X];
+
+					VectorRegister ColorScale = MakeVectorRegister(1.0f / 255.0f, 1.0f / 255.0f, 1.0f / 255.0f, 1.0f / 255.0f);
+
+					// simple, one pixel at a time vectorized version, we could do better
+					for (int32 X = 0; X < RawSize.X; X++)
+					{
+						VectorRegister Color = VectorLoadByte4(&SrcRowDataPtr[X*RawNumChan]);
+						Color = VectorMultiply(Color, ColorScale); // multiply by 1/255
+
+						const float* RawColorVec = reinterpret_cast<const float *>(&Color);
+						DstRowDataR[X] = RawColorVec[2];
+						DstRowDataG[X] = RawColorVec[1];
+						DstRowDataB[X] = RawColorVec[0];
+						DstRowDataA[X] = RawColorVec[3];
+					}
+				}
+			}
+		}
+		else if (Fmt == EImagePixelType::Float16 && RawNumChan == 4 && RawBitDepth == 16)
+		{
+			const uint16* RawDataPtr = static_cast<const uint16*>(SrcRawDataPtr);
+
+			// simple, slow unpack version for fp16, could make this faster
+			for (int32 Y = 0; Y < RawSize.Y; Y++)
+			{
+				for (int32 X = 0; X < RawSize.X; X++)
+				{
+					for (int32 ChanIter = 0; ChanIter < 4; ChanIter++)
+					{
+						FFloat16 RawValue;
+						RawValue.Encoded = RawDataPtr[(Y*RawSize.X + X)*RawNumChan + ChanIter];
+						
+						// c cast does the conversion from fp16 bits to float
+						float Value = float(RawValue);
+						RawData[ChanIter][Y*RawSize.X + X] = Value;
+					}
+				}
+			}
+		}
+		else if (Fmt == EImagePixelType::Float32 && RawNumChan == 4 && RawBitDepth == 32)
+		{
+			const float* RawDataPtr = static_cast<const float*>(SrcRawDataPtr);
+
+			// reference version for float
+			for (int32 Y = 0; Y < RawSize.Y; Y++)
+			{
+				for (int32 X = 0; X < RawSize.X; X++)
+				{
+					for (int32 ChanIter = 0; ChanIter < 4; ChanIter++)
+					{
+						float Value = RawDataPtr[(Y*RawSize.X + X)*RawNumChan + ChanIter];
+						RawData[ChanIter][Y*RawSize.X + X] = Value;
+					}
+				}
+			}
+		}
+		else
+		{
+			check(0);
+		}
+
+		const double AccumulateEndTime = FPlatformTime::Seconds();
+		{
+			const float ElapsedMs = float((AccumulateEndTime - AccumulateBeginTime)*1000.0f);
+			//UE_LOG(LogTemp, Log, TEXT("    [%8.2f] Unpack time."), ElapsedMs);
+		}
+	
+		if (AccumulationGamma != 1.0f)
+		{
+			// Unfortunately, we don't have an SSE optimized pow function. This function is quite slow (about 30-40ms).
+			float Gamma = AccumulationGamma;
+
+			for (int ChanIter = 0; ChanIter < NumChannels; ChanIter++)
+			{
+				float* DstData = RawData[ChanIter].GetData();
+				int32 DstSize = RawSize.X * RawSize.Y;
+				for (int32 Iter = 0; Iter < DstSize; Iter++)
+				{
+					DstData[Iter] = FMath::Pow(DstData[Iter], Gamma);
+				}
+			}
+		} 
+
+		const double GammaEndTime = FPlatformTime::Seconds();
+		{
+			const float ElapsedMs = float((GammaEndTime - AccumulateEndTime)*1000.0f);
+			//UE_LOG(LogTemp, Log, TEXT("        [%8.2f] Gamma time."), ElapsedMs);
+		}
+
+		// Calculate weight data for this tile.
+		TArray64<float> Weights;
+		GenerateTileWeight(Weights, RawSize.X, RawSize.Y);
+
+		for (int32 ChanIter = 0; ChanIter < NumChannels; ChanIter++)
+		{
+			ChannelPlanes[ChanIter].AccumulateSinglePlane(RawData[ChanIter], Weights, RawSize.X, RawSize.Y, InTileOffsetX, InTileOffsetY, InSubpixelOffset.X, InSubpixelOffset.Y);
+		}
+
+		// Accumulate weights as well.
+		{
+			TArray64<float> VecOne;
+			VecOne.SetNum(RawSize.X * RawSize.Y);
+			for (int32 Index = 0; Index < VecOne.Num(); Index++)
+			{
+				VecOne[Index] = 1.0f;
+			}
+
+			WeightPlane.AccumulateSinglePlane(VecOne, Weights, RawSize.X, RawSize.Y, InTileOffsetX, InTileOffsetY, InSubpixelOffset.X, InSubpixelOffset.Y);
+		}
+	}
+}
+
+void FImageOverlappedAccumulator::FetchFullImageValue(float Rgba[4], int32 FullX, int32 FullY) const
+{
+	Rgba[0] = 0.0f;
+	Rgba[1] = 0.0f;
+	Rgba[2] = 0.0f;
+	Rgba[3] = 1.0f;
+
+	float RawWeight = WeightPlane.ChannelData[FullY * PlaneSizeX + FullX];
+
+	float Scale = 1.0f / FMath::Max<float>(RawWeight, 0.0001f);
+
+	for (int64 Chan = 0; Chan < NumChannels; Chan++)
+	{
+		float Val = ChannelPlanes[Chan].ChannelData[FullY * PlaneSizeX + FullX];
+
+		Rgba[Chan] = Val * Scale;
+	}
+
+	if (AccumulationGamma != 1.0f)
+	{
+		float InvGamma = 1.0f / AccumulationGamma;
+		Rgba[0] = FMath::Pow(Rgba[0], InvGamma);
+		Rgba[1] = FMath::Pow(Rgba[1], InvGamma);
+		Rgba[2] = FMath::Pow(Rgba[2], InvGamma);
+		Rgba[3] = FMath::Pow(Rgba[3], InvGamma);
+	}
+}
+
+void FImageOverlappedAccumulator::FetchFinalPixelDataByte(TArray<FColor> & OutPixelData) const
+{
+	int32 FullSizeX = PlaneSizeX;
+	int32 FullSizeY = PlaneSizeY;
+	OutPixelData.SetNumUninitialized(FullSizeX * FullSizeY);
+
+	for (int32 FullY = 0L; FullY < FullSizeY; FullY++)
+	{
+		for (int32 FullX = 0L; FullX < FullSizeX; FullX++)
+		{
+			float Rgba[4] = { 0.0f, 0.0f, 0.0f, 1.0f };
+
+			FetchFullImageValue(Rgba, FullX, FullY);
+			FColor Color = FColor(Rgba[0]*255.0f,Rgba[1]*255.0f,Rgba[2]*255.0f,Rgba[3]*255.0f);
+
+			// be careful with this index, make sure to use 64bit math, not 32bit
+			int64 DstIndex = int64(FullY) * int64(FullSizeX) + int64(FullX);
+			OutPixelData[DstIndex] = Color;
+		}
+	}
+}
+
+void FImageOverlappedAccumulator::FetchFinalPixelDataLinearColor(TArray<FLinearColor> & OutPixelData) const
+{
+	int32 FullSizeX = PlaneSizeX;
+	int32 FullSizeY = PlaneSizeY;
+	OutPixelData.SetNumUninitialized(FullSizeX * FullSizeY);
+
+	for (int32 FullY = 0L; FullY < FullSizeY; FullY++)
+	{
+		for (int32 FullX = 0L; FullX < FullSizeX; FullX++)
+		{
+			float Rgba[4] = { 0.0f, 0.0f, 0.0f, 1.0f };
+
+			FetchFullImageValue(Rgba, FullX, FullY);
+			FLinearColor Color = FLinearColor(Rgba[0],Rgba[1],Rgba[2],Rgba[3]);
+
+			// be careful with this index, make sure to use 64bit math, not 32bit
+			int64 DstIndex = int64(FullY) * int64(FullSizeX) + int64(FullX);
+			OutPixelData[DstIndex] = Color;
+		}
+	}
+}
+
diff --git a/Engine/Plugins/MovieScene/MovieRenderPipeline/Source/MovieRenderPipelineCore/Private/MovieRenderPipelineCommandLine.cpp b/Engine/Plugins/MovieScene/MovieRenderPipeline/Source/MovieRenderPipelineCore/Private/MovieRenderPipelineCommandLine.cpp
new file mode 100644
--- /dev/null
+++ b/Engine/Plugins/MovieScene/MovieRenderPipeline/Source/MovieRenderPipelineCore/Private/MovieRenderPipelineCommandLine.cpp
@@ -0,0 +1,231 @@
+// Copyright 1998-2019 Epic Games, Inc. All Rights Reserved.
+
+#include "MovieRenderPipelineCoreModule.h"
+#include "Modules/ModuleInterface.h"
+#include "Misc/CommandLine.h"
+#include "Modules/ModuleManager.h"
+#include "MoviePipelineExecutor.h"
+#include "MoviePipeline.h"
+#include "MoviePipelineInProcessExecutor.h"
+#include "MovieRenderPipelineConfig.h"
+#include "LevelSequence.h"
+
+#if WITH_EDITOR
+//#include "Editor.h"
+#endif
+
+void FMovieRenderPipelineCoreModule::InitializeCommandLineMovieRender()
+{
+	// This should be called after the engine has a chance to load all of the modules so that all of our classes are loaded.
+	if (!ensureMsgf(SequenceAssetValue.Len() > 0, TEXT("InitializeCommandLineMovieRender shouldn't be called without an asset to render!")))
+	{
+		return;
+	}
+
+#if WITH_EDITOR
+	//const bool bIsGameMode = !GEditor;
+	//if (!bIsGameMode)
+	//{
+	//	UE_LOG(LogMovieRenderPipeline, Fatal, TEXT("Command Line Renders must be performed in -game mode, otherwise use the editor ui/python and PIE. Add -game to your command line arguments."));
+	//	FPlatformMisc::RequestExitWithStatus(false, MoviePipelineErrorCodes::Critical);
+	//	return;
+	//}
+#endif
+
+	// Attempt to convert their command line arguments into the required objects.
+	ULevelSequence* LevelSequence = nullptr;
+	UMovieRenderPipelineConfig* Config = nullptr;
+	UClass* ExecutorClass = nullptr;
+	UClass* PipelineClass = nullptr;
+
+	uint8 ReturnCode = GetMovieRenderClasses(SequenceAssetValue, SettingsAssetValue, MoviePipelineLocalExecutorClassType, MoviePipelineClassType, 
+											 /*Out*/ LevelSequence, /*Out*/ Config, /*Out*/ ExecutorClass, /*Out*/ PipelineClass);
+	if (!ensureMsgf(LevelSequence && Config && ExecutorClass && PipelineClass, TEXT("There was a failure parsing the command line and a movie render cannot be started. Check the log for more details.")))
+	{
+		// Take the failure return code from the detection of our command line arguments.
+		FPlatformMisc::RequestExitWithStatus(/*Force*/ false, /*ReturnCode*/ ReturnCode);
+		return;
+	}
+	else
+	{
+		UE_LOG(LogMovieRenderPipeline, Log, TEXT("Successfully detected and loaded required movie arguments. Rendering will begin once the map is loaded."));
+		UE_LOG(LogMovieRenderPipeline, Log, TEXT("LevelSequence: %s Config: %s ExecutorClass: %s PipelineClass: %s"), *LevelSequence->GetPathName(), *Config->GetPathName(), *ExecutorClass->GetPathName(), *PipelineClass->GetPathName());
+	}
+
+	// By this time, we know what assets you want to render, how to process the array of assets, and what runs an individual render. First we will create an executor.
+	UMoviePipelineExecutorBase* ExecutorBase = NewObject<UMoviePipelineExecutorBase>(GetTransientPackage());
+
+	// We add the Executor to the root set. It will own all of the configuration data so this keeps it nicely in memory until finished,
+	// and means we only have to add/remove one thing from the root set, everything else uses normal outer ownership.
+	ExecutorBase->AddToRoot();
+	ExecutorBase->OnExecutorFinished().AddRaw(this, &FMovieRenderPipelineCoreModule::OnCommandLineMovieRenderCompleted);
+
+	TArray<UMovieRenderPipelineConfig*> Pipelines;
+	Pipelines.Add(Config);
+	Config->Sequence = FSoftObjectPath(SequenceAssetValue);
+
+	ExecutorBase->Execute(Pipelines);
+}
+
+void FMovieRenderPipelineCoreModule::OnCommandLineMovieRenderCompleted(UMoviePipelineExecutorBase* InExecutor)
+{
+	if (InExecutor)
+	{
+		InExecutor->RemoveFromRoot();
+
+		// Return a success code. Ideally any errors detected during rendering will return different codes.
+		FPlatformMisc::RequestExitWithStatus(/*Force*/ false, /*ReturnCode*/ 0);
+	}
+}
+
+UClass* GetLocalExecutorClass(const FString& MoviePipelineLocalExecutorClassType, const FString ExecutorClassFormatString)
+{
+	if (MoviePipelineLocalExecutorClassType.Len() > 0)
+	{
+		UMoviePipelineExecutorBase* LoadedExecutor = LoadObject<UMoviePipelineExecutorBase>(GetTransientPackage(), *MoviePipelineLocalExecutorClassType);
+		if (LoadedExecutor)
+		{
+			UE_LOG(LogMovieRenderPipeline, Log, TEXT("Loaded explicitly specified Movie Pipeline Executor %s."), *MoviePipelineLocalExecutorClassType);
+			return LoadedExecutor->GetClass();
+		}
+		else
+		{
+			// They explicitly specified an object, but that couldn't be loaded so it's an error.
+			UE_LOG(LogMovieRenderPipeline, Fatal, TEXT("Failed to load specified Local Executor class. Executor Class: %s"), *MoviePipelineLocalExecutorClassType);
+			UE_LOG(LogMovieRenderPipeline, Warning, TEXT("%s"), *ExecutorClassFormatString);
+			return nullptr;
+		}
+	}
+	else
+	{
+		// Fall back to our provided one. This is okay because it doesn't come from a user preference,
+		// since it needs an different executor than the editor provided New Process
+		UE_LOG(LogMovieRenderPipeline, Log, TEXT("Using default Movie Pipeline Executor %s. See '-MoviePipelineLocalExecutorClass' if you need to override this."), *MoviePipelineLocalExecutorClassType);
+		return UMoviePipelineInProcessExecutor::StaticClass();
+	}
+}
+
+UClass* GetMoviePipelineClass(const FString& MoviePipelineLocalExecutorClassType, const FString ExecutorClassFormatString)
+{
+	if (MoviePipelineLocalExecutorClassType.Len() > 0)
+	{
+		UMoviePipelineExecutorBase* LoadedExecutor = LoadObject<UMoviePipelineExecutorBase>(GetTransientPackage(), *MoviePipelineLocalExecutorClassType);
+		if (LoadedExecutor)
+		{
+			UE_LOG(LogMovieRenderPipeline, Log, TEXT("Loaded explicitly specified Movie Pipeline Executor %s."), *MoviePipelineLocalExecutorClassType);
+			return LoadedExecutor->GetClass();
+		}
+		else
+		{
+			// They explicitly specified an object, but that couldn't be loaded so it's an error.
+			UE_LOG(LogMovieRenderPipeline, Fatal, TEXT("Failed to load specified Local Executor class. Executor Class: %s"), *MoviePipelineLocalExecutorClassType);
+			UE_LOG(LogMovieRenderPipeline, Warning, TEXT("%s"), *ExecutorClassFormatString);
+			return nullptr;
+		}
+	}
+	else
+	{
+		UE_LOG(LogMovieRenderPipeline, Log, TEXT("Using default Movie Pipeline %s. See '-MoviePipelineClass' if you need to override this."), *MoviePipelineLocalExecutorClassType);
+		return UMoviePipeline::StaticClass();
+	}
+}
+
+
+bool FMovieRenderPipelineCoreModule::IsTryingToRenderMovieFromCommandLine(FString& OutSequenceAssetPath, FString& OutConfigAssetPath, FString& OutExecutorType, FString& OutPipelineType) const
+{
+	// Look to see if they've specified a Movie Pipeline to run. This should be in the format:
+	// "/Game/MySequences/MySequence.MySequence"
+	FParse::Value(FCommandLine::Get(), TEXT("-MoviePipeline="), OutSequenceAssetPath);
+
+	// Look to see if they've specified a configuration to use. This should be in the format:
+	// "/Game/MyRenderSettings/MyHighQualitySetting.MyHighQualitySetting" or an absolute path 
+	// to a exported json file.
+	FParse::Value(FCommandLine::Get(), TEXT("-MoviePipelineConfig="), OutConfigAssetPath);
+
+	// The user may want to override the executor. By default, we use the one specified in the Project
+	// Settings, but we allow them to override it on the command line (for render farms, etc.) in case
+	// the one used when a human is running the box isn't appropriate. This should be in the format:
+	// "/Script/ModuleName.ClassNameNoUPrefix"
+	FParse::Value(FCommandLine::Get(), TEXT("-MoviePipelineLocalExecutorClass="), OutExecutorType);
+
+	// The user may want to override the Movie Pipeline itself. By default we use the one specified in the Project
+	// Setting, but we allow overriding it for consistency's sake really. This should be in the format:
+	// "/Script/ModuleName.ClassNameNoUPrefix"
+	FParse::Value(FCommandLine::Get(), TEXT("-MoviePipelineClass="), OutPipelineType);
+
+	return OutSequenceAssetPath.Len() > 0;
+}
+
+uint8 FMovieRenderPipelineCoreModule::GetMovieRenderClasses(const FString& InSequenceAssetPath, const FString& InConfigAssetPath, const FString& InExecutorType, const FString& InPipelineType,
+	ULevelSequence*& OutSequence, UMovieRenderPipelineConfig*& OutConfig, UClass*& OutExecutorClass, UClass*& OutPipelineClass) const
+{
+	if (!ensureMsgf(InSequenceAssetPath.Len() > 0, TEXT("GetMovieRenderClasses should not be called without an sequence asset path specified!")))
+	{
+		return MoviePipelineErrorCodes::NoAsset;
+	}
+
+	// Store off the messages that print the expected format since they're printed in a couple places.
+	const FString SequenceAssetFormatString = TEXT("Movie Pipeline Asset should be specified in the format '-MoviePipeline=\"/Game/MySequences/MyPipelineAsset\", or a relative (to executable)/absolute path to a JSON file in the format '-MoviePipeline=\"D:/MyMoviePipelineConfig.json\".");
+	const FString ConfigAssetFormatString	= TEXT("Pipeline Config Asset should be specified in the format '-MoviePipelineConfig=\"/Game/MyRenderSettings/MyHighQualitySetting.MyHighQualitySetting\"'");
+	const FString PipelineClassFormatString = TEXT("Movie Pipeline Class should be specified in the format '-MoviePipelineClass=\"/Script/ModuleName.ClassNameNoUPrefix\"'");
+	const FString ExecutorClassFormatString = TEXT("Pipeline Executor Class should be specified in the format '-MoviePipelineLocalExecutorClass=\"/Script/ModuleName.ClassNameNoUPrefix\"'");
+
+	// They need both an asset and a configuration so early out if we don't have that.
+	if (InConfigAssetPath.Len() == 0)
+	{
+		UE_LOG(LogMovieRenderPipeline, Fatal, TEXT("No Movie Pipeline Config was specified to use as settings."));
+		UE_LOG(LogMovieRenderPipeline, Warning, TEXT("%s"), *ConfigAssetFormatString);
+		return MoviePipelineErrorCodes::NoConfig;
+	}
+
+	OutSequence = nullptr;
+	OutConfig = nullptr;
+	OutExecutorClass = nullptr;
+	OutPipelineClass = nullptr;
+
+	// We try to detect executor/class overrides first because it'll affect deserialization of JSON based pipeline assets.
+	OutExecutorClass = GetLocalExecutorClass(MoviePipelineLocalExecutorClassType, ExecutorClassFormatString);
+	OutPipelineClass = GetMoviePipelineClass(MoviePipelineClassType, PipelineClassFormatString);
+
+	if (!ensureMsgf(OutPipelineClass && OutExecutorClass, TEXT("Attempted to render a movie pipeline but could not load executor or pipeline class, nor fall back to defaults.")))
+	{
+		return MoviePipelineErrorCodes::Critical;
+	}
+
+	// Locate and load the level sequence they wish to render.
+	{
+		// Convert it to a soft object path and use that load to ensure it follows redirectors, etc.
+		FSoftObjectPath AssetPath = FSoftObjectPath(InSequenceAssetPath);
+		OutSequence = CastChecked<ULevelSequence>(AssetPath.TryLoad());
+
+		if (!OutSequence)
+		{
+			UE_LOG(LogMovieRenderPipeline, Fatal, TEXT("Failed to find Movie Pipeline sequence asset to render. Please note that the /Content/ part of the on-disk structure is omitted. Looked for: %s"), *InSequenceAssetPath);
+			UE_LOG(LogMovieRenderPipeline, Warning, TEXT("%s"), *SequenceAssetFormatString);
+			return MoviePipelineErrorCodes::NoAsset;
+		}
+	}
+
+	// Now look for the configuration file to see how to render it.
+	if (InConfigAssetPath.StartsWith("/Game/"))
+	{
+		// Convert it to a soft object path and use that load to ensure it follows redirectors, etc.
+		FSoftObjectPath AssetPath = FSoftObjectPath(InConfigAssetPath);
+		OutConfig = CastChecked<UMovieRenderPipelineConfig>(AssetPath.TryLoad());
+
+		if (!OutConfig)
+		{
+			UE_LOG(LogMovieRenderPipeline, Fatal, TEXT("Failed to find Pipeline Configuration asset to render. Please note that the /Content/ part of the on-disk structure is omitted. Looked for: %s"), *InConfigAssetPath);
+			UE_LOG(LogMovieRenderPipeline, Warning, TEXT("%s"), *ConfigAssetFormatString);
+			return MoviePipelineErrorCodes::NoConfig;
+		}
+	}
+	else
+	{
+		// If they didn't pass a path that started with /Game/, we'll try to see if it is a JSON file we can understand.
+		UE_LOG(LogMovieRenderPipeline, Fatal, TEXT("Not supported yet"));
+		return MoviePipelineErrorCodes::Critical;
+	}
+
+	return MoviePipelineErrorCodes::Success;
+}
diff --git a/Engine/Plugins/MovieScene/MovieRenderPipeline/Source/MovieRenderPipelineCore/Private/MovieRenderPipelineConfig.cpp b/Engine/Plugins/MovieScene/MovieRenderPipeline/Source/MovieRenderPipelineCore/Private/MovieRenderPipelineConfig.cpp
new file mode 100644
--- /dev/null
+++ b/Engine/Plugins/MovieScene/MovieRenderPipeline/Source/MovieRenderPipelineCore/Private/MovieRenderPipelineConfig.cpp
@@ -0,0 +1,52 @@
+// Copyright 1998-2019 Epic Games, Inc. All Rights Reserved.
+
+#include "MovieRenderPipelineConfig.h"
+#include "MovieScene.h"
+#include "LevelSequence.h"
+
+#define LOCTEXT_NAMESPACE "MovieRenderPipelineConfig"
+
+FFrameRate UMovieRenderPipelineConfig::GetEffectiveFrameRate()
+{
+	LoadTargetSequence();
+
+	// Check to see if they overrode the frame rate.
+	if (bUseCustomFrameRate)
+	{
+		return OutputFrameRate;
+	}
+
+	// Pull it from the sequence if they didn't.
+	if (LoadedSequence)
+	{
+		return LoadedSequence->GetMovieScene()->GetDisplayRate();
+	}
+
+	return FFrameRate();
+}
+
+UMoviePipelineShotConfig* UMovieRenderPipelineConfig::GetConfigForShot(const FString& ShotName) const
+{
+	UMoviePipelineShotConfig* OutConfig = PerShotConfigMapping.FindRef(ShotName);
+
+	// They didn't customize this shot, return the global pipeline default
+	if (!OutConfig)
+	{
+		OutConfig = DefaultShotConfig;
+	}
+
+	return OutConfig;
+}
+
+void UMovieRenderPipelineConfig::LoadTargetSequence()
+{
+	if (LoadedSequence)
+	{
+		return;
+	}
+
+	LoadedSequence = LoadObject<ULevelSequence>(this, *Sequence.GetAssetPathString());
+	ensureMsgf(LoadedSequence, TEXT("Failed to load target sequence. Pipeline is not fully configured."));
+}
+
+#undef LOCTEXT_NAMESPACE // "MovieRenderPipelineConfig"
\ No newline at end of file
diff --git a/Engine/Plugins/MovieScene/MovieRenderPipeline/Source/MovieRenderPipelineCore/Private/MovieRenderPipelineCoreModule.cpp b/Engine/Plugins/MovieScene/MovieRenderPipeline/Source/MovieRenderPipelineCore/Private/MovieRenderPipelineCoreModule.cpp
new file mode 100644
--- /dev/null
+++ b/Engine/Plugins/MovieScene/MovieRenderPipeline/Source/MovieRenderPipelineCore/Private/MovieRenderPipelineCoreModule.cpp
@@ -0,0 +1,25 @@
+// Copyright 1998-2019 Epic Games, Inc. All Rights Reserved.
+
+#include "MovieRenderPipelineCoreModule.h"
+#include "Modules/ModuleInterface.h"
+#include "Misc/CoreDelegates.h"
+#include "Modules/ModuleManager.h"
+
+void FMovieRenderPipelineCoreModule::StartupModule()
+{
+	// Look to see if they supplied arguments on the command line indicating they wish to render a movie.
+	if (IsTryingToRenderMovieFromCommandLine(SequenceAssetValue, SettingsAssetValue, MoviePipelineLocalExecutorClassType, MoviePipelineClassType))
+	{
+		UE_LOG(LogMovieRenderPipeline, Log, TEXT("Detected that the user intends to render a movie. Waiting until engine loop init is complete to ensure "))
+		// Register a hook to wait until the engine has finished loading to increase the likelihood that the desired classes are loaded.
+		// Note: This is not called for a commandlet.
+		FCoreDelegates::OnFEngineLoopInitComplete.AddRaw(this, &FMovieRenderPipelineCoreModule::InitializeCommandLineMovieRender);
+	}
+}
+
+void FMovieRenderPipelineCoreModule::ShutdownModule()
+{
+}
+
+IMPLEMENT_MODULE(FMovieRenderPipelineCoreModule, MovieRenderPipelineCore);
+DEFINE_LOG_CATEGORY(LogMovieRenderPipeline); 
diff --git a/Engine/Plugins/MovieScene/MovieRenderPipeline/Source/MovieRenderPipelineCore/Public/MoviePipeline.h b/Engine/Plugins/MovieScene/MovieRenderPipeline/Source/MovieRenderPipelineCore/Public/MoviePipeline.h
new file mode 100644
--- /dev/null
+++ b/Engine/Plugins/MovieScene/MovieRenderPipeline/Source/MovieRenderPipelineCore/Public/MoviePipeline.h
@@ -0,0 +1,246 @@
+// Copyright 1998-2019 Epic Games, Inc. All Rights Reserved.
+#pragma once
+
+#include "UObject/Object.h"
+#include "Engine/EngineCustomTimeStep.h"
+#include "MovieRenderPipelineDataTypes.h"
+#include "MovieSceneTimeController.h"
+#include "MoviePipeline.generated.h"
+
+// Forward Declares
+class UMovieRenderPipelineConfig;
+class ULevelSequence;
+class UMovieSceneSequencePlayer;
+class UMoviePipelineCustomTimeStep;
+class UEngineCustomTimeStep;
+class ALevelSequenceActor;
+class UMovieRenderDebugWidget;
+class FMovieRenderViewport;
+class FMovieRenderViewportClient;
+struct FImagePixelPipe;
+struct FMoviePipelineTimeController;
+class FMoviePipelineOutputMerger;
+
+DECLARE_MULTICAST_DELEGATE_OneParam(FMoviePipelineFinished, UMoviePipeline*);
+
+UCLASS(Blueprintable)
+class MOVIERENDERPIPELINECORE_API UMoviePipeline : public UObject
+{
+	GENERATED_BODY()
+	
+public:
+	UMoviePipeline();
+
+	/** 
+	* Initialize the movie pipeline with the specified settings. This kicks off the rendering process. 
+	
+	* @param InInitSettings - Settings related to the context the pipeline is being executed in. Not configuration for the Pipeline.
+	* @param InConfig		- What settings should the Pipeline use to render with? This should be filled out completely
+	*						  before starting the pipeline. Any overrides based on command line or other should be done
+	*						  before the pipeline is started.
+	*/
+	void Initialize(UMovieRenderPipelineConfig* InConfig);
+
+	void LoadDebugWidget();
+
+	/** 
+	* Call to shut down the pipeline. This flushes any outstanding file writes and unregisters all delegates.
+	* This will not call the OnMoviePipelineFinished() delegate as it assumes the external party already knows
+	* it is "Finished". 
+	*/
+	void Shutdown();
+
+	/** 
+	* Called when we have completely finished this pipeline. This means that all frames have been rendered,
+	* all files written to disk, and any post-finalize exports have finished. This Pipeline will call
+	* Shutdown() on itself before calling this delegate to ensure we've unregistered from all delegates
+	* and are no longer trying to do anything (even if we still exist).
+	*/
+	FMoviePipelineFinished& OnMoviePipelineFinished()
+	{
+		return OnMoviePipelineFinishedDelegate;
+	}
+
+	UFUNCTION(BlueprintPure, Category = "Movie Render Pipeline")
+	FMoviePipelineShotCache GetCurrentShotSnapshot() const;
+
+	UFUNCTION(BlueprintPure, Category = "Movie Render Pipeline")
+	FMoviePipelineShotCutCache GetCurrentCameraCutSnapshot() const;
+
+	UFUNCTION(BlueprintPure, Category = "Movie Render Pipeline")
+	FMoviePipelineFrameOutputState GetOutputStateSnapshot() const;
+
+	UFUNCTION(BlueprintPure, Category = "Movie Render Pipeline")
+	ULevelSequence* GetTargetSequence() const { return TargetSequence; }
+
+	UFUNCTION(BlueprintPure, Category = "Movie Render Pipeline")
+	UMovieRenderPipelineConfig* GetPipelineConfig() const { return Config; }
+
+	UFUNCTION(BlueprintPure, Category = "Movie Render Pipeline")
+	FFrameNumber GetTotalOutputFrameCountEstimate() const;
+
+	/** 
+	* Returns an estimate based on the average time taken to render all previous frames.
+	* Will be incorrect if PlayRate tracks are in use. Will be inaccurate when different
+	* shots take significantly different amounts of time to render.
+	*
+	* @param OutTimespan: The returned estimated time left. Will be default initialized if there is no estimate.
+	* @return True if we can make a reasonable estimate, false otherwise (ie: no rendering has been done to estimate).
+	*/
+	UFUNCTION(BlueprintPure, Category = "Movie Render Pipeline")
+	bool GetRemainingTimeEstimate(FTimespan& OutTimespan) const;
+
+	const TArray<FMoviePipelineShotCache>& GetShotList() const { return ShotList; }
+	int32 GetCurrentShotIndex() const { return CurrentShotIndex; }
+
+	TSharedPtr<FImagePixelPipe, ESPMode::ThreadSafe> GetOutputPipe() const { return OutputPipe; }
+private:
+
+	void OnEngineTickBeginFrame();
+	void TickProducingFrames();
+
+	/** Return true if we should early out of the Tick function */
+	bool DebugFrameStepPreTick();
+	/** Returns true if we are idling because of debug frame stepping. */
+	bool IsDebugFrameStepIdling() const;
+
+	void OnSequenceEvaluated(const UMovieSceneSequencePlayer& Player, FFrameTime CurrentTime, FFrameTime PreviousTime);
+
+	void OnEngineTickEndFrame();
+	void TickFinalizeOutputContainers(const bool bInForceFinish);
+	void PostFinalizeExport();
+	void TickPostFinalizeExport();
+	void InitializeShot(FMoviePipelineShotCache& InShot);
+	void SetupRenderingPipelineForShot(FMoviePipelineShotCache& InShot);
+	void TeardownRenderingPipelineForShot(FMoviePipelineShotCache& InShot);
+
+	void BeginFinalize();
+	void FlushAsyncSystems();
+	/** 
+	* Renders the next frame in the Pipeline. This updates/ticks all scene view render states
+	* and produces data. This may not result in an output frame due to multiple renders 
+	* accumulating together to produce an output. frame.
+	* Should not be called if we're idling (debug) or not initialized yet. 
+	*/
+	void RenderFrame();
+private:
+	ULevelSequence* CreateCopyOfSequence(ULevelSequence* InSequence);
+	void InitializeLevelSequenceActor(ULevelSequence* OriginalLevelSequence, ULevelSequence* InSequenceToApply);
+
+	/** Modifies the TargetSequence to ensure that only the specified Shot has it's associated Cinematic Shot Section enabled. */
+	void SetSoloShot(const FMoviePipelineShotCache& InShot);
+	void SetSoloCameraCut(FMoviePipelineShotCutCache& InCameraCut);
+
+	/** This converts the sequence into a Shot List and expands bounds.*/
+	TArray<FMoviePipelineShotCache> BuildShotListFromSequence(const ULevelSequence* InSequence);
+
+	/* Tick Resolution Frames*/
+	FFrameNumber GetMotionBlurDuration(int32 InShotIndex) const;
+	/* Tick Resolution Frames*/
+	FFrameNumber CalculateHandleFrameDuration(const FMoviePipelineShotCache& InShot) const;
+	
+	MoviePipeline::FFrameConstantMetrics CalculateShotFrameMetrics(const FMoviePipelineShotCache& InShot) const;
+
+	void CalculateFrameNumbersForOutputState(const MoviePipeline::FFrameConstantMetrics& InFrameMetrics, const FMoviePipelineShotCutCache& InCameraCut, FMoviePipelineFrameOutputState& InOutOutputState) const;
+
+	/* Expands the specified shot by the specified offsets. Offsets should be in TickResolution */
+	void ExpandShot(FMoviePipelineShotCache& InShot);
+
+private:
+	UPROPERTY(Transient, Instanced)
+	UMoviePipelineCustomTimeStep* CustomTimeStep;
+
+	TSharedPtr<FMoviePipelineTimeController> CustomSequenceTimeController;
+
+	/** Hold a reference to the existing custom time step (if any) so we can restore it after we're done using our custom one. */
+	UPROPERTY(Transient)
+	UEngineCustomTimeStep* CachedCustomTimeStep;
+
+	UPROPERTY(Transient, BlueprintReadOnly, meta=(AllowPrivateAccess=true), Category = "Movie Render Pipeline")
+	ULevelSequence* TargetSequence;
+
+	UPROPERTY(Transient)
+	ALevelSequenceActor* LevelSequenceActor;
+
+	UPROPERTY(Transient)
+	UMovieRenderDebugWidget* DebugWidget;
+
+	/** A list of all of the shots we are going to render out from this sequence. */
+	UPROPERTY(Transient, BlueprintReadOnly, meta = (AllowPrivateAccess = true), Category = "Movie Render Pipeline")
+	TArray<FMoviePipelineShotCache> ShotList;
+
+	/** What state of the overall flow are we in? See enum for specifics. */
+	UPROPERTY(Transient, BlueprintReadOnly, meta = (AllowPrivateAccess = true), Category = "Movie Render Pipeline")
+	EMovieRenderPipelineState PipelineState;
+
+	/** What is the index of the shot we are working on. -1 if not initialized, may be greater than ShotList.Num() if we've reached the end. */
+	UPROPERTY(Transient, BlueprintReadOnly, meta = (AllowPrivateAccess = true), Category = "Movie Render Pipeline")
+	int32 CurrentShotIndex;
+
+	UPROPERTY(Transient, BlueprintReadOnly, meta = (AllowPrivateAccess = true), Category = "Movie Render Pipeline")
+	FMoviePipelineFrameOutputState CachedOutputState;
+
+	/** The time (in UTC) that Initialize was called. Used to track elapsed time. */
+	UPROPERTY(Transient, BlueprintReadOnly, meta = (AllowPrivateAccess = true), Category = "Movie Render Pipeline")
+	FDateTime InitializationTime;
+
+	/** 
+	* Have we hit the callback for the BeginFrame at least once? This solves an issue where the delegates
+	* get registered mid-frame so you end up calling EndFrame before BeginFrame which is undesirable.
+	*/
+	bool bHasRunBeginFrameOnce;
+
+	/** Should we pause the game at the end of the frame? Used to implement frame step debugger. */
+	bool bPauseAtEndOfFrame;
+
+	/** When using temporal sub-frame stepping common counts (such as 3) don't result in whole ticks. We keep track of how many ticks we lose so we can add them the next time there's a chance. */
+	float AccumulatedTickSubFrameDeltas;
+
+	/** Called when we have completely finished. This object will call Shutdown before this and stop ticking. */
+	FMoviePipelineFinished OnMoviePipelineFinishedDelegate;
+
+	/** This gathers all of the produced data for an output frame (which may come in async many frames later) before passing them onto the Output Containers. */
+	TSharedPtr<FMoviePipelineOutputMerger, ESPMode::ThreadSafe> OutputBuilder;
+
+	// Debug Rendering
+	TSharedPtr<FMovieRenderViewport> DummyViewport;
+	TSharedPtr<FMovieRenderViewportClient> ViewportClient;
+
+
+	TSharedPtr<FImagePixelPipe, ESPMode::ThreadSafe> OutputPipe;
+
+private:
+	/** The overall configuration for how this Pipeline executes. Should be immutable. */
+	UPROPERTY(Transient)
+	UMovieRenderPipelineConfig* Config;
+};
+
+UCLASS()
+class UMoviePipelineCustomTimeStep : public UEngineCustomTimeStep
+{
+	GENERATED_BODY()
+
+public:
+	// UEngineCustomTimeStep Interface
+	virtual bool Initialize(UEngine* InEngine) override { return true; }
+	virtual void Shutdown(UEngine* InEngine) override {}
+	virtual bool UpdateTimeStep(UEngine* InEngine) override;
+	virtual ECustomTimeStepSynchronizationState GetSynchronizationState() const override { return ECustomTimeStepSynchronizationState::Synchronized; }
+	// ~UEngineCustomTimeStep Interface
+
+	void SetCachedFrameTiming(const MoviePipeline::FFrameTimeStepCache& InTimeCache) { TimeCache = InTimeCache; }
+
+private:
+	/** We don't do any thinking on our own, instead we just spit out the numbers stored in our time cache. */
+	MoviePipeline::FFrameTimeStepCache TimeCache;
+};
+
+struct FMoviePipelineTimeController : public FMovieSceneTimeController
+{
+	virtual FFrameTime OnRequestCurrentTime(const FQualifiedFrameTime& InCurrentTime, float InPlayRate) override;
+	void SetCachedFrameTiming(const FQualifiedFrameTime& InTimeCache) { TimeCache = InTimeCache; }
+
+private:
+	/** Simply store the number calculated and return it when requested. */
+	FQualifiedFrameTime TimeCache;
+};
\ No newline at end of file
diff --git a/Engine/Plugins/MovieScene/MovieRenderPipeline/Source/MovieRenderPipelineCore/Public/MoviePipelineAccumulationSetting.h b/Engine/Plugins/MovieScene/MovieRenderPipeline/Source/MovieRenderPipelineCore/Public/MoviePipelineAccumulationSetting.h
new file mode 100644
--- /dev/null
+++ b/Engine/Plugins/MovieScene/MovieRenderPipeline/Source/MovieRenderPipelineCore/Public/MoviePipelineAccumulationSetting.h
@@ -0,0 +1,90 @@
+// Copyright 1998-2019 Epic Games, Inc. All Rights Reserved.
+#pragma once
+
+#include "MoviePipelineSetting.h"
+#include "MoviePipelineAccumulationSetting.generated.h"
+
+
+UCLASS(Blueprintable)
+class MOVIERENDERPIPELINECORE_API UMoviePipelineAccumulationSetting : public UMoviePipelineSetting
+{
+	GENERATED_BODY()
+public:
+	UMoviePipelineAccumulationSetting()
+		: TemporalSampleCount(1)
+		, CameraShutterAngle(180)
+		, TileCount(1)
+		, SpatialSampleCount(1)
+		, bIsUsingOverlappedTiles(false)
+		, PadRatioX(.50f)
+		, PadRatioY(.50f)
+	{
+	}
+	
+	/** 
+	* The number of frames we should combine together to produce each output frame. This blends the
+	* results of this many sub-steps together to produce one output frame. See CameraShutterAngle to
+	* control how much time passes between each sub-frame. See SpatialSampleCount to see how many
+	* samples we average together to produce a sub-step. (This means rendering complexity is
+	* SampleCount * TileCount^2 * SpatialSampleCount * NumPasses).
+	*/
+	UPROPERTY(EditAnywhere, BlueprintReadWrite, meta=(UIMin = 1, ClampMin = 1), Category = "Movie Pipeline")
+	int32 TemporalSampleCount;
+	
+	/** 
+	* The camera shutter angle determines how much of a given frame the accumulation frames span.
+	* For example, a 24fps film with a shutter angle of 180 will spread the accumulation frames
+	* out over (180/360) percent of the frame (1/24th of a second). With 2 accumulation frames,
+	* the first one will be at the very start of the frame and the second accumulation frame will
+	* be land on the half way point between the start of this frame and the start of the next.
+	*
+	* A shutter angle of 360 means continuous movement, while a shutter angle of zero means no
+	* motion blur.
+	*/
+	UPROPERTY(EditAnywhere, BlueprintReadWrite, meta = (UIMin = 0, UIMax = 360, ClampMin = 0, ClampMax = 360), Category = "Movie Pipeline")
+	int32 CameraShutterAngle;
+
+	/**
+	* How many tiles should the resulting movie render be broken into? A tile should be no larger than
+	* the maximum texture resolution supported by your graphics card (likely 16k), so NumTiles should be
+	* ceil(Width/MaxTextureSize). More tiles mean more individual passes over the whole scene at a smaller
+	* resolution which may help with gpu timeouts. Requires at least 1 tile. Tiling is applied evenly to
+	* both X and Y (to make final interlacing more straightforward).
+	*/
+	UPROPERTY(EditAnywhere, BlueprintReadWrite, meta = (UIMin = 1, ClampMin = 1), Category = "Movie Pipeline")
+	int32 TileCount;
+
+	/**
+	* This bias encourages the engine to use a higher detail texture when it would normally use a lower detail
+	* texture due to the size of the texture on screen. A more negative number means overall sharper images
+	* (up to the detail resolution of your texture). Too much sharpness will cause visual grain/noise in the
+	* resulting image, but this can be mitigated with more spatial samples.
+	*/
+	UPROPERTY(EditAnywhere, BlueprintReadWrite, meta = (UIMin = -1, UIMax = 0), Category = "Movie Pipeline")
+	float TextureSharpnessBias;
+
+	/**
+	* How many frames should we accumulate together before contributing to one overall sample. This lets you
+	* increase the anti-aliasing quality of an sample, or have high quality anti-aliasing if you don't want
+	* any motion blur due to accumulation over time in SampleCount.
+	*/
+	UPROPERTY(EditAnywhere, BlueprintReadWrite, meta = (UIMin = 1, ClampMin = 1), Category = "Movie Pipeline")
+	int32 SpatialSampleCount;
+
+
+	bool bIsUsingOverlappedTiles;
+
+	float PadRatioX;
+	float PadRatioY;
+	float AccumulationGamma;
+
+
+	/**
+	* Shutter Timing allows you to bias the timing of your shutter angle to either be before, during, or after
+	* a frame. When set to FrameClose, it means that the motion gathered up to produce frame N comes from 
+	* before and right up to frame N. When set to FrameCenter, the motion represents half the time before the
+	* frame and half the time after the frame. When set to FrameOpen, the motion represents the time from 
+	* Frame N onwards.
+	*/
+	EMoviePipelineShutterTiming ShutterTiming;
+};
\ No newline at end of file
diff --git a/Engine/Plugins/MovieScene/MovieRenderPipeline/Source/MovieRenderPipelineCore/Public/MoviePipelineBlueprintLibrary.h b/Engine/Plugins/MovieScene/MovieRenderPipeline/Source/MovieRenderPipelineCore/Public/MoviePipelineBlueprintLibrary.h
new file mode 100644
--- /dev/null
+++ b/Engine/Plugins/MovieScene/MovieRenderPipeline/Source/MovieRenderPipelineCore/Public/MoviePipelineBlueprintLibrary.h
@@ -0,0 +1,52 @@
+// Copyright 1998-2019 Epic Games, Inc. All Rights Reserved.
+
+#pragma once
+
+#include "Kismet/BlueprintFunctionLibrary.h"
+#include "MovieRenderPipelineDataTypes.h"
+
+#include "MoviePipelineBlueprintLibrary.generated.h"
+
+// Forward Declare
+class UMoviePipeline;
+
+UCLASS()
+class MOVIERENDERPIPELINECORE_API UMoviePipelineBlueprintLibrary : public UBlueprintFunctionLibrary
+{
+	GENERATED_BODY()
+
+public:
+	/** 
+	* Get the total number of Camera Cuts in a Movie Pipeline, and an index of which one is being processed.
+	* Note: Index doesn't index into any arrays (since internally we store CameraCuts inside of Shots)
+	*/
+	UFUNCTION(BlueprintPure, Category = "Movie Render Pipeline")
+	static void GetCameraCutCounts(const UMoviePipeline* InPipeline, int32& OutTotalCuts, int32& OutCurrentCutIndex);
+	
+	/** 
+	* Returns the number of expected frames to be produced by Initial Range + Handle Frames given a Frame Rate.
+	* This will be inaccurate when PlayRate tracks are involved.
+	*/
+	UFUNCTION(BlueprintPure, Category = "Movie Render Pipeline")
+	static int32 GetOutputFrameCountEstimate(const FMoviePipelineShotCutCache& InCameraCut);
+	
+	/**
+	* Returns the expected number of frames different frames that will be submitted for rendering. This is
+	* the number of output frames * temporal samples. This will be inaccurate when PlayRate tracks are involved.
+	*/
+	UFUNCTION(BlueprintPure, Category = "Movie Render Pipeline")
+	static int32 GetTemporalFrameCountEstimate(const FMoviePipelineShotCutCache& InCameraCut);
+	
+	/** 
+	* Returns misc. utility frame counts (Warm Up + MotionBlur Fix) since these are outside the ranges.
+	*/
+	UFUNCTION(BlueprintPure, Category = "Movie Render Pipeline")
+	static int32 GetUtilityFrameCountEstimate(const FMoviePipelineShotCutCache& InCameraCut);
+	
+	/**
+	* Returns the number of samples submitted to the GPU, optionally counting samples for the various parts.
+	* This will be inaccurate when PlayRate tracks are involved.
+	*/
+	UFUNCTION(BlueprintPure, Category = "Movie Render Pipeline")
+	static int32 GetSampleCountEstimate(const FMoviePipelineShotCutCache& InCameraCut, const bool bIncludeWarmup = true, const bool bIncludeMotionBlur = true);
+};
\ No newline at end of file
diff --git a/Engine/Plugins/MovieScene/MovieRenderPipeline/Source/MovieRenderPipelineCore/Public/MoviePipelineExecutor.h b/Engine/Plugins/MovieScene/MovieRenderPipeline/Source/MovieRenderPipelineCore/Public/MoviePipelineExecutor.h
new file mode 100644
--- /dev/null
+++ b/Engine/Plugins/MovieScene/MovieRenderPipeline/Source/MovieRenderPipelineCore/Public/MoviePipelineExecutor.h
@@ -0,0 +1,61 @@
+// Copyright 1998-2019 Epic Games, Inc. All Rights Reserved.
+#pragma once
+
+#include "UObject/Object.h"
+#include "MoviePipelineExecutor.generated.h"
+
+class UMovieRenderPipelineConfig;
+class UMoviePipelineExecutorBase;
+
+DECLARE_MULTICAST_DELEGATE_OneParam(FOnMoviePipelineExecutorFinishedNative, UMoviePipelineExecutorBase*);
+DECLARE_DYNAMIC_MULTICAST_DELEGATE_OneParam(FOnMoviePipelineExecutorFinished, UMoviePipelineExecutorBase*, PipelineExecutor);
+
+/**
+* A Movie Pipeline Executor is responsible for executing an array of Movie Pipelines,
+* and (optionally) reporting progress back for the movie pipelines. The entire array
+* is passed at once to allow the implementations to choose how to split up the work.
+* By default we provide a local execution (UMoviePipelineLocalExecutor) which works
+* on them serially, but you can create an implementation of this class, change the 
+* default in the Project Settings and use your own distribution logic. For example,
+* you may want to distribute the work to multiple computers over a network, which
+* may involve running command line options on each machine to sync the latest content
+* from the project before the execution starts.
+*/
+UCLASS(Blueprintable, Abstract)
+class MOVIERENDERPIPELINECORE_API UMoviePipelineExecutorBase : public UObject
+{
+	GENERATED_BODY()
+public:
+	UFUNCTION(BlueprintCallable, Category = "Movie Render Pipeline")
+	virtual void Execute(UPARAM(Ref) TArray<UMovieRenderPipelineConfig*>& InPipelines)
+	{
+		ExecuteImpl(InPipelines);
+	}
+
+	FOnMoviePipelineExecutorFinishedNative& OnExecutorFinished()
+	{
+		return OnExecutorFinishedDelegateNative;
+	}
+	
+protected:
+	/** 
+	* This should be called when the Executor has finished executing all of the things
+	* it has been asked to execute. This should be called in the event of a failure as 
+	* well, and passing in false for success to allow the caller to know failure.
+	*/
+	virtual void OnExecutorFinishedImpl(const bool bInSuccess)
+	{
+		// Broadcast to both Native and Python/BP
+		OnExecutorFinishedDelegateNative.Broadcast(this);
+		OnExecutorFinishedDelegate.Broadcast(this);
+	}
+	virtual void ExecuteImpl(TArray<UMovieRenderPipelineConfig*>& InPipelines) PURE_VIRTUAL(UMoviePipelineExecutorBase::ExecuteImpl, );
+
+private:
+	/** Exposed for Blueprints/Python. Called at the same time as the native one. */
+	UPROPERTY(BlueprintAssignable, Category = "Movie Render Pipeline")
+	FOnMoviePipelineExecutorFinished OnExecutorFinishedDelegate;
+
+	/** For native C++ code. Called at the same time as the Blueprint/Python one. */
+	FOnMoviePipelineExecutorFinishedNative OnExecutorFinishedDelegateNative;
+};
\ No newline at end of file
diff --git a/Engine/Plugins/MovieScene/MovieRenderPipeline/Source/MovieRenderPipelineCore/Public/MoviePipelineLinearExecutor.h b/Engine/Plugins/MovieScene/MovieRenderPipeline/Source/MovieRenderPipelineCore/Public/MoviePipelineLinearExecutor.h
new file mode 100644
--- /dev/null
+++ b/Engine/Plugins/MovieScene/MovieRenderPipeline/Source/MovieRenderPipelineCore/Public/MoviePipelineLinearExecutor.h
@@ -0,0 +1,50 @@
+// Copyright 1998-2019 Epic Games, Inc. All Rights Reserved.
+#pragma once
+
+#include "MoviePipelineExecutor.h"
+#include "MoviePipelineLinearExecutor.generated.h"
+
+class UMoviePipeline;
+class UMoviePipelineInEditorExecution;
+
+/**
+* This is an abstract base class designed for executing an array of movie pipelines in linear
+* fashion. It is generally the case that you only want to execute one at a time on a local machine
+* and a different executor approach should be taken for a render farm that distributes out jobs
+* to many different machines.
+*/
+UCLASS(Blueprintable, Abstract)
+class MOVIERENDERPIPELINECORE_API UMoviePipelineLinearExecutorBase : public UMoviePipelineExecutorBase
+{
+	GENERATED_BODY()
+	
+public:
+	UMoviePipelineLinearExecutorBase()
+		: UMoviePipelineExecutorBase()
+		, CurrentPipelineIndex(0)
+	{
+	}
+
+protected:
+	virtual void ExecuteImpl(TArray<UMovieRenderPipelineConfig*>& InPipelines) override;
+
+	virtual void StartPipelineByIndex(int32 InPipelineIndex);
+	virtual void Start(UMovieRenderPipelineConfig* InConfig, const int32 InConfigIndex, const int32 InNumConfigs) { check(InConfig); }
+
+public:
+	virtual void OnIndividualPipelineFinished(UMoviePipeline* /* FinishedPipeline */);
+	virtual void OnExecutorFinishedImpl(const bool bInSuccess);
+protected:
+	
+	/** List of Pipeline Configs we've been asked to execute. */
+	UPROPERTY(Transient)
+	TArray<UMovieRenderPipelineConfig*> Pipelines;
+
+	/** Which Pipeline Config are we currently working on. */
+	int32 CurrentPipelineIndex;
+private:
+	/** Used to track total processing duration. */
+	FDateTime InitializationTime;
+};
+
+
diff --git a/Engine/Plugins/MovieScene/MovieRenderPipeline/Source/MovieRenderPipelineCore/Public/MoviePipelineOutput.h b/Engine/Plugins/MovieScene/MovieRenderPipeline/Source/MovieRenderPipelineCore/Public/MoviePipelineOutput.h
new file mode 100644
--- /dev/null
+++ b/Engine/Plugins/MovieScene/MovieRenderPipeline/Source/MovieRenderPipelineCore/Public/MoviePipelineOutput.h
@@ -0,0 +1,84 @@
+// Copyright 1998-2019 Epic Games, Inc. All Rights Reserved.
+#pragma once
+
+#include "UObject/Object.h"
+#include "MovieRenderPipelineDataTypes.h"
+#include "Engine/EngineTypes.h"
+#include "MoviePipelineSetting.h"
+#include "MoviePipelineOutput.generated.h"
+
+class UMoviePipeline;
+
+/**
+* A base class for different output containers for the data generated by the Movie Pipeline,
+* i.e: image sequences, video containers, etc.
+*/
+UCLASS(Blueprintable, Abstract)
+class MOVIERENDERPIPELINECORE_API UMoviePipelineOutput : public UMoviePipelineSetting
+{
+	GENERATED_BODY()
+public:
+	/**
+	* Called when all data sources for a given output frame are available and need to be placed in the output stream. All input passes are
+	* given at once to make it easier to decide if they belong in separate containers or in combined containers.
+	*/
+	void ProcessFrame(TArray<MoviePipeline::FOutputFrameData> FrameData, FMoviePipelineFrameOutputState CachedOutputState, FDirectoryPath OutputDirectory)
+	{
+		ProcessFrameImpl(FrameData, CachedOutputState, OutputDirectory);
+	}
+
+	/** 
+	* Called once when all frames have been produced for the pipeline. Use this as an indicator to start flushing to disk. 
+	* Async processes can be started here and HasFinishedProcessing() will be called each frame until all all containers
+	* return true, at which point Finalize will be called.
+	*/
+	void BeginFinalize() { BeginFinalizeImpl(); }
+
+	/**
+	* Called each frame to check to see if this output container has finished processing. Finalize won't be called until
+	* all output containers return true for this. 
+	*/
+	bool HasFinishedProcessing() { return HasFinishedProcessingImpl(); }
+
+	/** 
+	* Called after all output containers have reported that they are done processing.
+	*/
+	void Finalize() { FinalizeImpl(); }
+
+	/**
+	* This is called when a new shot starts, before warm up happens. See OnFrameProductionStart.
+	*/
+	void OnShotInitialized(const TOptional<FMoviePipelineShotCache> PrevShot, const FMoviePipelineShotCache& NewShot) { OnShotInitializedImpl(PrevShot, NewShot); }
+	
+	/**
+	* This is called when a shot ends, right after the last frame is rendered.
+	*/
+	void OnShotFinished(const FMoviePipelineShotCache& Shot) { OnShotFinishedImpl(Shot); }
+
+	/**
+	* This is called during the Shutdown process of the Pipeline. This is after finalization.
+	*/
+	void OnPipelineFinished() { OnPipelineFinishedImpl(); }
+
+	/**
+	* This is called right before a shot starts producing frames, after warm up. See OnShotInitialized
+	*/
+	void OnFrameProductionStart() { OnFrameProductionStartImpl(); }
+
+	/**
+	* Called at the end of the frame, before rendering has happened for that frame.
+	*/
+	void OnPostTick() { OnPostTickImpl(); }
+
+protected:
+	virtual void OnShotInitializedImpl(const TOptional<FMoviePipelineShotCache> PrevShot, const FMoviePipelineShotCache& NewShot) {}
+	virtual void OnShotFinishedImpl(const FMoviePipelineShotCache& Shot) {}
+
+	virtual void ProcessFrameImpl(TArray<MoviePipeline::FOutputFrameData> FrameData, FMoviePipelineFrameOutputState CachedOutputState, FDirectoryPath OutputDirectory) {};
+	virtual void FinalizeImpl() {}
+	virtual void BeginFinalizeImpl() {}
+	virtual bool HasFinishedProcessingImpl() { return true; }
+	virtual void OnPipelineFinishedImpl() {}
+	virtual void OnFrameProductionStartImpl() {}
+	virtual void OnPostTickImpl() {}
+};
\ No newline at end of file
diff --git a/Engine/Plugins/MovieScene/MovieRenderPipeline/Source/MovieRenderPipelineCore/Public/MoviePipelineRenderPass.h b/Engine/Plugins/MovieScene/MovieRenderPipeline/Source/MovieRenderPipelineCore/Public/MoviePipelineRenderPass.h
new file mode 100644
--- /dev/null
+++ b/Engine/Plugins/MovieScene/MovieRenderPipeline/Source/MovieRenderPipelineCore/Public/MoviePipelineRenderPass.h
@@ -0,0 +1,56 @@
+// Copyright 1998-2019 Epic Games, Inc. All Rights Reserved.
+#pragma once
+
+#include "UObject/Object.h"
+#include "MovieRenderPipelineDataTypes.h"
+#include "MoviePipelineSetting.h"
+#include "ImageWriteStream.h"
+#include "MoviePipelineRenderPass.generated.h"
+
+UCLASS(Blueprintable, Abstract)
+class MOVIERENDERPIPELINECORE_API UMoviePipelineRenderPass : public UMoviePipelineSetting
+{
+	GENERATED_BODY()
+public:
+	void Setup(const FMoviePipelineRenderPassInitSettings& InInitSettings, TSharedRef<FImagePixelPipe, ESPMode::ThreadSafe> InOutputPipe)
+	{
+		InitSettings = InInitSettings;
+		SetupImpl(InInitSettings, InOutputPipe);
+	}
+
+	void CaptureFrame(const FMoviePipelineRenderPassMetrics& OutputFrameMetrics)
+	{
+		CaptureFrameImpl(OutputFrameMetrics);
+	}
+
+
+	void GetFrameData(MoviePipeline::FOutputFrameData& OutFrameData)
+	{
+		GetFrameDataImpl(OutFrameData);
+	}
+
+	void Teardown()
+	{
+		TeardownImpl();
+	}
+
+	void GatherOutputPasses(TArray<FMoviePipelinePassIdentifier>& ExpectedRenderPasses)
+	{
+		GatherOutputPassesImpl(ExpectedRenderPasses);
+	}
+
+protected:
+	virtual void SetupImpl(const FMoviePipelineRenderPassInitSettings& InInitSettings, TSharedRef<FImagePixelPipe, ESPMode::ThreadSafe> InOutputPipe) {  }
+
+	virtual void CaptureFrameImpl(const FMoviePipelineRenderPassMetrics& OutputFrameMetrics) { }
+
+	virtual void GetFrameDataImpl(MoviePipeline::FOutputFrameData& OutFrameData) { }
+
+	virtual void TeardownImpl() {}
+
+	virtual void GatherOutputPassesImpl(TArray<FMoviePipelinePassIdentifier>& ExpectedRenderPasses) {}
+		
+protected:
+	UPROPERTY()
+	FMoviePipelineRenderPassInitSettings InitSettings;
+};
\ No newline at end of file
diff --git a/Engine/Plugins/MovieScene/MovieRenderPipeline/Source/MovieRenderPipelineCore/Public/MoviePipelineSetting.h b/Engine/Plugins/MovieScene/MovieRenderPipeline/Source/MovieRenderPipelineCore/Public/MoviePipelineSetting.h
new file mode 100644
--- /dev/null
+++ b/Engine/Plugins/MovieScene/MovieRenderPipeline/Source/MovieRenderPipelineCore/Public/MoviePipelineSetting.h
@@ -0,0 +1,58 @@
+// Copyright 1998-2019 Epic Games, Inc. All Rights Reserved.
+#pragma once
+
+#include "UObject/Object.h"
+#include "MoviePipelineSetting.generated.h"
+
+class UMoviePipeline;
+struct FSlateBrush;
+
+/**
+* A base class for all Movie Render Pipeline settings.
+*/
+UCLASS(Blueprintable, Abstract)
+class MOVIERENDERPIPELINECORE_API UMoviePipelineSetting : public UObject
+{
+	GENERATED_BODY()
+		
+public:
+	/**
+	* Called when the Pipeline is initialized for the first time before the first frame is ticked.
+	*/
+	void OnInitializedForPipeline(UMoviePipeline* InPipeline);
+
+	// UObject Interface
+	virtual UWorld* GetWorld() const override;
+	// ~UObject Interface
+
+protected:
+	UMoviePipeline* GetPipeline() const;
+
+
+	virtual void OnInitializedForPipelineImpl(UMoviePipeline* InPipeline) {}
+
+public:
+	virtual FText GetDisplayText() const { return this->GetClass()->GetDisplayNameText(); }
+
+	/** Should the Pipeline automatically create an instance of this under the hood so calling code can rely on it existing? */
+	virtual bool IsRequired() const { return false; }
+	
+	/** Can only one of these settings objects be active in a valid pipeline? */
+	virtual bool IsSolo() const { return true; }
+	
+	/** Is this setting valid? Return false and add a reason it's not valid to the array if not. */
+	virtual bool ValidatePipeline(TArray<FText>& OutValidationErrors) const { return true; }
+	
+	/** What icon should this setting use when displayed in the tree list. */
+	const FSlateBrush* GetDisplayIcon() { return nullptr; }
+	
+	/** What tooltip should be displayed for this setting when hovered in the tree list? */
+	FText GetDescriptionText() { return FText(); }
+	
+	/** Is this setting currently enabled? Disabled settings are like they never existed. */
+	bool bEnabled;
+
+private:
+	UPROPERTY(Transient)
+	TWeakObjectPtr<UMoviePipeline> CachedPipeline;
+};
\ No newline at end of file
diff --git a/Engine/Plugins/MovieScene/MovieRenderPipeline/Source/MovieRenderPipelineCore/Public/MoviePipelineShotConfig.h b/Engine/Plugins/MovieScene/MovieRenderPipeline/Source/MovieRenderPipelineCore/Public/MoviePipelineShotConfig.h
new file mode 100644
--- /dev/null
+++ b/Engine/Plugins/MovieScene/MovieRenderPipeline/Source/MovieRenderPipelineCore/Public/MoviePipelineShotConfig.h
@@ -0,0 +1,130 @@
+// Copyright 1998-2019 Epic Games, Inc. All Rights Reserved.
+#pragma once
+
+#include "UObject/Object.h"
+#include "UObject/SoftObjectPath.h"
+#include "MoviePipelineSetting.h"
+#include "Engine/EngineTypes.h"
+
+#include "MoviePipelineShotConfig.generated.h"
+
+
+// Forward Declares
+class ULevelSequence;
+class UMoviePipelineSetting;
+class UMoviePipelineRenderPass;
+class UMoviePipelineOutput;
+
+UCLASS(Blueprintable)
+class MOVIERENDERPIPELINECORE_API UMoviePipelineShotConfig : public UObject
+{
+	GENERATED_BODY()
+	
+public:
+	UMoviePipelineShotConfig()
+	{
+		SettingsSerialNumber = -1;
+	}
+
+	bool ValidateConfig(TArray<FText>& OutValidationErrors) const;
+	void RemoveSetting(UMoviePipelineSetting* InSetting);
+	int32 GetSettingsSerialNumber() const { return SettingsSerialNumber; }
+	TArray<UMoviePipelineSetting*> GetSettings() const { return Settings; }
+
+public:
+	/**
+	* Find a setting of a particular type for this config.
+	* @param InClass - Class that you wish to find the setting object for.
+	* @return An instance of this class if it already exists as a setting on this config, otherwise null.
+	*/
+	UFUNCTION(BlueprintPure, Category = "Movie Render Pipeline")
+	UMoviePipelineSetting* FindSettingByClass(TSubclassOf<UMoviePipelineSetting> InClass) const
+	{
+		UMoviePipelineSetting* const* Found = Settings.FindByPredicate([InClass](UMoviePipelineSetting* In) { return In && In->GetClass() == InClass; });
+		return Found ? CastChecked<UMoviePipelineSetting>(*Found) : nullptr;
+	}
+	
+	/** Finds a setting of a particular type for this pipeline config, adding it if it doesn't already exist.
+	* @param InClass - Class you wish to find or create the setting object for.
+	* @return An instance of this class as a setting on this config. 
+	*/
+	UFUNCTION(BlueprintCallable, Category = "Movie Render Pipeline")
+	UMoviePipelineSetting* FindOrAddSettingByClass(TSubclassOf<UMoviePipelineSetting> InClass)
+	{
+		UMoviePipelineSetting* Found = FindSettingByClass(InClass);
+		if (!Found)
+		{
+			Modify();
+
+			Found = NewObject<UMoviePipelineSetting>(this, InClass);
+			Settings.Add(Found);
+			++SettingsSerialNumber;
+		}
+	
+		return Found;
+	}
+	
+	/**
+	* Find a setting of a particular type for this config
+	*/ 
+	template<typename SettingType>
+	SettingType* FindSetting() const
+	{
+		UClass* PredicateClass = SettingType::StaticClass();
+		UMoviePipelineSetting* const* Found = Settings.FindByPredicate([PredicateClass](UMoviePipelineSetting* In) { return In && In->GetClass() == PredicateClass; });
+		return Found ? CastChecked<SettingType>(*Found) : nullptr;
+	}
+	
+	/**
+	 * Find a setting of a particular type for this config instance, adding one if it was not found.
+	 */
+	template<typename SettingType>
+	SettingType* FindOrAddSetting()
+	{
+		SettingType* Found = FindSetting<SettingType>();
+		if (!Found)
+		{
+			Modify();
+
+			Found = NewObject<SettingType>(this);
+			Settings.Add(Found);
+			++SettingsSerialNumber;
+		}
+		return Found;
+	}
+public:
+	/** What Sequence do we want to render out. Can be overridden by command line launch arguments. */
+	UPROPERTY(EditAnywhere, BlueprintReadWrite, Category="Movie Pipeline", meta=(MetaClass="LevelSequence"))
+	FSoftObjectPath Sequence;
+	
+	/** What resolution should our output files be exported at? */
+	UPROPERTY(EditAnywhere, BlueprintReadWrite, Category = "Movie Pipeline")
+	FVector2D OutputResolution;
+	
+	/** What frame rate should the output files be exported at? This overrides the Display Rate of the target sequence. */
+	UPROPERTY(EditAnywhere, BlueprintReadWrite, Category = "Movie Pipeline")
+	FFrameRate OutputFrameRate;
+	
+	/** What directory should all of our output files be relative to. */
+	UPROPERTY(EditAnywhere, BlueprintReadWrite, Category = "Movie Pipeline")
+	FDirectoryPath OutputDirectory;
+	
+	/** If true, output containers should attempt to override any existing files with the same name. */
+	UPROPERTY(EditAnywhere, BlueprintReadWrite, Category = "Movie Pipeline")
+	bool bOverrideExistingOutput;
+	
+	/** Array of settings classes that affect various parts of the output pipeline. */
+	UPROPERTY(EditAnywhere, BlueprintReadWrite, Category = "Movie Pipeline")
+	TArray<UMoviePipelineSetting*> Settings;
+	
+	/** Array of Input Buffers. These specify how many render passes should happen each frame to get the data for that frame. */
+	UPROPERTY(EditAnywhere, BlueprintReadWrite, Category = "Movie Pipeline")
+	TArray<UMoviePipelineRenderPass*> InputBuffers;
+	
+	/** Array of Output Containers. Each output container is passed data for each Input Buffer every frame. */
+	UPROPERTY(EditAnywhere, BlueprintReadWrite, Category = "Movie Pipeline")
+	TArray<UMoviePipelineOutput*> OutputContainers;
+
+private:
+	int32 SettingsSerialNumber;
+};
\ No newline at end of file
