Initial merge of the new Sequencer Movie Render system. This provides some basic ways to output movies but is mostly scripting based without UI still.

diff --git a/Engine/Plugins/MovieScene/MovieRenderPipeline/Source/MovieRenderPipelineCore/Public/MovieRenderOverlappedImage.h b/Engine/Plugins/MovieScene/MovieRenderPipeline/Source/MovieRenderPipelineCore/Public/MovieRenderOverlappedImage.h
new file mode 100644
--- /dev/null
+++ b/Engine/Plugins/MovieScene/MovieRenderPipeline/Source/MovieRenderPipelineCore/Public/MovieRenderOverlappedImage.h
@@ -0,0 +1,167 @@
+// Copyright 1998-2019 Epic Games, Inc. All Rights Reserved.
+#pragma once
+
+
+#include "CoreMinimal.h"
+
+#include "ImagePixelData.h"
+
+/* Types
+ *****************************************************************************/
+
+/**
+ * Structure for a single overlapped image.
+ */
+struct MOVIERENDERPIPELINECORE_API FImageOverlappedPlane
+{
+public:
+	/**
+	 * Default constructor.
+	 */
+	FImageOverlappedPlane()
+		: SizeX(0)
+		, SizeY(0)
+	{
+	}
+
+	/**
+	 * Initialize the memory. Before using the memory, we also will need a call to ZeroPlane.
+	 *
+	 * @param InSizeX - Horizontal size of the tile.
+	 * @param InSizeY - Vertical size of the tile.
+	 */
+	void Init(int32 InSizeX, int32 InSizeY);
+
+	/**
+	 * Zeroes the accumulation. Assumes the data is already allocated.
+	 */
+	void ZeroPlane();
+
+	/**
+	 * Frees the memory and resets the sizes
+	 */
+	void Reset();
+
+	/**
+	 * Accumulate a single tile to this plane. The raw data will in general be smaller than the full plane.
+	 *
+	 * @param InRawData - Raw data to accumulate
+	 * @param InRawSizeX - Width of the tile. Must exactly match.
+	 * @param InRawSizeY - Height of the tile. Must exactly match.
+	 * @param InSampleOffsetX - Pixel offset of the tile (in X)
+	 * @param InSampleOffsetY - Pixel offset of the tile (in Y)
+	 */
+	void AccumulateSinglePlane(const TArray64<float>& InRawData, const TArray64<float>& InWeightData, int32 InSizeX, int32 InSizeY, int32 InOffsetX, int32 InOffsetY,
+								float SubpixelOffsetX, float SubpixelOffsetY);
+
+	/** Actual channel data*/
+	TArray64<float> ChannelData;
+
+	/** Width of the image. */
+	int32 SizeX;
+
+	/** Height of the image. */
+	int32 SizeY;
+};
+
+/**
+ * Contains all the image planes for the tiles.
+ */
+struct MOVIERENDERPIPELINECORE_API FImageOverlappedAccumulator
+{
+public:
+	/**
+	 * Default constructor.
+	 */
+	FImageOverlappedAccumulator()
+		: PlaneSizeX(0)
+		, PlaneSizeY(0)
+		, NumChannels(0)
+		, AccumulationGamma(1.0f)
+	{
+	}
+
+	/**
+	 * Allocates memory.
+	 *
+	 * @param InTileSizeX - Horizontal tile size.
+	 * @param InTileSizeY - Vertical tile size.
+	 * @param InNumTilesX - Num horizontal tiles.
+	 * @param InNumTilesY - Num vertical tiles.
+	 * @param InNumChannels - Num Channels.
+	 */
+	void InitMemory(int InPlaneSizeX, int InPlaneSizeY, int InNumChannels);
+
+	/**
+	 * Initializes memory.
+	 *
+	 * Resets the memory to 0s so that we can start a new frame.
+	 */
+	void ZeroPlanes();
+
+	/**
+	 * Resets the memory.
+	 */
+	void Reset();
+
+
+	static void GenerateTileWeight(TArray64<float>& Weights, int32 SizeX, int32 SizeY);
+
+
+	/**
+	 * Given a rendered tile, accumulate the data to the full size image.
+	 * 
+	 * @param InPixelData - Raw pixel data.
+	 * @param InTileX - Tile index in X.
+	 * @param InTileY - Tile index in Y.
+	 * @param InSubPixelOffset - SubPixel offset, should be in the range [0,1)
+	 */
+	void AccumulatePixelData(const FImagePixelData& InPixelData, int32 InTileOffsetX, int32 InTileOffsetY, FVector2D InSubpixelOffset);
+
+	/**
+	 * After accumulation is finished, fetch the final image as bytes. In theory we don't need this, because we could
+	 * just fetch as LinearColor and convert to bytes. But the largest size asked for is 45k x 22.5k, which is 1B pixels.
+	 * So fetching as LinearColor would create a 16GB intermediary image, so it's worth having an option to fetch
+	 * straight to FColors.
+	 * 
+	 * @param FImagePixelData - Finished pixel data.
+	 */
+	void FetchFinalPixelDataByte(TArray<FColor>& OutPixelData) const;
+
+	/**
+	 * After accumulation is finished, fetch the final image as linear colors
+	 * 
+	 * @param FImagePixelData - Finished pixel data.
+	 */
+	void FetchFinalPixelDataLinearColor(TArray<FLinearColor>& OutPixelData) const;
+
+	/**
+	 * Grab a single pixel from the full res tile and scale it by the appropriate Scale value.
+	 * 
+	 * @param Rgba - Found pixel value.
+	 * @param PlaneScale - The scales of each channel plane.
+	 * @param FullX - X position of the full res image.
+	 * @param FullY - Y position of the full res image.
+	 */
+	void FetchFullImageValue(float Rgba[4], int32 FullX, int32 FullY) const;
+
+public:
+	/** Width of each tile in pixels */
+	int32 PlaneSizeX;
+
+	/** Height of each tile in pixels */
+	int32 PlaneSizeY;
+
+	/** Number of channels in the tiles. Typical will be 3 (RGB). */
+	int32 NumChannels;
+
+	/** Gamma for accumulation. Typical values are 1.0 and 2.2. */
+	float AccumulationGamma;
+
+	TArray64<FImageOverlappedPlane> ChannelPlanes;
+	
+	FImageOverlappedPlane WeightPlane;
+
+};
+
+
diff --git a/Engine/Plugins/MovieScene/MovieRenderPipeline/Source/MovieRenderPipelineCore/Public/MovieRenderPipelineConfig.h b/Engine/Plugins/MovieScene/MovieRenderPipeline/Source/MovieRenderPipelineCore/Public/MovieRenderPipelineConfig.h
new file mode 100644
--- /dev/null
+++ b/Engine/Plugins/MovieScene/MovieRenderPipeline/Source/MovieRenderPipelineCore/Public/MovieRenderPipelineConfig.h
@@ -0,0 +1,115 @@
+// Copyright 1998-2019 Epic Games, Inc. All Rights Reserved.
+#pragma once
+
+#include "UObject/Object.h"
+#include "UObject/SoftObjectPath.h"
+#include "Misc/FrameRate.h"
+#include "Misc/FrameNumber.h"
+#include "Engine/EngineTypes.h"
+
+#include "MovieRenderPipelineConfig.generated.h"
+
+
+// Forward Declares
+class ULevelSequence;
+class UMoviePipelineSetting;
+class UMoviePipelineRenderPass;
+class UMoviePipelineOutput;
+class UMoviePipelineShotConfig;
+
+/**
+* This class describes the main configuration for a Movie Render Pipeline.
+* Only settings that apply to the entire output should be stored here,
+* anything that is changed on a per-shot basis should be stored inside of 
+* UMovieRenderShotConfig instead.
+*
+* THIS CLASS SHOULD BE IMMUTABLE ONCE PASSED TO THE PIPELINE FOR PROCESSING.
+* (Otherwise you will be modifying the instance that exists in the UI)
+*/
+UCLASS(Blueprintable)
+class MOVIERENDERPIPELINECORE_API UMovieRenderPipelineConfig : public UObject
+{
+	GENERATED_BODY()
+	
+public:
+	UMovieRenderPipelineConfig()
+		: OutputResolution(FIntPoint(1920, 1080))
+		, bUseCustomFrameRate(false)
+		, OutputFrameRate(FFrameRate(24, 1))
+		, OutputFrameStep(FFrameNumber(1))
+		, bOverrideExistingOutput(true)
+		, bUseCustomPlaybackRange(false)
+	{
+		
+	}
+
+public:
+	/** Gets the effective frame rate (respecting framerate overrides). Should not be called unless the Sequence is set. */
+	FFrameRate GetEffectiveFrameRate();
+
+	/** Returns a pointer to the config specified for the shot, otherwise the default for this pipeline. */
+	UMoviePipelineShotConfig* GetConfigForShot(const FString& ShotName) const;
+
+public:
+	/** What Sequence do we want to render out. Can be overridden by command line launch arguments. */
+	UPROPERTY(EditAnywhere, BlueprintReadWrite, Category="Movie Render Pipeline", meta=(MetaClass="LevelSequence"))
+	FSoftObjectPath Sequence;
+	
+	/** What resolution should our output files be exported at? */
+	UPROPERTY(EditAnywhere, BlueprintReadWrite, Category = "Movie Render Pipeline")
+	FIntPoint OutputResolution;
+	
+	/** Should we use the custom frame rate specified by OutputFrameRate? Otherwise defaults to Sequence frame rate. */
+	UPROPERTY(EditAnywhere, BlueprintReadWrite, Category = "Movie Render Pipeline")
+	bool bUseCustomFrameRate;
+	
+	/** What frame rate should the output files be exported at? This overrides the Display Rate of the target sequence. */
+	UPROPERTY(EditAnywhere, BlueprintReadWrite, meta=(EditCondition=bUseCustomFrameRate), Category = "Movie Render Pipeline")
+	FFrameRate OutputFrameRate;
+	
+	/** How many frames (in Display Rate) should we step at a time. Can be used to render every other frame for faster drafts. */
+	UPROPERTY(EditAnywhere, BlueprintReadWrite, Category = "Movie Render Pipeline")
+	FFrameNumber OutputFrameStep;
+	
+	/** What directory should all of our output files be relative to. */
+	UPROPERTY(EditAnywhere, BlueprintReadWrite, Category = "Movie Render Pipeline")
+	FDirectoryPath OutputDirectory;
+	
+	/** If true, output containers should attempt to override any existing files with the same name. */
+	UPROPERTY(EditAnywhere, BlueprintReadWrite, Category = "Movie Render Pipeline")
+	bool bOverrideExistingOutput;
+public:
+	/** If true, override the Playback Range start/end bounds with the bounds specified below.*/
+	UPROPERTY(EditAnywhere, BlueprintReadWrite, Category = "Movie Render Pipeline")
+	bool bUseCustomPlaybackRange;
+	
+	/** Used when overriding the playback range. In Display Rate frames. If bUseCustomPlaybackRange is false range will come from Sequence. */
+	UPROPERTY(EditAnywhere, BlueprintReadWrite, meta=(EditCondition=bUseCustomPlaybackRange), Category = "Movie Render Pipeline")
+	FFrameNumber CustomStartFrame;
+	
+	/** Used when overriding the playback range. In Display Rate frames. If bUseCustomPlaybackRange is false range will come from Sequence. */
+	UPROPERTY(EditAnywhere, BlueprintReadWrite, meta=(EditCondition=bUseCustomPlaybackRange), Category = "Movie Render Pipeline")
+	FFrameNumber CustomEndFrame;
+
+public:
+	
+	/** The default shot-setup to use for any shot that doesn't a specific implementation. This is required! */
+	UPROPERTY(EditAnywhere, BlueprintReadWrite, Category = "Movie Render Pipeline")
+	UMoviePipelineShotConfig* DefaultShotConfig;
+	
+	/** A mapping of Shot Name -> Shot Config to use for rendering specific shots with specific configs. */
+	UPROPERTY(EditAnywhere, BlueprintReadWrite, Category = "Movie Render Pipeline")
+	TMap<FString, UMoviePipelineShotConfig*> PerShotConfigMapping;
+	
+	/** Array of Output Containers. Each output container is passed data for each Input Buffer every frame. */
+	UPROPERTY(EditAnywhere, BlueprintReadWrite, Category = "Movie Render Pipeline")
+	TArray<UMoviePipelineOutput*> OutputContainers;
+
+private:
+	/** Attempts to load the target sequence. Will trip an ensure if the sequence is not found. */
+	void LoadTargetSequence();
+	
+	/** We want to store the loaded asset to avoid having to load it each time we need a setting from it. */
+	UPROPERTY(Transient)
+	ULevelSequence* LoadedSequence;
+};
\ No newline at end of file
diff --git a/Engine/Plugins/MovieScene/MovieRenderPipeline/Source/MovieRenderPipelineCore/Public/MovieRenderPipelineCoreModule.h b/Engine/Plugins/MovieScene/MovieRenderPipeline/Source/MovieRenderPipelineCore/Public/MovieRenderPipelineCoreModule.h
new file mode 100644
--- /dev/null
+++ b/Engine/Plugins/MovieScene/MovieRenderPipeline/Source/MovieRenderPipelineCore/Public/MovieRenderPipelineCoreModule.h
@@ -0,0 +1,46 @@
+// Copyright 1998-2019 Epic Games, Inc. All Rights Reserved.
+
+#pragma once
+
+#include "CoreMinimal.h"
+#include "Modules/ModuleInterface.h"
+
+// Forward Declare
+class UMoviePipelineExecutorBase;
+class UMovieRenderPipelineConfig;
+class ULevelSequence;
+
+namespace MoviePipelineErrorCodes
+{
+	/** Everything completed as expected or we (unfortunately) couldn't detect the error. */
+	constexpr uint8 Success = 0;
+	/** Fallback for any generic critical failure. This should be used for "Unreal concepts aren't working as expected" severity errors. */
+	constexpr uint8 Critical = 1;
+	/** The specified level sequence asset could not be found. Check the logs for details on what it looked for. */
+	constexpr uint8 NoAsset = 2;
+	/** The specified pipeline configuration asset could not be found. Check the logs for details on what it looked for. */
+	constexpr uint8 NoConfig = 3;
+
+}
+
+class FMovieRenderPipelineCoreModule : public IModuleInterface
+{
+public:
+	virtual void StartupModule() override;
+	virtual void ShutdownModule() override;
+
+private:
+	bool IsTryingToRenderMovieFromCommandLine(FString& OutSequenceAssetPath, FString& OutConfigAssetPath, FString& OutExecutorType, FString& OutPipelineType) const;
+	void InitializeCommandLineMovieRender();
+	void OnCommandLineMovieRenderCompleted(UMoviePipelineExecutorBase* InExecutor);
+	uint8 GetMovieRenderClasses(const FString& InAssetPath, const FString& InConfigAssetPath, const FString& InExecutorType, const FString& InPipelineType, 
+								ULevelSequence*& OutAsset, UMovieRenderPipelineConfig*& OutConfig, UClass*& OutExecutorClass, UClass*& OutPipelineClass) const;
+
+private:
+	FString MoviePipelineLocalExecutorClassType;
+	FString MoviePipelineClassType;
+	FString SequenceAssetValue;
+	FString SettingsAssetValue;
+};
+
+MOVIERENDERPIPELINECORE_API DECLARE_LOG_CATEGORY_EXTERN(LogMovieRenderPipeline, Log, All);
\ No newline at end of file
diff --git a/Engine/Plugins/MovieScene/MovieRenderPipeline/Source/MovieRenderPipelineCore/Public/MovieRenderPipelineDataTypes.cpp b/Engine/Plugins/MovieScene/MovieRenderPipeline/Source/MovieRenderPipelineCore/Public/MovieRenderPipelineDataTypes.cpp
new file mode 100644
--- /dev/null
+++ b/Engine/Plugins/MovieScene/MovieRenderPipeline/Source/MovieRenderPipelineCore/Public/MovieRenderPipelineDataTypes.cpp
@@ -0,0 +1,108 @@
+// Copyright 1998-2019 Epic Games, Inc. All Rights Reserved.
+
+#include "MovieRenderPipelineDataTypes.h"
+#include "MovieRenderPipelineCoreModule.h"
+#include "Sections/MovieSceneCinematicShotSection.h"
+#include "MoviePipelineAccumulationSetting.h"
+
+FString FMoviePipelineShotCache::GetDisplayName() const
+{
+	if (CinematicShotSection.Get())
+	{
+		return CinematicShotSection->GetShotDisplayName();
+	}
+
+	return TEXT("Unnamed");
+}
+
+FFrameNumber FMoviePipelineShotCutCache::GetOutputFrameCountEstimate() const
+{
+	// TotalRange is stored in Tick Resolution, so we convert 1 frame of Frame Rate to the number of ticks.
+	FFrameNumber OneFrameInTicks = FFrameRate::TransformTime(FFrameTime(FFrameNumber(1)), CachedFrameRate, CachedTickResolution).FloorToFrame();
+
+	// Find out how many ticks long our total output range is.
+	FFrameNumber TotalOutputRangeTicks = TotalOutputRange.Size<FFrameNumber>();
+	int32 NumFrames = FMath::CeilToInt(TotalOutputRangeTicks.Value / (double)OneFrameInTicks.Value);
+
+	return FFrameNumber(NumFrames);
+}
+
+FFrameNumber FMoviePipelineShotCutCache::GetTemporalFrameCountEstimate() const
+{
+	FFrameNumber OutputEstimate = GetOutputFrameCountEstimate();
+	FFrameNumber NumRenderedSamples = OutputEstimate * NumTemporalSamples;
+
+	return NumRenderedSamples;
+}
+
+FFrameNumber FMoviePipelineShotCutCache::GetUtilityFrameCountEstimate() const
+{
+	return FFrameNumber(NumWarmUpFrames + (bAccurateFirstFrameHistory ? 1 : 0));
+}
+
+FFrameNumber FMoviePipelineShotCutCache::GetSampleCountEstimate(const bool bIncludeWarmup, const bool bIncludeMotionBlur) const
+{
+	int32 TotalSamples = 0;
+	
+	// Warm Up Frames currently don't submit anything to the GPU.
+	if (bIncludeWarmup)
+	{
+		TotalSamples += 0;
+	}
+
+	// Motion blur only adds one frame
+	if (bIncludeMotionBlur)
+	{
+		TotalSamples += 1;
+	}
+
+	// Only the main frames are rendered with Temporal sub-sampling and Spatial sub-sampling.
+	TotalSamples += GetOutputFrameCountEstimate().Value * NumTemporalSamples * NumSpatialSamples;
+
+	return FFrameNumber(TotalSamples);
+}
+
+void FMoviePipelineShotCutCache::SetNextState(const EMovieRenderShotState InCurrentState)
+{
+	switch (InCurrentState)
+	{
+		// This may get called multiple times so we just do nothing until it's appropriate to move on from WarmingUp.
+	case EMovieRenderShotState::WarmingUp:
+		// Warming Up can jump directly to either Rendering or to MotionBlur depending on if fixes are applied.
+		if (NumWarmUpFramesRemaining == 0)
+		{
+			if (bAccurateFirstFrameHistory)
+			{
+				UE_LOG(LogMovieRenderPipeline, Verbose, TEXT("[%d] Shot WarmUp finished. Setting state to MotionBlur."), GFrameCounter);
+				State = EMovieRenderShotState::MotionBlur;
+			}
+			else
+			{
+				UE_LOG(LogMovieRenderPipeline, Verbose, TEXT("[%d] Shot WarmUp finished. Setting state to Rendering due to no MotionBlur pre-frames."), GFrameCounter);
+				State = EMovieRenderShotState::Rendering;
+			}
+		}
+		break;
+		// This should only be called once with the Uninitalized state.
+	case EMovieRenderShotState::Uninitialized:
+		// Uninitialized can either jump to WarmUp, MotionBlur, or straight to Rendering if no fixes are applied.
+		if (NumWarmUpFramesRemaining > 0)
+		{
+			UE_LOG(LogMovieRenderPipeline, Log, TEXT("[%d] Initialization set state to WarmingUp due to having %d warm up frames."), GFrameCounter, NumWarmUpFramesRemaining);
+			State = EMovieRenderShotState::WarmingUp;
+		}
+		// If they didn't want any warm-up frames we'll still check to see if they want to fix motion blur on the first frame.
+		else if (bAccurateFirstFrameHistory)
+		{
+			UE_LOG(LogMovieRenderPipeline, Log, TEXT("[%d] Initialization set state to MotionBlur due to having no warm up frames."), GFrameCounter);
+			State = EMovieRenderShotState::MotionBlur;
+		}
+		else
+		{
+			UE_LOG(LogMovieRenderPipeline, Verbose, TEXT("[%d] Initialization set state to Rendering due to no MotionBlur pre-frames."), GFrameCounter);
+			State = EMovieRenderShotState::Rendering;
+		}
+		break;
+	}
+
+}
\ No newline at end of file
diff --git a/Engine/Plugins/MovieScene/MovieRenderPipeline/Source/MovieRenderPipelineCore/Public/MovieRenderPipelineDataTypes.h b/Engine/Plugins/MovieScene/MovieRenderPipeline/Source/MovieRenderPipelineCore/Public/MovieRenderPipelineDataTypes.h
new file mode 100644
--- /dev/null
+++ b/Engine/Plugins/MovieScene/MovieRenderPipeline/Source/MovieRenderPipelineCore/Public/MovieRenderPipelineDataTypes.h
@@ -0,0 +1,583 @@
+// Copyright 1998-2019 Epic Games, Inc. All Rights Reserved.
+#pragma once
+
+#include "Misc/Timecode.h"
+#include "ImagePixelData.h"
+#include "MovieRenderPipelineDataTypes.generated.h"
+
+class UMovieSceneCinematicShotSection;
+
+template<typename ElementType> class TRange;
+template<class T, class TWeakObjectPtrBase> struct TWeakObjectPtr;
+
+class UWorld;
+class ULevelSequence;
+class UMovieSceneCinematicShotSection;
+class UMovieRenderPipelineConfig;
+class UMoviePipelineShotConfig;
+class UMovieSceneCameraCutSection;
+
+
+/**
+* What is the current overall state of the Pipeline? States are processed in order from first to
+* last and all states will be hit (though there is no guarantee the state will not be transitioned
+* away from on the same frame it entered it). Used to help track overall progress and validate
+* code flow.
+*/
+UENUM(BlueprintType)
+enum class EMovieRenderPipelineState : uint8
+{
+	/** The pipeline has not been initialized yet. Only valid operation is to call Initialize. */
+	Uninitialized = 0,
+	/** The pipeline has been initialized and is now controlling time and working on producing frames. */
+	ProducingFrames = 1,
+	/** All desired frames have been produced. Audio is already finalized. Outputs have a chance to finish long processing tasks. */
+	Finalize = 2,
+	/** All outputs have finished writing to disk or otherwise processing. Additional exports that may have needed information about the produced file can now be run. */
+	Export = 3,
+	/** The pipeline has been shut down. It is an error to shut it down again. */
+	Shutdown = 4,
+};
+
+/**
+* What is the current state of a shot? States are processed in order from first to last but not
+* all states are required, ie: WarmUp and MotionBlur can be disabled and the shot will never
+* pass through this state.
+*/
+UENUM(BlueprintType)
+enum class EMovieRenderShotState : uint8
+{
+	/** The shot has not been initialized yet.*/
+	Uninitialized = 0,
+
+	/** The shot is warming up. Engine ticks are passing but no frames are being produced. */
+	WarmingUp = 1,
+	/*
+	* The shot is doing additional pre-roll for motion blur. No frames are being produced,
+	* but the rendering pipeline is being run to seed histories.
+	*/
+	MotionBlur = 2,
+	/*
+	* The shot is working on producing frames and may be currently doing a sub-frame or
+	* a whole frame.
+	*/
+	Rendering = 3,
+	/*
+	* The shot has produced all frames it will produce. No more evaluation should be
+	* done for this shot once it reaches this state.
+	*/
+	Finished = 4
+};
+
+struct FMoviePipelinePassIdentifier
+{
+	FMoviePipelinePassIdentifier(const FName& InPassName)
+		: Name(InPassName)
+	{
+	}
+
+	FName Name;
+};
+
+namespace MoviePipeline
+{
+	/**
+	* Utility function for the Pipeline to calculate the time metrics for the current
+	* frame that can then be used by the custom time step to avoid custom time logic.
+	*/
+	struct FFrameTimeStepCache
+	{
+		FFrameTimeStepCache()
+			: DeltaTime(0.0)
+		{}
+
+		FFrameTimeStepCache(double InDeltaTime)
+			: DeltaTime(InDeltaTime)
+		{}
+
+		double DeltaTime;
+	};
+
+	struct FOutputFrameData
+	{
+		FIntPoint Resolution;
+
+		TArray<FColor> ColorBuffer;
+
+		FString PassName;
+	};
+
+	struct FFrameConstantMetrics
+	{
+		/** What is the tick resolution fo the sequence */
+		FFrameRate TickResolution;
+		/** What is the effective frame rate of the output */
+		FFrameRate FrameRate;
+		/** How many ticks per output frame. */
+		FFrameTime TicksPerOutputFrame;
+		/** How many ticks per sub-frame sample. */
+		FFrameTime TicksPerSample;
+		/** How many ticks while the shutter is closed */
+		FFrameTime TicksWhileShutterClosed;
+		/** How many ticks while the shutter is opened */
+		FFrameTime TicksWhileShutterOpen;
+
+		/** Fraction of the output frame that the accumulation frames should cover. */
+		double ShutterAnglePercentage;
+
+		/** Fraction of the output frame that the shutter is closed (that accumulation frames should NOT cover). */
+		double ShutterClosedFraction;
+
+		/**
+		* Our time samples are offset by this many ticks. For a given output frame N, we can bias the time N represents
+		* to be before, during or after the data stored on time N.
+		*/
+		FFrameTime ShutterOffsetTicks;
+
+		/** 
+		* How many ticks do we have to go to center our motion blur? This makes it so the blurred area matches the object
+		* position during the time it represents, and not before.
+		*/
+		FFrameTime MotionBlurCenteringOffsetTicks;
+
+		FFrameTime GetFinalEvalTime(const FFrameNumber InTime) const
+		{
+			// We just use a consistent offset from the given time to take motion blur centering and
+			// shutter timing offsets into account. Written here just to consolidate various usages.
+			return FFrameTime(InTime) + MotionBlurCenteringOffsetTicks + ShutterOffsetTicks;
+		}
+
+
+	};
+}
+
+UENUM(BlueprintType)
+enum class EMoviePipelineShutterTiming : uint8
+{
+	FrameOpen,
+	FrameCenter,
+	FrameClose
+};
+
+USTRUCT(BlueprintType)
+struct FMoviePipelineShotCutCache
+{
+	GENERATED_BODY()
+public:
+	FMoviePipelineShotCutCache()
+		: OriginalRange(TRange<FFrameNumber>::Empty())
+		, TotalOutputRange(TRange<FFrameNumber>::Empty())
+		, NumWarmUpFrames(0)
+		, bAccurateFirstFrameHistory(true)
+		, NumTemporalSamples(0)
+		, NumSpatialSamples(0)
+		, State(EMovieRenderShotState::Uninitialized)
+		, CurrentTick(FFrameNumber(0))
+		, NumFramesProduced(FFrameNumber(0))
+		, NumSamplesRendered(FFrameNumber(0))
+		, CurrentFrameSamples(0)
+		, bHasEvaluatedMotionBlurFrame(false)
+		, bHasNotifiedFrameProductionStarted(false)
+		, NumWarmUpFramesRemaining(0)
+	{}
+
+	bool IsInitialized() const { return State != EMovieRenderShotState::Uninitialized; }
+	void SetNextState(const EMovieRenderShotState InCurrentState);
+
+	/** 
+	* Returns the number of expected frames to be produced by Initial Range + Handle Frames given a Frame Rate.
+	* This will be inaccurate when PlayRate tracks are involved.
+	*/
+	FFrameNumber GetOutputFrameCountEstimate() const;
+
+	/**
+	* Returns the expected number of frames different frames that will be submitted for rendering. This is
+	* the number of output frames * temporal samples. This will be inaccurate when PlayRate tracks are involved.
+	*/
+	FFrameNumber GetTemporalFrameCountEstimate() const;
+
+	/** 
+	* Returns misc. utility frame counts (Warm Up + MotionBlur Fix) since these are outside the ranges.
+	*/
+	FFrameNumber GetUtilityFrameCountEstimate() const;
+
+	/**
+	* Returns the number of samples submitted to the GPU, optionally counting samples for the various parts.
+	* This will be inaccurate when PlayRate tracks are involved.
+	*/
+	FFrameNumber GetSampleCountEstimate(const bool bIncludeWarmup = true, const bool bIncludeMotionBlur = true) const;
+	
+public:
+	/** The original non-modified range for this shot that will be rendered. */
+	TRange<FFrameNumber> OriginalRange;
+
+	/** The range for this shot including handle frames that will be rendered. */
+	TRange<FFrameNumber> TotalOutputRange;
+
+	/** How many warm up frames should this shot wait for. */
+	UPROPERTY(BlueprintReadOnly, Category = "Movie Render Pipeline")
+	int32 NumWarmUpFrames;
+
+	/** Should we evaluate/render an extra frame at the start of this shot to show correct motion blur on the first frame? */
+	UPROPERTY(BlueprintReadOnly, Category = "Movie Render Pipeline")
+	bool bAccurateFirstFrameHistory;
+
+	UPROPERTY(BlueprintReadOnly, Transient, Category = "Movie Render Pipeline")
+	TWeakObjectPtr<UMovieSceneCameraCutSection> CameraCutSection;
+
+	/** How many temporal samples is each frame broken up into? */
+	UPROPERTY(BlueprintReadOnly, Category = "Movie Render Pipeline")
+	int32 NumTemporalSamples;
+
+	/** For each temporal sample, how many spatial samples do we take? */
+	UPROPERTY(BlueprintReadOnly, Category = "Movie Render Pipeline")
+	int32 NumSpatialSamples;
+
+	/** Display name for UI Purposes. */
+	UPROPERTY(BlueprintReadOnly, Category = "Movie Render Pipeline")
+	FString DisplayName;
+
+	/** Cached Frame Rate these are being rendered at. Simplifies some APIs. */
+	FFrameRate CachedFrameRate;
+
+	/** Cached Tick Resolution our numbers are in. Simplifies some APIs. */
+	FFrameRate CachedTickResolution;
+public:
+	/** The current state of processing this Shot is in. Not all states will be passed through. */
+	UPROPERTY(BlueprintReadOnly, Category = "Movie Render Pipeline")
+	EMovieRenderShotState State;
+
+	UPROPERTY(BlueprintReadOnly, Category = "Movie Render Pipeline")
+	FFrameNumber CurrentTick;
+
+	/** Metrics - How many frames have been rendered for this shot so far? Can be more than the estimated total due to PlayRate. */
+	UPROPERTY(BlueprintReadOnly, Category = "Movie Render Pipeline")
+	FFrameNumber NumFramesProduced;
+
+	/** Metrics - How many samples have been rendered for this shot so far? Can be more than the estimated total due to PlayRate. */
+	UPROPERTY(BlueprintReadOnly, Category = "Movie Render Pipeline")
+	FFrameNumber NumSamplesRendered;
+
+	/** 
+	* Metrics - How many samples of the total have been generated for this frame. Only as granular as SpatialSamples since
+	* we can't refresh the screen between individual spatial samples.
+	*/
+	UPROPERTY(BlueprintReadOnly, Category = "Movie Render Pipeline")
+	int32 CurrentFrameSamples;
+
+	/** Have we evaluated the motion blur frame? Only used if bEvaluateMotionBlurOnFirstFrame is set */
+	UPROPERTY(BlueprintReadOnly, Category = "Movie Render Pipeline")
+	bool bHasEvaluatedMotionBlurFrame;
+
+	/** Have we notified the output containers that we're about to begin producing frames? */
+	UPROPERTY(BlueprintReadOnly, Category = "Movie Render Pipeline")
+	bool bHasNotifiedFrameProductionStarted;
+
+	/** How many warm up frames are left to process for this shot. May always be zero. */
+	UPROPERTY(BlueprintReadOnly, Category = "Movie Render Pipeline")
+	int32 NumWarmUpFramesRemaining;
+
+	bool operator == (const FMoviePipelineShotCutCache& InRHS) const
+	{
+		return
+			OriginalRange == InRHS.OriginalRange &&
+			TotalOutputRange == InRHS.TotalOutputRange &&
+			State == InRHS.State &&
+			NumWarmUpFramesRemaining == InRHS.NumWarmUpFramesRemaining &&
+			bAccurateFirstFrameHistory == InRHS.bAccurateFirstFrameHistory &&
+			bHasEvaluatedMotionBlurFrame == InRHS.bHasEvaluatedMotionBlurFrame &&
+			bHasNotifiedFrameProductionStarted == InRHS.bHasNotifiedFrameProductionStarted &&
+			CurrentTick == InRHS.CurrentTick &&
+			CameraCutSection == InRHS.CameraCutSection;
+	}
+
+	bool operator != (const FMoviePipelineShotCutCache& InRHS) const
+	{
+		return !(*this == InRHS);
+	}
+};
+
+/**
+* Pre-calculated information about a shot we are going to produce. This lets us build
+* the expected output at the start of the process and just read from it later. Having
+* all information in advanced aids in debugging and visualization of progress.
+*/
+USTRUCT(BlueprintType)
+struct FMoviePipelineShotCache
+{
+	GENERATED_BODY()
+public:
+	FMoviePipelineShotCache()
+		: NumHandleFrames(0)
+		, OriginalRange(TRange<FFrameNumber>::Empty())
+		, TotalOutputRange(TRange<FFrameNumber>::Empty())
+		, HandleFrameRangeStart(TRange<FFrameNumber>::Empty())
+		, HandleFrameRangeEnd(TRange<FFrameNumber>::Empty())
+		, ShotConfig(nullptr)
+		, CinematicShotSection(nullptr)
+		, CurrentCameraCutIndex(0)
+	{}
+
+	FString GetDisplayName() const;
+
+public:
+	/** How many handle frames (in Display Rate of the Master Sequence) */
+	UPROPERTY(BlueprintReadOnly, Category = "Movie Render Pipeline")
+	int32 NumHandleFrames;
+
+	/** The original non-modified (overall) range for this shot that will be rendered. */
+	TRange<FFrameNumber> OriginalRange;
+
+	/** The range for this shot including handle frames that will be rendered. */
+	TRange<FFrameNumber> TotalOutputRange;
+
+	TRange<FFrameNumber> HandleFrameRangeStart;
+	TRange<FFrameNumber> HandleFrameRangeEnd;
+	
+	UPROPERTY(Transient, BlueprintReadOnly, Category = "Movie Render Pipeline")
+	UMoviePipelineShotConfig* ShotConfig;
+
+	UPROPERTY(BlueprintReadOnly, Transient, Category = "Movie Render Pipeline")
+	TWeakObjectPtr<UMovieSceneCinematicShotSection> CinematicShotSection;
+
+	UPROPERTY(BlueprintReadOnly, Category = "Movie Render Pipeline")
+	TArray<FMoviePipelineShotCutCache> CameraCuts;
+
+	UPROPERTY(BlueprintReadOnly, Category = "Movie Render Pipeline")
+	int32 CurrentCameraCutIndex;
+
+	FMoviePipelineShotCutCache& GetCurrentCameraCut()
+	{
+		check(CurrentCameraCutIndex >= 0 && CurrentCameraCutIndex < CameraCuts.Num());
+		return CameraCuts[CurrentCameraCutIndex];
+	}
+
+	bool SetNextShotActive()
+	{
+		if (CurrentCameraCutIndex + 1 < CameraCuts.Num() - 1)
+		{
+			CurrentCameraCutIndex++;
+			return false;
+		}
+
+		// We're on the last camera cut, we can't make another one active.
+		return true;
+	}
+
+	bool operator == (const FMoviePipelineShotCache& InRHS) const
+	{ 
+		return
+			NumHandleFrames == InRHS.NumHandleFrames &&
+			OriginalRange == InRHS.OriginalRange &&
+			TotalOutputRange == InRHS.TotalOutputRange &&
+			HandleFrameRangeStart == InRHS.HandleFrameRangeStart &&
+			HandleFrameRangeEnd == InRHS.HandleFrameRangeEnd &&
+			ShotConfig == InRHS.ShotConfig &&
+			CinematicShotSection == InRHS.CinematicShotSection &&
+			CameraCuts == InRHS.CameraCuts;
+	}
+
+	bool operator != (const FMoviePipelineShotCache& InRHS) const
+	{
+		return !(*this == InRHS);
+	}
+};
+
+/**
+* The Tick/Render loops are decoupled from the actual desired output. In some cases
+* we may render but not desire an output (such as filling temporal histories) or
+* we may be accumulating the results into a target which does not produce an output.
+* Finally, we may be running the Tick/Render loop but not wanting to do anything!
+* This is the case when we are using the frame-step debugging, the engine is still
+* live and running, but we don't want to advance the output process forwards at all.
+*/
+USTRUCT(BlueprintType)
+struct FMoviePipelineFrameOutputState
+{
+	GENERATED_BODY()
+public:
+	FMoviePipelineFrameOutputState()
+		: OutputFrameNumber(-1)
+		, TemporalSampleIndex(-1)
+		, FrameDeltaTime(0)
+		, WorldSeconds(0)
+		, bWasAffectedByTimeDilation(false)
+	{
+		// There is other code which relies on this starting on -1.
+		check(OutputFrameNumber == -1);
+
+		// Likewise, the code assumes you start on sample index 0.
+		check(TemporalSampleIndex == -1);
+	}
+
+	/**
+	* The expected output frame count that the render is working towards creating.
+	* This number accurately tracks the number of frames we have produced even if
+	* the file written to disk uses a different number (due to relative frame numbers
+	* or offset frames being added.
+	*/
+	UPROPERTY(BlueprintReadOnly, Transient, Category = "Movie Render Pipeline")
+	int32 OutputFrameNumber;
+
+	/**
+	* This is an approximation of the number of samples submitted to the GPU to be rendered.
+	* This only an approximation and is used to create an estimated time remaining approximation.
+	*/
+	UPROPERTY(BlueprintReadOnly, Transient, Category = "Movie Render Pipeline")
+	int32 TotalSamplesRendered;
+
+	/** Which sub-frame are we on when using Accumulation Frame rendering. */
+	UPROPERTY(BlueprintReadOnly, Transient, Category = "Movie Render Pipeline")
+	int32 TemporalSampleIndex;
+
+	/** A 0-1 fraction of how much motion blur should be applied. This is tied to the shutter angle but more complicated when using sub-sampling */
+	UPROPERTY(BlueprintReadOnly, Transient, Category = "Movie Render Pipeline")
+	float MotionBlurFraction;
+
+	/** The delta time in seconds used for this frame. */
+	double FrameDeltaTime;
+
+	/** The total elapsed time since the pipeline started.*/
+	double WorldSeconds;
+	
+	/** The closest frame number (in Display Rate) on the Sequence. May be duplicates in the case of different output framerate or Play Rate tracks. */
+	FFrameNumber SourceFrameNumber;
+
+	/** The closest time code version of the SourceFrameNumber on the Sequence. May be a duplicate in the case of different output framerate or Play Rate tracks. */
+	FTimecode SourceTimeCode;
+
+	/** 
+	* The closest frame number (in Display Rate) on the Sequence adjusted for the effective output rate. These numbers will not line up with the frame
+	* in the source Sequence if the output frame rate differs from the Sequence display rate. May be a duplicate in the event of Play Rate tracks.
+	*/
+	FFrameNumber EffectiveFrameNumber;
+
+	/** The closest time code version of the EffectiveFrameNumber. May be a duplicate in the event of Play Rate tracks. */
+	FTimecode EffectiveTimeCode;
+
+	/**
+	* If true, there was a non-1.0 Time Dilation in effect when this frame was produced. This indicates that there 
+	* may be duplicate frame Source/Effective frame numbers as they find the closest ideal time to the current.
+	*/
+	bool bWasAffectedByTimeDilation;
+
+	bool operator == (const FMoviePipelineFrameOutputState& InRHS) const
+	{
+		return
+			OutputFrameNumber == InRHS.OutputFrameNumber &&
+			TemporalSampleIndex == InRHS.TemporalSampleIndex &&
+			FMath::IsNearlyEqual(MotionBlurFraction,InRHS.MotionBlurFraction) &&
+			FMath::IsNearlyEqual(FrameDeltaTime, InRHS.FrameDeltaTime) &&
+			FMath::IsNearlyEqual(WorldSeconds, InRHS.WorldSeconds) &&
+			SourceFrameNumber == InRHS.SourceFrameNumber &&
+			SourceTimeCode == InRHS.SourceTimeCode &&
+			EffectiveFrameNumber == InRHS.EffectiveFrameNumber &&
+			EffectiveTimeCode == InRHS.EffectiveTimeCode &&
+			bWasAffectedByTimeDilation == InRHS.bWasAffectedByTimeDilation;
+	}
+
+	bool operator != (const FMoviePipelineFrameOutputState& InRHS) const
+	{
+		return !(*this == InRHS);
+	}
+
+	friend uint32 GetTypeHash(FMoviePipelineFrameOutputState OutputState)
+	{
+		return GetTypeHash(OutputState.OutputFrameNumber);
+	}
+};
+
+USTRUCT(BlueprintType)
+struct FMoviePipelineRenderPassInitSettings
+{
+	GENERATED_BODY()
+public:
+	FMoviePipelineRenderPassInitSettings()
+		: ShotConfig(nullptr)
+		, TileResolution(0, 0)
+		, TileCount(0)
+	{
+	}
+
+public:
+	/** The config associated with the shot currently being rendered. */
+	UPROPERTY()
+	UMoviePipelineShotConfig* ShotConfig;
+
+	// These are derived out of other settings and cached directly to needing to re-derive values.
+	
+	/** The back buffer should match this setting and not the final output resolution.*/
+	FIntPoint TileResolution;
+	
+	/** How many tiles (in each direction) are we rendering with. */
+	int32 TileCount;
+};
+
+USTRUCT(BlueprintType)
+struct FMoviePipelineRenderPassMetrics
+{
+	GENERATED_BODY()
+
+public:
+	FVector2D SpatialShift;
+	int32 TileIndex;
+	int32 TileIndexX;
+	int32 TileIndexY;
+	int32 NumTilesX;
+	int32 NumTilesY;
+	FMoviePipelineFrameOutputState OutputState;
+	bool bIsHistoryOnlyFrame;
+	int32 SpatialJitterIndex;
+	int32 NumSpatialJitters;
+
+	int32 TemporalJitterIndex;
+	int32 NumTemporalJitters;
+
+	float JitterOffsetX;
+	float JitterOffsetY;
+
+	float AccumulationGamma;
+
+	bool bIsUsingOverlappedTiles;
+	int32 OverlappedOffsetX;
+	int32 OverlappedOffsetY;
+	int32 OverlappedSizeX;
+	int32 OverlappedSizeY;
+	int32 OverlappedPadX;
+	int32 OverlappedPadY;
+	FVector2D OverlappedSubpixelShift;
+};
+
+struct FImagePixelDataPayload : IImagePixelDataPayload
+{
+	FMoviePipelineFrameOutputState OutputState;
+	FVector2D SpatialShift;
+	int32 TileIndexX;
+	int32 TileIndexY;
+	int32 TileSizeX;
+	int32 TileSizeY;
+	int32 NumTilesX;
+	int32 NumTilesY;
+	int32 SpatialJitterIndex;
+	int32 NumSpatialJitters;
+
+	int32 TemporalJitterIndex;
+	int32 NumTemporalJitters;
+
+	float JitterOffsetX;
+	float JitterOffsetY;
+
+	float AccumulationGamma;
+
+	bool bIsUsingOverlappedTiles;
+	int32 OverlappedOffsetX;
+	int32 OverlappedOffsetY;
+	int32 OverlappedSizeX;
+	int32 OverlappedSizeY;
+	int32 OverlappedPadX;
+	int32 OverlappedPadY;
+	FVector2D OverlappedSubpixelShift;
+
+
+	FString PassName;
+};
\ No newline at end of file
diff --git a/Engine/Plugins/MovieScene/MovieRenderPipeline/Source/MovieRenderPipelineEditor/Private/MoviePipelineEditorBlueprintLibrary.cpp b/Engine/Plugins/MovieScene/MovieRenderPipeline/Source/MovieRenderPipelineEditor/Private/MoviePipelineEditorBlueprintLibrary.cpp
new file mode 100644
--- /dev/null
+++ b/Engine/Plugins/MovieScene/MovieRenderPipeline/Source/MovieRenderPipelineEditor/Private/MoviePipelineEditorBlueprintLibrary.cpp
@@ -0,0 +1,50 @@
+// Copyright 1998-2019 Epic Games, Inc. All Rights Reserved.
+
+#include "MoviePipelineEditorBlueprintLibrary.h"
+#include "UObject/Package.h"
+#include "Misc/PackageName.h"
+#include "MovieRenderPipelineConfig.h"
+#include "AssetRegistryModule.h"
+#include "FileHelpers.h"
+
+#define LOCTEXT_NAMESPACE "MoviePipelineEditorBlueprintLibrary"
+
+bool UMoviePipelineEditorBlueprintLibrary::ExportConfigToAsset(const UMovieRenderPipelineConfig* InConfig, const FString& InPackagePath, const FString& InFileName, const bool bInSaveAsset, UMovieRenderPipelineConfig*& OutAsset, FText& OutErrorReason)
+{
+	if(!InConfig)
+	{
+		OutErrorReason = LOCTEXT("CantExportNullConfigToPackage", "Can't export a null configuration to a package.");
+		return false;
+	}
+	
+	
+	if (!FPackageName::IsValidLongPackageName(InPackagePath, false, &OutErrorReason))
+	{
+		return false;
+	}
+
+
+	UPackage* NewPackage = CreatePackage(nullptr, *InPackagePath);
+	NewPackage->AddToRoot();
+	
+	// Duplicate the provided config into this package.
+	UMovieRenderPipelineConfig* NewConfig = Cast<UMovieRenderPipelineConfig>(StaticDuplicateObject(InConfig, NewPackage, FName(*InFileName), RF_NoFlags));
+	NewConfig->SetFlags(RF_Public | RF_Transactional | RF_Standalone);
+	NewConfig->MarkPackageDirty();
+
+	// Mark it so it shows up in the Content Browser immediately
+	FAssetRegistryModule::AssetCreated(NewConfig);
+
+	// If they want to save, ask them to save (and add to version control)
+	if (bInSaveAsset)
+	{
+		TArray<UPackage*> Packages;
+		Packages.Add(NewConfig->GetOutermost());
+
+		return UEditorLoadingAndSavingUtils::SavePackages(Packages, true);
+	}
+
+	return true;
+}
+
+#undef LOCTEXT_NAMESPACE // "MoviePipelineEditorBlueprintLibrary"
\ No newline at end of file
diff --git a/Engine/Plugins/MovieScene/MovieRenderPipeline/Source/MovieRenderPipelineEditor/Private/MoviePipelinePIEExecutor.cpp b/Engine/Plugins/MovieScene/MovieRenderPipeline/Source/MovieRenderPipelineEditor/Private/MoviePipelinePIEExecutor.cpp
new file mode 100644
--- /dev/null
+++ b/Engine/Plugins/MovieScene/MovieRenderPipeline/Source/MovieRenderPipelineEditor/Private/MoviePipelinePIEExecutor.cpp
@@ -0,0 +1,132 @@
+// Copyright 1998-2019 Epic Games, Inc. All Rights Reserved.
+#include "MoviePipelinePIEExecutor.h"
+#include "Settings/LevelEditorPlaySettings.h"
+#include "MovieRenderPipelineCoreModule.h"
+#include "MovieRenderPipelineConfig.h"
+#include "MoviePipelineShotConfig.h"
+#include "MoviePipelineSeparateProcessSetting.h"
+#include "Framework/Application/SlateApplication.h"
+#include "Editor/EditorEngine.h"
+#include "Editor.h"
+#include "MoviePipelineGameOverrideSetting.h"
+#include "MovieRenderPipelineDataTypes.h"
+#include "MoviePipeline.h"
+#include "MovieRenderPipelineSettings.h"
+
+#define LOCTEXT_NAMESPACE "MoviePipelinePIEExecutor"
+
+FText UMoviePipelinePIEExecutor::GetWindowTitle(const int32 InConfigIndex, const int32 InNumConfigs) const
+{
+	FNumberFormattingOptions PercentFormatOptions;
+	PercentFormatOptions.MinimumIntegralDigits = 2;
+	PercentFormatOptions.MaximumIntegralDigits = 3;
+
+	FText TitleFormatString = LOCTEXT("MoviePreviewWindowTitleFormat", "Movie Pipeline Render (Preview) [{CurrentCount}/{TotalCount} Total] {PercentComplete}% Completed.");
+	FText WindowTitle = FText::FormatNamed(TitleFormatString, TEXT("CurrentCount"), FText::AsNumber(InConfigIndex + 1), TEXT("TotalCount"), FText::AsNumber(InNumConfigs), TEXT("PercentComplete"), FText::AsNumber(12.f, &PercentFormatOptions));
+	return WindowTitle;
+}
+
+void UMoviePipelinePIEExecutor::Start(UMovieRenderPipelineConfig* InConfig, const int32 InConfigIndex, const int32 InNumConfigs)
+{
+	Super::Start(InConfig, InConfigIndex, InNumConfigs);
+
+	ActiveConfig = InConfig;
+
+	// Create a Slate window to hold our preview.
+	TSharedRef<SWindow> CustomWindow = SNew(SWindow)
+		.Title_UObject(this, &UMoviePipelinePIEExecutor::GetWindowTitle, InConfigIndex, InNumConfigs)
+		.ClientSize(FVector2D(1280, 720))
+		.AutoCenter(EAutoCenter::PrimaryWorkArea)
+		.UseOSWindowBorder(true)
+		.FocusWhenFirstShown(false)
+		.ActivationPolicy(EWindowActivationPolicy::Never)
+		.HasCloseButton(true)
+		.SupportsMaximize(false)
+		.SupportsMinimize(true)
+		.SizingRule(ESizingRule::UserSized);
+
+	FSlateApplication::Get().AddWindow(CustomWindow);
+
+	// Initialize our own copy of the Editor Play settings which we will adjust defaults on.
+	ULevelEditorPlaySettings* PlayInEditorSettings = NewObject<ULevelEditorPlaySettings>();
+	PlayInEditorSettings->SetPlayNetMode(EPlayNetMode::PIE_Standalone);
+	PlayInEditorSettings->SetPlayNumberOfClients(1);
+	PlayInEditorSettings->bLaunchSeparateServer = false;
+
+	FRequestPlaySessionParams Params;
+	Params.EditorPlaySettings = PlayInEditorSettings;
+	Params.CustomPIEWindow = CustomWindow;
+
+	UMoviePipelineGameOverrideSetting* GameOverrides = InConfig->DefaultShotConfig->FindSetting<UMoviePipelineGameOverrideSetting>();
+	if (GameOverrides)
+	{
+		Params.GameModeOverride = GameOverrides->GameModeOverride;
+	}
+
+	// Kick off an async request to start a play session. This won't happen until the next frame.
+	GEditor->RequestPlaySession(Params);
+
+	// Listen for PIE startup since there's no current way to pass a delegate through the request.
+	FEditorDelegates::PostPIEStarted.AddUObject(this, &UMoviePipelinePIEExecutor::OnPIEStartupFinished);
+}
+
+void UMoviePipelinePIEExecutor::OnPIEStartupFinished(bool)
+{
+	// Immediately un-bind our delegate so that we don't catch all PIE startup requests in the future.
+	FEditorDelegates::PostPIEStarted.RemoveAll(this);
+
+
+	// Hack to find out the PIE world since it is not provided by the delegate.
+	UWorld* ExecutingWorld = nullptr;
+	
+	for (const FWorldContext& Context : GEngine->GetWorldContexts())
+	{
+		if (Context.WorldType == EWorldType::PIE)
+		{
+			ExecutingWorld = Context.World();
+		}
+	}
+	
+	check(ExecutingWorld);
+
+	// Allow the user to have overridden which Pipeline is actually run. This is an unlikely scenario but allows
+	// the user to create their own implementation while still re-using the rest of our UI and infrastructure.
+	const UMovieRenderPipelineProjectSettings* ProjectSettings = GetDefault<UMovieRenderPipelineProjectSettings>();
+	TSubclassOf<UMoviePipeline> PipelineClass = ProjectSettings->DefaultPipeline;
+
+	// This Pipeline belongs to the world being created so that they have context for things they execute.
+	ActiveMoviePipeline = NewObject<UMoviePipeline>(ExecutingWorld, PipelineClass);
+	
+	ActiveMoviePipeline->Initialize(ActiveConfig);
+
+	// Listen for when the pipeline thinks it has finished.
+	ActiveMoviePipeline->OnMoviePipelineFinished().AddUObject(this, &UMoviePipelinePIEExecutor::OnPIEMoviePipelineFinished);
+}
+
+void UMoviePipelinePIEExecutor::OnPIEMoviePipelineFinished(UMoviePipeline* InMoviePipeline)
+{
+	if (ActiveMoviePipeline)
+	{
+		// Unsubscribe in the event that it gets called twice we don't have issues.
+		ActiveMoviePipeline->OnMoviePipelineFinished().RemoveAll(this);
+	}
+
+	// The End Play will happen on the next frame.
+	GEditor->RequestEndPlayMap();
+
+	// Delay for one frame so that PIE can finish shut down. It's not a huge fan of us starting up on the same frame.
+	FTimerHandle Handle;
+	GEditor->GetTimerManager()->SetTimer(Handle, FTimerDelegate::CreateUObject(this, &UMoviePipelinePIEExecutor::DelayedFinishNotification), 1.f, false);
+}
+
+void UMoviePipelinePIEExecutor::DelayedFinishNotification()
+{
+	// Now that another frame has passed and we should be OK to start another PIE session, notify our owner.
+	OnIndividualPipelineFinished(ActiveMoviePipeline);
+	
+
+	ActiveMoviePipeline = nullptr;
+	ActiveConfig = nullptr;
+}
+
+#undef LOCTEXT_NAMESPACE // "MoviePipelinePIEExecutor"
diff --git a/Engine/Plugins/MovieScene/MovieRenderPipeline/Source/MovieRenderPipelineEditor/Private/MovieRenderPipelineEditorModule.cpp b/Engine/Plugins/MovieScene/MovieRenderPipeline/Source/MovieRenderPipelineEditor/Private/MovieRenderPipelineEditorModule.cpp
new file mode 100644
--- /dev/null
+++ b/Engine/Plugins/MovieScene/MovieRenderPipeline/Source/MovieRenderPipelineEditor/Private/MovieRenderPipelineEditorModule.cpp
@@ -0,0 +1,138 @@
+// Copyright 1998-2019 Epic Games, Inc. All Rights Reserved.
+
+#include "MovieRenderPipelineEditorModule.h"
+#include "Widgets/Docking/SDockTab.h"
+#include "Widgets/Text/STextBlock.h"
+#include "EditorStyleSet.h"
+#include "WorkspaceMenuStructureModule.h"
+#include "WorkspaceMenuStructure.h"
+#include "Widgets/SMoviePipelineTabContent.h"
+#include "Editor.h"
+#include "MoviePipeline.h"
+#include "MovieRenderPipelineDataTypes.h"
+#include "MovieRenderPipelineConfig.h"
+#include "MoviePipelineBackbufferPass.h"
+#include "MoviePipelineImageSequenceContainer.h"
+#include "MoviePipelineShotConfig.h"
+#include "MovieRenderPipelineCoreModule.h"
+#include "MoviePipelineAudioOutput.h"
+#include "MoviePipelineBackbufferPass.h"
+#include "MoviePipelineGameOverrideSetting.h"
+#include "GameFramework/GameMode.h"
+#include "Engine/World.h"
+#include "GameFramework/WorldSettings.h"
+#include "MoviePipelineAccumulationSetting.h"
+#include "ISettingsModule.h"
+#include "MovieRenderPipelineSettings.h"
+#include "HAL/IConsoleManager.h"
+#include "MovieRenderPipelineSettings.h"
+#include "MoviePipelineExecutor.h"
+
+#define LOCTEXT_NAMESPACE "FMovieRenderPipelineEditorModule"
+
+FName IMovieRenderPipelineEditorModule::MovieRenderPipelineTabName = "MovieRenderPipeline";
+FText IMovieRenderPipelineEditorModule::MovieRenderPipelineTabLabel = LOCTEXT("MovieRenderPipelineTab_Label", "Movie Render Pipeline");
+										 
+namespace
+{
+	static TSharedRef<SDockTab> SpawnMovieRenderPipelineTab(const FSpawnTabArgs& SpawnTabArgs)
+	{
+		TSharedRef<SDockTab> NewTab = SNew(SDockTab)
+			.TabRole(ETabRole::NomadTab)
+			[
+				SNew(SMoviePipelineTabContent)
+			];
+
+		return NewTab;
+	}
+
+	static void RegisterTabImpl()
+	{
+		FTabSpawnerEntry& MRPTabSpawner = FGlobalTabmanager::Get()->RegisterNomadTabSpawner(IMovieRenderPipelineEditorModule::MovieRenderPipelineTabName, FOnSpawnTab::CreateStatic(SpawnMovieRenderPipelineTab));
+
+		MRPTabSpawner
+			.SetGroup(WorkspaceMenu::GetMenuStructure().GetLevelEditorCinematicsCategory())
+			.SetDisplayName(IMovieRenderPipelineEditorModule::MovieRenderPipelineTabLabel)
+			.SetTooltipText(LOCTEXT("MovieRenderPipelineTab_Tooltip", "Open the Movie Render Pipeline UI"))
+			.SetIcon(FSlateIcon(FEditorStyle::GetStyleSetName(), "ContentBrowser.TabIcon"));
+	}
+}
+
+void FMovieRenderPipelineEditorModule::PerformTestPipelineRender(const TArray<FString>& Args)
+{
+	if (Args.Num() == 0)
+	{
+		UE_LOG(LogMovieRenderPipeline, Warning, TEXT("No arguments specified. Specify the path to a sequence asset to render."));
+		return;
+	}
+
+	FSoftObjectPath SequencePath = FSoftObjectPath(Args[0]);
+
+	const UMovieRenderPipelineProjectSettings* ProjectSettings = GetDefault<UMovieRenderPipelineProjectSettings>();
+	TSubclassOf<UMoviePipelineExecutorBase> ExecutorClass = ProjectSettings->DefaultLocalExecutor;
+
+	// Create an instance of the Executor and add it to the root so it never dies.
+	Executor = NewObject<UMoviePipelineExecutorBase>(GetTransientPackage(), ExecutorClass);
+	Executor->AddToRoot();
+	Executor->OnExecutorFinished().AddRaw(this, &FMovieRenderPipelineEditorModule::OnTestPipelineExecutorFinished);
+
+	TArray<UMovieRenderPipelineConfig*> Pipelines = GenerateTestPipelineConfigs(SequencePath);
+	Executor->Execute(Pipelines);
+}
+
+void FMovieRenderPipelineEditorModule::OnTestPipelineExecutorFinished(UMoviePipelineExecutorBase* InExecutor)
+{
+	if (Executor)
+	{
+		Executor->RemoveFromRoot();
+		Executor = nullptr;
+	}
+}
+
+void FMovieRenderPipelineEditorModule::RegisterSettings()
+{
+	ISettingsModule& SettingsModule = FModuleManager::LoadModuleChecked<ISettingsModule>("Settings");
+
+	SettingsModule.RegisterSettings("Project", "Plugins", "Movie Render Pipeline",
+		LOCTEXT("ProjectSettings_Label", "Movie Render Pipeline"),
+		LOCTEXT("ProjectSettings_Description", "Configure project-wide defaults for the movie render pipeline."),
+		GetMutableDefault<UMovieRenderPipelineProjectSettings>()
+	);
+
+	IConsoleManager::Get().RegisterConsoleCommand(
+		TEXT("MovieRenderPipeline.TestRenderSequence"),
+		TEXT("Renders the specified sequence asset using the default executor and a pre-made configuration."),
+		FConsoleCommandWithArgsDelegate::CreateRaw(this, &FMovieRenderPipelineEditorModule::PerformTestPipelineRender),
+		ECVF_Default
+	);
+}
+
+void FMovieRenderPipelineEditorModule::UnregisterSettings()
+{
+	ISettingsModule* SettingsModule = FModuleManager::GetModulePtr<ISettingsModule>("Settings");
+	if (SettingsModule)
+	{
+		SettingsModule->UnregisterSettings("Project", "Plugins", "Movie Render Pipeline");
+	}
+
+	IConsoleManager::Get().UnregisterConsoleObject(TEXT("MovieRenderPipeline.TestRenderSequence"));
+}
+
+void FMovieRenderPipelineEditorModule::StartupModule()
+{
+	RegisterTabImpl();
+	RegisterSettings();
+}
+
+void FMovieRenderPipelineEditorModule::ShutdownModule()
+{
+	FEditorDelegates::PostPIEStarted.RemoveAll(this);
+	FEditorDelegates::PrePIEEnded.RemoveAll(this);
+	UGameViewportClient::OnViewportCreated().RemoveAll(this);
+
+	UnregisterSettings();
+}
+
+#undef LOCTEXT_NAMESPACE
+
+IMPLEMENT_MODULE(FMovieRenderPipelineEditorModule, MovieRenderPipelineEditor)
\ No newline at end of file
diff --git a/Engine/Plugins/MovieScene/MovieRenderPipeline/Source/MovieRenderPipelineEditor/Private/MovieRenderPipelineEditorModule.h b/Engine/Plugins/MovieScene/MovieRenderPipeline/Source/MovieRenderPipelineEditor/Private/MovieRenderPipelineEditorModule.h
new file mode 100644
--- /dev/null
+++ b/Engine/Plugins/MovieScene/MovieRenderPipeline/Source/MovieRenderPipelineEditor/Private/MovieRenderPipelineEditorModule.h
@@ -0,0 +1,29 @@
+// Copyright 1998-2019 Epic Games, Inc. All Rights Reserved.
+
+#pragma once
+
+#include "IMovieRenderPipelineEditorModule.h"
+class UMoviePipelineShotConfig;
+class UMovieRenderPipelineConfig;
+class UMoviePipelineExecutorBase;
+
+class FMovieRenderPipelineEditorModule : public IMovieRenderPipelineEditorModule
+{
+public:
+
+	/** IModuleInterface implementation */
+	virtual void StartupModule() override;
+	virtual void ShutdownModule() override;
+private:
+	void RegisterSettings();
+	void UnregisterSettings();
+	void PerformTestPipelineRender(const TArray<FString>& Args);
+	void OnTestPipelineExecutorFinished(UMoviePipelineExecutorBase* InExecutor);
+
+private:
+	TArray<UMovieRenderPipelineConfig*> GenerateTestPipelineConfigs(FSoftObjectPath InSequence);
+private:
+	class UMoviePipeline* ActiveMoviePipeline;
+	class UMovieRenderPipelineConfig* PipelineConfig;
+	UMoviePipelineExecutorBase* Executor;
+};
diff --git a/Engine/Plugins/MovieScene/MovieRenderPipeline/Source/MovieRenderPipelineEditor/Private/MovieRenderPipelineEditorTest.cpp b/Engine/Plugins/MovieScene/MovieRenderPipeline/Source/MovieRenderPipelineEditor/Private/MovieRenderPipelineEditorTest.cpp
new file mode 100644
--- /dev/null
+++ b/Engine/Plugins/MovieScene/MovieRenderPipeline/Source/MovieRenderPipelineEditor/Private/MovieRenderPipelineEditorTest.cpp
@@ -0,0 +1,86 @@
+#include "MovieRenderPipelineEditorModule.h"
+#include "MoviePipelineBackbufferPass.h"
+#include "MoviePipelineShotConfig.h"
+#include "MoviePipelineGameOverrideSetting.h"
+#include "GameFramework/GameMode.h"
+#include "MoviePipelineAccumulationSetting.h"
+#include "MovieRenderPipelineConfig.h"
+#include "MoviePipelineImageSequenceContainer.h"
+#include "MoviePipelineAudioOutput.h"
+#include "Misc/Paths.h"
+
+UMoviePipelineShotConfig* GenerateTestShotConfig(UObject* InOwner, int32 InSampleCount,
+	int32 InShutterAngle, EMoviePipelineShutterTiming InFrameTiming, int32 InTileCount,
+	int32 InSpatialSampleCount)
+{
+	UMoviePipelineShotConfig* OutConfig = NewObject<UMoviePipelineShotConfig>(InOwner);
+	UMoviePipelineBackbufferPass* DefaultPass = NewObject<UMoviePipelineBackbufferPass>(InOwner);
+	OutConfig->InputBuffers.Add(DefaultPass);
+
+	UMoviePipelineGameOverrideSetting* GamemodeOverride = OutConfig->FindOrAddSetting<UMoviePipelineGameOverrideSetting>();
+	GamemodeOverride->GameModeOverride = AGameMode::StaticClass();
+		
+	UMoviePipelineAccumulationSetting* Accumulation = OutConfig->FindOrAddSetting< UMoviePipelineAccumulationSetting>();
+	Accumulation->TemporalSampleCount = InSampleCount;
+	Accumulation->CameraShutterAngle = InShutterAngle;
+	Accumulation->TileCount = InTileCount;
+	Accumulation->ShutterTiming = InFrameTiming;
+	Accumulation->SpatialSampleCount = InSpatialSampleCount;
+
+	return OutConfig;
+}
+
+TArray<UMovieRenderPipelineConfig*> FMovieRenderPipelineEditorModule::GenerateTestPipelineConfigs(FSoftObjectPath InSequence)
+{
+	// int32 ShutterAngles[] = { 180, 360 };
+	// EMoviePipelineShutterTiming ShutterTimings[] = { EMoviePipelineShutterTiming::FrameOpen, EMoviePipelineShutterTiming::FrameCenter, EMoviePipelineShutterTiming::FrameClose };
+	// FString ShutterTimingNames[] = { TEXT("FOpen"), TEXT("FCenter"), TEXT("FClose") };
+	// int32 TemporalSampleCounts[] = { 1, 5 };
+	   
+	int32 ShutterAngles[] = { 180 };
+	EMoviePipelineShutterTiming ShutterTimings[] = { EMoviePipelineShutterTiming::FrameCenter};
+	FString ShutterTimingNames[] = { TEXT("FCenter") };
+	int32 TemporalSampleCounts[] = { 1 };
+
+	TArray<UMovieRenderPipelineConfig*> Configs;
+
+	for (int32 ShutterAngleIndex = 0; ShutterAngleIndex < UE_ARRAY_COUNT(ShutterAngles); ShutterAngleIndex++)
+	{
+		for (int32 ShutterTimingIndex = 0; ShutterTimingIndex < UE_ARRAY_COUNT(ShutterTimings); ShutterTimingIndex++)
+		{
+			for (int32 TemporalSampleCountIndex = 0; TemporalSampleCountIndex < UE_ARRAY_COUNT(TemporalSampleCounts); TemporalSampleCountIndex++)
+			{
+				// Build a folder string to place this particular test in.
+				FString DirectoryName = FString::Printf(TEXT("SA_%d_ST_%s_SC_%d"), ShutterAngles[ShutterAngleIndex], *ShutterTimingNames[ShutterTimingIndex], TemporalSampleCounts[TemporalSampleCountIndex]);
+
+				UMovieRenderPipelineConfig* OutPipeline = NewObject<UMovieRenderPipelineConfig>(GetTransientPackage());
+				OutPipeline->Sequence = InSequence;
+				OutPipeline->OutputResolution = FIntPoint(1920, 1080);
+				OutPipeline->OutputDirectory.Path = FPaths::ProjectSavedDir() / TEXT("/VideoCaptures/") / *DirectoryName;
+
+				const int32 NumSpatialSamples = 1;
+				const int32 NumTiles = 1;
+
+				UMoviePipelineShotConfig* DefaultConfig = GenerateTestShotConfig(OutPipeline,
+					TemporalSampleCounts[TemporalSampleCountIndex],
+					ShutterAngles[ShutterAngleIndex],
+					ShutterTimings[ShutterTimingIndex],
+					NumTiles, NumSpatialSamples);
+
+				OutPipeline->DefaultShotConfig = DefaultConfig;
+
+				// Add some Outputs
+				UMoviePipelineImageSequenceContainerBase* ImageSequenceOutput = NewObject<UMoviePipelineImageSequenceContainerBase>(OutPipeline);
+				OutPipeline->OutputContainers.Add(ImageSequenceOutput);
+
+				UMoviePipelineAudioOutput* AudioOutput = NewObject<UMoviePipelineAudioOutput>(OutPipeline);
+				OutPipeline->OutputContainers.Add(AudioOutput);
+				
+				// Add it to our list to be processed.
+				Configs.Add(OutPipeline);
+			}
+		}
+	}
+
+	return Configs;
+}
\ No newline at end of file
diff --git a/Engine/Plugins/MovieScene/MovieRenderPipeline/Source/MovieRenderPipelineEditor/Private/Widgets/SMoviePipelineEditor.cpp b/Engine/Plugins/MovieScene/MovieRenderPipeline/Source/MovieRenderPipelineEditor/Private/Widgets/SMoviePipelineEditor.cpp
new file mode 100644
--- /dev/null
+++ b/Engine/Plugins/MovieScene/MovieRenderPipeline/Source/MovieRenderPipelineEditor/Private/Widgets/SMoviePipelineEditor.cpp
@@ -0,0 +1,299 @@
+// Copyright 1998-2019 Epic Games, Inc. All Rights Reserved.
+
+#include "Widgets/SMoviePipelineEditor.h"
+#include "Widgets/MoviePipelineWidgetConstants.h"
+
+#include "MoviePipelineShotConfig.h"
+#include "MoviePipelineSetting.h"
+
+// Core includes
+#include "UObject/UObjectIterator.h"
+#include "Templates/SubclassOf.h"
+#include "ClassIconFinder.h"
+// 
+// AssetRegistry includes
+#include "AssetRegistryModule.h"
+#include "AssetData.h"
+ 
+// ContentBrowser includes
+#include "IContentBrowserSingleton.h"
+#include "ContentBrowserModule.h"
+ 
+// AssetTools includes
+#include "AssetToolsModule.h"
+ 
+// Slate includes
+#include "Widgets/SBoxPanel.h"
+#include "Widgets/Text/STextBlock.h"
+#include "Widgets/Input/SButton.h"
+#include "Widgets/Input/SComboButton.h"
+#include "Widgets/Layout/SSplitter.h"
+#include "Widgets/Layout/SScrollBox.h"
+#include "Widgets/Layout/SBorder.h"
+#include "Framework/MultiBox/MultiBoxBuilder.h"
+#include "Styling/SlateIconFinder.h"
+
+// EditorStyle includes
+#include "EditorStyleSet.h"
+#include "EditorFontGlyphs.h"
+#include "ScopedTransaction.h"
+#include "PropertyEditorModule.h"
+#include "IDetailsView.h"
+#include "SMoviePipelineSettings.h"
+
+#define LOCTEXT_NAMESPACE "SMoviePipelineEditor"
+
+
+TArray<UClass*> FindRecordingSourceClasses()
+{
+	FAssetRegistryModule& AssetRegistryModule = FModuleManager::LoadModuleChecked<FAssetRegistryModule>(TEXT("AssetRegistry"));
+
+	TArray<FAssetData> ClassList;
+
+	FARFilter Filter;
+	Filter.ClassNames.Add(UMoviePipelineSetting::StaticClass()->GetFName());
+
+	// Include any Blueprint based objects as well, this includes things like Blutilities, UMG, and GameplayAbility objects
+	Filter.bRecursiveClasses = true;
+	AssetRegistryModule.Get().GetAssets(Filter, ClassList);
+
+	TArray<UClass*> Classes;
+
+	for (const FAssetData& Data : ClassList)
+	{
+		UClass* Class = Data.GetClass();
+		if (Class)
+		{
+			Classes.Add(Class);
+		}
+	}
+
+	for (TObjectIterator<UClass> ClassIterator; ClassIterator; ++ClassIterator)
+	{
+		if (ClassIterator->IsChildOf(UMoviePipelineSetting::StaticClass()) && !ClassIterator->HasAnyClassFlags(CLASS_Abstract | CLASS_Deprecated | CLASS_NewerVersionExists))
+		{
+			Classes.Add(*ClassIterator);
+		}
+	}
+
+	return Classes;
+}
+
+PRAGMA_DISABLE_OPTIMIZATION
+void SMoviePipelineEditor::Construct(const FArguments& InArgs)
+{
+	bRequestDetailsRefresh = true;
+	MoviePipelineAttribute = InArgs._MoviePipeline;
+    
+	DetailsBox = SNew(SScrollBox);
+	DetailsBox->SetScrollBarRightClickDragAllowed(true);
+     
+	SettingsWidget = SNew(SMoviePipelineSettings)
+	.OnSelectionChanged(this, &SMoviePipelineEditor::OnSettingsSelectionChanged);
+    
+	CheckForNewMoviePipeline();
+     
+	ChildSlot
+	[
+		SNew(SSplitter)
+		.Orientation(Orient_Horizontal)
+    
+		+ SSplitter::Slot()
+		.Value(.25f)
+		[
+			SNew(SBorder)
+			.Padding(4)
+			.BorderImage( FEditorStyle::GetBrush("ToolPanel.GroupBorder") )
+			[
+				SettingsWidget.ToSharedRef()
+			]
+		]
+    
+		+ SSplitter::Slot()
+		.Value(.75f)
+		[
+			DetailsBox.ToSharedRef()
+		]
+	];
+}
+
+TSharedRef<SWidget> SMoviePipelineEditor::MakeAddSettingButton()
+{
+	return SNew(SComboButton)
+		.ContentPadding(MoviePipeline::ButtonPadding)
+		// .ButtonStyle(FTakeRecorderStyle::Get(), "FlatButton.Success")
+		.OnGetMenuContent(this, &SMoviePipelineEditor::OnGenerateSettingsMenu)
+		.ForegroundColor(FSlateColor::UseForeground())
+		.HasDownArrow(false)
+		.ButtonContent()
+		[
+			SNew(SHorizontalBox)
+    
+			// Plus Icon
+			+ SHorizontalBox::Slot()
+			.VAlign(VAlign_Center)
+			.AutoWidth()
+			[
+				SNew(STextBlock)
+				.TextStyle(FEditorStyle::Get(), "NormalText.Important")
+				.Font(FEditorStyle::Get().GetFontStyle("FontAwesome.10"))
+				.Text(FEditorFontGlyphs::Plus)
+			]
+    
+			// "Setting" Text
+			+ SHorizontalBox::Slot()
+			.VAlign(VAlign_Center)
+			.AutoWidth()
+			.Padding(4, 0, 0, 0)
+			[
+				SNew(STextBlock)
+				.TextStyle(FEditorStyle::Get(), "NormalText.Important")
+				.Text(LOCTEXT("AddNewSetting_Text", "Setting"))
+			]
+    
+			// Non-Default Down Caret arrow.
+			+ SHorizontalBox::Slot()
+			.VAlign(VAlign_Center)
+			.AutoWidth()
+			.Padding(4, 0, 0, 0)
+			[
+				SNew(STextBlock)
+				.TextStyle(FEditorStyle::Get(), "NormalText.Important")
+				.Font(FEditorStyle::Get().GetFontStyle("FontAwesome.10"))
+				.Text(FEditorFontGlyphs::Caret_Down)
+			]
+		];
+}
+PRAGMA_ENABLE_OPTIMIZATION
+
+void SMoviePipelineEditor::CheckForNewMoviePipeline()
+{
+	UMoviePipelineShotConfig* NewMoviePipeline = MoviePipelineAttribute.Get();
+	if (CachedMoviePipeline != NewMoviePipeline)
+	{
+		CachedMoviePipeline = NewMoviePipeline;
+
+		SettingsWidget->SetShotConfigObject(NewMoviePipeline);
+		bRequestDetailsRefresh = true;
+	}
+}
+
+void SMoviePipelineEditor::Tick(const FGeometry& AllottedGeometry, const double InCurrentTime, const float InDeltaTime)
+{
+	CheckForNewMoviePipeline();
+	if (bRequestDetailsRefresh)
+	{
+		UpdateDetails();
+		bRequestDetailsRefresh = false;
+	}
+}
+
+TSharedRef<SWidget> SMoviePipelineEditor::OnGenerateSettingsMenu()
+{
+	TSharedRef<FExtender> Extender = MakeShared<FExtender>();
+	{
+		// ULevelSequence*       LevelSequence = LevelSequenceAttribute.Get();
+		// UTakeRecorderSources* Sources = LevelSequence ? LevelSequence->FindOrAddMetaData<UTakeRecorderSources>() : nullptr;
+		// 
+		// if (Sources)
+		// {
+		// 	FTakeRecorderModule& TakeRecorderModule = FModuleManager::GetModuleChecked<FTakeRecorderModule>("TakeRecorder");
+		// 	TakeRecorderModule.PopulateSourcesMenu(Extender, Sources);
+		// }
+	}
+
+	FMenuBuilder MenuBuilder(true, nullptr, Extender);
+
+	MenuBuilder.BeginSection("Settings", LOCTEXT("SettingsMenuSection", "Available Settings"));
+	{
+		TArray<UClass*> SourceClasses = FindRecordingSourceClasses();
+		Algo::SortBy(SourceClasses, &UClass::GetDisplayNameText, FText::FSortPredicate());
+
+		for (UClass* Class : SourceClasses)
+		{
+			TSubclassOf<UMoviePipelineSetting> SubclassOf = Class;
+
+			MenuBuilder.AddMenuEntry(
+				FText::FromString(Class->GetMetaData(TEXT("MovieRenderPipelineDisplayName"))),
+				Class->GetToolTipText(true),
+				FSlateIconFinder::FindIconForClass(Class),
+				FUIAction(
+					FExecuteAction::CreateSP(this, &SMoviePipelineEditor::AddSettingFromClass, SubclassOf),
+					FCanExecuteAction::CreateSP(this, &SMoviePipelineEditor::CanAddSettingFromClass, SubclassOf)
+				)
+			);
+		}
+	}
+	MenuBuilder.EndSection();
+
+	return MenuBuilder.MakeWidget();
+}
+
+void SMoviePipelineEditor::UpdateDetails()
+{
+	FPropertyEditorModule& PropertyEditorModule = FModuleManager::Get().LoadModuleChecked<FPropertyEditorModule>("PropertyEditor");
+	FDetailsViewArgs DetailsViewArgs(false, false, false, FDetailsViewArgs::HideNameArea, true);
+	DetailsViewArgs.bShowScrollBar = false;
+	 
+	TArray<UMoviePipelineSetting*> SelectedSources;
+	SettingsWidget->GetSelectedSettings(SelectedSources);
+
+	
+	// Create 1 details panel per source class type
+	TSortedMap<const UClass*, TArray<UObject*>> ClassToSources;
+	for (UMoviePipelineSetting* Source : SelectedSources)
+	{
+		ClassToSources.FindOrAdd(Source->GetClass()).Add(Source);
+	}
+	
+	TArray<FObjectKey> PreviousClasses;
+	ClassToDetailsView.GenerateKeyArray(PreviousClasses);
+	
+	for (auto& Pair : ClassToSources)
+	{
+		PreviousClasses.Remove(Pair.Key);
+	
+		TSharedPtr<IDetailsView> ExistingDetails = ClassToDetailsView.FindRef(Pair.Key);
+		if (ExistingDetails.IsValid())
+		{
+			ExistingDetails->SetObjects(Pair.Value);
+		}
+		else
+		{
+			TSharedRef<IDetailsView> Details = PropertyEditorModule.CreateDetailView(DetailsViewArgs);
+			Details->SetObjects(Pair.Value);
+			DetailsBox->AddSlot()
+				[
+					Details
+				];
+	
+			ClassToDetailsView.Add(Pair.Key, Details);
+		}
+	}
+	
+	for (FObjectKey StaleClass : PreviousClasses)
+	{
+		TSharedPtr<IDetailsView> Details = ClassToDetailsView.FindRef(StaleClass);
+		DetailsBox->RemoveSlot(Details.ToSharedRef());
+		ClassToDetailsView.Remove(StaleClass);
+	}
+}
+
+void SMoviePipelineEditor::AddSettingFromClass(TSubclassOf<UMoviePipelineSetting> SettingClass)
+{
+	UMoviePipelineShotConfig* Pipeline = MoviePipelineAttribute.Get();
+
+	if (*SettingClass && Pipeline)
+	{
+		FScopedTransaction Transaction(FText::Format(LOCTEXT("AddNewSetting", "Add New {0} Setting"), SettingClass->GetDisplayNameText()));
+		Pipeline->FindOrAddSettingByClass(SettingClass);
+	}
+}
+
+bool SMoviePipelineEditor::CanAddSettingFromClass(TSubclassOf<UMoviePipelineSetting> SettingClass)
+{
+	return true;
+}
+
+
+#undef LOCTEXT_NAMESPACE
\ No newline at end of file
diff --git a/Engine/Plugins/MovieScene/MovieRenderPipeline/Source/MovieRenderPipelineEditor/Private/Widgets/SMoviePipelineEditor.h b/Engine/Plugins/MovieScene/MovieRenderPipeline/Source/MovieRenderPipelineEditor/Private/Widgets/SMoviePipelineEditor.h
new file mode 100644
--- /dev/null
+++ b/Engine/Plugins/MovieScene/MovieRenderPipeline/Source/MovieRenderPipelineEditor/Private/Widgets/SMoviePipelineEditor.h
@@ -0,0 +1,77 @@
+// Copyright 1998-2019 Epic Games, Inc. All Rights Reserved.
+
+#pragma once
+
+#include "Widgets/DeclarativeSyntaxSupport.h"
+#include "Widgets/SCompoundWidget.h"
+#include "UObject/GCObject.h"
+#include "UObject/ObjectKey.h"
+
+struct FAssetData;
+struct IMoviePipelineSettingTreeItem;
+
+class SScrollBox;
+class SMoviePipelineSettings;
+class UMoviePipelineShotConfig;
+class IDetailsView;
+class UMoviePipelineSetting;
+
+template<class> class TSubclassOf;
+
+/**
+ * Widget used to edit a Movie Render Pipeline Shot Config.
+ */
+class SMoviePipelineEditor : public SCompoundWidget
+{
+public:
+
+	SLATE_BEGIN_ARGS(SMoviePipelineEditor)
+		: _MoviePipeline(nullptr)
+		{}
+
+		SLATE_ATTRIBUTE(UMoviePipelineShotConfig*, MoviePipeline)
+
+	SLATE_END_ARGS()
+
+	void Construct(const FArguments& InArgs);
+
+	/**
+	 * Construct a button that can add sources to this widget's preset
+	 */
+	TSharedRef<SWidget> MakeAddSettingButton();
+
+private:
+
+	virtual void Tick(const FGeometry& AllottedGeometry, const double InCurrentTime, const float InDeltaTime) override;
+
+private:
+
+	/**
+	 * Check to see whether the movie pipeline ptr has changed, and propagate that change if necessary
+	 */
+	void CheckForNewMoviePipeline();
+
+	/**
+	 * Update the details panel for the current selection
+	 */
+	void UpdateDetails();
+
+private:
+
+	TSharedRef<SWidget> OnGenerateSettingsMenu();
+     
+	void AddSettingFromClass(TSubclassOf<UMoviePipelineSetting> SettingClass);
+	bool CanAddSettingFromClass(TSubclassOf<UMoviePipelineSetting> SettingClass);
+	
+	void OnSettingsSelectionChanged(TSharedPtr<IMoviePipelineSettingTreeItem>, ESelectInfo::Type) { bRequestDetailsRefresh = true; }
+
+private:
+
+	bool bRequestDetailsRefresh;
+	TAttribute<UMoviePipelineShotConfig*> MoviePipelineAttribute;
+	TWeakObjectPtr<UMoviePipelineShotConfig> CachedMoviePipeline;
+    
+	TSharedPtr<SMoviePipelineSettings> SettingsWidget;
+	TSharedPtr<SScrollBox> DetailsBox;
+	TMap<FObjectKey, TSharedPtr<IDetailsView>> ClassToDetailsView;
+};
\ No newline at end of file
diff --git a/Engine/Plugins/MovieScene/MovieRenderPipeline/Source/MovieRenderPipelineEditor/Private/Widgets/SMoviePipelinePanel.cpp b/Engine/Plugins/MovieScene/MovieRenderPipeline/Source/MovieRenderPipelineEditor/Private/Widgets/SMoviePipelinePanel.cpp
new file mode 100644
--- /dev/null
+++ b/Engine/Plugins/MovieScene/MovieRenderPipeline/Source/MovieRenderPipelineEditor/Private/Widgets/SMoviePipelinePanel.cpp
@@ -0,0 +1,288 @@
+// Copyright 1998-2019 Epic Games, Inc. All Rights Reserved.
+
+#include "Widgets/SMoviePipelinePanel.h"
+#include "Widgets/SMoviePipelineEditor.h"
+#include "Widgets/MoviePipelineWidgetConstants.h"
+#include "Misc/MessageDialog.h"
+#include "MoviePipelineShotConfig.h"
+
+// Slate Includes
+#include "Widgets/Text/STextBlock.h"
+#include "Widgets/SBoxPanel.h"
+#include "Widgets/Layout/SBorder.h"
+#include "Widgets/Layout/SBox.h"
+#include "Widgets/Layout/SSpacer.h"
+#include "Widgets/Input/SButton.h"
+#include "Widgets/Input/SComboButton.h"
+#include "Widgets/Images/SImage.h"
+#include "Framework/MultiBox/MultiBoxBuilder.h"
+#include "Styling/SlateIconFinder.h"
+#include "EditorStyleSet.h"
+#include "EditorFontGlyphs.h"
+
+
+// ContentBrowser Includes
+#include "IContentBrowserSingleton.h"
+#include "ContentBrowserModule.h"
+
+// UnrealEd Includes
+#include "ScopedTransaction.h"
+
+#define LOCTEXT_NAMESPACE "SMoviePipelinePanel"
+
+SMoviePipelinePanel::~SMoviePipelinePanel()
+{
+}
+
+PRAGMA_DISABLE_OPTIMIZATION
+void SMoviePipelinePanel::Construct(const FArguments& InArgs)
+{
+	// Allocate a transient preset automatically so they can start editing without having to create an asset.
+	TransientPreset = AllocateTransientPreset();
+
+	// Copy the base preset into the transient preset if it was provided.
+	if (InArgs._BasePreset)
+	{
+		// TransientPreset->CopyFrom(InArgs._BasePreset);
+	}
+
+	// Create the child widgets that need to know about our pipeline
+	MoviePipelineEditorWidget = SNew(SMoviePipelineEditor)
+	.MoviePipeline(this, &SMoviePipelinePanel::GetMoviePipeline);
+
+	ChildSlot
+	[
+		SNew(SVerticalBox)
+
+		// Create the toolbar for adding new stuff, choosing/saving a preset and resetting to the preset defaults.
+		+ SVerticalBox::Slot()
+		.Padding(FMargin(0.f, 1.0f))
+		.AutoHeight()
+		[
+			SNew(SBorder)
+			.BorderImage(FEditorStyle::GetBrush("DetailsView.CategoryTop"))
+			.BorderBackgroundColor(FLinearColor(.6, .6, .6, 1.0f))
+			// .IsEnabled_Lambda([this]() { return !CockpitWidget->Reviewing(); })
+			[
+				SNew(SHorizontalBox)
+
+				// Adding things to the Pipeline
+				+ SHorizontalBox::Slot()
+				.Padding(MoviePipeline::ButtonOffset)
+				.VAlign(VAlign_Fill)
+				.AutoWidth()
+				[
+					MoviePipelineEditorWidget->MakeAddSettingButton()
+				]
+
+				// Presets Management Button
+				+ SHorizontalBox::Slot()
+				.Padding(MoviePipeline::ButtonOffset)
+				.VAlign(VAlign_Fill)
+				.AutoWidth()
+				[
+					SNew(SComboButton)
+					.ContentPadding(MoviePipeline::ButtonPadding)
+					// .ComboButtonStyle(FTakeRecorderStyle::Get(), "ComboButton")
+					.OnGetMenuContent(this, &SMoviePipelinePanel::OnGeneratePresetsMenu)
+					.ForegroundColor(FSlateColor::UseForeground())
+					.ButtonContent()
+					[
+						SNew(SHorizontalBox)
+
+						+ SHorizontalBox::Slot()
+						.AutoWidth()
+						[
+							SNew(SImage)
+							.Image(FSlateIconFinder::FindIconBrushForClass(UMoviePipelineShotConfig::StaticClass()))
+						]
+
+						+ SHorizontalBox::Slot()
+						[
+							SNew(STextBlock)
+							.Text(LOCTEXT("PresetsToolbarButton", "Presets"))
+						]
+					]
+				]
+
+				// Spacer to make the next widget right aligned
+				+ SHorizontalBox::Slot()
+				[
+					SNew(SSpacer)
+				]
+
+				// A button to undo your mistakes
+				+ SHorizontalBox::Slot()
+				.Padding(MoviePipeline::ButtonOffset)
+				.VAlign(VAlign_Fill)
+				.AutoWidth()
+				[
+					SNew(SButton)
+					.ContentPadding(MoviePipeline::ButtonPadding)
+					.ToolTipText(LOCTEXT("RevertChanges_Text", "Revert all changes made to this take back its original state (either its original preset, or an empty take)."))
+					.ForegroundColor(FSlateColor::UseForeground())
+					.ButtonStyle(FEditorStyle::Get(), "HoverHintOnly")
+					.OnClicked(this, &SMoviePipelinePanel::OnRevertChanges)
+					[
+						SNew(STextBlock)
+						.Font(FEditorStyle::Get().GetFontStyle("FontAwesome.11"))
+						.Text(FEditorFontGlyphs::Undo)
+					]
+				]
+			]
+		]
+
+		+ SVerticalBox::Slot()
+		.FillHeight(1.0f)
+		[
+			SNew(SHorizontalBox)
+			+ SHorizontalBox::Slot()
+			[
+				MoviePipelineEditorWidget.ToSharedRef()
+			]
+		]
+
+		// Footer Bar
+		+ SVerticalBox::Slot()
+		.AutoHeight()
+		[
+			SNew(SBorder)
+			.BorderImage(FEditorStyle::GetBrush("DetailsView.CategoryTop"))
+			.BorderBackgroundColor(FLinearColor(.6, .6, .6, 1.0f))
+			// .IsEnabled_Lambda([this]() { return !CockpitWidget->Reviewing(); })
+			[
+				SNew(SHorizontalBox)
+
+				// Spacer to make the next widget right aligned
+				+ SHorizontalBox::Slot()
+				[
+					SNew(STextBlock)
+					.Text(LOCTEXT("DummyValidationError", "This is where a validation error might show up?"))
+				]
+
+				// Attempt to kick off the pipeline if it's valid
+				+ SHorizontalBox::Slot()
+				.Padding(MoviePipeline::ButtonOffset)
+				.VAlign(VAlign_Fill)
+				.HAlign(HAlign_Right)
+				.AutoWidth()
+				[
+					SNew(SButton)
+					.ContentPadding(MoviePipeline::ButtonPadding)
+					.Text(LOCTEXT("RenderPipeline_Text", "Render Movie"))
+					.ToolTipText(LOCTEXT("RenderPipelineTooltip_Text", "Renders the selected sequence with the current output pipeline."))
+					.ForegroundColor(FSlateColor::UseForeground())
+					// .ButtonStyle(FEditorStyle::Get(), "HoverHintOnly")
+					// .OnClicked(this, &SMoviePipelinePanel::OnRevertChanges)
+
+				]
+			]
+		]
+	];
+}
+
+PRAGMA_ENABLE_OPTIMIZATION
+
+UMoviePipelineShotConfig* SMoviePipelinePanel::AllocateTransientPreset()
+{
+	static const TCHAR* PackageName = TEXT("/Temp/MoviePipeline/PendingPipeline");
+
+	// Return a cached transient if it exists
+	UMoviePipelineShotConfig* ExistingPreset = FindObject<UMoviePipelineShotConfig>(nullptr, TEXT("/Temp/MoviePipeline/PendingPipeline.PendingPipeline"));
+	if (ExistingPreset)
+	{
+		return ExistingPreset;
+	}
+
+	static FName DesiredName = "PendingMoviePipeline";
+	
+	UPackage* NewPackage = CreatePackage(nullptr, PackageName);
+	NewPackage->SetFlags(RF_Transient);
+	NewPackage->AddToRoot();
+
+	UMoviePipelineShotConfig* NewPreset = NewObject<UMoviePipelineShotConfig>(NewPackage, DesiredName, RF_Transient | RF_Transactional | RF_Standalone);
+
+	return NewPreset;
+}
+
+TSharedRef<SWidget> SMoviePipelinePanel::OnGeneratePresetsMenu()
+{
+	FMenuBuilder MenuBuilder(true, nullptr);
+
+	IContentBrowserSingleton& ContentBrowser = FModuleManager::LoadModuleChecked<FContentBrowserModule>("ContentBrowser").Get();
+
+	MenuBuilder.AddMenuEntry(
+		LOCTEXT("SaveAsPreset_Text", "Save As Preset"),
+		LOCTEXT("SaveAsPreset_Tip", "Save the current setup as a new preset that can be imported at a later date"),
+		FSlateIcon(FEditorStyle::Get().GetStyleSetName(), "AssetEditor.SaveAsset.Greyscale"),
+		FUIAction(
+			// FExecuteAction::CreateSP(this, &STakeRecorderPanel::OnSaveAsPreset)
+		)
+	);
+
+	FAssetPickerConfig AssetPickerConfig;
+	{
+		AssetPickerConfig.SelectionMode = ESelectionMode::Single;
+		AssetPickerConfig.InitialAssetViewType = EAssetViewType::Column;
+		AssetPickerConfig.bFocusSearchBoxWhenOpened = true;
+		AssetPickerConfig.bAllowNullSelection = false;
+		AssetPickerConfig.bShowBottomToolbar = true;
+		AssetPickerConfig.bAutohideSearchBar = false;
+		AssetPickerConfig.bAllowDragging = false;
+		AssetPickerConfig.bCanShowClasses = false;
+		AssetPickerConfig.bShowPathInColumnView = true;
+		AssetPickerConfig.bShowTypeInColumnView = false;
+		AssetPickerConfig.bSortByPathInColumnView = false;
+
+		AssetPickerConfig.AssetShowWarningText = LOCTEXT("NoPresets_Warning", "No Presets Found");
+		AssetPickerConfig.Filter.ClassNames.Add(UMoviePipelineShotConfig::StaticClass()->GetFName());
+		// AssetPickerConfig.OnAssetSelected = FOnAssetSelected::CreateSP(this, &SMoviePipelinePanel::OnImportPreset);
+	}
+
+	MenuBuilder.BeginSection(NAME_None, LOCTEXT("ImportPreset_MenuSection", "Import Preset"));
+	{
+		TSharedRef<SWidget> PresetPicker = SNew(SBox)
+			.MinDesiredWidth(400.f)
+			.MinDesiredHeight(400.f)
+			[
+				ContentBrowser.CreateAssetPicker(AssetPickerConfig)
+			];
+
+		MenuBuilder.AddWidget(PresetPicker, FText(), true, false);
+	}
+	MenuBuilder.EndSection();
+
+	return MenuBuilder.MakeWidget();
+}
+
+FReply SMoviePipelinePanel::OnRevertChanges()
+{
+	FText WarningMessage(LOCTEXT("Warning_RevertChanges", "Are you sure you want to revert changes? Your current changes will be discarded."));
+	if (EAppReturnType::No == FMessageDialog::Open(EAppMsgType::YesNo, WarningMessage))
+	{
+		return FReply::Handled();
+	}
+
+	// UMovieRenderPipeline* PresetOrigin = CockpitWidget->GetMetaData()->GetPresetOrigin();
+
+	FScopedTransaction Transaction(LOCTEXT("RevertChanges_Transaction", "Revert Changes"));
+
+	// TransientPreset->Modify();
+	// TransientPreset->CopyFrom(PresetOrigin);
+
+	return FReply::Handled();
+}
+
+void SMoviePipelinePanel::AddReferencedObjects(FReferenceCollector& Collector)
+{
+	// Collector.AddReferencedObject(TransientPreset);
+	// Collector.AddReferencedObject(SuppliedLevelSequence);
+	// Collector.AddReferencedObject(RecordingLevelSequence);
+}
+
+UMoviePipelineShotConfig* SMoviePipelinePanel::GetMoviePipeline() const
+{
+	return TransientPreset;
+}
+
+#undef LOCTEXT_NAMESPACE // SMoviePipelinePanel
\ No newline at end of file
diff --git a/Engine/Plugins/MovieScene/MovieRenderPipeline/Source/MovieRenderPipelineEditor/Private/Widgets/SMoviePipelinePanel.h b/Engine/Plugins/MovieScene/MovieRenderPipeline/Source/MovieRenderPipelineEditor/Private/Widgets/SMoviePipelinePanel.h
new file mode 100644
--- /dev/null
+++ b/Engine/Plugins/MovieScene/MovieRenderPipeline/Source/MovieRenderPipelineEditor/Private/Widgets/SMoviePipelinePanel.h
@@ -0,0 +1,56 @@
+// Copyright 1998-2019 Epic Games, Inc. All Rights Reserved.
+
+#pragma once
+
+#include "Widgets/DeclarativeSyntaxSupport.h"
+#include "Widgets/SCompoundWidget.h"
+#include "UObject/GCObject.h"
+
+class UMoviePipelineShotConfig;
+class SMoviePipelineEditor;
+
+/**
+ * Outermost widget that is used for setting up a new movie render pipeline config. Operates on a transient UMovieRenderShotConfig that is internally owned and maintained 
+ */
+class SMoviePipelinePanel : public SCompoundWidget, public FGCObject
+{
+public:
+
+	~SMoviePipelinePanel();
+
+	SLATE_BEGIN_ARGS(SMoviePipelinePanel)
+		: _BasePreset(nullptr)
+
+		{}
+
+		/*~ All following arguments are mutually-exclusive */
+		/*-------------------------------------------------*/
+		/** A preset asset to base the pipeline off */
+		SLATE_ARGUMENT(UMoviePipelineShotConfig*, BasePreset)
+
+	SLATE_END_ARGS()
+
+	void Construct(const FArguments& InArgs);
+	UMoviePipelineShotConfig* GetMoviePipeline() const;
+
+private:
+	virtual void AddReferencedObjects(FReferenceCollector& Collector) override;
+
+private:
+	/** Generate the widget that is visible in the Choose Preset dropdown. */
+	TSharedRef<SWidget> OnGeneratePresetsMenu();
+
+	/** Called when the user wishes to revert their changes to the current preset. */
+	FReply OnRevertChanges();
+
+	/** Allocates a transient preset so that the user can use the pipeline without saving it to an asset first. */
+	UMoviePipelineShotConfig* AllocateTransientPreset();
+
+
+private:
+	/** The transient preset that we use - kept alive by AddReferencedObjects */
+	UMoviePipelineShotConfig* TransientPreset;
+
+	/** The main movie pipeline editor widget */
+	TSharedPtr<SMoviePipelineEditor> MoviePipelineEditorWidget;
+};
\ No newline at end of file
diff --git a/Engine/Plugins/MovieScene/MovieRenderPipeline/Source/MovieRenderPipelineEditor/Private/Widgets/SMoviePipelineSettings.cpp b/Engine/Plugins/MovieScene/MovieRenderPipeline/Source/MovieRenderPipelineEditor/Private/Widgets/SMoviePipelineSettings.cpp
new file mode 100644
--- /dev/null
+++ b/Engine/Plugins/MovieScene/MovieRenderPipeline/Source/MovieRenderPipelineEditor/Private/Widgets/SMoviePipelineSettings.cpp
@@ -0,0 +1,517 @@
+// Copyright 1998-2019 Epic Games, Inc. All Rights Reserved.
+
+#include "Widgets/SMoviePipelineSettings.h"
+#include "MoviePipelineShotConfig.h"
+#include "MoviePipelineSetting.h"
+
+// Core includes
+#include "Algo/Sort.h"
+#include "UObject/UObjectIterator.h"
+
+// AssetRegistry includes
+#include "AssetRegistryModule.h"
+
+// Slate includes
+#include "Widgets/SBoxPanel.h"
+#include "Widgets/Colors/SColorBlock.h"
+// #include "Widgets/Colors/SColorPicker.h"
+#include "Widgets/Input/SButton.h"
+#include "Widgets/Input/SCheckBox.h"
+#include "Widgets/Images/SThrobber.h"
+#include "Widgets/Layout/SWidgetSwitcher.h"
+#include "Widgets/Views/STreeView.h"
+#include "Widgets/Input/SComboButton.h"
+#include "Framework/Commands/UICommandList.h"
+#include "Framework/Commands/GenericCommands.h"
+
+// UnrealEd includes
+#include "ScopedTransaction.h"
+
+#define LOCTEXT_NAMESPACE "SMoviePipelineSettings"
+
+struct FMoviePipelineSettingCategory;
+struct FMoviePipelineSettingTreeItem;
+
+struct IMoviePipelineSettingTreeItem : TSharedFromThis<IMoviePipelineSettingTreeItem>
+{
+	virtual ~IMoviePipelineSettingTreeItem() {}
+
+	virtual void Delete(UMoviePipelineShotConfig* Owner) = 0;
+
+	virtual TSharedPtr<FMoviePipelineSettingCategory> AsCategory() { return nullptr; }
+	virtual TSharedPtr<FMoviePipelineSettingTreeItem> AsSetting()   { return nullptr; }
+
+	virtual TSharedRef<SWidget> ConstructWidget(TWeakPtr<SMoviePipelineSettings> SettingsWidget) = 0;
+};
+
+struct FMoviePipelineSettingTreeItem : IMoviePipelineSettingTreeItem
+{
+	/** Weak pointer to the setting that this tree item represents */
+	TWeakObjectPtr<UMoviePipelineSetting> WeakSetting;
+
+	explicit FMoviePipelineSettingTreeItem(UMoviePipelineSetting* InSetting)
+		: WeakSetting(InSetting)
+	{}
+
+	FText GetLabel() const
+	{
+		UMoviePipelineSetting* Setting = WeakSetting.Get();
+		return Setting ? Setting->GetDisplayText() : FText();
+	}
+
+	virtual void Delete(UMoviePipelineShotConfig* Owner) override
+	{
+		UMoviePipelineSetting* Setting = WeakSetting.Get();
+		if (Setting)
+		{
+			Owner->RemoveSetting(Setting);
+		}
+	}
+
+	virtual TSharedPtr<FMoviePipelineSettingTreeItem> AsSetting()
+	{
+		return SharedThis(this);
+	}
+
+	virtual TSharedRef<SWidget> ConstructWidget(TWeakPtr<SMoviePipelineSettings> SettingsWidget) override
+	{
+		return SNew(SOverlay)
+
+		+SOverlay::Slot()
+		[
+		
+			SNew(SHorizontalBox)
+
+			+ SHorizontalBox::Slot()
+			.AutoWidth()
+			.VAlign(VAlign_Center)
+			[
+				SNew(SBox)
+				.WidthOverride(32)
+				.HeightOverride(32)
+		  		[
+					SNew(SCheckBox)
+					//.Style(FTakeRecorderStyle::Get(), "TakeRecorder.Source.Switch")
+					.IsFocusable(false)
+					.IsChecked(this, &FMoviePipelineSettingTreeItem::GetCheckState)
+					.OnCheckStateChanged(this, &FMoviePipelineSettingTreeItem::SetCheckState, SettingsWidget)
+				]
+			]
+
+			+ SHorizontalBox::Slot()
+			.Padding(8, 4, 8, 4)
+			.AutoWidth()
+			.VAlign(VAlign_Center)
+			.HAlign(HAlign_Center)
+			[
+				SNew(SBox)
+				.WidthOverride(32)
+				.HeightOverride(32)
+				[
+					SNew(SImage)
+					.Image(this, &FMoviePipelineSettingTreeItem::GetIcon)
+					.ColorAndOpacity(this, &FMoviePipelineSettingTreeItem::GetImageColorAndOpacity)
+				]
+			]
+
+			+ SHorizontalBox::Slot()
+			.AutoWidth()
+			.VAlign(VAlign_Center)
+			.Padding(8, 0)
+			[
+				SNew(STextBlock)
+				.Text(this, &FMoviePipelineSettingTreeItem::GetLabel)
+				//.TextStyle(FTakeRecorderStyle::Get(), "TakeRecorder.Source.Label")
+				.ColorAndOpacity(this, &FMoviePipelineSettingTreeItem::GetColorAndOpacity)
+			]
+
+			+ SHorizontalBox::Slot()
+			.FillWidth(1)
+			[
+				SNew(SBox)
+			]
+
+			+ SHorizontalBox::Slot()
+			.AutoWidth()
+			.VAlign(VAlign_Center)
+			.Padding(8, 0, 24, 0)
+			[
+				SNew(STextBlock)
+				.Text(this, &FMoviePipelineSettingTreeItem::GetDescription)
+				//.TextStyle(FTakeRecorderStyle::Get(), "TakeRecorder.Source.Label")
+				.Font(FCoreStyle::GetDefaultFontStyle("Regular", 8))
+				.ColorAndOpacity(this, &FMoviePipelineSettingTreeItem::GetColorAndOpacity)
+			]
+		]
+
+		+SOverlay::Slot()
+		[
+			SNew(SHorizontalBox)
+
+			+ SHorizontalBox::Slot()
+			.FillWidth(1)
+			[
+				SNew(SBox)
+			]
+
+			+ SHorizontalBox::Slot()
+			.AutoWidth()
+			.VAlign(VAlign_Center)
+			.Padding(0, 0, 0, 0)
+			[
+				SNew(SColorBlock)
+				.Size(FVector2D(6.0, 38.0))
+				.Color(FLinearColor::White) 
+			]
+
+		];
+
+	}
+
+private:
+
+	const FSlateBrush* GetIcon() const
+	{
+		UMoviePipelineSetting* Setting = WeakSetting.Get();
+		return Setting ? Setting->GetDisplayIcon() : nullptr;
+	}
+
+	FText GetDescription() const
+	{
+		UMoviePipelineSetting* Setting = WeakSetting.Get();
+		return Setting ? Setting->GetDescriptionText() : FText();
+	}
+
+	ECheckBoxState GetCheckState() const
+	{
+		UMoviePipelineSetting* Setting = WeakSetting.Get();
+		return Setting && Setting->bEnabled ? ECheckBoxState::Checked : ECheckBoxState::Unchecked;
+	}
+
+	void SetCheckState(const ECheckBoxState NewState, TWeakPtr<SMoviePipelineSettings> WeakSettingsWidget)
+	{
+		const bool bEnable = NewState == ECheckBoxState::Checked;
+
+		TSharedPtr<SMoviePipelineSettings>		SettingsWidget = WeakSettingsWidget.Pin();
+		UMoviePipelineSetting*            ThisSetting    = WeakSetting.Get();
+
+		if (ThisSetting && SettingsWidget.IsValid())
+		{
+			TArray<UMoviePipelineSetting*> SelectedSettings;
+			SettingsWidget->GetSelectedSettings(SelectedSettings);
+
+			FText TransactionFormat = bEnable
+				? LOCTEXT("EnableSetting", "Enable {0}|plural(one=Setting, other=Settings)")
+				: LOCTEXT("DisableSetting", "Disable {0}|plural(one=Setting, other=Settings)");
+
+			if (!SelectedSettings.Contains(ThisSetting))
+			{
+				FScopedTransaction Transaction(FText::Format(TransactionFormat, 1));
+
+				ThisSetting->Modify();
+				ThisSetting->bEnabled = bEnable;
+			}
+			else 
+			{
+				FScopedTransaction Transaction(FText::Format(TransactionFormat, SelectedSettings.Num()));
+
+				for (UMoviePipelineSetting* SelectedSetting : SelectedSettings)
+				{
+					SelectedSetting->Modify();
+					SelectedSetting->bEnabled = bEnable;
+				}
+			}
+		}
+	}
+
+	FSlateColor GetColorAndOpacity() const
+	{
+		UMoviePipelineSetting* Setting = WeakSetting.Get();
+		return Setting && Setting->bEnabled ? FSlateColor::UseForeground() : FSlateColor::UseSubduedForeground();
+	}
+
+	FSlateColor GetImageColorAndOpacity() const
+	{
+		UMoviePipelineSetting* Setting = WeakSetting.Get();
+		return Setting && Setting->bEnabled ? FLinearColor::White : FLinearColor::White.CopyWithNewOpacity(0.3f);
+	}
+};
+
+struct FMoviePipelineSettingCategory : IMoviePipelineSettingTreeItem
+{
+	/** The title of this category */
+	FText Category;
+
+	/** Sorted list of this category's children */
+	TArray<TSharedPtr<FMoviePipelineSettingTreeItem>> Children;
+
+	explicit FMoviePipelineSettingCategory(const FString& InCategory)
+		: Category(FText::FromString(InCategory))
+	{}
+
+	virtual void Delete(UMoviePipelineShotConfig* Owner) override
+	{
+		for (TSharedPtr<FMoviePipelineSettingTreeItem> Child : Children)
+		{
+			Child->Delete(Owner);
+		}
+	}
+
+	virtual TSharedPtr<FMoviePipelineSettingCategory> AsCategory()
+	{
+		return SharedThis(this);
+	}
+
+	virtual TSharedRef<SWidget> ConstructWidget(TWeakPtr<SMoviePipelineSettings> SettingsWidget) override
+	{
+		return SNew(SHorizontalBox)
+
+			+SHorizontalBox::Slot()
+			.AutoWidth()
+			.Padding(20, 4)
+			.VAlign(VAlign_Center)
+			[
+				SNew(STextBlock)
+				.Text(this, &FMoviePipelineSettingCategory::GetLabel)
+			];
+	}
+
+private:
+
+	FText GetLabel() const
+	{
+		return FText::Format(LOCTEXT("CategoryFormatString", "{0} ({1})"),
+			Category, Children.Num());
+	}
+};
+
+
+void SMoviePipelineSettings::Construct(const FArguments& InArgs)
+{
+	CachedSettingsSerialNumber = uint32(-1);
+
+	TreeView = SNew(STreeView<TSharedPtr<IMoviePipelineSettingTreeItem>>)
+		.TreeItemsSource(&RootNodes)
+		.OnSelectionChanged(InArgs._OnSelectionChanged)
+		.OnGenerateRow(this, &SMoviePipelineSettings::OnGenerateRow)
+		.OnGetChildren(this, &SMoviePipelineSettings::OnGetChildren);
+
+	CommandList = MakeShared<FUICommandList>();
+	CommandList->MapAction(
+		FGenericCommands::Get().Delete,
+		FExecuteAction::CreateSP(this, &SMoviePipelineSettings::OnDeleteSelected),
+		FCanExecuteAction()
+	);
+
+	ChildSlot
+	[
+		TreeView.ToSharedRef()
+	];
+}
+
+void SMoviePipelineSettings::Tick(const FGeometry& AllottedGeometry, const double InCurrentTime, const float InDeltaTime)
+{
+	UMoviePipelineShotConfig* ShotConfig = WeakShotConfig.Get();
+
+	// If we have a settings ptr, we expect its serial number to match our cached one, if not, we rebuild the tree
+	if (ShotConfig)
+	{
+		if (CachedSettingsSerialNumber != ShotConfig->GetSettingsSerialNumber())
+		{
+			ReconstructTree();
+		}
+	}
+	// The settings are no longer valid, so we expect our cached serial number to be -1. If not, we haven't reset the tree yet.
+	else if (CachedSettingsSerialNumber != uint32(-1))
+	{
+		ReconstructTree();
+	}
+}
+
+FReply SMoviePipelineSettings::OnKeyDown(const FGeometry& MyGeometry, const FKeyEvent& InKeyEvent)
+{
+	if (CommandList->ProcessCommandBindings(InKeyEvent))
+	{
+		return FReply::Handled();
+	}
+
+	return FReply::Unhandled();
+}
+
+void SMoviePipelineSettings::GetSelectedSettings(TArray<UMoviePipelineSetting*>& OutSettings) const
+{
+	TArray<TSharedPtr<IMoviePipelineSettingTreeItem>> SelectedItems;
+
+	TreeView->GetSelectedItems(SelectedItems);
+	for (TSharedPtr<IMoviePipelineSettingTreeItem> Item : SelectedItems)
+	{
+		TSharedPtr<FMoviePipelineSettingTreeItem> SettingItem = Item->AsSetting();
+		UMoviePipelineSetting* SettingPtr = SettingItem.IsValid() ? SettingItem->WeakSetting.Get() : nullptr;
+		if (SettingPtr)
+		{
+			OutSettings.Add(SettingPtr);
+		}
+	}
+}
+
+void SMoviePipelineSettings::SetShotConfigObject(UMoviePipelineShotConfig* InShotConfig)
+{
+	WeakShotConfig = InShotConfig;
+	ReconstructTree();
+}
+
+void SMoviePipelineSettings::ReconstructTree()
+{
+	UMoviePipelineShotConfig* ShotConfig = WeakShotConfig.Get();
+	if (!ShotConfig)
+	{
+		CachedSettingsSerialNumber = uint32(-1);
+		RootNodes.Reset();
+		return;
+	}
+
+	CachedSettingsSerialNumber = ShotConfig->GetSettingsSerialNumber();
+
+	TSortedMap<FString, TSharedPtr<FMoviePipelineSettingCategory>> RootCategories;
+	for (TSharedPtr<IMoviePipelineSettingTreeItem> RootItem : RootNodes)
+	{
+		TSharedPtr<FMoviePipelineSettingCategory> RootCategory = RootItem->AsCategory();
+		if (RootCategory.IsValid())
+		{
+			RootCategory->Children.Reset();
+			RootCategories.Add(RootCategory->Category.ToString(), RootCategory);
+		}
+	}
+
+	RootNodes.Reset();
+
+	// We attempt to re-use tree items in order to maintain selection states on them
+	TMap<FObjectKey, TSharedPtr<FMoviePipelineSettingTreeItem>> OldSettingToTreeItem;
+	Swap(SettingToTreeItem, OldSettingToTreeItem);
+
+
+	static const FName CategoryName = "Category";
+
+	for (UMoviePipelineSetting* Setting : ShotConfig->GetSettings())
+	{
+		if (!Setting)
+		{
+			continue;
+		}
+
+		// The category in the UI is taken from the class itself
+		FString Category = TEXT("Default"); // Source->GetCategoryText().ToString();
+		if (Category.IsEmpty())
+		{
+			Category = Setting->GetClass()->GetMetaData(CategoryName);
+		}
+
+		// Attempt to find an existing category node, creating one if necessary
+		TSharedPtr<FMoviePipelineSettingCategory> CategoryNode = RootCategories.FindRef(Category);
+		if (!CategoryNode.IsValid())
+		{
+			CategoryNode = RootCategories.Add(Category, MakeShared<FMoviePipelineSettingCategory>(Category));
+
+			TreeView->SetItemExpansion(CategoryNode, true);
+		}
+
+		// Attempt to find an existing setting item node from the previous data, creating one if necessary
+		FObjectKey SettingKey(Setting);
+		TSharedPtr<FMoviePipelineSettingTreeItem> SettingItem = SettingToTreeItem.FindRef(SettingKey);
+		if (!SettingItem.IsValid())
+		{
+			SettingItem = OldSettingToTreeItem.FindRef(Setting);
+			if (SettingItem.IsValid())
+			{
+				SettingToTreeItem.Add(SettingKey, SettingItem);
+			}
+			else
+			{
+				SettingItem = MakeShared<FMoviePipelineSettingTreeItem>(Setting);
+				SettingToTreeItem.Add(SettingKey, SettingItem);
+
+				TreeView->SetItemExpansion(SettingItem, true);
+			}
+		}
+
+		check(SettingItem.IsValid());
+		CategoryNode->Children.Add(SettingItem);
+	}
+
+
+	RootNodes.Reset(RootCategories.Num());
+	for(TTuple<FString, TSharedPtr<FMoviePipelineSettingCategory>>& Pair : RootCategories)
+	{
+		if (Pair.Value->Children.Num() == 0)
+		{
+			continue;
+		}
+
+		// Sort children by name. Work with a tuple of index and string to avoid excessively calling GetLabel().ToString()
+		TArray<TTuple<int32, FString>> SortData;
+		SortData.Reserve(Pair.Value->Children.Num());
+		for (TSharedPtr<FMoviePipelineSettingTreeItem> Item : Pair.Value->Children)
+		{
+			SortData.Add(MakeTuple(SortData.Num(), Item->GetLabel().ToString()));
+		}
+
+		auto SortPredicate = [](TTuple<int32, FString>& A, TTuple<int32, FString>& B)
+		{
+			return A.Get<1>() < B.Get<1>();
+		};
+
+		Algo::Sort(SortData, SortPredicate);
+
+		// Create a new sorted list of the child entries
+		TArray<TSharedPtr<FMoviePipelineSettingTreeItem>> NewChildren;
+		NewChildren.Reserve(SortData.Num());
+		for (const TTuple<int32, FString>& Item : SortData)
+		{
+			NewChildren.Add(Pair.Value->Children[Item.Get<0>()]);
+		}
+
+		Swap(Pair.Value->Children, NewChildren);
+
+		// Add the category
+		RootNodes.Add(Pair.Value);
+	}
+
+	TreeView->RequestTreeRefresh();
+}
+
+TSharedRef<ITableRow> SMoviePipelineSettings::OnGenerateRow(TSharedPtr<IMoviePipelineSettingTreeItem> Item, const TSharedRef<STableViewBase>& Tree)
+{
+	return
+		SNew(STableRow<TSharedPtr<IMoviePipelineSettingTreeItem>>, Tree)
+		[
+			Item->ConstructWidget(SharedThis(this))
+		];
+}
+
+void SMoviePipelineSettings::OnGetChildren(TSharedPtr<IMoviePipelineSettingTreeItem> Item, TArray<TSharedPtr<IMoviePipelineSettingTreeItem>>& OutChildItems)
+{
+	TSharedPtr<FMoviePipelineSettingCategory> Category = Item->AsCategory();
+	if (Category.IsValid())
+	{
+		OutChildItems.Append(Category->Children);
+	}
+}
+
+void SMoviePipelineSettings::OnDeleteSelected()
+{
+	UMoviePipelineShotConfig* ShotConfig = WeakShotConfig.Get();
+	if (ShotConfig)
+	{
+		TArray<TSharedPtr<IMoviePipelineSettingTreeItem>> Items = TreeView->GetSelectedItems();
+
+		FScopedTransaction Transaction(FText::Format(LOCTEXT("DeleteSelection", "Delete Selected {0}|plural(one=Setting, other=Settings)"), Items.Num()));
+		ShotConfig->Modify();
+
+		for (TSharedPtr<IMoviePipelineSettingTreeItem> Item : Items)
+		{
+			Item->Delete(ShotConfig);
+		}
+	}
+}
+
+#undef LOCTEXT_NAMESPACE
\ No newline at end of file
diff --git a/Engine/Plugins/MovieScene/MovieRenderPipeline/Source/MovieRenderPipelineEditor/Private/Widgets/SMoviePipelineSettings.h b/Engine/Plugins/MovieScene/MovieRenderPipeline/Source/MovieRenderPipelineEditor/Private/Widgets/SMoviePipelineSettings.h
new file mode 100644
--- /dev/null
+++ b/Engine/Plugins/MovieScene/MovieRenderPipeline/Source/MovieRenderPipelineEditor/Private/Widgets/SMoviePipelineSettings.h
@@ -0,0 +1,67 @@
+// Copyright 1998-2019 Epic Games, Inc. All Rights Reserved.
+
+#pragma once
+
+#include "Containers/Array.h"
+#include "Templates/SharedPointer.h"
+#include "Widgets/DeclarativeSyntaxSupport.h"
+#include "Widgets/SCompoundWidget.h"
+#include "UObject/ObjectKey.h"
+#include "Styling/SlateTypes.h" 
+
+class ITableRow;
+class FUICommandList;
+class STableViewBase;
+class UMoviePipelineSetting;
+class UMoviePipelineShotConfig;
+struct IMoviePipelineSettingTreeItem;
+struct FMoviePipelineSettingTreeItem;
+template<typename> class STreeView;
+
+DECLARE_DELEGATE_TwoParams(
+	FOnSettingSelectionChanged,
+	TSharedPtr<IMoviePipelineSettingTreeItem>,
+	ESelectInfo::Type
+)
+
+/** Main widget for the Movie Render Shot Config panel */
+class SMoviePipelineSettings : public SCompoundWidget
+{
+public:
+
+	SLATE_BEGIN_ARGS(SMoviePipelineSettings){}
+		SLATE_EVENT(FOnSettingSelectionChanged, OnSelectionChanged)
+	SLATE_END_ARGS()
+
+	void Construct(const FArguments& InArgs);
+
+	void SetShotConfigObject(UMoviePipelineShotConfig* InShotConfig);
+
+	void GetSelectedSettings(TArray<UMoviePipelineSetting*>& OutSettings) const;
+
+private:
+
+	virtual FReply OnKeyDown(const FGeometry& MyGeometry, const FKeyEvent& InKeyEvent) override;
+	virtual void Tick(const FGeometry& AllottedGeometry, const double InCurrentTime, const float InDeltaTime) override;
+
+private:
+
+	void ReconstructTree();
+
+	TSharedRef<ITableRow> OnGenerateRow(TSharedPtr<IMoviePipelineSettingTreeItem> Item, const TSharedRef<STableViewBase>& Tree);
+
+	void OnGetChildren(TSharedPtr<IMoviePipelineSettingTreeItem> Item, TArray<TSharedPtr<IMoviePipelineSettingTreeItem>>& OutChildItems);
+
+	void OnDeleteSelected();
+
+private:
+
+	uint32 CachedSettingsSerialNumber;
+	TWeakObjectPtr<UMoviePipelineShotConfig> WeakShotConfig;
+
+	TArray<TSharedPtr<IMoviePipelineSettingTreeItem>> RootNodes;
+	TMap<FObjectKey, TSharedPtr<FMoviePipelineSettingTreeItem>> SettingToTreeItem;
+	TSharedPtr<STreeView<TSharedPtr<IMoviePipelineSettingTreeItem>>> TreeView;
+
+	TSharedPtr<FUICommandList> CommandList;
+};
\ No newline at end of file
