Initial merge of the new Sequencer Movie Render system. This provides some basic ways to output movies but is mostly scripting based without UI still.

diff --git a/Engine/Plugins/MovieScene/MovieRenderPipeline/Source/MovieRenderPipelineEditor/Private/Widgets/SMoviePipelineTabContent.cpp b/Engine/Plugins/MovieScene/MovieRenderPipeline/Source/MovieRenderPipelineEditor/Private/Widgets/SMoviePipelineTabContent.cpp
new file mode 100644
--- /dev/null
+++ b/Engine/Plugins/MovieScene/MovieRenderPipeline/Source/MovieRenderPipelineEditor/Private/Widgets/SMoviePipelineTabContent.cpp
@@ -0,0 +1,51 @@
+// Copyright 1998-2019 Epic Games, Inc. All Rights Reserved.
+
+#include "Widgets/SMoviePipelineTabContent.h"
+#include "Widgets/SMoviePipelinePanel.h"
+
+
+// Analytics
+#include "EngineAnalytics.h"
+#include "Interfaces/IAnalyticsProvider.h"
+
+#define LOCTEXT_NAMESPACE "SMoviePipelineTabContent"
+
+
+void SMoviePipelineTabContent::Construct(const FArguments& InArgs)
+{    
+	// Delay one tick before opening the default pipeline setup panel.
+	// this allows anything that just invoked the tab to customize it without the default UI being created
+	RegisterActiveTimer(0.f, FWidgetActiveTimerDelegate::CreateSP(this, &SMoviePipelineTabContent::OnActiveTimer));
+     
+	if (FEngineAnalytics::IsAvailable())
+	{
+		FEngineAnalytics::GetProvider().RecordEvent(TEXT("MoviePipeline.PanelOpened"));
+	}
+}
+
+EActiveTimerReturnType SMoviePipelineTabContent::OnActiveTimer(double InCurrentTime, float InDeltaTime)
+{
+	SetupForPipeline((UMoviePipelineShotConfig*)nullptr);
+	return EActiveTimerReturnType::Stop;
+}
+
+void SMoviePipelineTabContent::SetupForPipeline(UMoviePipelineShotConfig* BasePreset)
+{
+	// Null out the tab content to ensure that all references have been cleaned up before constructing the new one
+	ChildSlot [ SNullWidget::NullWidget ];
+
+	ChildSlot
+	.HAlign(HAlign_Fill)
+	.VAlign(VAlign_Fill)
+	[
+		SAssignNew(WeakPanel, SMoviePipelinePanel)
+		.BasePreset(BasePreset)
+	];
+
+	if (FEngineAnalytics::IsAvailable())
+	{
+		FEngineAnalytics::GetProvider().RecordEvent(TEXT("MoviePipeline.SetupForPipelineFromPreset"));
+	}
+}
+
+#undef LOCTEXT_NAMESPACE // SMoviePipelineTabContent
diff --git a/Engine/Plugins/MovieScene/MovieRenderPipeline/Source/MovieRenderPipelineEditor/Private/Widgets/SMoviePipelineTabContent.h b/Engine/Plugins/MovieScene/MovieRenderPipeline/Source/MovieRenderPipelineEditor/Private/Widgets/SMoviePipelineTabContent.h
new file mode 100644
--- /dev/null
+++ b/Engine/Plugins/MovieScene/MovieRenderPipeline/Source/MovieRenderPipelineEditor/Private/Widgets/SMoviePipelineTabContent.h
@@ -0,0 +1,28 @@
+// Copyright 1998-2019 Epic Games, Inc. All Rights Reserved.
+
+#pragma once
+
+#include "Widgets/DeclarativeSyntaxSupport.h"
+#include "Widgets/SCompoundWidget.h"
+
+class UMoviePipelineShotConfig;
+class SMoviePipelinePanel;
+
+class SMoviePipelineTabContent : public SCompoundWidget
+{
+public:
+
+	SLATE_BEGIN_ARGS(SMoviePipelineTabContent){}
+	SLATE_END_ARGS()
+
+	void Construct(const FArguments& InArgs);
+
+	/** Set this tab up for a specific Pipeline preset. Passing nullptr will initialize a default pipeline. */
+	void SetupForPipeline(UMoviePipelineShotConfig* BasePreset);
+
+private:
+	EActiveTimerReturnType OnActiveTimer(double InCurrentTime, float InDeltaTime);
+
+private:
+	TWeakPtr<SMoviePipelinePanel> WeakPanel;
+};
diff --git a/Engine/Plugins/MovieScene/MovieRenderPipeline/Source/MovieRenderPipelineEditor/Public/IMovieRenderPipelineEditorModule.h b/Engine/Plugins/MovieScene/MovieRenderPipeline/Source/MovieRenderPipelineEditor/Public/IMovieRenderPipelineEditorModule.h
new file mode 100644
--- /dev/null
+++ b/Engine/Plugins/MovieScene/MovieRenderPipeline/Source/MovieRenderPipelineEditor/Public/IMovieRenderPipelineEditorModule.h
@@ -0,0 +1,19 @@
+// Copyright 1998-2019 Epic Games, Inc. All Rights Reserved.
+
+#pragma once
+
+#include "CoreMinimal.h"
+#include "Modules/ModuleManager.h"
+
+/**
+ * Public module interface for the Movie Render Pipeline module
+ */
+class IMovieRenderPipelineEditorModule : public IModuleInterface
+{
+public:
+	/** The tab name for the movie render pipeline tab */
+	static FName MovieRenderPipelineTabName;
+
+	/** The default label for the movie render pipeline tab. */
+	static FText MovieRenderPipelineTabLabel;
+};
diff --git a/Engine/Plugins/MovieScene/MovieRenderPipeline/Source/MovieRenderPipelineEditor/Public/MoviePipelineEditorBlueprintLibrary.h b/Engine/Plugins/MovieScene/MovieRenderPipeline/Source/MovieRenderPipelineEditor/Public/MoviePipelineEditorBlueprintLibrary.h
new file mode 100644
--- /dev/null
+++ b/Engine/Plugins/MovieScene/MovieRenderPipeline/Source/MovieRenderPipelineEditor/Public/MoviePipelineEditorBlueprintLibrary.h
@@ -0,0 +1,25 @@
+// Copyright 1998-2019 Epic Games, Inc. All Rights Reserved.
+
+#pragma once
+
+#include "Kismet/BlueprintFunctionLibrary.h"
+#include "UObject/ObjectMacros.h"
+#include "UObject/TextProperty.h"
+#include "Internationalization/Text.h"
+#include "Containers/UnrealString.h"
+#include "MovieRenderPipelineConfig.h"
+
+#include "MoviePipelineEditorBlueprintLibrary.generated.h"
+
+// Forward Declare
+class UMovieRenderPipelineConfig;
+
+UCLASS()
+class MOVIERENDERPIPELINEEDITOR_API UMoviePipelineEditorBlueprintLibrary : public UBlueprintFunctionLibrary
+{
+	GENERATED_BODY()
+	
+public:
+	UFUNCTION(BlueprintCallable, Category = "Movie Render Pipeline")
+	static bool ExportConfigToAsset(const UMovieRenderPipelineConfig* InConfig, const FString& InPackagePath, const FString& InFileName, const bool bInSaveAsset, UMovieRenderPipelineConfig*& OutAsset, FText& OutErrorReason);
+};
\ No newline at end of file
diff --git a/Engine/Plugins/MovieScene/MovieRenderPipeline/Source/MovieRenderPipelineEditor/Public/MoviePipelineNewProcessExecutor.h b/Engine/Plugins/MovieScene/MovieRenderPipeline/Source/MovieRenderPipelineEditor/Public/MoviePipelineNewProcessExecutor.h
new file mode 100644
--- /dev/null
+++ b/Engine/Plugins/MovieScene/MovieRenderPipeline/Source/MovieRenderPipelineEditor/Public/MoviePipelineNewProcessExecutor.h
@@ -0,0 +1,24 @@
+// Copyright 1998-2019 Epic Games, Inc. All Rights Reserved.
+#pragma once
+
+#include "MoviePipelineLinearExecutor.h"
+#include "MoviePipelineNewProcessExecutor.generated.h"
+
+/**
+* This is the implementation responsible for executing the rendering of
+* multiple movie pipelines on the local machine in an external process.
+* This simply handles launching and managing the external processes and 
+* acts as a proxy to them where possible. This internally uses the
+* UMoviePipelineInProcessExecutor on the launched instances.
+*/
+UCLASS(Blueprintable)
+class MOVIERENDERPIPELINEEDITOR_API UMoviePipelineNewProcessExecutor : public UMoviePipelineLinearExecutorBase
+{
+	GENERATED_BODY()
+	
+public:
+	UMoviePipelineNewProcessExecutor()
+		: UMoviePipelineLinearExecutorBase()
+	{
+	}
+};
\ No newline at end of file
diff --git a/Engine/Plugins/MovieScene/MovieRenderPipeline/Source/MovieRenderPipelineEditor/Public/MoviePipelinePIEExecutor.h b/Engine/Plugins/MovieScene/MovieRenderPipeline/Source/MovieRenderPipelineEditor/Public/MoviePipelinePIEExecutor.h
new file mode 100644
--- /dev/null
+++ b/Engine/Plugins/MovieScene/MovieRenderPipeline/Source/MovieRenderPipelineEditor/Public/MoviePipelinePIEExecutor.h
@@ -0,0 +1,49 @@
+// Copyright 1998-2019 Epic Games, Inc. All Rights Reserved.
+#pragma once
+
+#include "MoviePipelineLinearExecutor.h"
+#include "MoviePipelinePIEExecutor.generated.h"
+
+class UMoviePipeline;
+
+/**
+* This is the implementation responsible for executing the rendering of
+* multiple movie pipelines in the currently running Editor process. This
+* involves launching a Play in Editor session for each Movie Pipeline to
+* process.
+*/
+UCLASS(Blueprintable)
+class MOVIERENDERPIPELINEEDITOR_API UMoviePipelinePIEExecutor : public UMoviePipelineLinearExecutorBase
+{
+	GENERATED_BODY()
+	
+public:
+	UMoviePipelinePIEExecutor()
+		: UMoviePipelineLinearExecutorBase()
+	{
+	}
+
+protected:
+	virtual void Start(UMovieRenderPipelineConfig* InConfig, const int32 InConfigIndex, const int32 InNumConfigs) override;
+
+private:
+	/** Called when PIE finishes booting up and it is safe for us to spawn an object into that world. */
+	void OnPIEStartupFinished(bool);
+	/** Called when the instance of the pipeline in the PIE world has finished. */
+	void OnPIEMoviePipelineFinished(UMoviePipeline* InMoviePipeline);
+	/** Called a short period of time after OnPIEMoviePipelineFinished to allow Editor the time to fully close PIE before we make a new request. */
+	void DelayedFinishNotification();
+
+	/** Gets the title for the Window. Used to show progress in the title. */
+	FText GetWindowTitle(const int32 InConfigIndex, const int32 InNumConfigs) const;
+
+
+private:
+	/** Instance of the Pipeline that exists in the world that is currently processing (if any) */
+	UPROPERTY(Transient)
+	UMoviePipeline* ActiveMoviePipeline;
+
+	/** Instance of the Config we're supposed to be working on */
+	UPROPERTY(Transient)
+	UMovieRenderPipelineConfig* ActiveConfig;
+};
\ No newline at end of file
diff --git a/Engine/Plugins/MovieScene/MovieRenderPipeline/Source/MovieRenderPipelineRenderPasses/Private/MoviePipelineBackbufferPass.cpp b/Engine/Plugins/MovieScene/MovieRenderPipeline/Source/MovieRenderPipelineRenderPasses/Private/MoviePipelineBackbufferPass.cpp
new file mode 100644
--- /dev/null
+++ b/Engine/Plugins/MovieScene/MovieRenderPipeline/Source/MovieRenderPipelineRenderPasses/Private/MoviePipelineBackbufferPass.cpp
@@ -0,0 +1,336 @@
+// Copyright 1998-2019 Epic Games, Inc. All Rights Reserved.
+#pragma once
+
+#include "MoviePipelineBackbufferPass.h"
+#include "Slate/SceneViewport.h"
+#include "Engine/World.h"
+#include "Engine/GameViewportClient.h"
+#include "MovieRenderPipelineDataTypes.h"
+#include "CanvasTypes.h"
+#include "Engine/TextureRenderTarget2D.h"
+#include "SceneView.h"
+#include "LegacyScreenPercentageDriver.h"
+#include "RendererInterface.h"
+#include "EngineModule.h"
+#include "ImagePixelData.h"
+#include "RHICommandList.h"
+#include "ImageWriteStream.h"
+#include "GameFramework/PlayerController.h"
+#include "Math/Vector.h"
+#include "MovieRenderPipelineConfig.h"
+#include "Misc/FrameRate.h"
+#include "MoviePipelineAccumulationSetting.h"
+#include "MoviePipelineRenderPass.h"
+#include "MoviePipelineShotConfig.h"
+#include "MoviePipelineSetting.h"
+#include "MoviePipeline.h"
+
+
+void UMoviePipelineBackbufferPass::SetupImpl(const FMoviePipelineRenderPassInitSettings& InInitSettings, TSharedRef<FImagePixelPipe, ESPMode::ThreadSafe> InOutputPipe)
+{
+	OutputPipe = InOutputPipe;
+
+	TileRenderTarget = NewObject<UTextureRenderTarget2D>(this);
+	TileRenderTarget->ClearColor = FLinearColor(0.0f, 0.0f, 0.0f, 0.0f);
+
+	// Initialize to the tile size (not final size) and use a 16 bit back buffer to avoid precision issues
+			// when accumulating later.
+
+	int32 TargetSizeX = InInitSettings.TileResolution.X;
+	int32 TargetSizeY = InInitSettings.TileResolution.Y;
+/*
+	*/
+	TileRenderTarget->InitCustomFormat(TargetSizeX, TargetSizeY, EPixelFormat::PF_A16B16G16R16, false);
+
+	// Create a Render Target to hold our back buffer.
+	ViewState.Allocate();
+}
+
+
+void UMoviePipelineBackbufferPass::TeardownImpl()
+{
+	ViewState.Destroy();
+}
+
+void UMoviePipelineBackbufferPass::GatherOutputPassesImpl(TArray<FMoviePipelinePassIdentifier>& ExpectedRenderPasses)
+{
+	ExpectedRenderPasses.Add(FName("Backbuffer"));
+}
+
+
+void UMoviePipelineBackbufferPass::CaptureFrameImpl(const FMoviePipelineRenderPassMetrics& InPassMetrics)
+{
+	float TimeSeconds = InPassMetrics.OutputState.WorldSeconds;
+	float RealTimeSeconds = InPassMetrics.OutputState.WorldSeconds;
+	float DeltaTimeSeconds = InPassMetrics.OutputState.FrameDeltaTime;
+
+	// Reuse the common RenderTarget
+	UTextureRenderTarget2D* RenderTarget = TileRenderTarget;
+
+	FEngineShowFlags ShowFlags(EShowFlagInitMode::ESFIM_Game);
+	if (InPassMetrics.bIsUsingOverlappedTiles)
+	{
+		// Disable these for now
+		//ShowFlags.SetDepthOfField(false);
+		ShowFlags.SetMotionBlur(false);
+	}
+
+	FSceneViewFamilyContext ViewFamily(FSceneViewFamily::ConstructionValues(
+		RenderTarget->GameThread_GetRenderTargetResource(),
+		GetWorld()->Scene,
+		ShowFlags)
+		.SetWorldTimes(TimeSeconds, DeltaTimeSeconds, RealTimeSeconds)
+		.SetRealtimeUpdate(true));
+
+	ViewFamily.SetScreenPercentageInterface(new FLegacyScreenPercentageDriver(ViewFamily, 1.0f, false));
+	// View is added as a child of the ViewFamily
+	FSceneView* View = CalcSceneView(&ViewFamily, InPassMetrics);
+	ViewFamily.FrameNumber += InPassMetrics.SpatialJitterIndex;
+	
+	// Override the Motion Blur settings since these are controlled by the movie pipeline.
+	{
+		View->FinalPostProcessSettings.MotionBlurTargetFPS = FMath::RoundToInt(GetPipeline()->GetPipelineConfig()->GetEffectiveFrameRate().AsDecimal());
+		View->FinalPostProcessSettings.MotionBlurAmount = InPassMetrics.OutputState.MotionBlurFraction;
+		View->FinalPostProcessSettings.MotionBlurMax = 100.f;
+		View->FinalPostProcessSettings.bOverride_MotionBlurAmount = true;
+		View->FinalPostProcessSettings.bOverride_MotionBlurTargetFPS = true;
+		View->FinalPostProcessSettings.bOverride_MotionBlurMax = true;
+	}
+
+	// Disable TAA while we're working on pixel accuracy testing
+	{
+		View->AntiAliasingMethod = EAntiAliasingMethod::AAM_None;
+	}
+
+	// Override the view matrix with our offsets provided by tiling/jittering. This requires
+	// TAA to be disabled.
+	{
+		View->ViewMatrices.HackAddTemporalAAProjectionJitter(InPassMetrics.SpatialShift);
+	}
+
+	// Bias all mip-mapping to pretend to be working at our target resolution and not our tile resolution
+	// so that the images don't end up soft.
+	{
+		float EffectivePrimaryResolutionFraction = 1.f / InitSettings.TileCount;
+		View->MaterialTextureMipBias = FMath::Log2(EffectivePrimaryResolutionFraction);
+
+		// Add an additional bias per user settings. This allows them to choose to make the textures sharper if it
+		// looks better with their particular settings.
+		UMoviePipelineAccumulationSetting* AccumulationSettings = InitSettings.ShotConfig->FindOrAddSetting<UMoviePipelineAccumulationSetting>();
+		View->MaterialTextureMipBias += AccumulationSettings->TextureSharpnessBias;
+	}
+
+	// If this is a history only frame it means we just came from a camera cut, flag it as such
+	// which will have knock-on effects in histories (such as clearing old TAA buffers).
+	if (InPassMetrics.bIsHistoryOnlyFrame)
+	{
+		View->bCameraCut = true;
+	}
+	
+	// Pause the view family on renders after the first one. This prevents motion blur/view matrices from being
+	// updated
+	ViewFamily.bWorldIsPaused = InPassMetrics.SpatialJitterIndex != InPassMetrics.NumSpatialJitters - 1;
+
+	FCanvas Canvas = FCanvas(RenderTarget->GameThread_GetRenderTargetResource(), nullptr, GetWorld(), ERHIFeatureLevel::SM5, FCanvas::CDM_DeferDrawing, 1.0f);
+
+
+	// Draw the world into this View Family
+	GetRendererModule().BeginRenderingViewFamily(&Canvas, &ViewFamily);
+
+	FRenderTarget* CanvasRenderTarget = Canvas.GetRenderTarget();
+	TSharedPtr<FImagePixelPipe, ESPMode::ThreadSafe> LocalOutputPipe = OutputPipe;
+	
+	ENQUEUE_RENDER_COMMAND(CanvasRenderTargetResolveCommand)(
+		[CanvasRenderTarget, InPassMetrics, LocalOutputPipe](FRHICommandListImmediate& RHICmdList)
+	{
+		// If they only needed this frame for the history, we don't actually care about trying to output it,
+		// so we can skip the copy and pushing to output.
+		if (InPassMetrics.bIsHistoryOnlyFrame)
+		{
+			return;
+		}
+
+		FReadSurfaceDataFlags ReadDataFlags;
+		ReadDataFlags.SetLinearToGamma(false);
+
+		FIntRect SourceRect = FIntRect(0, 0, CanvasRenderTarget->GetSizeXY().X, CanvasRenderTarget->GetSizeXY().Y);
+
+		// Build our metadata to associate with this frame.
+		TSharedRef<FImagePixelDataPayload, ESPMode::ThreadSafe> FrameData = MakeShared<FImagePixelDataPayload, ESPMode::ThreadSafe>();
+		FrameData->OutputState = InPassMetrics.OutputState;
+		FrameData->PassName = TEXT("Backbuffer");
+		FrameData->TileIndexX = InPassMetrics.TileIndexX;
+		FrameData->TileIndexY = InPassMetrics.TileIndexY;
+		FrameData->TileSizeX = SourceRect.Width();
+		FrameData->TileSizeY = SourceRect.Height();
+		FrameData->NumTilesX = InPassMetrics.NumTilesX;
+		FrameData->NumTilesY = InPassMetrics.NumTilesY;
+		FrameData->SpatialShift = InPassMetrics.SpatialShift;
+		FrameData->SpatialJitterIndex = InPassMetrics.SpatialJitterIndex;
+		FrameData->NumSpatialJitters = InPassMetrics.NumSpatialJitters;
+		FrameData->TemporalJitterIndex = InPassMetrics.TemporalJitterIndex;
+		FrameData->NumTemporalJitters = InPassMetrics.NumTemporalJitters;
+		FrameData->JitterOffsetX = InPassMetrics.JitterOffsetX;
+		FrameData->JitterOffsetY = InPassMetrics.JitterOffsetY;
+		FrameData->bIsUsingOverlappedTiles = InPassMetrics.bIsUsingOverlappedTiles;
+
+		FrameData->AccumulationGamma = InPassMetrics.AccumulationGamma;
+
+		FrameData->OverlappedOffsetX = InPassMetrics.OverlappedOffsetX;
+		FrameData->OverlappedOffsetY = InPassMetrics.OverlappedOffsetY;
+		FrameData->OverlappedSizeX = InPassMetrics.OverlappedSizeX;
+		FrameData->OverlappedSizeY = InPassMetrics.OverlappedSizeY;
+		FrameData->OverlappedPadX = InPassMetrics.OverlappedPadX;
+		FrameData->OverlappedPadY = InPassMetrics.OverlappedPadY;
+		FrameData->OverlappedSubpixelShift = InPassMetrics.OverlappedSubpixelShift;
+
+		TUniquePtr<TImagePixelData<FColor>> PixelData = MakeUnique<TImagePixelData<FColor>>(SourceRect.Size(), FrameData);
+		RHICmdList.ReadSurfaceData(CanvasRenderTarget->GetRenderTargetTexture(), SourceRect, PixelData->Pixels, ReadDataFlags);
+
+		check(PixelData->IsDataWellFormed());
+
+		const bool bIsOutputFrame = true;
+		if (bIsOutputFrame)
+		{
+			LocalOutputPipe->Push(MoveTemp(PixelData));
+		}
+	}
+	);
+}
+
+void UMoviePipelineBackbufferPass::GetFrameDataImpl(MoviePipeline::FOutputFrameData& OutFrameData)
+{
+
+}
+
+FSceneView* UMoviePipelineBackbufferPass::CalcSceneView(FSceneViewFamily* ViewFamily, const FMoviePipelineRenderPassMetrics& InPassMetrics)
+{
+	FVector ViewLocation;
+	FRotator ViewRotation;
+
+	APlayerController* LocalPlayerController = GetWorld()->GetFirstPlayerController();
+	LocalPlayerController->GetPlayerViewPoint(ViewLocation, ViewRotation);
+
+	UE_LOG(LogTemp, Log, TEXT("%d,%s"), GFrameCounter, *ViewLocation.ToString());
+
+	int32 TileSizeX = InitSettings.TileResolution.X;
+	int32 TileSizeY = InitSettings.TileResolution.Y;
+
+	FSceneViewInitOptions ViewInitOptions;
+	ViewInitOptions.ViewFamily = ViewFamily;
+	ViewInitOptions.ViewOrigin = ViewLocation;
+	ViewInitOptions.SetViewRectangle(FIntRect(FIntPoint(0, 0), FIntPoint(TileSizeX, TileSizeY)));
+	ViewInitOptions.ViewRotationMatrix = FInverseRotationMatrix(ViewRotation);
+
+	// Rotate the view 90 degrees (reason: unknown)
+	ViewInitOptions.ViewRotationMatrix = ViewInitOptions.ViewRotationMatrix * FMatrix(
+		FPlane(0, 0, 1, 0),
+		FPlane(1, 0, 0, 0),
+		FPlane(0, 1, 0, 0),
+		FPlane(0, 0, 0, 1));
+	float ViewFOV = 90.f;
+	if (GetWorld()->GetFirstPlayerController()->PlayerCameraManager)
+	{
+		ViewFOV = GetWorld()->GetFirstPlayerController()->PlayerCameraManager->GetFOVAngle();
+	}
+
+	// Calculate a Projection Matrix
+	{
+		float XAxisMultiplier;
+		float YAxisMultiplier;
+
+		if (ViewInitOptions.GetViewRect().Width() > ViewInitOptions.GetViewRect().Height())
+		{
+			//if the viewport is wider than it is tall
+			XAxisMultiplier = 1.0f;
+			YAxisMultiplier = ViewInitOptions.GetViewRect().Width() / (float)ViewInitOptions.GetViewRect().Height();
+		}
+		else
+		{
+			//if the viewport is taller than it is wide
+			XAxisMultiplier = ViewInitOptions.GetViewRect().Height() / (float)ViewInitOptions.GetViewRect().Width();
+			YAxisMultiplier = 1.0f;
+		}
+
+		const float MinZ = GNearClippingPlane;
+		const float MaxZ = MinZ;
+		// Avoid zero ViewFOV's which cause divide by zero's in projection matrix
+		const float MatrixFOV = FMath::Max(0.001f, ViewFOV) * (float)PI / 360.0f;
+
+		FMatrix BaseProjMatrix;
+
+		if ((bool)ERHIZBuffer::IsInverted)
+		{
+			BaseProjMatrix = FReversedZPerspectiveMatrix(
+				MatrixFOV,
+				MatrixFOV,
+				XAxisMultiplier,
+				YAxisMultiplier,
+				MinZ,
+				MaxZ
+			);
+		}
+		else
+		{
+			BaseProjMatrix = FPerspectiveMatrix(
+				MatrixFOV,
+				MatrixFOV,
+				XAxisMultiplier,
+				YAxisMultiplier,
+				MinZ,
+				MaxZ
+			);
+		}
+
+		if (InPassMetrics.bIsUsingOverlappedTiles)
+		{
+			float PadRatioX = float(InPassMetrics.OverlappedPadX * 2 + InPassMetrics.OverlappedSizeX) / float (InPassMetrics.OverlappedSizeX);
+			float PadRatioY = float(InPassMetrics.OverlappedPadY * 2 + InPassMetrics.OverlappedSizeY) / float (InPassMetrics.OverlappedSizeY);
+
+			float ScaleX = PadRatioX / float(InitSettings.TileCount);
+			float ScaleY = PadRatioY / float(InitSettings.TileCount);
+
+			BaseProjMatrix.M[0][0] /= ScaleX;
+			BaseProjMatrix.M[1][1] /= ScaleY;
+
+			// this offset would be correct with no pad
+			float OffsetX = -((float(InPassMetrics.TileIndexX) + 0.5f - float(InitSettings.TileCount) / 2.0f)* 2.0f);
+			float OffsetY =  ((float(InPassMetrics.TileIndexY) + 0.5f - float(InitSettings.TileCount) / 2.0f)* 2.0f);
+
+			BaseProjMatrix.M[2][0] += OffsetX / PadRatioX;
+			BaseProjMatrix.M[2][1] += OffsetY / PadRatioX;
+		}
+
+		ViewInitOptions.ProjectionMatrix = BaseProjMatrix;
+	}
+
+	ViewInitOptions.SceneViewStateInterface = ViewState.GetReference();
+	ViewInitOptions.FOV = ViewFOV;
+
+	FSceneView* View = new FSceneView(ViewInitOptions);
+	ViewFamily->Views.Add(View);
+	View->ViewLocation = ViewLocation;
+	View->ViewRotation = ViewRotation;
+	View->StartFinalPostprocessSettings(View->ViewLocation);
+
+	// CameraAnim override
+	if (LocalPlayerController->PlayerCameraManager)
+	{
+		TArray<FPostProcessSettings> const* CameraAnimPPSettings;
+		TArray<float> const* CameraAnimPPBlendWeights;
+		LocalPlayerController->PlayerCameraManager->GetCachedPostProcessBlends(CameraAnimPPSettings, CameraAnimPPBlendWeights);
+
+		FMinimalViewInfo ViewInfo = LocalPlayerController->PlayerCameraManager->GetCameraCachePOV();
+		for (int32 PPIdx = 0; PPIdx < CameraAnimPPBlendWeights->Num(); ++PPIdx)
+		{
+			View->OverridePostProcessSettings((*CameraAnimPPSettings)[PPIdx], (*CameraAnimPPBlendWeights)[PPIdx]);
+		}
+
+		View->OverridePostProcessSettings(ViewInfo.PostProcessSettings, ViewInfo.PostProcessBlendWeight);
+
+	}
+
+	View->EndFinalPostprocessSettings(ViewInitOptions);
+
+	return View;
+}
\ No newline at end of file
diff --git a/Engine/Plugins/MovieScene/MovieRenderPipeline/Source/MovieRenderPipelineRenderPasses/Private/MoviePipelineImageSequenceContainer.cpp b/Engine/Plugins/MovieScene/MovieRenderPipeline/Source/MovieRenderPipelineRenderPasses/Private/MoviePipelineImageSequenceContainer.cpp
new file mode 100644
--- /dev/null
+++ b/Engine/Plugins/MovieScene/MovieRenderPipeline/Source/MovieRenderPipelineRenderPasses/Private/MoviePipelineImageSequenceContainer.cpp
@@ -0,0 +1,210 @@
+// Copyright 1998-2019 Epic Games, Inc. All Rights Reserved.
+#pragma once
+
+#include "MoviePipelineImageSequenceContainer.h"
+#include "ImageWriteTask.h"
+#include "ImagePixelData.h"
+#include "Modules/ModuleManager.h"
+#include "ImageWriteQueue.h"
+#include "MoviePipeline.h"
+#include "ImageWriteStream.h"
+#include "MovieRenderPipelineConfig.h"
+#include "MovieRenderTileImage.h"
+#include "MovieRenderOverlappedImage.h"
+#include "MovieRenderPipelineCoreModule.h"
+#include "Misc/FrameRate.h"
+
+
+UMoviePipelineImageSequenceContainerBase::UMoviePipelineImageSequenceContainerBase()
+{
+	// ToDo: Move this to a specific function to get out of CDO land.
+	ImageWriteQueue = &FModuleManager::Get().LoadModuleChecked<IImageWriteQueueModule>("ImageWriteQueue").GetWriteQueue();
+}
+
+void UMoviePipelineImageSequenceContainerBase::BeginFinalizeImpl()
+{
+	FinalizeFence = ImageWriteQueue->CreateFence();
+}
+
+bool UMoviePipelineImageSequenceContainerBase::HasFinishedProcessingImpl()
+{ 
+	// Wait until the finalization fence is reached meaning we've written everything to disk.
+	return Super::HasFinishedProcessingImpl() && (!FinalizeFence.IsValid() || FinalizeFence.WaitFor(0));
+}
+
+static FImageTileAccumulator TestAccumulator;
+static FImageOverlappedAccumulator TestOverlappedAccumulator;
+
+void UMoviePipelineImageSequenceContainerBase::OnInitializedForPipelineImpl(UMoviePipeline* InPipeline)
+{
+	FString OutputDirectory = GetPipeline()->GetPipelineConfig()->OutputDirectory.Path;
+	IImageWriteQueue* WriteQueue = ImageWriteQueue;
+
+	FImageTileAccumulator* TileAccumulator = &TestAccumulator;
+	FImageOverlappedAccumulator* OverlappedAccumulator = &TestOverlappedAccumulator;
+
+
+	auto OnImageReceived = [OutputDirectory, WriteQueue, TileAccumulator, OverlappedAccumulator](TUniquePtr<FImagePixelData>&& InOwnedImage)
+	{
+		FImagePixelDataPayload* TilePayload = InOwnedImage->GetPayload<FImagePixelDataPayload>();
+
+		// Init if first one, shut down if last one
+		bool bIsFirstTile = (TilePayload->TileIndexX == 0 && TilePayload->TileIndexY == 0 && TilePayload->SpatialJitterIndex == 0);
+		bool bIsLastTile =  (TilePayload->TileIndexX == TilePayload->NumTilesX-1  &&
+							TilePayload->TileIndexY == TilePayload->NumTilesY-1 &&
+							TilePayload->SpatialJitterIndex == TilePayload->NumSpatialJitters-1);
+		bool bIsLastTemporalSample = TilePayload->TemporalJitterIndex == TilePayload->NumTemporalJitters - 1;
+		bool bIsFirstTemporalSample = TilePayload->TemporalJitterIndex == 0;
+
+
+		if (!TilePayload->bIsUsingOverlappedTiles)
+		{
+			if (bIsFirstTile && bIsFirstTemporalSample)
+			{
+				TileAccumulator->InitMemory(TilePayload->TileSizeX, TilePayload->TileSizeY, TilePayload->NumTilesX, TilePayload->NumTilesY, 3);
+				TileAccumulator->ZeroPlanes();
+				TileAccumulator->AccumulationGamma = TilePayload->AccumulationGamma;
+			}
+
+			const double AccumulateBeginTime = FPlatformTime::Seconds();
+			TileAccumulator->AccumulatePixelData(*InOwnedImage.Get(), TilePayload->TileIndexX, TilePayload->TileIndexY, FVector2D(TilePayload->JitterOffsetX,TilePayload->JitterOffsetY));
+			const double AccumulateEndTime = FPlatformTime::Seconds();
+			const float ElapsedMs = float((AccumulateEndTime - AccumulateBeginTime)*1000.0f);
+
+		}
+		else
+		{
+			if (bIsFirstTile && bIsFirstTemporalSample)
+			{
+				OverlappedAccumulator->InitMemory(TilePayload->OverlappedSizeX * TilePayload->NumTilesX, TilePayload->OverlappedSizeY * TilePayload->NumTilesY, 3);
+				OverlappedAccumulator->ZeroPlanes();
+				OverlappedAccumulator->AccumulationGamma = TilePayload->AccumulationGamma;
+			}
+
+			FIntPoint RawSize = InOwnedImage.Get()->GetSize();
+
+			check(TilePayload->OverlappedSizeX + 2 * TilePayload->OverlappedPadX == RawSize.X);
+			check(TilePayload->OverlappedSizeY + 2 * TilePayload->OverlappedPadY == RawSize.Y);
+			OverlappedAccumulator->AccumulatePixelData(*InOwnedImage.Get(), TilePayload->OverlappedOffsetX, TilePayload->OverlappedOffsetY, TilePayload->OverlappedSubpixelShift);
+		}
+
+		//UE_LOG(LogTemp, Log, TEXT("[%8.2f] Accumulation time."), ElapsedMs);
+
+		static bool bWriteTiles = true;
+		if (bWriteTiles)
+		{
+			TUniquePtr<FImageWriteTask> TileImageTask = MakeUnique<FImageWriteTask>();
+			TileImageTask->PixelPreProcessors.Add(TAsyncAlphaWrite<FColor>(255));
+			// ImageTask->PixelPreProcessors.Add(TAsyncGammaCorrect<FColor>(2.2f));
+			TileImageTask->Format = EImageFormat::PNG;
+
+			// ImageTask->CompressionQuality = GetCompressionQuality();
+			FString OutputName = FString::Printf(TEXT("/%s_%d_SS_%d_TS_%d_TileX_%d_TileY_%d.png"),
+				*TilePayload->PassName, TilePayload->OutputState.OutputFrameNumber, TilePayload->SpatialJitterIndex, TilePayload->OutputState.TemporalSampleIndex,
+				TilePayload->TileIndexX, TilePayload->TileIndexY);
+			FString OutputPath = OutputDirectory + OutputName;
+			TileImageTask->Filename = OutputPath;
+
+			TileImageTask->PixelData = MoveTemp(InOwnedImage);
+
+			WriteQueue->Enqueue(MoveTemp(TileImageTask));
+		}
+
+		if (bIsLastTile && bIsLastTemporalSample)
+		{
+
+			int32 FullSizeX = TilePayload->bIsUsingOverlappedTiles ? OverlappedAccumulator->PlaneSizeX : TileAccumulator->TileSizeX * TileAccumulator->NumTilesX;
+			int32 FullSizeY = TilePayload->bIsUsingOverlappedTiles ? OverlappedAccumulator->PlaneSizeY : TileAccumulator->TileSizeY * TileAccumulator->NumTilesY;
+
+			{
+				TUniquePtr<FImageWriteTask> ImageTask = MakeUnique<FImageWriteTask>();
+
+				static bool bFullLinearColor = false; // Make this an option. For now, if it's linear, write EXR. Otherwise PNGs for 8bit.
+				if (bFullLinearColor)
+				{
+					ImageTask->Format = EImageFormat::EXR;
+
+					TUniquePtr<TImagePixelData<FLinearColor> > PixelData = MakeUnique<TImagePixelData<FLinearColor> >(FIntPoint(FullSizeX,FullSizeY));
+					if (!TilePayload->bIsUsingOverlappedTiles)
+					{
+						TileAccumulator->FetchFinalPixelDataLinearColor(PixelData->Pixels);
+					}
+					else
+					{
+						OverlappedAccumulator->FetchFinalPixelDataLinearColor(PixelData->Pixels);
+					}
+
+					ImageTask->PixelPreProcessors.Add(TAsyncAlphaWrite<FLinearColor>(1.0f));
+					ImageTask->PixelData = MoveTemp(PixelData);
+
+					FString TempOutputName = FString::Printf(TEXT("/FINAL_%s_SS.%d.exr"),
+						*TilePayload->PassName, TilePayload->OutputState.OutputFrameNumber);
+					FString TempOutputPath = OutputDirectory + TempOutputName;
+					ImageTask->Filename = TempOutputPath;
+				}
+				else
+				{
+					ImageTask->Format = EImageFormat::PNG;
+
+					// 8bit FColors
+					TUniquePtr<TImagePixelData<FColor> > PixelData = MakeUnique<TImagePixelData<FColor> >(FIntPoint(FullSizeX,FullSizeY));
+					if (!TilePayload->bIsUsingOverlappedTiles)
+					{
+						TileAccumulator->FetchFinalPixelDataByte(PixelData->Pixels);
+					}
+					else
+					{
+						OverlappedAccumulator->FetchFinalPixelDataByte(PixelData->Pixels);
+					}
+
+					ImageTask->PixelPreProcessors.Add(TAsyncAlphaWrite<FColor>(255));
+					ImageTask->PixelData = MoveTemp(PixelData);
+
+					FString TempOutputName = FString::Printf(TEXT("/FINAL_%s_SS.%d.png"),
+						*TilePayload->PassName, TilePayload->OutputState.OutputFrameNumber);
+					FString TempOutputPath = OutputDirectory + TempOutputName;
+					ImageTask->Filename = TempOutputPath;
+				}
+
+				WriteQueue->Enqueue(MoveTemp(ImageTask));
+			}
+
+			TileAccumulator->Reset();
+		}
+	};
+
+	InPipeline->GetOutputPipe()->AddEndpoint(OnImageReceived);
+}
+
+void UMoviePipelineImageSequenceContainerBase::ProcessFrameImpl(TArray<MoviePipeline::FOutputFrameData> FrameData, FMoviePipelineFrameOutputState CachedOutputState, FDirectoryPath OutputDirectory)
+{
+	if (FrameData.Num() == 0)
+	{
+		return;
+	}
+
+	for (MoviePipeline::FOutputFrameData& Frame : FrameData)
+	{
+		if (Frame.Resolution.X <= 0 || Frame.Resolution.Y <= 0)
+		{
+			continue;
+		}
+
+		TUniquePtr<FImageWriteTask> ImageTask = MakeUnique<FImageWriteTask>();
+
+		// Move the color buffer into a raw image data container that we can pass to the write queue
+		ImageTask->PixelData = MakeUnique<TImagePixelData<FColor>>(Frame.Resolution, MoveTemp(Frame.ColorBuffer));
+	
+		// Always write full alpha for PNGs
+		ImageTask->PixelPreProcessors.Add(TAsyncAlphaWrite<FColor>(255));
+		ImageTask->Format = EImageFormat::PNG;
+	
+		// ImageTask->CompressionQuality = GetCompressionQuality();
+	
+		FString OutputName = FString::Printf(TEXT("/Frame_%s_%d.png"), *Frame.PassName, CachedOutputState.OutputFrameNumber);
+		FString OutputPath = OutputDirectory.Path + OutputName;
+		ImageTask->Filename = OutputPath;
+
+		ImageWriteQueue->Enqueue(MoveTemp(ImageTask));
+	}
+}
diff --git a/Engine/Plugins/MovieScene/MovieRenderPipeline/Source/MovieRenderPipelineRenderPasses/Public/MoviePipelineBackbufferPass.h b/Engine/Plugins/MovieScene/MovieRenderPipeline/Source/MovieRenderPipelineRenderPasses/Public/MoviePipelineBackbufferPass.h
new file mode 100644
--- /dev/null
+++ b/Engine/Plugins/MovieScene/MovieRenderPipeline/Source/MovieRenderPipelineRenderPasses/Public/MoviePipelineBackbufferPass.h
@@ -0,0 +1,40 @@
+// Copyright 1998-2019 Epic Games, Inc. All Rights Reserved.
+#pragma once
+
+#include "MoviePipelineRenderPass.h"
+#include "SceneTypes.h"
+#include "MoviePipelineBackbufferPass.generated.h"
+
+// Forward Declare
+class FSceneViewport;
+class UTextureRenderTarget2D;
+class FCanvas;
+class FSceneViewFamily;
+class FSceneView;
+class UWorld;
+struct FImagePixelPipe;
+
+UCLASS(Blueprintable)
+class MOVIERENDERPIPELINERENDERPASSES_API UMoviePipelineBackbufferPass : public UMoviePipelineRenderPass
+{
+	GENERATED_BODY()
+
+protected:
+
+	virtual void SetupImpl(const FMoviePipelineRenderPassInitSettings& InInitSettings, TSharedRef<FImagePixelPipe, ESPMode::ThreadSafe> InOutputPipe) override;
+	virtual void CaptureFrameImpl(const FMoviePipelineRenderPassMetrics& OutputFrameMetrics) override;
+	virtual void GetFrameDataImpl(MoviePipeline::FOutputFrameData& OutFrameData) override;
+	virtual void TeardownImpl() override;
+	virtual void GatherOutputPassesImpl(TArray<FMoviePipelinePassIdentifier>& ExpectedRenderPasses) override;
+
+	FSceneView* CalcSceneView(FSceneViewFamily* ViewFamily, const FMoviePipelineRenderPassMetrics& InPassMetrics);
+
+private:
+	FSceneViewStateReference ViewState;
+
+
+	UPROPERTY(Transient)
+	UTextureRenderTarget2D* TileRenderTarget;
+
+	TSharedPtr<FImagePixelPipe, ESPMode::ThreadSafe> OutputPipe;
+};
\ No newline at end of file
diff --git a/Engine/Plugins/MovieScene/MovieRenderPipeline/Source/MovieRenderPipelineRenderPasses/Public/MoviePipelineImageSequenceContainer.h b/Engine/Plugins/MovieScene/MovieRenderPipeline/Source/MovieRenderPipelineRenderPasses/Public/MoviePipelineImageSequenceContainer.h
new file mode 100644
--- /dev/null
+++ b/Engine/Plugins/MovieScene/MovieRenderPipeline/Source/MovieRenderPipelineRenderPasses/Public/MoviePipelineImageSequenceContainer.h
@@ -0,0 +1,31 @@
+// Copyright 1998-2019 Epic Games, Inc. All Rights Reserved.
+#pragma once
+
+#include "MoviePipelineOutput.h"
+#include "Async/Future.h"
+#include "MoviePipelineImageSequenceContainer.generated.h"
+
+// Forward Declare
+class IImageWriteQueue;
+
+UCLASS(Blueprintable)
+class MOVIERENDERPIPELINERENDERPASSES_API UMoviePipelineImageSequenceContainerBase : public UMoviePipelineOutput
+{
+	GENERATED_BODY()
+public:
+	UMoviePipelineImageSequenceContainerBase();
+
+protected:
+	// UMovieRenderPipelineOutputContainer interface
+	virtual void OnInitializedForPipelineImpl(UMoviePipeline* InPipeline);
+	virtual void ProcessFrameImpl(TArray<MoviePipeline::FOutputFrameData> FrameData, FMoviePipelineFrameOutputState CachedOutputState, FDirectoryPath OutputDirectory) override;
+	virtual void BeginFinalizeImpl() override;
+	virtual bool HasFinishedProcessingImpl() override;
+	// ~UMovieRenderPipelineOutputContainer interface
+private:
+	/** A pointer to the image write queue used for asynchronously writing images */
+	IImageWriteQueue* ImageWriteQueue;
+
+	/** A fence to keep track of when the Image Write queue has fully flushed. */
+	TFuture<void> FinalizeFence;
+};
\ No newline at end of file
diff --git a/Engine/Plugins/MovieScene/MovieRenderPipeline/Source/MovieRenderPipelineSettings/Public/MoviePipelineAudioOutput.h b/Engine/Plugins/MovieScene/MovieRenderPipeline/Source/MovieRenderPipelineSettings/Public/MoviePipelineAudioOutput.h
new file mode 100644
--- /dev/null
+++ b/Engine/Plugins/MovieScene/MovieRenderPipeline/Source/MovieRenderPipelineSettings/Public/MoviePipelineAudioOutput.h
@@ -0,0 +1,35 @@
+// Copyright 1998-2019 Epic Games, Inc. All Rights Reserved.
+#pragma once
+
+#include "MoviePipelineOutput.h"
+#include "MoviePipelineAudioOutput.generated.h"
+
+UCLASS(BlueprintType, meta=(MovieRenderPipelineDisplayName="Audio"))
+class MOVIERENDERPIPELINESETTINGS_API UMoviePipelineAudioOutput : public UMoviePipelineOutput
+{
+	GENERATED_BODY()
+public:
+	UMoviePipelineAudioOutput()
+		: bExportWavFile(true)
+		, CachedExpectedDuration(0.f)
+		, PrevRenderEveryTickValue(false)
+	{
+	}
+
+	// virtual FText GetDisplayText() const override { return NSLOCTEXT("MovieRenderPipeline", "AudioSettingDisplayName", "Audio"); }
+	virtual void OnInitializedForPipelineImpl(UMoviePipeline* InPipeline) override;
+	virtual void OnPipelineFinishedImpl() override;
+	virtual void OnPostTickImpl() override;
+	virtual void OnFrameProductionStartImpl() override;
+	virtual void OnShotFinishedImpl(const FMoviePipelineShotCache& Shot) override;
+public:
+	/** If true audio will be written to a separate wave file in addition to being embedded into any supported containers */
+	UPROPERTY(EditAnywhere, BlueprintReadWrite, Category = "Movie Render Pipeline")
+	bool bExportWavFile;
+
+private:
+	float CachedExpectedDuration;
+	int32 PrevRenderEveryTickValue;
+	FString PrevAudioDevicePlatform;
+	float PrevUnfocusedAudioMultiplier;
+};
\ No newline at end of file
diff --git a/Engine/Plugins/MovieScene/MovieRenderPipeline/Source/MovieRenderPipelineSettings/Public/MoviePipelineBurnInSetting.h b/Engine/Plugins/MovieScene/MovieRenderPipeline/Source/MovieRenderPipelineSettings/Public/MoviePipelineBurnInSetting.h
new file mode 100644
--- /dev/null
+++ b/Engine/Plugins/MovieScene/MovieRenderPipeline/Source/MovieRenderPipelineSettings/Public/MoviePipelineBurnInSetting.h
@@ -0,0 +1,20 @@
+// Copyright 1998-2019 Epic Games, Inc. All Rights Reserved.
+#pragma once
+
+#include "MoviePipelineSetting.h"
+#include "LevelSequenceActor.h"
+#include "MoviePipelineBurnInSetting.generated.h"
+
+UCLASS(Blueprintable, meta=(MovieRenderPipelineDisplayName="Burn In"))
+class UMoviePipelineBurnInSetting : public UMoviePipelineSetting
+{
+	GENERATED_BODY()
+public:
+	virtual FText GetDisplayText() const override { return NSLOCTEXT("MovieRenderPipeline", "BurnInSettingDisplayName", "Burn In"); }
+public:
+	UPROPERTY(EditAnywhere, BlueprintReadWrite, meta=(MetaClass="LevelSequenceBurnIn"), Category = "Movie Pipeline")
+	FSoftClassPath BurnInClass;
+
+	UPROPERTY(Instanced, EditAnywhere, BlueprintReadWrite, Category="Movie Pipeline")
+	ULevelSequenceBurnInInitSettings* Settings;
+};
\ No newline at end of file
diff --git a/Engine/Plugins/MovieScene/MovieRenderPipeline/Source/MovieRenderPipelineSettings/Public/MoviePipelineGameOverrideSetting.h b/Engine/Plugins/MovieScene/MovieRenderPipeline/Source/MovieRenderPipelineSettings/Public/MoviePipelineGameOverrideSetting.h
new file mode 100644
--- /dev/null
+++ b/Engine/Plugins/MovieScene/MovieRenderPipeline/Source/MovieRenderPipelineSettings/Public/MoviePipelineGameOverrideSetting.h
@@ -0,0 +1,32 @@
+// Copyright 1998-2019 Epic Games, Inc. All Rights Reserved.
+#pragma once
+
+#include "MoviePipelineSetting.h"
+#include "MoviePipelineGameOverrideSetting.generated.h"
+
+// Forward Declares
+class AGameModeBase;
+
+UCLASS(Blueprintable)
+class MOVIERENDERPIPELINESETTINGS_API UMoviePipelineGameOverrideSetting : public UMoviePipelineSetting
+{
+	GENERATED_BODY()
+public:
+	UMoviePipelineGameOverrideSetting()
+		: bDisableCinematicMode(false)
+		, bShowPlayer(false)
+	{
+	}
+	
+	/** Optional Game Mode to override the map's default game mode with. This can be useful if the game's normal mode displays UI elements or loading screens that you don't want captured. */
+	UPROPERTY(EditAnywhere, BlueprintReadWrite, Category = "Movie Pipeline")
+	TSubclassOf<AGameModeBase> GameModeOverride;
+		
+	/** Cinematic Mode makes various systems to flush async loading to ensure first frame accuracy, which also disables Player Controller input, and the HUD spawned by AGameModeBase::HUDClass */
+	UPROPERTY(EditAnywhere, BlueprintReadWrite, Category = "Movie Pipeline")
+	bool bDisableCinematicMode;
+	
+	/** Should the Player Pawn spawned by the maps gamemode be visible? */
+	UPROPERTY(EditAnywhere, BlueprintReadWrite, Category = "Movie Pipeline")
+	bool bShowPlayer;
+};
\ No newline at end of file
diff --git a/Engine/Plugins/MovieScene/MovieRenderPipeline/Source/MovieRenderPipelineSettings/Public/MoviePipelineSeparateProcessSetting.h b/Engine/Plugins/MovieScene/MovieRenderPipeline/Source/MovieRenderPipelineSettings/Public/MoviePipelineSeparateProcessSetting.h
new file mode 100644
--- /dev/null
+++ b/Engine/Plugins/MovieScene/MovieRenderPipeline/Source/MovieRenderPipelineSettings/Public/MoviePipelineSeparateProcessSetting.h
@@ -0,0 +1,24 @@
+// Copyright 1998-2019 Epic Games, Inc. All Rights Reserved.
+#pragma once
+
+#include "MoviePipelineSetting.h"
+#include "MoviePipelineSeparateProcessSetting.generated.h"
+
+UCLASS(Blueprintable, MinimalAPI)
+class UMoviePipelineSeparateProcessSetting : public UMoviePipelineSetting
+{
+	GENERATED_BODY()
+public:
+	
+	/** If enabled the editor will close itself when a new process is started. This can be used to gain some performance. */
+	UPROPERTY(EditAnywhere, BlueprintReadWrite, Category = "Movie Pipeline")
+	bool bCloseEditor;
+	
+	/** A list of additional command line arguments to be appended to the new process startup. Can be useful if your game requires certain arguments to start such as disabling log-in screens. */
+	UPROPERTY(EditAnywhere, BlueprintReadWrite, Category = "Movie Pipeline")
+	FString AdditionalCommandLineArguments;
+	
+	/** A list of command line arguments which are inherited from the currently running Editor instance that will be automatically appended to the new process. */
+	UPROPERTY(VisibleAnywhere, BlueprintReadOnly, Category = "Movie Pipeline")
+	FString InheritedCommandLineArguments;
+};
\ No newline at end of file
diff --git a/Engine/Source/Runtime/Engine/Private/SceneView.cpp b/Engine/Source/Runtime/Engine/Private/SceneView.cpp
--- a/Engine/Source/Runtime/Engine/Private/SceneView.cpp
+++ b/Engine/Source/Runtime/Engine/Private/SceneView.cpp
@@ -634,6 +634,7 @@ FSceneView::FSceneView(const FSceneViewInitOptions& InitOptions)
 	, SpecularOverrideParameter(FVector4(0,0,0,1))
 	, NormalOverrideParameter(FVector4(0,0,0,1))
 	, RoughnessOverrideParameter(FVector2D(0,1))
+	, MaterialTextureMipBias(0.f)
 	, HiddenPrimitives(InitOptions.HiddenPrimitives)
 	, ShowOnlyPrimitives(InitOptions.ShowOnlyPrimitives)
 	, OriginOffsetThisFrame(InitOptions.OriginOffsetThisFrame)
diff --git a/Engine/Source/Runtime/Engine/Public/SceneView.h b/Engine/Source/Runtime/Engine/Public/SceneView.h
--- a/Engine/Source/Runtime/Engine/Public/SceneView.h
+++ b/Engine/Source/Runtime/Engine/Public/SceneView.h
@@ -929,6 +929,9 @@ public:
 	FVector4 NormalOverrideParameter;
 	FVector2D RoughnessOverrideParameter;
 
+	/** Mip bias to apply in material's samplers. */
+	float MaterialTextureMipBias;
+
 	/** The primitives which are hidden for this view. */
 	TSet<FPrimitiveComponentId> HiddenPrimitives;
 
diff --git a/Engine/Source/Runtime/ImageWriteQueue/Public/ImagePixelData.h b/Engine/Source/Runtime/ImageWriteQueue/Public/ImagePixelData.h
--- a/Engine/Source/Runtime/ImageWriteQueue/Public/ImagePixelData.h
+++ b/Engine/Source/Runtime/ImageWriteQueue/Public/ImagePixelData.h
@@ -18,6 +18,13 @@ enum class EImagePixelType
 	Float32,
 };
 
+struct IImagePixelDataPayload
+{
+	virtual ~IImagePixelDataPayload() {}
+};
+
+typedef TSharedPtr<IImagePixelDataPayload, ESPMode::ThreadSafe> FImagePixelPayloadPtr;
+
 struct FImagePixelData
 {
 	virtual ~FImagePixelData() {}
@@ -83,7 +90,8 @@ struct FImagePixelData
 
 		RetrieveData(RawPtr, SizeBytes);
 
-		if (RawPtr && SizeBytes == Size.X*Size.Y*BitDepth/8*NumChannels)
+		int64 FoundTotalSize = int64(Size.X)*int64(Size.Y)*int64(BitDepth / 8)*int64(NumChannels);
+		if (RawPtr && SizeBytes == FoundTotalSize)
 		{
 			OutRawData = RawPtr;
 			OutSizeBytes = SizeBytes;
@@ -108,14 +116,21 @@ struct FImagePixelData
 		return Move();
 	}
 
+	/**
+	* Return a pointer to the Payload stored in this data.
+	*/
+	template<typename T>
+	T* GetPayload() { return static_cast<T*>(Payload.Get()); }
+
 protected:
 
-	FImagePixelData(const FIntPoint& InSize, EImagePixelType InPixelType, ERGBFormat InPixelLayout, uint8 InBitDepth, uint8 InNumChannels)
+	FImagePixelData(const FIntPoint& InSize, EImagePixelType InPixelType, ERGBFormat InPixelLayout, uint8 InBitDepth, uint8 InNumChannels, FImagePixelPayloadPtr InPayload)
 		: Size(InSize)
 		, Type(InPixelType)
 		, PixelLayout(InPixelLayout)
 		, BitDepth(InBitDepth)
 		, NumChannels(InNumChannels)
+		, Payload(InPayload)
 	{}
 
 private:
@@ -141,6 +156,9 @@ private:
 
 	/** Number of channels in the data */
 	uint8 NumChannels;
+
+	/** Optional user-specified payload */
+	FImagePixelPayloadPtr Payload;
 };
 
 /**
@@ -152,11 +170,20 @@ struct TImagePixelData : FImagePixelData
 	TArray<PixelType> Pixels;
 
 	TImagePixelData(const FIntPoint& InSize)
-		: FImagePixelData(InSize, TImagePixelDataTraits<PixelType>::PixelType, TImagePixelDataTraits<PixelType>::PixelLayout, TImagePixelDataTraits<PixelType>::BitDepth, TImagePixelDataTraits<PixelType>::NumChannels)
+		: FImagePixelData(InSize, TImagePixelDataTraits<PixelType>::PixelType, TImagePixelDataTraits<PixelType>::PixelLayout, TImagePixelDataTraits<PixelType>::BitDepth, TImagePixelDataTraits<PixelType>::NumChannels, nullptr)
 	{}
 
 	TImagePixelData(const FIntPoint& InSize, TArray<PixelType>&& InPixels)
-		: FImagePixelData(InSize, TImagePixelDataTraits<PixelType>::PixelType, TImagePixelDataTraits<PixelType>::PixelLayout, TImagePixelDataTraits<PixelType>::BitDepth, TImagePixelDataTraits<PixelType>::NumChannels)
+		: FImagePixelData(InSize, TImagePixelDataTraits<PixelType>::PixelType, TImagePixelDataTraits<PixelType>::PixelLayout, TImagePixelDataTraits<PixelType>::BitDepth, TImagePixelDataTraits<PixelType>::NumChannels, nullptr)
+		, Pixels(MoveTemp(InPixels))
+	{}
+
+	TImagePixelData(const FIntPoint& InSize, FImagePixelPayloadPtr InPayload)
+		: FImagePixelData(InSize, TImagePixelDataTraits<PixelType>::PixelType, TImagePixelDataTraits<PixelType>::PixelLayout, TImagePixelDataTraits<PixelType>::BitDepth, TImagePixelDataTraits<PixelType>::NumChannels, InPayload)
+	{}
+
+	TImagePixelData(const FIntPoint& InSize, TArray<PixelType>&& InPixels, FImagePixelPayloadPtr InPayload)
+		: FImagePixelData(InSize, TImagePixelDataTraits<PixelType>::PixelType, TImagePixelDataTraits<PixelType>::PixelLayout, TImagePixelDataTraits<PixelType>::BitDepth, TImagePixelDataTraits<PixelType>::NumChannels, InPayload)
 		, Pixels(MoveTemp(InPixels))
 	{}
 
diff --git a/Engine/Source/Runtime/MediaUtils/Private/MediaRecorder.cpp b/Engine/Source/Runtime/MediaUtils/Private/MediaRecorder.cpp
--- a/Engine/Source/Runtime/MediaUtils/Private/MediaRecorder.cpp
+++ b/Engine/Source/Runtime/MediaUtils/Private/MediaRecorder.cpp
@@ -62,7 +62,7 @@ struct FMediaImagePixelData : FImagePixelData
 		, ERGBFormat InPixelLayout
 		, uint8 InBitDepth
 		, uint8 InNumChannels)
-		: FImagePixelData(InSize, InPixelType, InPixelLayout, InBitDepth, InNumChannels)
+		: FImagePixelData(InSize, InPixelType, InPixelLayout, InBitDepth, InNumChannels, nullptr)
 		, Sample(InSample)
 	{
 	}
diff --git a/Engine/Source/Runtime/Renderer/Private/SceneRendering.cpp b/Engine/Source/Runtime/Renderer/Private/SceneRendering.cpp
--- a/Engine/Source/Runtime/Renderer/Private/SceneRendering.cpp
+++ b/Engine/Source/Runtime/Renderer/Private/SceneRendering.cpp
@@ -886,7 +886,6 @@ void FViewInfo::Init()
 	TemporalJitterPixels = FVector2D::ZeroVector;
 
 	PreExposure = 1.0f;
-	MaterialTextureMipBias = 0.0f;
 
 	// Cache TEXTUREGROUP_World's for the render thread to create the material textures' shared sampler.
 	if (IsInGameThread())
@@ -1376,7 +1375,7 @@ void FViewInfo::SetupUniformBufferParameters(
 
 		float FinalMaterialTextureMipBias = GlobalMipBias;
 
-		if (bIsValidWorldTextureGroupSamplerFilter && PrimaryScreenPercentageMethod == EPrimaryScreenPercentageMethod::TemporalUpscale)
+		if (bIsValidWorldTextureGroupSamplerFilter && !FMath::IsNearlyZero(MaterialTextureMipBias))
 		{
 			ViewUniformShaderParameters.MaterialTextureMipBias = MaterialTextureMipBias;
 			ViewUniformShaderParameters.MaterialTextureDerivativeMultiply = FMath::Pow(2.0f, MaterialTextureMipBias);
diff --git a/Engine/Source/Runtime/Renderer/Private/SceneRendering.h b/Engine/Source/Runtime/Renderer/Private/SceneRendering.h
--- a/Engine/Source/Runtime/Renderer/Private/SceneRendering.h
+++ b/Engine/Source/Runtime/Renderer/Private/SceneRendering.h
@@ -1088,9 +1088,6 @@ public:
 	/** Frame's exposure. Always > 0. */
 	float PreExposure;
 
-	/** Mip bias to apply in material's samplers. */
-	float MaterialTextureMipBias;
-
 	/** Precomputed visibility data, the bits are indexed by VisibilityId of a primitive component. */
 	const uint8* PrecomputedVisibilityData;
 
diff --git a/Engine/Source/Runtime/Renderer/Private/SceneVisibility.cpp b/Engine/Source/Runtime/Renderer/Private/SceneVisibility.cpp
--- a/Engine/Source/Runtime/Renderer/Private/SceneVisibility.cpp
+++ b/Engine/Source/Runtime/Renderer/Private/SceneVisibility.cpp
@@ -3406,8 +3406,8 @@ void FSceneViewState::UpdateMotionBlurTimeScale(const FViewInfo& View)
 {
 	const int32 MotionBlurTargetFPS = View.FinalPostProcessSettings.MotionBlurTargetFPS;
 
-	// Frame rates over 120 FPS are clamped to avoid creating huge motion vectors.
-	float DeltaWorldTime = FMath::Max(View.Family->DeltaWorldTime, 1.0f / 120.0f);
+	// Ensure we can divide by the Delta Time later without a divide by zero.
+	float DeltaWorldTime = FMath::Max(View.Family->DeltaWorldTime, SMALL_NUMBER);
 
 	// Track the current FPS by using an exponential moving average of the current delta time.
 	if (MotionBlurTargetFPS <= 0)
