Missing Include fixes.

#rnx

diff --git a/Engine/Plugins/MovieScene/MovieRenderPipeline/Source/MovieRenderPipelineCore/Private/MoviePipelineRendering.cpp b/Engine/Plugins/MovieScene/MovieRenderPipeline/Source/MovieRenderPipelineCore/Private/MoviePipelineRendering.cpp
--- a/Engine/Plugins/MovieScene/MovieRenderPipeline/Source/MovieRenderPipelineCore/Private/MoviePipelineRendering.cpp
+++ b/Engine/Plugins/MovieScene/MovieRenderPipeline/Source/MovieRenderPipelineCore/Private/MoviePipelineRendering.cpp
@@ -9,6 +9,7 @@
 #include "MoviePipelineShotConfig.h"
 #include "MoviePipelineRenderPass.h"
 #include "MoviePipelineOutputBuilder.h"
+#include "RenderingThread.h"
 
 void UMoviePipeline::SetupRenderingPipelineForShot(FMoviePipelineShotCache& Shot)
 {
diff --git a/Engine/Plugins/MovieScene/MovieRenderPipeline/Source/MovieRenderPipelineCore/Private/MoviePipelineShotConfig.cpp b/Engine/Plugins/MovieScene/MovieRenderPipeline/Source/MovieRenderPipelineCore/Private/MoviePipelineShotConfig.cpp
--- a/Engine/Plugins/MovieScene/MovieRenderPipeline/Source/MovieRenderPipelineCore/Private/MoviePipelineShotConfig.cpp
+++ b/Engine/Plugins/MovieScene/MovieRenderPipeline/Source/MovieRenderPipelineCore/Private/MoviePipelineShotConfig.cpp
@@ -4,6 +4,7 @@
 #include "MoviePipelineSetting.h"
 #include "MoviePipelineRenderPass.h"
 #include "MoviePipelineOutput.h"
+#include "Misc/FrameRate.h"
 
 #define LOCTEXT_NAMESPACE "MovieRenderShotConfig"
 
diff --git a/Engine/Plugins/MovieScene/MovieRenderPipeline/Source/MovieRenderPipelineCore/Private/MovieRenderOverlappedImage.cpp b/Engine/Plugins/MovieScene/MovieRenderPipeline/Source/MovieRenderPipelineCore/Private/MovieRenderOverlappedImage.cpp
--- a/Engine/Plugins/MovieScene/MovieRenderPipeline/Source/MovieRenderPipelineCore/Private/MovieRenderOverlappedImage.cpp
+++ b/Engine/Plugins/MovieScene/MovieRenderPipeline/Source/MovieRenderPipelineCore/Private/MovieRenderOverlappedImage.cpp
@@ -7,6 +7,7 @@
 
 #include "MovieRenderPipelineCoreModule.h"
 #include "Math/Float16.h"
+#include "HAL/PlatformTime.h"
 
 void FImageOverlappedPlane::Init(int32 InSizeX, int32 InSizeY)
 {
diff --git a/Engine/Plugins/MovieScene/MovieRenderPipeline/Source/MovieRenderPipelineCore/Private/MovieRenderTileImage.cpp b/Engine/Plugins/MovieScene/MovieRenderPipeline/Source/MovieRenderPipelineCore/Private/MovieRenderTileImage.cpp
--- a/Engine/Plugins/MovieScene/MovieRenderPipeline/Source/MovieRenderPipelineCore/Private/MovieRenderTileImage.cpp
+++ b/Engine/Plugins/MovieScene/MovieRenderPipeline/Source/MovieRenderPipelineCore/Private/MovieRenderTileImage.cpp
@@ -7,6 +7,7 @@
 
 #include "MovieRenderPipelineCoreModule.h"
 #include "Math/Float16.h"
+#include "HAL/PlatformTime.h"
 
 void FImageTilePlane::Init(int32 InSizeX, int32 InSizeY)
 {
diff --git a/Engine/Plugins/MovieScene/MovieRenderPipeline/Source/MovieRenderPipelineCore/Public/MoviePipelineAccumulationSetting.h b/Engine/Plugins/MovieScene/MovieRenderPipeline/Source/MovieRenderPipelineCore/Public/MoviePipelineAccumulationSetting.h
--- a/Engine/Plugins/MovieScene/MovieRenderPipeline/Source/MovieRenderPipelineCore/Public/MoviePipelineAccumulationSetting.h
+++ b/Engine/Plugins/MovieScene/MovieRenderPipeline/Source/MovieRenderPipelineCore/Public/MoviePipelineAccumulationSetting.h
@@ -2,6 +2,7 @@
 #pragma once
 
 #include "MoviePipelineSetting.h"
+#include "MovieRenderPipelineDataTypes.h"
 #include "MoviePipelineAccumulationSetting.generated.h"
 
 
diff --git a/Engine/Plugins/MovieScene/MovieRenderPipeline/Source/MovieRenderPipelineCore/Public/MoviePipelineShotConfig.h b/Engine/Plugins/MovieScene/MovieRenderPipeline/Source/MovieRenderPipelineCore/Public/MoviePipelineShotConfig.h
--- a/Engine/Plugins/MovieScene/MovieRenderPipeline/Source/MovieRenderPipelineCore/Public/MoviePipelineShotConfig.h
+++ b/Engine/Plugins/MovieScene/MovieRenderPipeline/Source/MovieRenderPipelineCore/Public/MoviePipelineShotConfig.h
@@ -5,6 +5,7 @@
 #include "UObject/SoftObjectPath.h"
 #include "MoviePipelineSetting.h"
 #include "Engine/EngineTypes.h"
+#include "Misc/FrameRate.h"
 
 #include "MoviePipelineShotConfig.generated.h"
 
