Additional CIS fixes.

#rnx

diff --git a/Engine/Plugins/MovieScene/MovieRenderPipeline/Source/MovieRenderPipelineCore/Private/MoviePipelineTiming.cpp b/Engine/Plugins/MovieScene/MovieRenderPipeline/Source/MovieRenderPipelineCore/Private/MoviePipelineTiming.cpp
--- a/Engine/Plugins/MovieScene/MovieRenderPipeline/Source/MovieRenderPipelineCore/Private/MoviePipelineTiming.cpp
+++ b/Engine/Plugins/MovieScene/MovieRenderPipeline/Source/MovieRenderPipelineCore/Private/MoviePipelineTiming.cpp
@@ -12,6 +12,9 @@
 #include "MovieScene.h"
 #include "LevelSequence.h"
 #include "LevelSequenceActor.h"
+#include "Engine/World.h"
+#include "GameFramework/PlayerController.h"
+#include "Object.h"
 
 void UMoviePipeline::TickProducingFrames()
 {
diff --git a/Engine/Plugins/MovieScene/MovieRenderPipeline/Source/MovieRenderPipelineEditor/Private/Widgets/SMoviePipelineSettings.cpp b/Engine/Plugins/MovieScene/MovieRenderPipeline/Source/MovieRenderPipelineEditor/Private/Widgets/SMoviePipelineSettings.cpp
--- a/Engine/Plugins/MovieScene/MovieRenderPipeline/Source/MovieRenderPipelineEditor/Private/Widgets/SMoviePipelineSettings.cpp
+++ b/Engine/Plugins/MovieScene/MovieRenderPipeline/Source/MovieRenderPipelineEditor/Private/Widgets/SMoviePipelineSettings.cpp
@@ -21,6 +21,8 @@
 #include "Widgets/Layout/SWidgetSwitcher.h"
 #include "Widgets/Views/STreeView.h"
 #include "Widgets/Input/SComboButton.h"
+#include "Widgets/SOverlay.h"
+#include "Widgets/Images/SImage.h"
 #include "Framework/Commands/UICommandList.h"
 #include "Framework/Commands/GenericCommands.h"
 
diff --git a/Engine/Plugins/MovieScene/MovieRenderPipeline/Source/MovieRenderPipelineSettings/Public/MoviePipelineGameOverrideSetting.h b/Engine/Plugins/MovieScene/MovieRenderPipeline/Source/MovieRenderPipelineSettings/Public/MoviePipelineGameOverrideSetting.h
--- a/Engine/Plugins/MovieScene/MovieRenderPipeline/Source/MovieRenderPipelineSettings/Public/MoviePipelineGameOverrideSetting.h
+++ b/Engine/Plugins/MovieScene/MovieRenderPipeline/Source/MovieRenderPipelineSettings/Public/MoviePipelineGameOverrideSetting.h
@@ -2,6 +2,8 @@
 #pragma once
 
 #include "MoviePipelineSetting.h"
+#include "GameFramework/GameModeBase.h"
+#include "Templates/SubclassOf.h"
 #include "MoviePipelineGameOverrideSetting.generated.h"
 
 // Forward Declares
