Additional compile fixes.

#rnx

diff --git a/Engine/Plugins/MovieScene/MovieRenderPipeline/Source/MovieRenderPipelineCore/Private/MoviePipelineTiming.cpp b/Engine/Plugins/MovieScene/MovieRenderPipeline/Source/MovieRenderPipelineCore/Private/MoviePipelineTiming.cpp
--- a/Engine/Plugins/MovieScene/MovieRenderPipeline/Source/MovieRenderPipelineCore/Private/MoviePipelineTiming.cpp
+++ b/Engine/Plugins/MovieScene/MovieRenderPipeline/Source/MovieRenderPipelineCore/Private/MoviePipelineTiming.cpp
@@ -14,7 +14,7 @@
 #include "LevelSequenceActor.h"
 #include "Engine/World.h"
 #include "GameFramework/PlayerController.h"
-#include "Object.h"
+#include "GameFramework/WorldSettings.h"
 
 void UMoviePipeline::TickProducingFrames()
 {
