MoviePipeline: Changed plugin modules to "UncookedOnly" instead of "DeveloperTool" so that they didn't get compiled without the editor (ie: non-shipping Server/Client targets). Uninitialized variables fixes.

diff --git a/Engine/Plugins/MovieScene/MovieRenderPipeline/MovieRenderPipeline.uplugin b/Engine/Plugins/MovieScene/MovieRenderPipeline/MovieRenderPipeline.uplugin
--- a/Engine/Plugins/MovieScene/MovieRenderPipeline/MovieRenderPipeline.uplugin
+++ b/Engine/Plugins/MovieScene/MovieRenderPipeline/MovieRenderPipeline.uplugin
@@ -19,17 +19,17 @@
 	"Modules": [
 		{
 			"Name": "MovieRenderPipelineCore",
-			"Type": "DeveloperTool",
+			"Type": "UncookedOnly",
 			"LoadingPhase": "Default"
 		},
 		{
 			"Name": "MovieRenderPipelineSettings",
-			"Type": "DeveloperTool",
+			"Type": "UncookedOnly",
 			"LoadingPhase": "Default"
 		}, 
 		{
 			"Name" : "MovieRenderPipelineRenderPasses",
-			"Type": "DeveloperTool",
+			"Type": "UncookedOnly",
 			"LoadingPhase": "Default"
 		},
 		{
diff --git a/Engine/Plugins/MovieScene/MovieRenderPipeline/Source/MovieRenderPipelineCore/Private/MoviePipeline.cpp b/Engine/Plugins/MovieScene/MovieRenderPipeline/Source/MovieRenderPipelineCore/Private/MoviePipeline.cpp
--- a/Engine/Plugins/MovieScene/MovieRenderPipeline/Source/MovieRenderPipelineCore/Private/MoviePipeline.cpp
+++ b/Engine/Plugins/MovieScene/MovieRenderPipeline/Source/MovieRenderPipelineCore/Private/MoviePipeline.cpp
@@ -1,6 +1,5 @@
 // Copyright 1998-2019 Epic Games, Inc. All Rights Reserved.
 #include "MoviePipeline.h"
-#include "Editor.h"
 #include "MovieRenderPipelineCoreModule.h"
 #include "LevelSequence.h"
 #include "Tracks/MovieSceneSubTrack.h"
diff --git a/Engine/Plugins/MovieScene/MovieRenderPipeline/Source/MovieRenderPipelineCore/Public/MoviePipelineAccumulationSetting.h b/Engine/Plugins/MovieScene/MovieRenderPipeline/Source/MovieRenderPipelineCore/Public/MoviePipelineAccumulationSetting.h
--- a/Engine/Plugins/MovieScene/MovieRenderPipeline/Source/MovieRenderPipelineCore/Public/MoviePipelineAccumulationSetting.h
+++ b/Engine/Plugins/MovieScene/MovieRenderPipeline/Source/MovieRenderPipelineCore/Public/MoviePipelineAccumulationSetting.h
@@ -19,6 +19,7 @@ public:
 		, bIsUsingOverlappedTiles(false)
 		, PadRatioX(.50f)
 		, PadRatioY(.50f)
+		, AccumulationGamma(1.f)
 	{
 	}
 	
@@ -72,7 +73,10 @@ public:
 	UPROPERTY(EditAnywhere, BlueprintReadWrite, meta = (UIMin = 1, ClampMin = 1), Category = "Movie Pipeline")
 	int32 SpatialSampleCount;
 
-
+	/**
+	* 
+	*/
+	UPROPERTY(EditAnywhere, BlueprintReadWrite, meta = (UIMin = 1, ClampMin = 1), Category = "Movie Pipeline")
 	bool bIsUsingOverlappedTiles;
 
 	float PadRatioX;
diff --git a/Engine/Plugins/MovieScene/MovieRenderPipeline/Source/MovieRenderPipelineCore/Public/MoviePipelineSetting.h b/Engine/Plugins/MovieScene/MovieRenderPipeline/Source/MovieRenderPipelineCore/Public/MoviePipelineSetting.h
--- a/Engine/Plugins/MovieScene/MovieRenderPipeline/Source/MovieRenderPipelineCore/Public/MoviePipelineSetting.h
+++ b/Engine/Plugins/MovieScene/MovieRenderPipeline/Source/MovieRenderPipelineCore/Public/MoviePipelineSetting.h
@@ -32,8 +32,9 @@ protected:
 	virtual void OnInitializedForPipelineImpl(UMoviePipeline* InPipeline) {}
 
 public:
+#if WITH_EDITOR
 	virtual FText GetDisplayText() const { return this->GetClass()->GetDisplayNameText(); }
-
+#endif
 	/** Should the Pipeline automatically create an instance of this under the hood so calling code can rely on it existing? */
 	virtual bool IsRequired() const { return false; }
 	
diff --git a/Engine/Plugins/MovieScene/MovieRenderPipeline/Source/MovieRenderPipelineCore/Public/MovieRenderPipelineDataTypes.h b/Engine/Plugins/MovieScene/MovieRenderPipeline/Source/MovieRenderPipelineCore/Public/MovieRenderPipelineDataTypes.h
--- a/Engine/Plugins/MovieScene/MovieRenderPipeline/Source/MovieRenderPipelineCore/Public/MovieRenderPipelineDataTypes.h
+++ b/Engine/Plugins/MovieScene/MovieRenderPipeline/Source/MovieRenderPipelineCore/Public/MovieRenderPipelineDataTypes.h
@@ -398,7 +398,9 @@ struct FMoviePipelineFrameOutputState
 public:
 	FMoviePipelineFrameOutputState()
 		: OutputFrameNumber(-1)
+		, TotalSamplesRendered(0)
 		, TemporalSampleIndex(-1)
+		, MotionBlurFraction(0.5f)
 		, FrameDeltaTime(0)
 		, WorldSeconds(0)
 		, bWasAffectedByTimeDilation(false)
diff --git a/Engine/Plugins/MovieScene/MovieRenderPipeline/Source/MovieRenderPipelineRenderPasses/Private/MoviePipelineBackbufferPass.cpp b/Engine/Plugins/MovieScene/MovieRenderPipeline/Source/MovieRenderPipelineRenderPasses/Private/MoviePipelineBackbufferPass.cpp
--- a/Engine/Plugins/MovieScene/MovieRenderPipeline/Source/MovieRenderPipelineRenderPasses/Private/MoviePipelineBackbufferPass.cpp
+++ b/Engine/Plugins/MovieScene/MovieRenderPipeline/Source/MovieRenderPipelineRenderPasses/Private/MoviePipelineBackbufferPass.cpp
@@ -1,5 +1,4 @@
 // Copyright 1998-2019 Epic Games, Inc. All Rights Reserved.
-#pragma once
 
 #include "MoviePipelineBackbufferPass.h"
 #include "Slate/SceneViewport.h"
diff --git a/Engine/Plugins/MovieScene/MovieRenderPipeline/Source/MovieRenderPipelineRenderPasses/Private/MoviePipelineImageSequenceContainer.cpp b/Engine/Plugins/MovieScene/MovieRenderPipeline/Source/MovieRenderPipelineRenderPasses/Private/MoviePipelineImageSequenceContainer.cpp
--- a/Engine/Plugins/MovieScene/MovieRenderPipeline/Source/MovieRenderPipelineRenderPasses/Private/MoviePipelineImageSequenceContainer.cpp
+++ b/Engine/Plugins/MovieScene/MovieRenderPipeline/Source/MovieRenderPipelineRenderPasses/Private/MoviePipelineImageSequenceContainer.cpp
@@ -1,5 +1,4 @@
 // Copyright 1998-2019 Epic Games, Inc. All Rights Reserved.
-#pragma once
 
 #include "MoviePipelineImageSequenceContainer.h"
 #include "ImageWriteTask.h"
