MoviePipeline: Command Line renders can now be executed if a configuration asset has been exported through editor scripting before hand.

#rnx

diff --git a/Engine/Plugins/MovieScene/MovieRenderPipeline/MovieRenderPipeline.uplugin b/Engine/Plugins/MovieScene/MovieRenderPipeline/MovieRenderPipeline.uplugin
--- a/Engine/Plugins/MovieScene/MovieRenderPipeline/MovieRenderPipeline.uplugin
+++ b/Engine/Plugins/MovieScene/MovieRenderPipeline/MovieRenderPipeline.uplugin
@@ -15,6 +15,7 @@
 	"Installed": false,
 	"EnabledByDefault": true,
 	"Plugins": [
+		"EditorScriptingUtilities"
 	],
 	"Modules": [
 		{
diff --git a/Engine/Plugins/MovieScene/MovieRenderPipeline/Source/MovieRenderPipelineCore/Private/MoviePipelineInProcessExecutor.cpp b/Engine/Plugins/MovieScene/MovieRenderPipeline/Source/MovieRenderPipelineCore/Private/MoviePipelineInProcessExecutor.cpp
--- a/Engine/Plugins/MovieScene/MovieRenderPipeline/Source/MovieRenderPipelineCore/Private/MoviePipelineInProcessExecutor.cpp
+++ b/Engine/Plugins/MovieScene/MovieRenderPipeline/Source/MovieRenderPipelineCore/Private/MoviePipelineInProcessExecutor.cpp
@@ -1,6 +1,49 @@
 // Copyright 1998-2019 Epic Games, Inc. All Rights Reserved.
 #include "MoviePipelineInProcessExecutor.h"
+#include "MoviePipeline.h"
+#include "Engine/World.h"
 
 #define LOCTEXT_NAMESPACE "MoviePipelineInProcessExecutor"
+void UMoviePipelineInProcessExecutor::Start(UMovieRenderPipelineConfig* InConfig, const int32 InConfigIndex, const int32 InNumConfigs)
+{
+	ActiveConfig = InConfig;
+	FCoreUObjectDelegates::PostLoadMapWithWorld.AddUObject(this, &UMoviePipelineInProcessExecutor::OnMapLoadFinished);
+}
 
+void UMoviePipelineInProcessExecutor::OnMapLoadFinished(UWorld* NewWorld)
+{
+	// NewWorld can be null if a world is being destroyed.
+	if (!NewWorld)
+	{
+		return;
+	}
+
+	// Stop listening for map load until we're done and know we want to start the next config.
+	FCoreUObjectDelegates::PostLoadMapWithWorld.RemoveAll(this);
+
+	ActiveMoviePipeline = NewObject<UMoviePipeline>(NewWorld, TargetPipelineClass);
+	ActiveMoviePipeline->Initialize(ActiveConfig);
+
+	// Listen for when the pipeline thinks it has finished.
+	ActiveMoviePipeline->OnMoviePipelineFinished().AddUObject(this, &UMoviePipelineInProcessExecutor::OnMoviePipelineFinished);
+}
+
+void UMoviePipelineInProcessExecutor::OnMoviePipelineFinished(UMoviePipeline* InMoviePipeline)
+{
+	UMoviePipeline* MoviePipeline = ActiveMoviePipeline;
+
+	if (ActiveMoviePipeline)
+	{
+		// Unsubscribe in the event that it gets called twice we don't have issues.
+		ActiveMoviePipeline->OnMoviePipelineFinished().RemoveAll(this);
+	}
+
+	// Null these out now since OnIndividualPipelineFinished might invoke something that causes a GC
+	// and we want them to go away with the GC.
+	ActiveMoviePipeline = nullptr;
+	ActiveConfig = nullptr;
+
+	// Now that another frame has passed and we should be OK to start another PIE session, notify our owner.
+	OnIndividualPipelineFinished(MoviePipeline);
+}
 #undef LOCTEXT_NAMESPACE // "MoviePipelineInProcessExecutor"
diff --git a/Engine/Plugins/MovieScene/MovieRenderPipeline/Source/MovieRenderPipelineCore/Private/MovieRenderPipelineCommandLine.cpp b/Engine/Plugins/MovieScene/MovieRenderPipeline/Source/MovieRenderPipelineCore/Private/MovieRenderPipelineCommandLine.cpp
--- a/Engine/Plugins/MovieScene/MovieRenderPipeline/Source/MovieRenderPipelineCore/Private/MovieRenderPipelineCommandLine.cpp
+++ b/Engine/Plugins/MovieScene/MovieRenderPipeline/Source/MovieRenderPipelineCore/Private/MovieRenderPipelineCommandLine.cpp
@@ -53,12 +53,13 @@ void FMovieRenderPipelineCoreModule::InitializeCommandLineMovieRender()
 	}
 
 	// By this time, we know what assets you want to render, how to process the array of assets, and what runs an individual render. First we will create an executor.
-	UMoviePipelineExecutorBase* ExecutorBase = NewObject<UMoviePipelineExecutorBase>(GetTransientPackage());
+	UMoviePipelineExecutorBase* ExecutorBase = NewObject<UMoviePipelineExecutorBase>(GetTransientPackage(), ExecutorClass);
 
 	// We add the Executor to the root set. It will own all of the configuration data so this keeps it nicely in memory until finished,
 	// and means we only have to add/remove one thing from the root set, everything else uses normal outer ownership.
 	ExecutorBase->AddToRoot();
 	ExecutorBase->OnExecutorFinished().AddRaw(this, &FMovieRenderPipelineCoreModule::OnCommandLineMovieRenderCompleted);
+	ExecutorBase->SetMoviePipelineClass(PipelineClass);
 
 	TArray<UMovieRenderPipelineConfig*> Pipelines;
 	Pipelines.Add(Config);
@@ -72,10 +73,10 @@ void FMovieRenderPipelineCoreModule::OnCommandLineMovieRenderCompleted(UMoviePip
 	if (InExecutor)
 	{
 		InExecutor->RemoveFromRoot();
-
-		// Return a success code. Ideally any errors detected during rendering will return different codes.
-		FPlatformMisc::RequestExitWithStatus(/*Force*/ false, /*ReturnCode*/ 0);
 	}
+
+	// Return a success code. Ideally any errors detected during rendering will return different codes.
+	FPlatformMisc::RequestExitWithStatus(/*Force*/ false, /*ReturnCode*/ 0);
 }
 
 UClass* GetLocalExecutorClass(const FString& MoviePipelineLocalExecutorClassType, const FString ExecutorClassFormatString)
diff --git a/Engine/Plugins/MovieScene/MovieRenderPipeline/Source/MovieRenderPipelineCore/Private/MovieRenderPipelineCoreModule.cpp b/Engine/Plugins/MovieScene/MovieRenderPipeline/Source/MovieRenderPipelineCore/Private/MovieRenderPipelineCoreModule.cpp
--- a/Engine/Plugins/MovieScene/MovieRenderPipeline/Source/MovieRenderPipelineCore/Private/MovieRenderPipelineCoreModule.cpp
+++ b/Engine/Plugins/MovieScene/MovieRenderPipeline/Source/MovieRenderPipelineCore/Private/MovieRenderPipelineCoreModule.cpp
@@ -12,8 +12,7 @@ void FMovieRenderPipelineCoreModule::StartupModule()
 	{
 		UE_LOG(LogMovieRenderPipeline, Log, TEXT("Detected that the user intends to render a movie. Waiting until engine loop init is complete to ensure "))
 		// Register a hook to wait until the engine has finished loading to increase the likelihood that the desired classes are loaded.
-		// Note: This is not called for a commandlet.
-		FCoreDelegates::OnFEngineLoopInitComplete.AddRaw(this, &FMovieRenderPipelineCoreModule::InitializeCommandLineMovieRender);
+		FCoreDelegates::OnPostEngineInit.AddRaw(this, &FMovieRenderPipelineCoreModule::InitializeCommandLineMovieRender);
 	}
 }
 
diff --git a/Engine/Plugins/MovieScene/MovieRenderPipeline/Source/MovieRenderPipelineCore/Public/MoviePipelineExecutor.h b/Engine/Plugins/MovieScene/MovieRenderPipeline/Source/MovieRenderPipelineCore/Public/MoviePipelineExecutor.h
--- a/Engine/Plugins/MovieScene/MovieRenderPipeline/Source/MovieRenderPipelineCore/Public/MoviePipelineExecutor.h
+++ b/Engine/Plugins/MovieScene/MovieRenderPipeline/Source/MovieRenderPipelineCore/Public/MoviePipelineExecutor.h
@@ -2,10 +2,12 @@
 #pragma once
 
 #include "UObject/Object.h"
+#include "Templates/SubclassOf.h"
 #include "MoviePipelineExecutor.generated.h"
 
 class UMovieRenderPipelineConfig;
 class UMoviePipelineExecutorBase;
+class UMoviePipeline;
 
 DECLARE_MULTICAST_DELEGATE_OneParam(FOnMoviePipelineExecutorFinishedNative, UMoviePipelineExecutorBase*);
 DECLARE_DYNAMIC_MULTICAST_DELEGATE_OneParam(FOnMoviePipelineExecutorFinished, UMoviePipelineExecutorBase*, PipelineExecutor);
@@ -26,6 +28,11 @@ class MOVIERENDERPIPELINECORE_API UMoviePipelineExecutorBase : public UObject
 {
 	GENERATED_BODY()
 public:
+	UMoviePipelineExecutorBase()
+		: TargetPipelineClass(nullptr)
+	{
+	}
+
 	UFUNCTION(BlueprintCallable, Category = "Movie Render Pipeline")
 	virtual void Execute(UPARAM(Ref) TArray<UMovieRenderPipelineConfig*>& InPipelines)
 	{
@@ -36,7 +43,11 @@ public:
 	{
 		return OnExecutorFinishedDelegateNative;
 	}
-	
+
+	void SetMoviePipelineClass(UClass* InPipelineClass)
+	{
+		TargetPipelineClass = InPipelineClass;
+	}
 protected:
 	/** 
 	* This should be called when the Executor has finished executing all of the things
@@ -58,4 +69,7 @@ private:
 
 	/** For native C++ code. Called at the same time as the Blueprint/Python one. */
 	FOnMoviePipelineExecutorFinishedNative OnExecutorFinishedDelegateNative;
+protected:
+	UPROPERTY()
+	TSubclassOf<UMoviePipeline> TargetPipelineClass;
 };
\ No newline at end of file
diff --git a/Engine/Plugins/MovieScene/MovieRenderPipeline/Source/MovieRenderPipelineCore/Public/MoviePipelineInProcessExecutor.h b/Engine/Plugins/MovieScene/MovieRenderPipeline/Source/MovieRenderPipelineCore/Public/MoviePipelineInProcessExecutor.h
--- a/Engine/Plugins/MovieScene/MovieRenderPipeline/Source/MovieRenderPipelineCore/Public/MoviePipelineInProcessExecutor.h
+++ b/Engine/Plugins/MovieScene/MovieRenderPipeline/Source/MovieRenderPipelineCore/Public/MoviePipelineInProcessExecutor.h
@@ -21,4 +21,11 @@ public:
 		: UMoviePipelineLinearExecutorBase()
 	{
 	}
+
+protected:
+	virtual void Start(UMovieRenderPipelineConfig* InConfig, const int32 InConfigIndex, const int32 InNumConfigs) override;
+
+private:
+	void OnMapLoadFinished(UWorld* NewWorld);
+	void OnMoviePipelineFinished(UMoviePipeline* InMoviePipeline);
 };
\ No newline at end of file
diff --git a/Engine/Plugins/MovieScene/MovieRenderPipeline/Source/MovieRenderPipelineCore/Public/MoviePipelineLinearExecutor.h b/Engine/Plugins/MovieScene/MovieRenderPipeline/Source/MovieRenderPipelineCore/Public/MoviePipelineLinearExecutor.h
--- a/Engine/Plugins/MovieScene/MovieRenderPipeline/Source/MovieRenderPipelineCore/Public/MoviePipelineLinearExecutor.h
+++ b/Engine/Plugins/MovieScene/MovieRenderPipeline/Source/MovieRenderPipelineCore/Public/MoviePipelineLinearExecutor.h
@@ -42,6 +42,15 @@ protected:
 
 	/** Which Pipeline Config are we currently working on. */
 	int32 CurrentPipelineIndex;
+
+	/** Instance of the Pipeline that exists in the world that is currently processing (if any) */
+	UPROPERTY(Transient)
+	UMoviePipeline* ActiveMoviePipeline;
+
+	/** Instance of the Config we're supposed to be working on */
+	UPROPERTY(Transient)
+	UMovieRenderPipelineConfig* ActiveConfig;
+
 private:
 	/** Used to track total processing duration. */
 	FDateTime InitializationTime;
diff --git a/Engine/Plugins/MovieScene/MovieRenderPipeline/Source/MovieRenderPipelineCore/Public/MovieRenderPipelineConfig.h b/Engine/Plugins/MovieScene/MovieRenderPipeline/Source/MovieRenderPipelineCore/Public/MovieRenderPipelineConfig.h
--- a/Engine/Plugins/MovieScene/MovieRenderPipeline/Source/MovieRenderPipelineCore/Public/MovieRenderPipelineConfig.h
+++ b/Engine/Plugins/MovieScene/MovieRenderPipeline/Source/MovieRenderPipelineCore/Public/MovieRenderPipelineConfig.h
@@ -94,15 +94,15 @@ public:
 public:
 	
 	/** The default shot-setup to use for any shot that doesn't a specific implementation. This is required! */
-	UPROPERTY(EditAnywhere, BlueprintReadWrite, Category = "Movie Render Pipeline")
+	UPROPERTY(EditAnywhere, Instanced, BlueprintReadWrite, Category = "Movie Render Pipeline")
 	UMoviePipelineShotConfig* DefaultShotConfig;
 	
 	/** A mapping of Shot Name -> Shot Config to use for rendering specific shots with specific configs. */
-	UPROPERTY(EditAnywhere, BlueprintReadWrite, Category = "Movie Render Pipeline")
+	UPROPERTY(EditAnywhere, Instanced, BlueprintReadWrite, Category = "Movie Render Pipeline")
 	TMap<FString, UMoviePipelineShotConfig*> PerShotConfigMapping;
 	
 	/** Array of Output Containers. Each output container is passed data for each Input Buffer every frame. */
-	UPROPERTY(EditAnywhere, BlueprintReadWrite, Category = "Movie Render Pipeline")
+	UPROPERTY(EditAnywhere, Instanced, BlueprintReadWrite, Category = "Movie Render Pipeline")
 	TArray<UMoviePipelineOutput*> OutputContainers;
 
 private:
diff --git a/Engine/Plugins/MovieScene/MovieRenderPipeline/Source/MovieRenderPipelineEditor/Private/MoviePipelineEditorBlueprintLibrary.cpp b/Engine/Plugins/MovieScene/MovieRenderPipeline/Source/MovieRenderPipelineEditor/Private/MoviePipelineEditorBlueprintLibrary.cpp
--- a/Engine/Plugins/MovieScene/MovieRenderPipeline/Source/MovieRenderPipelineEditor/Private/MoviePipelineEditorBlueprintLibrary.cpp
+++ b/Engine/Plugins/MovieScene/MovieRenderPipeline/Source/MovieRenderPipelineEditor/Private/MoviePipelineEditorBlueprintLibrary.cpp
@@ -6,6 +6,8 @@
 #include "MovieRenderPipelineConfig.h"
 #include "AssetRegistryModule.h"
 #include "FileHelpers.h"
+#include "ObjectTools.h"
+#include "PackageTools.h"
 
 #define LOCTEXT_NAMESPACE "MoviePipelineEditorBlueprintLibrary"
 
@@ -17,14 +19,15 @@ bool UMoviePipelineEditorBlueprintLibrary::ExportConfigToAsset(const UMovieRende
 		return false;
 	}
 	
-	
-	if (!FPackageName::IsValidLongPackageName(InPackagePath, false, &OutErrorReason))
+	FString FixedAssetName = ObjectTools::SanitizeObjectName(InFileName);
+	FString NewPackageName = FPackageName::GetLongPackagePath(InPackagePath) + TEXT("/") + FixedAssetName;
+
+	if (!FPackageName::IsValidLongPackageName(NewPackageName, false, &OutErrorReason))
 	{
 		return false;
 	}
 
-
-	UPackage* NewPackage = CreatePackage(nullptr, *InPackagePath);
+	UPackage* NewPackage = CreatePackage(nullptr, *NewPackageName);
 	NewPackage->AddToRoot();
 	
 	// Duplicate the provided config into this package.
diff --git a/Engine/Plugins/MovieScene/MovieRenderPipeline/Source/MovieRenderPipelineEditor/Private/MoviePipelinePIEExecutor.cpp b/Engine/Plugins/MovieScene/MovieRenderPipeline/Source/MovieRenderPipelineEditor/Private/MoviePipelinePIEExecutor.cpp
--- a/Engine/Plugins/MovieScene/MovieRenderPipeline/Source/MovieRenderPipelineEditor/Private/MoviePipelinePIEExecutor.cpp
+++ b/Engine/Plugins/MovieScene/MovieRenderPipeline/Source/MovieRenderPipelineEditor/Private/MoviePipelinePIEExecutor.cpp
@@ -121,12 +121,15 @@ void UMoviePipelinePIEExecutor::OnPIEMoviePipelineFinished(UMoviePipeline* InMov
 
 void UMoviePipelinePIEExecutor::DelayedFinishNotification()
 {
-	// Now that another frame has passed and we should be OK to start another PIE session, notify our owner.
-	OnIndividualPipelineFinished(ActiveMoviePipeline);
+	UMoviePipeline* MoviePipeline = ActiveMoviePipeline;
 	
-
+	// Null these out now since OnIndividualPipelineFinished might invoke something that causes a GC
+	// and we want them to go away with the GC.
 	ActiveMoviePipeline = nullptr;
 	ActiveConfig = nullptr;
+	
+	// Now that another frame has passed and we should be OK to start another PIE session, notify our owner.
+	OnIndividualPipelineFinished(MoviePipeline);
 }
 
 #undef LOCTEXT_NAMESPACE // "MoviePipelinePIEExecutor"
diff --git a/Engine/Plugins/MovieScene/MovieRenderPipeline/Source/MovieRenderPipelineEditor/Private/MovieRenderPipelineEditorTest.cpp b/Engine/Plugins/MovieScene/MovieRenderPipeline/Source/MovieRenderPipelineEditor/Private/MovieRenderPipelineEditorTest.cpp
--- a/Engine/Plugins/MovieScene/MovieRenderPipeline/Source/MovieRenderPipelineEditor/Private/MovieRenderPipelineEditorTest.cpp
+++ b/Engine/Plugins/MovieScene/MovieRenderPipeline/Source/MovieRenderPipelineEditor/Private/MovieRenderPipelineEditorTest.cpp
@@ -1,3 +1,4 @@
+// Copyright 1998-2019 Epic Games, Inc. All Rights Reserved.
 #include "MovieRenderPipelineEditorModule.h"
 #include "MoviePipelineBackbufferPass.h"
 #include "MoviePipelineShotConfig.h"
@@ -11,7 +12,7 @@
 
 UMoviePipelineShotConfig* GenerateTestShotConfig(UObject* InOwner, int32 InSampleCount,
 	int32 InShutterAngle, EMoviePipelineShutterTiming InFrameTiming, int32 InTileCount,
-	int32 InSpatialSampleCount)
+	int32 InSpatialSampleCount, bool bIsUsingOverlappedTiles, float PadRatioX, float PadRatioY, float AccumulationGamma)
 {
 	UMoviePipelineShotConfig* OutConfig = NewObject<UMoviePipelineShotConfig>(InOwner);
 	UMoviePipelineBackbufferPass* DefaultPass = NewObject<UMoviePipelineBackbufferPass>(InOwner);
@@ -19,26 +20,31 @@ UMoviePipelineShotConfig* GenerateTestShotConfig(UObject* InOwner, int32 InSampl
 
 	UMoviePipelineGameOverrideSetting* GamemodeOverride = OutConfig->FindOrAddSetting<UMoviePipelineGameOverrideSetting>();
 	GamemodeOverride->GameModeOverride = AGameMode::StaticClass();
-		
+
 	UMoviePipelineAccumulationSetting* Accumulation = OutConfig->FindOrAddSetting< UMoviePipelineAccumulationSetting>();
 	Accumulation->TemporalSampleCount = InSampleCount;
 	Accumulation->CameraShutterAngle = InShutterAngle;
 	Accumulation->TileCount = InTileCount;
 	Accumulation->ShutterTiming = InFrameTiming;
 	Accumulation->SpatialSampleCount = InSpatialSampleCount;
+	Accumulation->bIsUsingOverlappedTiles = bIsUsingOverlappedTiles;
+	Accumulation->PadRatioX = PadRatioX;
+	Accumulation->PadRatioY = PadRatioY;
+	Accumulation->AccumulationGamma = AccumulationGamma;
 
 	return OutConfig;
 }
 
+PRAGMA_DISABLE_OPTIMIZATION
 TArray<UMovieRenderPipelineConfig*> FMovieRenderPipelineEditorModule::GenerateTestPipelineConfigs(FSoftObjectPath InSequence)
 {
 	// int32 ShutterAngles[] = { 180, 360 };
 	// EMoviePipelineShutterTiming ShutterTimings[] = { EMoviePipelineShutterTiming::FrameOpen, EMoviePipelineShutterTiming::FrameCenter, EMoviePipelineShutterTiming::FrameClose };
 	// FString ShutterTimingNames[] = { TEXT("FOpen"), TEXT("FCenter"), TEXT("FClose") };
 	// int32 TemporalSampleCounts[] = { 1, 5 };
-	   
+
 	int32 ShutterAngles[] = { 180 };
-	EMoviePipelineShutterTiming ShutterTimings[] = { EMoviePipelineShutterTiming::FrameCenter};
+	EMoviePipelineShutterTiming ShutterTimings[] = { EMoviePipelineShutterTiming::FrameCenter };
 	FString ShutterTimingNames[] = { TEXT("FCenter") };
 	int32 TemporalSampleCounts[] = { 1 };
 
@@ -55,17 +61,28 @@ TArray<UMovieRenderPipelineConfig*> FMovieRenderPipelineEditorModule::GenerateTe
 
 				UMovieRenderPipelineConfig* OutPipeline = NewObject<UMovieRenderPipelineConfig>(GetTransientPackage());
 				OutPipeline->Sequence = InSequence;
-				OutPipeline->OutputResolution = FIntPoint(1920, 1080);
+
+				int32 SizeX = 1080;
+				int32 SizeY = 1920;
+				int32 TileX = 2;
+				int32 TileY = 2;
+				int32 TestNumSamples = 2;
+				float PadRatioX = 0.5f;
+				float PadRatioY = 0.5f;
+				float AccumulationGamma = 1.0f;
+
+				OutPipeline->OutputResolution = FIntPoint(SizeX*TileX, SizeY*TileY);
 				OutPipeline->OutputDirectory.Path = FPaths::ProjectSavedDir() / TEXT("/VideoCaptures/") / *DirectoryName;
 
-				const int32 NumSpatialSamples = 1;
-				const int32 NumTiles = 1;
+				const int32 NumSpatialSamples = TestNumSamples;
+				const int32 NumTiles = TileX;
+				static bool bIsUsingOverlappedTiles = true;
 
 				UMoviePipelineShotConfig* DefaultConfig = GenerateTestShotConfig(OutPipeline,
 					TemporalSampleCounts[TemporalSampleCountIndex],
 					ShutterAngles[ShutterAngleIndex],
 					ShutterTimings[ShutterTimingIndex],
-					NumTiles, NumSpatialSamples);
+					NumTiles, NumSpatialSamples, bIsUsingOverlappedTiles, PadRatioX, PadRatioY, AccumulationGamma);
 
 				OutPipeline->DefaultShotConfig = DefaultConfig;
 
@@ -75,7 +92,7 @@ TArray<UMovieRenderPipelineConfig*> FMovieRenderPipelineEditorModule::GenerateTe
 
 				UMoviePipelineAudioOutput* AudioOutput = NewObject<UMoviePipelineAudioOutput>(OutPipeline);
 				OutPipeline->OutputContainers.Add(AudioOutput);
-				
+
 				// Add it to our list to be processed.
 				Configs.Add(OutPipeline);
 			}
diff --git a/Engine/Plugins/MovieScene/MovieRenderPipeline/Source/MovieRenderPipelineEditor/Public/MoviePipelinePIEExecutor.h b/Engine/Plugins/MovieScene/MovieRenderPipeline/Source/MovieRenderPipelineEditor/Public/MoviePipelinePIEExecutor.h
--- a/Engine/Plugins/MovieScene/MovieRenderPipeline/Source/MovieRenderPipelineEditor/Public/MoviePipelinePIEExecutor.h
+++ b/Engine/Plugins/MovieScene/MovieRenderPipeline/Source/MovieRenderPipelineEditor/Public/MoviePipelinePIEExecutor.h
@@ -36,14 +36,4 @@ private:
 
 	/** Gets the title for the Window. Used to show progress in the title. */
 	FText GetWindowTitle(const int32 InConfigIndex, const int32 InNumConfigs) const;
-
-
-private:
-	/** Instance of the Pipeline that exists in the world that is currently processing (if any) */
-	UPROPERTY(Transient)
-	UMoviePipeline* ActiveMoviePipeline;
-
-	/** Instance of the Config we're supposed to be working on */
-	UPROPERTY(Transient)
-	UMovieRenderPipelineConfig* ActiveConfig;
 };
\ No newline at end of file
diff --git a/Engine/Plugins/MovieScene/MovieRenderPipeline/Source/MovieRenderPipelineRenderPasses/Private/MoviePipelineImageSequenceContainer.cpp b/Engine/Plugins/MovieScene/MovieRenderPipeline/Source/MovieRenderPipelineRenderPasses/Private/MoviePipelineImageSequenceContainer.cpp
--- a/Engine/Plugins/MovieScene/MovieRenderPipeline/Source/MovieRenderPipelineRenderPasses/Private/MoviePipelineImageSequenceContainer.cpp
+++ b/Engine/Plugins/MovieScene/MovieRenderPipeline/Source/MovieRenderPipelineRenderPasses/Private/MoviePipelineImageSequenceContainer.cpp
@@ -95,10 +95,11 @@ void UMoviePipelineImageSequenceContainerBase::OnInitializedForPipelineImpl(UMov
 			TUniquePtr<FImageWriteTask> TileImageTask = MakeUnique<FImageWriteTask>();
 			TileImageTask->PixelPreProcessors.Add(TAsyncAlphaWrite<FColor>(255));
 			// ImageTask->PixelPreProcessors.Add(TAsyncGammaCorrect<FColor>(2.2f));
-			TileImageTask->Format = EImageFormat::PNG;
+			TileImageTask->Format = EImageFormat::JPEG;
+			TileImageTask->CompressionQuality = 100;
 
 			// ImageTask->CompressionQuality = GetCompressionQuality();
-			FString OutputName = FString::Printf(TEXT("/%s_%d_SS_%d_TS_%d_TileX_%d_TileY_%d.png"),
+			FString OutputName = FString::Printf(TEXT("/%s_%d_SS_%d_TS_%d_TileX_%d_TileY_%d.jpeg"),
 				*TilePayload->PassName, TilePayload->OutputState.OutputFrameNumber, TilePayload->SpatialJitterIndex, TilePayload->OutputState.TemporalSampleIndex,
 				TilePayload->TileIndexX, TilePayload->TileIndexY);
 			FString OutputPath = OutputDirectory + OutputName;
@@ -143,7 +144,7 @@ void UMoviePipelineImageSequenceContainerBase::OnInitializedForPipelineImpl(UMov
 				}
 				else
 				{
-					ImageTask->Format = EImageFormat::PNG;
+					ImageTask->Format = EImageFormat::JPEG;
 
 					// 8bit FColors
 					TUniquePtr<TImagePixelData<FColor> > PixelData = MakeUnique<TImagePixelData<FColor> >(FIntPoint(FullSizeX,FullSizeY));
@@ -158,8 +159,9 @@ void UMoviePipelineImageSequenceContainerBase::OnInitializedForPipelineImpl(UMov
 
 					ImageTask->PixelPreProcessors.Add(TAsyncAlphaWrite<FColor>(255));
 					ImageTask->PixelData = MoveTemp(PixelData);
+					ImageTask->CompressionQuality = 100;
 
-					FString TempOutputName = FString::Printf(TEXT("/FINAL_%s_SS.%d.png"),
+					FString TempOutputName = FString::Printf(TEXT("/FINAL_%s_SS.%d.jpeg"),
 						*TilePayload->PassName, TilePayload->OutputState.OutputFrameNumber);
 					FString TempOutputPath = OutputDirectory + TempOutputName;
 					ImageTask->Filename = TempOutputPath;
