remove optimization pragma to fix build

#rnx

diff --git a/Engine/Plugins/MovieScene/MovieRenderPipeline/Source/MovieRenderPipelineEditor/Private/MovieRenderPipelineEditorTest.cpp b/Engine/Plugins/MovieScene/MovieRenderPipeline/Source/MovieRenderPipelineEditor/Private/MovieRenderPipelineEditorTest.cpp
--- a/Engine/Plugins/MovieScene/MovieRenderPipeline/Source/MovieRenderPipelineEditor/Private/MovieRenderPipelineEditorTest.cpp
+++ b/Engine/Plugins/MovieScene/MovieRenderPipeline/Source/MovieRenderPipelineEditor/Private/MovieRenderPipelineEditorTest.cpp
@@ -35,7 +35,6 @@ UMoviePipelineShotConfig* GenerateTestShotConfig(UObject* InOwner, int32 InSampl
 	return OutConfig;
 }
 
-PRAGMA_DISABLE_OPTIMIZATION
 TArray<UMovieRenderPipelineConfig*> FMovieRenderPipelineEditorModule::GenerateTestPipelineConfigs(FSoftObjectPath InSequence)
 {
 	// int32 ShutterAngles[] = { 180, 360 };
