Progress on optimizing the overlapped tile accumulation.

diff --git a/Engine/Plugins/MovieScene/MovieRenderPipeline/Source/MovieRenderPipelineCore/Private/MovieRenderOverlappedImage.cpp b/Engine/Plugins/MovieScene/MovieRenderPipeline/Source/MovieRenderPipelineCore/Private/MovieRenderOverlappedImage.cpp
--- a/Engine/Plugins/MovieScene/MovieRenderPipeline/Source/MovieRenderPipelineCore/Private/MovieRenderOverlappedImage.cpp
+++ b/Engine/Plugins/MovieScene/MovieRenderPipeline/Source/MovieRenderPipelineCore/Private/MovieRenderOverlappedImage.cpp
@@ -36,7 +36,9 @@ void FImageOverlappedPlane::Reset()
 
 // a subpixel offset of (0.5,0.5) means that the raw data is exactly centered on destination pixels.
 void FImageOverlappedPlane::AccumulateSinglePlane(const TArray64<float>& InRawData, const TArray64<float>& InWeightData, int32 InSizeX, int32 InSizeY, int32 InOffsetX, int32 InOffsetY,
-								float SubpixelOffsetX, float SubpixelOffsetY)
+								float SubpixelOffsetX, float SubpixelOffsetY,
+								int32 SubRectOffsetX, int32 SubRectOffsetY,
+								int32 SubRectSizeX, int32 SubRectSizeY)
 {
 	check(InRawData.Num() == InSizeX * InSizeY);
 	check(InWeightData.Num() == InSizeX * InSizeY);
@@ -63,33 +65,85 @@ void FImageOverlappedPlane::AccumulateSinglePlane(const TArray64<float>& InRawDa
 		PixelWeight[1][1] = (       WeightX) * (       WeightY);
 	}
 
-	// Slow, reference version. Maybe optimize later.
-	for (int CurrY = 0; CurrY < InSizeY; CurrY++)
+	static bool bSlowReference = false;
+	if (bSlowReference)
 	{
-		for (int CurrX = 0; CurrX < InSizeX; CurrX++)
+		// Slow, reference version. Maybe optimize later.
+		for (int CurrY = 0; CurrY < InSizeY; CurrY++)
 		{
-			float Val = InRawData[CurrY * InSizeX + CurrX];
-			float BaseWeight = InWeightData[CurrY * InSizeX + CurrX];
-
-			for (int OffsetY = 0; OffsetY < 2; OffsetY++)
+			for (int CurrX = 0; CurrX < InSizeX; CurrX++)
 			{
-				for (int OffsetX = 0; OffsetX < 2; OffsetX++)
+				float Val = InRawData[CurrY * InSizeX + CurrX];
+				float BaseWeight = InWeightData[CurrY * InSizeX + CurrX];
+
+				for (int OffsetY = 0; OffsetY < 2; OffsetY++)
 				{
-					int DstY = StartY + CurrY + OffsetY;
-					int DstX = StartX + CurrX + OffsetX;
+					for (int OffsetX = 0; OffsetX < 2; OffsetX++)
+					{
+						int DstY = StartY + CurrY + OffsetY;
+						int DstX = StartX + CurrX + OffsetX;
 
-					float Weight = BaseWeight * PixelWeight[OffsetY][OffsetX];
+						float Weight = BaseWeight * PixelWeight[OffsetY][OffsetX];
 
-					if (DstX >= 0 && DstY >= 0 &&
-						DstX < SizeX && DstY < SizeY)
-					{
-						ChannelData[DstY * SizeX + DstX] += Weight * Val;
+						if (DstX >= 0 && DstY >= 0 &&
+							DstX < SizeX && DstY < SizeY)
+						{
+							ChannelData[DstY * SizeX + DstX] += Weight * Val;
+						}
 					}
 				}
 			}
 		}
 	}
+	else
+	{
+		// The initial destination corners, may go off the destination window.
+		int DstX0 = StartX;
+		int DstY0 = StartY;
+		int DstX1 = StartX + InSizeX;
+		int DstY1 = StartY + InSizeY;
+
+		// Initial destination corners
+		
+		// Given a position on the curre tile (x,y), we apply the sum of the 4 points (x+0,y+0), (x+1,y+0), (x+0,y+1), (x+1,y+1) to the destination index.
+		// So in reverse, given a destination position, our source sample positions are (x-1,y-1), (x+0,y-1), (x-1,y+0), (x+0,y+0).
+		// That's why we make sure to start at a minimum of index 1, instead of 0.
+
+		for (int CurrY = 1; CurrY < InSizeY; CurrY++)
+		{
+			for (int CurrX = 1; CurrX < InSizeX; CurrX++)
+			{
+				int DstY = StartY + CurrY;
+				int DstX = StartX + CurrX;
+
+				if (DstX >= 0 && DstY >= 0 &&
+					DstX < SizeX && DstY < SizeY)
+				{
+					float OriginalValue = ChannelData[DstY * SizeX + DstX];
+					float SumValue = OriginalValue;
+
+					for (int OffsetY = 0; OffsetY < 2; OffsetY++)
+					{
+						for (int OffsetX = 0; OffsetX < 2; OffsetX++)
+						{
+							int SrcX = CurrX + (OffsetX - 1);
+							int SrcY = CurrY + (OffsetY - 1);
+						
+							float Val = InRawData[SrcY * InSizeX + CurrX];
+							float BaseWeight = InWeightData[SrcY * InSizeX + SrcX];
 
+							float Weight = BaseWeight * PixelWeight[OffsetY][OffsetX];
+
+							SumValue += Weight * Val;
+						}
+					}
+
+					ChannelData[DstY * SizeX + DstX] = SumValue;
+				}
+			}
+		}
+
+	}
 }
 
 void FImageOverlappedAccumulator::InitMemory(int InPlaneSizeX, int InPlaneSizeY, int InNumChannels)
@@ -158,6 +212,76 @@ void FImageOverlappedAccumulator::GenerateTileWeight(TArray64<float>& Weights, i
 	}
 }
 
+// Any changes to GenerateTileWeight() need to propagate to this function. The subrect should be large enough to incorporate all nonzero values
+void FImageOverlappedAccumulator::GetTileSubRect(int32 & SubRectOffsetX, int32 & SubRectOffsetY, int32 & SubRectSizeX, int32 & SubRectSizeY, const TArray64<float>& Weights, const int32 SizeX, const int32 SizeY)
+{
+	static bool bBruteForce = false;
+	if (bBruteForce)
+	{
+		int32 MinX = SizeX-1;
+		int32 MinY = SizeY-1;
+
+		int32 MaxX = 0;
+		int32 MaxY = 0;
+
+		for (int32 Y = 0; Y < SizeY; Y++)
+		{
+			for (int32 X = 0; X < SizeX; X++)
+			{
+				const float Val = Weights[Y*SizeX + X];
+				if (Val > 0.0f)
+				{
+					MinX = FMath::Min<int32>(MinX, X);
+					MinY = FMath::Min<int32>(MinY, Y);
+
+					MaxX = FMath::Max<int32>(MaxX, X);
+					MaxY = FMath::Max<int32>(MaxY, Y);
+				}
+			}
+		}
+
+		SubRectOffsetX = MinX;
+		SubRectOffsetY = MinY;
+		SubRectSizeX = (MaxX - MinX) + 1;
+		SubRectSizeY = (MaxY - MinY) + 1;
+	}
+	else
+	{
+		// We're using a traingle filter using the middle 75% of the screen. The filter starts at 12.5%, and ends at 87.5%. Then add one pixel
+		// just in case for roundoff/precision error.
+		int32 MinX = FMath::FloorToInt(0.125f * SizeX) - 1;
+		int32 MinY = FMath::FloorToInt(0.125f * SizeY) - 1;
+		int32 MaxX = FMath::CeilToInt(0.875f * SizeX) + 1;
+		int32 MaxY = FMath::CeilToInt(0.875f * SizeY) + 1;
+
+		SubRectOffsetX = MinX;
+		SubRectOffsetY = MinY;
+		SubRectSizeX = (MaxX - MinX) + 1;
+		SubRectSizeY = (MaxY - MinY) + 1;
+	}
+}
+
+void FImageOverlappedAccumulator::CheckTileSubRect(const TArray64<float>& Weights, const int32 SizeX, const int32 SizeY, int32 SubRectOffsetX, int32 SubRectOffsetY, int32 SubRectSizeX, int32 SubRectSizeY)
+{
+	for (int32 Y = 0; Y < SizeY; Y++)
+	{
+		for (int32 X = 0; X < SizeX; X++)
+		{
+			const float Val = Weights[Y*SizeX + X];
+			if (Val > 0.0f)
+			{
+				check(SubRectOffsetX <= X);
+				check(SubRectOffsetY <= Y);
+				check(X < SubRectOffsetX + SubRectSizeX);
+				check(Y < SubRectOffsetY + SubRectSizeY);
+			}
+		}
+	}
+}
+
+// temporary for debugging
+//PRAGMA_DISABLE_OPTIMIZATION
+
 void FImageOverlappedAccumulator::AccumulatePixelData(const FImagePixelData& InPixelData, int32 InTileOffsetX, int32 InTileOffsetY, FVector2D InSubpixelOffset)
 {
 	EImagePixelType Fmt = InPixelData.GetType();
@@ -283,10 +407,13 @@ void FImageOverlappedAccumulator::AccumulatePixelData(const FImagePixelData& InP
 			check(0);
 		}
 
+		static bool bLogTiming = false;
+
 		const double AccumulateEndTime = FPlatformTime::Seconds();
+		if(bLogTiming)
 		{
 			const float ElapsedMs = float((AccumulateEndTime - AccumulateBeginTime)*1000.0f);
-			//UE_LOG(LogTemp, Log, TEXT("    [%8.2f] Unpack time."), ElapsedMs);
+			UE_LOG(LogTemp, Log, TEXT("    [%8.2f] Unpack time."), ElapsedMs);
 		}
 	
 		if (AccumulationGamma != 1.0f)
@@ -306,18 +433,28 @@ void FImageOverlappedAccumulator::AccumulatePixelData(const FImagePixelData& InP
 		} 
 
 		const double GammaEndTime = FPlatformTime::Seconds();
+		if (bLogTiming)
 		{
 			const float ElapsedMs = float((GammaEndTime - AccumulateEndTime)*1000.0f);
-			//UE_LOG(LogTemp, Log, TEXT("        [%8.2f] Gamma time."), ElapsedMs);
+			UE_LOG(LogTemp, Log, TEXT("        [%8.2f] Gamma time."), ElapsedMs);
 		}
 
 		// Calculate weight data for this tile.
 		TArray64<float> Weights;
 		GenerateTileWeight(Weights, RawSize.X, RawSize.Y);
 
+		int SubRectOffsetX = 0;
+		int SubRectOffsetY = 0;
+		int SubRectSizeX = RawSize.X;
+		int SubRectSizeY = RawSize.Y;
+		GetTileSubRect(SubRectOffsetX, SubRectOffsetY, SubRectSizeX, SubRectSizeY, Weights, RawSize.X, RawSize.Y);
+		CheckTileSubRect(Weights, RawSize.X, RawSize.Y, SubRectOffsetX, SubRectOffsetY, SubRectSizeX, SubRectSizeY);
+
+
 		for (int32 ChanIter = 0; ChanIter < NumChannels; ChanIter++)
 		{
-			ChannelPlanes[ChanIter].AccumulateSinglePlane(RawData[ChanIter], Weights, RawSize.X, RawSize.Y, InTileOffsetX, InTileOffsetY, InSubpixelOffset.X, InSubpixelOffset.Y);
+			ChannelPlanes[ChanIter].AccumulateSinglePlane(RawData[ChanIter], Weights, RawSize.X, RawSize.Y, InTileOffsetX, InTileOffsetY, InSubpixelOffset.X, InSubpixelOffset.Y,
+				SubRectOffsetX, SubRectOffsetY, SubRectSizeX, SubRectSizeY);
 		}
 
 		// Accumulate weights as well.
@@ -329,8 +466,17 @@ void FImageOverlappedAccumulator::AccumulatePixelData(const FImagePixelData& InP
 				VecOne[Index] = 1.0f;
 			}
 
-			WeightPlane.AccumulateSinglePlane(VecOne, Weights, RawSize.X, RawSize.Y, InTileOffsetX, InTileOffsetY, InSubpixelOffset.X, InSubpixelOffset.Y);
+			WeightPlane.AccumulateSinglePlane(VecOne, Weights, RawSize.X, RawSize.Y, InTileOffsetX, InTileOffsetY, InSubpixelOffset.X, InSubpixelOffset.Y,
+				SubRectOffsetX, SubRectOffsetY, SubRectSizeX, SubRectSizeY);
 		}
+
+		const double PlaneEndTime = FPlatformTime::Seconds();
+		if (bLogTiming)
+		{
+			const float ElapsedMs = float((PlaneEndTime - GammaEndTime)*1000.0f);
+			UE_LOG(LogTemp, Log, TEXT("        [%8.2f] Plane time."), ElapsedMs);
+		}
+
 	}
 }
 
diff --git a/Engine/Plugins/MovieScene/MovieRenderPipeline/Source/MovieRenderPipelineCore/Public/MovieRenderOverlappedImage.h b/Engine/Plugins/MovieScene/MovieRenderPipeline/Source/MovieRenderPipelineCore/Public/MovieRenderOverlappedImage.h
--- a/Engine/Plugins/MovieScene/MovieRenderPipeline/Source/MovieRenderPipelineCore/Public/MovieRenderOverlappedImage.h
+++ b/Engine/Plugins/MovieScene/MovieRenderPipeline/Source/MovieRenderPipelineCore/Public/MovieRenderOverlappedImage.h
@@ -44,6 +44,8 @@ public:
 
 	/**
 	 * Accumulate a single tile to this plane. The raw data will in general be smaller than the full plane.
+	 * Addtionally, we are only going to be using part of the input plane because the input plane will
+	 * have unused border areas. The SubRect parameters describe the area inside the source plane.
 	 *
 	 * @param InRawData - Raw data to accumulate
 	 * @param InRawSizeX - Width of the tile. Must exactly match.
@@ -52,7 +54,9 @@ public:
 	 * @param InSampleOffsetY - Pixel offset of the tile (in Y)
 	 */
 	void AccumulateSinglePlane(const TArray64<float>& InRawData, const TArray64<float>& InWeightData, int32 InSizeX, int32 InSizeY, int32 InOffsetX, int32 InOffsetY,
-								float SubpixelOffsetX, float SubpixelOffsetY);
+								float SubpixelOffsetX, float SubpixelOffsetY,
+								int32 SubRectOffsetX, int32 SubRectOffsetY,
+								int32 SubRectSizeX, int32 SubRectSizeY);
 
 	/** Actual channel data*/
 	TArray64<float> ChannelData;
@@ -107,6 +111,9 @@ public:
 
 	static void GenerateTileWeight(TArray64<float>& Weights, int32 SizeX, int32 SizeY);
 
+	static void GetTileSubRect(int32 & SubRectOffsetX, int32 & SubRectOffsetY, int32 & SubRectSizeX, int32 & SubRectSizeY, const TArray64<float>& Weights, const int32 SizeX, const int32 SizeY);
+
+	static void CheckTileSubRect(const TArray64<float>& Weights, const int32 SizeX, const int32 SizeY, int32 SubRectOffsetX, int32 SubRectOffsetY, int32 SubRectSizeX, int32 SubRectSizeY);
 
 	/**
 	 * Given a rendered tile, accumulate the data to the full size image.
