Optimizing overlapped tile accumulation, mostly vectorization.

diff --git a/Engine/Plugins/MovieScene/MovieRenderPipeline/Source/MovieRenderPipelineCore/Private/MovieRenderOverlappedImage.cpp b/Engine/Plugins/MovieScene/MovieRenderPipeline/Source/MovieRenderPipelineCore/Private/MovieRenderOverlappedImage.cpp
--- a/Engine/Plugins/MovieScene/MovieRenderPipeline/Source/MovieRenderPipelineCore/Private/MovieRenderOverlappedImage.cpp
+++ b/Engine/Plugins/MovieScene/MovieRenderPipeline/Source/MovieRenderPipelineCore/Private/MovieRenderOverlappedImage.cpp
@@ -68,27 +68,66 @@ void FImageOverlappedPlane::AccumulateSinglePlane(const TArray64<float>& InRawDa
 	static bool bSlowReference = false;
 	if (bSlowReference)
 	{
-		// Slow, reference version. Maybe optimize later.
-		for (int CurrY = 0; CurrY < InSizeY; CurrY++)
+		static bool bReversedReference = false;
+		if (bReversedReference)
 		{
-			for (int CurrX = 0; CurrX < InSizeX; CurrX++)
+			// This is the reversed reference. Instead of taking a tile and applying it to the accumulation, we start from the
+			// accumulation point and figure out the source pixels that affect it. I.e. all the math is backwards, but it should
+			// be faster.
+			//
+			// Given a position on the current tile (x,y), we apply the sum of the 4 points (x+0,y+0), (x+1,y+0), (x+0,y+1), (x+1,y+1) to the destination index.
+			// So in reverse, given a destination position, our source sample positions are (x-1,y-1), (x+0,y-1), (x-1,y+0), (x+0,y+0).
+			// That's why we make sure to start at a minimum of index 1, instead of 0.
+			for (int CurrY = 1; CurrY < InSizeY; CurrY++)
 			{
-				float Val = InRawData[CurrY * InSizeX + CurrX];
-				float BaseWeight = InWeightData[CurrY * InSizeX + CurrX];
-
-				for (int OffsetY = 0; OffsetY < 2; OffsetY++)
+				for (int CurrX = 1; CurrX < InSizeX; CurrX++)
 				{
-					for (int OffsetX = 0; OffsetX < 2; OffsetX++)
+					int DstY = StartY + CurrY;// +OffsetY;
+					int DstX = StartX + CurrX;// +OffsetX;
+
+					if (DstX >= 0 && DstY >= 0 &&
+						DstX < SizeX && DstY < SizeY)
 					{
-						int DstY = StartY + CurrY + OffsetY;
-						int DstX = StartX + CurrX + OffsetX;
+						for (int OffsetY = 0; OffsetY < 2; OffsetY++)
+						{
+							for (int OffsetX = 0; OffsetX < 2; OffsetX++)
+							{
+								float Val = InRawData[(CurrY - 1 + OffsetY) * InSizeX + (CurrX - 1 + OffsetX)];
+								float BaseWeight = InWeightData[(CurrY - 1 + OffsetY) * InSizeX + (CurrX - 1 + OffsetX)];
 
-						float Weight = BaseWeight * PixelWeight[OffsetY][OffsetX];
+								float Weight = BaseWeight * PixelWeight[1-OffsetY][1-OffsetX];
+
+								ChannelData[DstY * SizeX + DstX] += Weight * Val;
+							}
+						}
+					}
+				}
+			}
+		}
+		else
+		{
+			// Slow, reference version. This is the main one.
+			for (int CurrY = 0; CurrY < InSizeY; CurrY++)
+			{
+				for (int CurrX = 0; CurrX < InSizeX; CurrX++)
+				{
+					float Val = InRawData[CurrY * InSizeX + CurrX];
+					float BaseWeight = InWeightData[CurrY * InSizeX + CurrX];
 
-						if (DstX >= 0 && DstY >= 0 &&
-							DstX < SizeX && DstY < SizeY)
+					for (int OffsetY = 0; OffsetY < 2; OffsetY++)
+					{
+						for (int OffsetX = 0; OffsetX < 2; OffsetX++)
 						{
-							ChannelData[DstY * SizeX + DstX] += Weight * Val;
+							int DstY = StartY + CurrY + OffsetY;
+							int DstX = StartX + CurrX + OffsetX;
+
+							float Weight = BaseWeight * PixelWeight[OffsetY][OffsetX];
+
+							if (DstX >= 0 && DstY >= 0 &&
+								DstX < SizeX && DstY < SizeY)
+							{
+								ChannelData[DstY * SizeX + DstX] += Weight * Val;
+							}
 						}
 					}
 				}
@@ -98,51 +137,146 @@ void FImageOverlappedPlane::AccumulateSinglePlane(const TArray64<float>& InRawDa
 	else
 	{
 		// The initial destination corners, may go off the destination window.
-		int DstX0 = StartX;
-		int DstY0 = StartY;
-		int DstX1 = StartX + InSizeX;
-		int DstY1 = StartY + InSizeY;
+		int32 ActualDstX0 = StartX + SubRectOffsetX;
+		int32 ActualDstY0 = StartY + SubRectOffsetY;
+		int32 ActualDstX1 = StartX + SubRectOffsetX + SubRectSizeX;
+		int32 ActualDstY1 = StartY + SubRectOffsetY + SubRectSizeY;
 
-		// Initial destination corners
-		
-		// Given a position on the curre tile (x,y), we apply the sum of the 4 points (x+0,y+0), (x+1,y+0), (x+0,y+1), (x+1,y+1) to the destination index.
-		// So in reverse, given a destination position, our source sample positions are (x-1,y-1), (x+0,y-1), (x-1,y+0), (x+0,y+0).
-		// That's why we make sure to start at a minimum of index 1, instead of 0.
+		ActualDstX0 = FMath::Clamp<int32>(ActualDstX0, 0, SizeX);
+		ActualDstX1 = FMath::Clamp<int32>(ActualDstX1, 0, SizeX);
+
+		ActualDstY0 = FMath::Clamp<int32>(ActualDstY0, 0, SizeY);
+		ActualDstY1 = FMath::Clamp<int32>(ActualDstY1, 0, SizeY);
 
-		for (int CurrY = 1; CurrY < InSizeY; CurrY++)
+		const float PixelWeight00 = PixelWeight[0][0];
+		const float PixelWeight01 = PixelWeight[0][1];
+		const float PixelWeight10 = PixelWeight[1][0];
+		const float PixelWeight11 = PixelWeight[1][1];
+
+		static bool bRunVectorized = false;
+		if (!bRunVectorized)
 		{
-			for (int CurrX = 1; CurrX < InSizeX; CurrX++)
+			for (int32 DstY = ActualDstY0; DstY < ActualDstY1; DstY++)
 			{
-				int DstY = StartY + CurrY;
-				int DstX = StartX + CurrX;
+				int32 CurrY = DstY - StartY;
+				int32 SrcY0 = CurrY + (0 - 1);
+				int32 SrcY1 = CurrY + (1 - 1);
+
+				const float * SrcLineRaw0 = &InRawData[SrcY0 * InSizeX];
+				const float * SrcLineRaw1 = &InRawData[SrcY1 * InSizeX];
 
-				if (DstX >= 0 && DstY >= 0 &&
-					DstX < SizeX && DstY < SizeY)
+				const float * SrcLineWeight0 = &InWeightData[SrcY0 * InSizeX];
+				const float * SrcLineWeight1 = &InWeightData[SrcY1 * InSizeX];
+
+				float * DstLine = &ChannelData[DstY * SizeX];
+				for (int32 DstX = ActualDstX0; DstX < ActualDstX1; DstX++)
 				{
-					float OriginalValue = ChannelData[DstY * SizeX + DstX];
-					float SumValue = OriginalValue;
+					int32 CurrX = DstX - StartX;
+					int32 X0 = CurrX - 1 + 0;
+					int32 X1 = CurrX - 1 + 1;
 
-					for (int OffsetY = 0; OffsetY < 2; OffsetY++)
+					float Sum = DstLine[DstX];
+
+					Sum += SrcLineWeight0[X0] * PixelWeight[1][1] * SrcLineRaw0[X0];
+					Sum += SrcLineWeight0[X1] * PixelWeight[1][0] * SrcLineRaw0[X1];
+					Sum += SrcLineWeight1[X0] * PixelWeight[0][1] * SrcLineRaw1[X0];
+					Sum += SrcLineWeight1[X1] * PixelWeight[0][0] * SrcLineRaw1[X1];
+
+					DstLine[DstX] = Sum;
+				}
+			}
+		}
+		else
+		{
+			int32 ActualWidth = ActualDstX1 - ActualDstX0;
+
+			// how many groups of 4
+			int32 ActualWidthGroup4 = ActualWidth / 4;
+
+			// extra 0-3
+			int32 ActualWidthExtra4 = ActualWidth - 4 * ActualWidthGroup4;
+
+			VectorRegister VecPixelWeight00 = VectorSetFloat1(PixelWeight[0][0]);
+			VectorRegister VecPixelWeight01 = VectorSetFloat1(PixelWeight[0][1]);
+			VectorRegister VecPixelWeight10 = VectorSetFloat1(PixelWeight[1][0]);
+			VectorRegister VecPixelWeight11 = VectorSetFloat1(PixelWeight[1][1]);
+
+			for (int32 DstY = ActualDstY0; DstY < ActualDstY1; DstY++)
+			{
+				int32 CurrY = DstY - StartY;
+				int32 SrcY0 = CurrY + (0 - 1);
+				int32 SrcY1 = CurrY + (1 - 1);
+
+				const float * SrcLineRaw0 = &InRawData[SrcY0 * InSizeX];
+				const float * SrcLineRaw1 = &InRawData[SrcY1 * InSizeX];
+
+				const float * SrcLineWeight0 = &InWeightData[SrcY0 * InSizeX];
+				const float * SrcLineWeight1 = &InWeightData[SrcY1 * InSizeX];
+
+				float * DstLine = &ChannelData[DstY * SizeX];
+
+				for (int32 GroupIter = 0; GroupIter < ActualWidthGroup4; GroupIter++)
+				{
 					{
-						for (int OffsetX = 0; OffsetX < 2; OffsetX++)
-						{
-							int SrcX = CurrX + (OffsetX - 1);
-							int SrcY = CurrY + (OffsetY - 1);
+						int32 BaseDstX = ActualDstX0 + GroupIter * 4;
+
+						int32 BaseCurrX = BaseDstX - StartX;
+
+						int32 BaseX0 = BaseCurrX - 1 + 0;
+						int32 BaseX1 = BaseCurrX - 1 + 1;
+
+						VectorRegister Sum = VectorLoad(&DstLine[BaseDstX]);
 						
-							float Val = InRawData[SrcY * InSizeX + CurrX];
-							float BaseWeight = InWeightData[SrcY * InSizeX + SrcX];
+						Sum = VectorMultiplyAdd(VecPixelWeight11, VectorMultiply(VectorLoad(&SrcLineWeight0[BaseX0]), VectorLoad(&SrcLineRaw0[BaseX0])), Sum);
+						Sum = VectorMultiplyAdd(VecPixelWeight10, VectorMultiply(VectorLoad(&SrcLineWeight0[BaseX1]), VectorLoad(&SrcLineRaw0[BaseX1])), Sum);
+						Sum = VectorMultiplyAdd(VecPixelWeight01, VectorMultiply(VectorLoad(&SrcLineWeight1[BaseX0]), VectorLoad(&SrcLineRaw1[BaseX0])), Sum);
+						Sum = VectorMultiplyAdd(VecPixelWeight00, VectorMultiply(VectorLoad(&SrcLineWeight1[BaseX1]), VectorLoad(&SrcLineRaw1[BaseX1])), Sum);
 
-							float Weight = BaseWeight * PixelWeight[OffsetY][OffsetX];
+						VectorStore(Sum, &DstLine[BaseDstX]);
+					}
+#if 0
+					// Scalar version of the same code for reference
+					for (int32 Iter = 0; Iter < 4; Iter++)
+					{
+						int32 DstX = ActualDstX0 + GroupIter * 4 + Iter;
 
-							SumValue += Weight * Val;
-						}
+						int32 CurrX = DstX - StartX;
+						int32 X0 = CurrX - 1 + 0;
+						int32 X1 = CurrX - 1 + 1;
+
+						float Sum = DstLine[DstX];
+
+						Sum += PixelWeight11 * SrcLineWeight0[X0] * SrcLineRaw0[X0];
+						Sum += PixelWeight10 * SrcLineWeight0[X1] * SrcLineRaw0[X1];
+						Sum += PixelWeight01 * SrcLineWeight1[X0] * SrcLineRaw1[X0];
+						Sum += PixelWeight00 * SrcLineWeight1[X1] * SrcLineRaw1[X1];
+
+						DstLine[DstX] = Sum;
 					}
+#endif
 
-					ChannelData[DstY * SizeX + DstX] = SumValue;
 				}
+
+				for (int32 IterExtra = 0; IterExtra < ActualWidthExtra4; IterExtra++)
+				{
+					int32 DstX = ActualDstX0 + ActualWidthExtra4 * 4 + IterExtra;
+
+					int32 CurrX = DstX - StartX;
+					int32 X0 = CurrX - 1 + 0;
+					int32 X1 = CurrX - 1 + 1;
+
+					float Sum = DstLine[DstX];
+
+					Sum += PixelWeight11 * SrcLineWeight0[X0] * SrcLineRaw0[X0];
+					Sum += PixelWeight10 * SrcLineWeight0[X1] * SrcLineRaw0[X1];
+					Sum += PixelWeight01 * SrcLineWeight1[X0] * SrcLineRaw1[X0];
+					Sum += PixelWeight00 * SrcLineWeight1[X1] * SrcLineRaw1[X1];
+
+					DstLine[DstX] = Sum;
+				}
+
 			}
 		}
-
 	}
 }
 
@@ -279,9 +413,6 @@ void FImageOverlappedAccumulator::CheckTileSubRect(const TArray64<float>& Weight
 	}
 }
 
-// temporary for debugging
-//PRAGMA_DISABLE_OPTIMIZATION
-
 void FImageOverlappedAccumulator::AccumulatePixelData(const FImagePixelData& InPixelData, int32 InTileOffsetX, int32 InTileOffsetY, FVector2D InSubpixelOffset)
 {
 	EImagePixelType Fmt = InPixelData.GetType();
@@ -436,7 +567,7 @@ void FImageOverlappedAccumulator::AccumulatePixelData(const FImagePixelData& InP
 		if (bLogTiming)
 		{
 			const float ElapsedMs = float((GammaEndTime - AccumulateEndTime)*1000.0f);
-			UE_LOG(LogTemp, Log, TEXT("        [%8.2f] Gamma time."), ElapsedMs);
+			UE_LOG(LogTemp, Log, TEXT("    [%8.2f] Gamma time."), ElapsedMs);
 		}
 
 		// Calculate weight data for this tile.
@@ -448,7 +579,19 @@ void FImageOverlappedAccumulator::AccumulatePixelData(const FImagePixelData& InP
 		int SubRectSizeX = RawSize.X;
 		int SubRectSizeY = RawSize.Y;
 		GetTileSubRect(SubRectOffsetX, SubRectOffsetY, SubRectSizeX, SubRectSizeY, Weights, RawSize.X, RawSize.Y);
-		CheckTileSubRect(Weights, RawSize.X, RawSize.Y, SubRectOffsetX, SubRectOffsetY, SubRectSizeX, SubRectSizeY);
+		
+		const bool bVerifySubRect = false;
+		if (bVerifySubRect)
+		{
+			CheckTileSubRect(Weights, RawSize.X, RawSize.Y, SubRectOffsetX, SubRectOffsetY, SubRectSizeX, SubRectSizeY);
+		}
+
+		const double TileEndTime = FPlatformTime::Seconds();
+		if (bLogTiming)
+		{
+			const float ElapsedMs = float((TileEndTime - GammaEndTime)*1000.0f);
+			UE_LOG(LogTemp, Log, TEXT("    [%8.2f] Tile time."), ElapsedMs);
+		}
 
 
 		for (int32 ChanIter = 0; ChanIter < NumChannels; ChanIter++)
@@ -473,8 +616,8 @@ void FImageOverlappedAccumulator::AccumulatePixelData(const FImagePixelData& InP
 		const double PlaneEndTime = FPlatformTime::Seconds();
 		if (bLogTiming)
 		{
-			const float ElapsedMs = float((PlaneEndTime - GammaEndTime)*1000.0f);
-			UE_LOG(LogTemp, Log, TEXT("        [%8.2f] Plane time."), ElapsedMs);
+			const float ElapsedMs = float((PlaneEndTime - TileEndTime)*1000.0f);
+			UE_LOG(LogTemp, Log, TEXT("    [%8.2f] Plane time."), ElapsedMs);
 		}
 
 	}
