Caching the tile mask in the overlapped accumulator since it doesn't change.

diff --git a/Engine/Plugins/MovieScene/MovieRenderPipeline/Source/MovieRenderPipelineCore/Private/MovieRenderOverlappedImage.cpp b/Engine/Plugins/MovieScene/MovieRenderPipeline/Source/MovieRenderPipelineCore/Private/MovieRenderOverlappedImage.cpp
--- a/Engine/Plugins/MovieScene/MovieRenderPipeline/Source/MovieRenderPipelineCore/Private/MovieRenderOverlappedImage.cpp
+++ b/Engine/Plugins/MovieScene/MovieRenderPipeline/Source/MovieRenderPipelineCore/Private/MovieRenderOverlappedImage.cpp
@@ -294,6 +294,10 @@ void FImageOverlappedAccumulator::InitMemory(int InPlaneSizeX, int InPlaneSizeY,
 	}
 
 	WeightPlane.Init(PlaneSizeX, PlaneSizeY);
+
+	TileMaskData.Empty();
+	TileMaskSizeX = 0;
+	TileMaskSizeY = 0;
 }
 
 void FImageOverlappedAccumulator::ZeroPlanes()
@@ -316,6 +320,11 @@ void FImageOverlappedAccumulator::Reset()
 
 	// Let the desctructor clean up
 	ChannelPlanes.Empty();
+	WeightPlane.Reset();
+
+	TileMaskData.Empty();
+	TileMaskSizeX = 0;
+	TileMaskSizeY = 0;
 }
 
 
@@ -570,20 +579,28 @@ void FImageOverlappedAccumulator::AccumulatePixelData(const FImagePixelData& InP
 			UE_LOG(LogTemp, Log, TEXT("    [%8.2f] Gamma time."), ElapsedMs);
 		}
 
-		// Calculate weight data for this tile.
-		TArray64<float> Weights;
-		GenerateTileWeight(Weights, RawSize.X, RawSize.Y);
+		// Calculate weight data for this tile. If it is the same, use the cached size. In theory it is allowed
+		// to have a bunch of tiles of different sizes, but so far we are only using the same size for all tiles.
+		if (RawSize.X != TileMaskSizeX ||
+			RawSize.Y != TileMaskSizeY ||
+			TileMaskSizeX * TileMaskSizeY != TileMaskData.Num())
+		{
+			TileMaskSizeX = RawSize.X;
+			TileMaskSizeY = RawSize.Y;
 
+			GenerateTileWeight(TileMaskData, RawSize.X, RawSize.Y);
+		}
+		
 		int SubRectOffsetX = 0;
 		int SubRectOffsetY = 0;
 		int SubRectSizeX = RawSize.X;
 		int SubRectSizeY = RawSize.Y;
-		GetTileSubRect(SubRectOffsetX, SubRectOffsetY, SubRectSizeX, SubRectSizeY, Weights, RawSize.X, RawSize.Y);
+		GetTileSubRect(SubRectOffsetX, SubRectOffsetY, SubRectSizeX, SubRectSizeY, TileMaskData, RawSize.X, RawSize.Y);
 		
 		const bool bVerifySubRect = false;
 		if (bVerifySubRect)
 		{
-			CheckTileSubRect(Weights, RawSize.X, RawSize.Y, SubRectOffsetX, SubRectOffsetY, SubRectSizeX, SubRectSizeY);
+			CheckTileSubRect(TileMaskData, RawSize.X, RawSize.Y, SubRectOffsetX, SubRectOffsetY, SubRectSizeX, SubRectSizeY);
 		}
 
 		const double TileEndTime = FPlatformTime::Seconds();
@@ -596,7 +613,7 @@ void FImageOverlappedAccumulator::AccumulatePixelData(const FImagePixelData& InP
 
 		for (int32 ChanIter = 0; ChanIter < NumChannels; ChanIter++)
 		{
-			ChannelPlanes[ChanIter].AccumulateSinglePlane(RawData[ChanIter], Weights, RawSize.X, RawSize.Y, InTileOffsetX, InTileOffsetY, InSubpixelOffset.X, InSubpixelOffset.Y,
+			ChannelPlanes[ChanIter].AccumulateSinglePlane(RawData[ChanIter], TileMaskData, RawSize.X, RawSize.Y, InTileOffsetX, InTileOffsetY, InSubpixelOffset.X, InSubpixelOffset.Y,
 				SubRectOffsetX, SubRectOffsetY, SubRectSizeX, SubRectSizeY);
 		}
 
@@ -609,7 +626,7 @@ void FImageOverlappedAccumulator::AccumulatePixelData(const FImagePixelData& InP
 				VecOne[Index] = 1.0f;
 			}
 
-			WeightPlane.AccumulateSinglePlane(VecOne, Weights, RawSize.X, RawSize.Y, InTileOffsetX, InTileOffsetY, InSubpixelOffset.X, InSubpixelOffset.Y,
+			WeightPlane.AccumulateSinglePlane(VecOne, TileMaskData, RawSize.X, RawSize.Y, InTileOffsetX, InTileOffsetY, InSubpixelOffset.X, InSubpixelOffset.Y,
 				SubRectOffsetX, SubRectOffsetY, SubRectSizeX, SubRectSizeY);
 		}
 
diff --git a/Engine/Plugins/MovieScene/MovieRenderPipeline/Source/MovieRenderPipelineCore/Public/MovieRenderOverlappedImage.h b/Engine/Plugins/MovieScene/MovieRenderPipeline/Source/MovieRenderPipelineCore/Public/MovieRenderOverlappedImage.h
--- a/Engine/Plugins/MovieScene/MovieRenderPipeline/Source/MovieRenderPipelineCore/Public/MovieRenderOverlappedImage.h
+++ b/Engine/Plugins/MovieScene/MovieRenderPipeline/Source/MovieRenderPipelineCore/Public/MovieRenderOverlappedImage.h
@@ -82,6 +82,8 @@ public:
 		, PlaneSizeY(0)
 		, NumChannels(0)
 		, AccumulationGamma(1.0f)
+		, TileMaskSizeX(0)
+		, TileMaskSizeY(0)
 	{
 	}
 
@@ -169,6 +171,14 @@ public:
 	
 	FImageOverlappedPlane WeightPlane;
 
+	/** Weights image. Recalculate if it changes. */
+	TArray64<float> TileMaskData;
+
+	/** Width of weight image. */
+	int32 TileMaskSizeX;
+
+	/** Height of weight image. */
+	int32 TileMaskSizeY;
 };
 
 
