In overlapped image accumulator, replacing pairs of ints with FIntPoint structs.

diff --git a/Engine/Plugins/MovieScene/MovieRenderPipeline/Source/MovieRenderPipelineCore/Private/MovieRenderOverlappedImage.cpp b/Engine/Plugins/MovieScene/MovieRenderPipeline/Source/MovieRenderPipelineCore/Private/MovieRenderOverlappedImage.cpp
--- a/Engine/Plugins/MovieScene/MovieRenderPipeline/Source/MovieRenderPipelineCore/Private/MovieRenderOverlappedImage.cpp
+++ b/Engine/Plugins/MovieScene/MovieRenderPipeline/Source/MovieRenderPipelineCore/Private/MovieRenderOverlappedImage.cpp
@@ -9,18 +9,17 @@
 #include "Math/Float16.h"
 #include "HAL/PlatformTime.h"
 
-void FImageOverlappedPlane::Init(int32 InSizeX, int32 InSizeY)
+void FImageOverlappedPlane::Init(FIntPoint InSize)
 {
-	SizeX = InSizeX;
-	SizeY = InSizeY;
+	Size = InSize;
 
 	// leaves the data as is
-	ChannelData.SetNumUninitialized(SizeX * SizeY);
+	ChannelData.SetNumUninitialized(Size.X * Size.Y);
 }
 
 void FImageOverlappedPlane::ZeroPlane()
 {
-	int32 Num = SizeX * SizeY;
+	int32 Num = Size.X * Size.Y;
 	for (int32 Index = 0; Index < Num; Index++)
 	{
 		ChannelData[Index] = 0.0f;
@@ -29,19 +28,19 @@ void FImageOverlappedPlane::ZeroPlane()
 
 void FImageOverlappedPlane::Reset()
 {
-	SizeX = 0;
-	SizeY = 0;
+	Size.X = 0;
+	Size.Y = 0;
 	ChannelData.Empty();
 }
 
 // a subpixel offset of (0.5,0.5) means that the raw data is exactly centered on destination pixels.
-void FImageOverlappedPlane::AccumulateSinglePlane(const TArray64<float>& InRawData, const TArray64<float>& InWeightData, int32 InSizeX, int32 InSizeY, int32 InOffsetX, int32 InOffsetY,
+void FImageOverlappedPlane::AccumulateSinglePlane(const TArray64<float>& InRawData, const TArray64<float>& InWeightData, FIntPoint InSize, FIntPoint InOffset,
 								float SubpixelOffsetX, float SubpixelOffsetY,
-								int32 SubRectOffsetX, int32 SubRectOffsetY,
-								int32 SubRectSizeX, int32 SubRectSizeY)
+								FIntPoint SubRectOffset,
+								FIntPoint SubRectSize)
 {
-	check(InRawData.Num() == InSizeX * InSizeY);
-	check(InWeightData.Num() == InSizeX * InSizeY);
+	check(InRawData.Num() == InSize.X * InSize.Y);
+	check(InWeightData.Num() == InSize.X * InSize.Y);
 
 	check(SubpixelOffsetX >= 0.0f);
 	check(SubpixelOffsetX <= 1.0f);
@@ -49,8 +48,8 @@ void FImageOverlappedPlane::AccumulateSinglePlane(const TArray64<float>& InRawDa
 	check(SubpixelOffsetY <= 1.0f);
 
 	// if the subpixel offset is less than 0.5, go back one pixel
-	int32 StartX = (SubpixelOffsetX >= 0.5 ? InOffsetX : InOffsetX - 1);
-	int32 StartY = (SubpixelOffsetY >= 0.5 ? InOffsetY : InOffsetY - 1);
+	int32 StartX = (SubpixelOffsetX >= 0.5 ? InOffset.X : InOffset.X - 1);
+	int32 StartY = (SubpixelOffsetY >= 0.5 ? InOffset.Y : InOffset.Y - 1);
 
 	float PixelWeight[2][2];
 	{
@@ -78,26 +77,26 @@ void FImageOverlappedPlane::AccumulateSinglePlane(const TArray64<float>& InRawDa
 			// Given a position on the current tile (x,y), we apply the sum of the 4 points (x+0,y+0), (x+1,y+0), (x+0,y+1), (x+1,y+1) to the destination index.
 			// So in reverse, given a destination position, our source sample positions are (x-1,y-1), (x+0,y-1), (x-1,y+0), (x+0,y+0).
 			// That's why we make sure to start at a minimum of index 1, instead of 0.
-			for (int CurrY = 1; CurrY < InSizeY; CurrY++)
+			for (int CurrY = 1; CurrY < InSize.Y; CurrY++)
 			{
-				for (int CurrX = 1; CurrX < InSizeX; CurrX++)
+				for (int CurrX = 1; CurrX < InSize.X; CurrX++)
 				{
 					int DstY = StartY + CurrY;// +OffsetY;
 					int DstX = StartX + CurrX;// +OffsetX;
 
 					if (DstX >= 0 && DstY >= 0 &&
-						DstX < SizeX && DstY < SizeY)
+						DstX < Size.X && DstY < Size.Y)
 					{
 						for (int OffsetY = 0; OffsetY < 2; OffsetY++)
 						{
 							for (int OffsetX = 0; OffsetX < 2; OffsetX++)
 							{
-								float Val = InRawData[(CurrY - 1 + OffsetY) * InSizeX + (CurrX - 1 + OffsetX)];
-								float BaseWeight = InWeightData[(CurrY - 1 + OffsetY) * InSizeX + (CurrX - 1 + OffsetX)];
+								float Val = InRawData[(CurrY - 1 + OffsetY) * InSize.X + (CurrX - 1 + OffsetX)];
+								float BaseWeight = InWeightData[(CurrY - 1 + OffsetY) * InSize.X + (CurrX - 1 + OffsetX)];
 
 								float Weight = BaseWeight * PixelWeight[1-OffsetY][1-OffsetX];
 
-								ChannelData[DstY * SizeX + DstX] += Weight * Val;
+								ChannelData[DstY * Size.X + DstX] += Weight * Val;
 							}
 						}
 					}
@@ -107,12 +106,12 @@ void FImageOverlappedPlane::AccumulateSinglePlane(const TArray64<float>& InRawDa
 		else
 		{
 			// Slow, reference version. This is the main one.
-			for (int CurrY = 0; CurrY < InSizeY; CurrY++)
+			for (int CurrY = 0; CurrY < InSize.Y; CurrY++)
 			{
-				for (int CurrX = 0; CurrX < InSizeX; CurrX++)
+				for (int CurrX = 0; CurrX < InSize.X; CurrX++)
 				{
-					float Val = InRawData[CurrY * InSizeX + CurrX];
-					float BaseWeight = InWeightData[CurrY * InSizeX + CurrX];
+					float Val = InRawData[CurrY * InSize.X + CurrX];
+					float BaseWeight = InWeightData[CurrY * InSize.X + CurrX];
 
 					for (int OffsetY = 0; OffsetY < 2; OffsetY++)
 					{
@@ -124,9 +123,9 @@ void FImageOverlappedPlane::AccumulateSinglePlane(const TArray64<float>& InRawDa
 							float Weight = BaseWeight * PixelWeight[OffsetY][OffsetX];
 
 							if (DstX >= 0 && DstY >= 0 &&
-								DstX < SizeX && DstY < SizeY)
+								DstX < Size.X && DstY < Size.Y)
 							{
-								ChannelData[DstY * SizeX + DstX] += Weight * Val;
+								ChannelData[DstY * Size.X + DstX] += Weight * Val;
 							}
 						}
 					}
@@ -137,16 +136,16 @@ void FImageOverlappedPlane::AccumulateSinglePlane(const TArray64<float>& InRawDa
 	else
 	{
 		// The initial destination corners, may go off the destination window.
-		int32 ActualDstX0 = StartX + SubRectOffsetX;
-		int32 ActualDstY0 = StartY + SubRectOffsetY;
-		int32 ActualDstX1 = StartX + SubRectOffsetX + SubRectSizeX;
-		int32 ActualDstY1 = StartY + SubRectOffsetY + SubRectSizeY;
+		int32 ActualDstX0 = StartX + SubRectOffset.X;
+		int32 ActualDstY0 = StartY + SubRectOffset.Y;
+		int32 ActualDstX1 = StartX + SubRectOffset.X + SubRectSize.X;
+		int32 ActualDstY1 = StartY + SubRectOffset.Y + SubRectSize.Y;
 
-		ActualDstX0 = FMath::Clamp<int32>(ActualDstX0, 0, SizeX);
-		ActualDstX1 = FMath::Clamp<int32>(ActualDstX1, 0, SizeX);
+		ActualDstX0 = FMath::Clamp<int32>(ActualDstX0, 0, Size.X);
+		ActualDstX1 = FMath::Clamp<int32>(ActualDstX1, 0, Size.X);
 
-		ActualDstY0 = FMath::Clamp<int32>(ActualDstY0, 0, SizeY);
-		ActualDstY1 = FMath::Clamp<int32>(ActualDstY1, 0, SizeY);
+		ActualDstY0 = FMath::Clamp<int32>(ActualDstY0, 0, Size.Y);
+		ActualDstY1 = FMath::Clamp<int32>(ActualDstY1, 0, Size.Y);
 
 		const float PixelWeight00 = PixelWeight[0][0];
 		const float PixelWeight01 = PixelWeight[0][1];
@@ -162,13 +161,13 @@ void FImageOverlappedPlane::AccumulateSinglePlane(const TArray64<float>& InRawDa
 				int32 SrcY0 = CurrY + (0 - 1);
 				int32 SrcY1 = CurrY + (1 - 1);
 
-				const float * SrcLineRaw0 = &InRawData[SrcY0 * InSizeX];
-				const float * SrcLineRaw1 = &InRawData[SrcY1 * InSizeX];
+				const float * SrcLineRaw0 = &InRawData[SrcY0 * InSize.X];
+				const float * SrcLineRaw1 = &InRawData[SrcY1 * InSize.X];
 
-				const float * SrcLineWeight0 = &InWeightData[SrcY0 * InSizeX];
-				const float * SrcLineWeight1 = &InWeightData[SrcY1 * InSizeX];
+				const float * SrcLineWeight0 = &InWeightData[SrcY0 * InSize.X];
+				const float * SrcLineWeight1 = &InWeightData[SrcY1 * InSize.X];
 
-				float * DstLine = &ChannelData[DstY * SizeX];
+				float * DstLine = &ChannelData[DstY * Size.X];
 				for (int32 DstX = ActualDstX0; DstX < ActualDstX1; DstX++)
 				{
 					int32 CurrX = DstX - StartX;
@@ -207,13 +206,13 @@ void FImageOverlappedPlane::AccumulateSinglePlane(const TArray64<float>& InRawDa
 				int32 SrcY0 = CurrY + (0 - 1);
 				int32 SrcY1 = CurrY + (1 - 1);
 
-				const float * SrcLineRaw0 = &InRawData[SrcY0 * InSizeX];
-				const float * SrcLineRaw1 = &InRawData[SrcY1 * InSizeX];
+				const float * SrcLineRaw0 = &InRawData[SrcY0 * InSize.X];
+				const float * SrcLineRaw1 = &InRawData[SrcY1 * InSize.X];
 
-				const float * SrcLineWeight0 = &InWeightData[SrcY0 * InSizeX];
-				const float * SrcLineWeight1 = &InWeightData[SrcY1 * InSizeX];
+				const float * SrcLineWeight0 = &InWeightData[SrcY0 * InSize.X];
+				const float * SrcLineWeight1 = &InWeightData[SrcY1 * InSize.X];
 
-				float * DstLine = &ChannelData[DstY * SizeX];
+				float * DstLine = &ChannelData[DstY * Size.X];
 
 				for (int32 GroupIter = 0; GroupIter < ActualWidthGroup4; GroupIter++)
 				{
@@ -280,24 +279,24 @@ void FImageOverlappedPlane::AccumulateSinglePlane(const TArray64<float>& InRawDa
 	}
 }
 
-void FImageOverlappedAccumulator::InitMemory(int InPlaneSizeX, int InPlaneSizeY, int InNumChannels)
+void FImageOverlappedAccumulator::InitMemory(FIntPoint InPlaneSize, int InNumChannels)
 {
-	PlaneSizeX = InPlaneSizeX;
-	PlaneSizeY = InPlaneSizeY;
+	PlaneSize.X = InPlaneSize.X;
+	PlaneSize.Y = InPlaneSize.Y;
 	NumChannels = InNumChannels;
 
 	ChannelPlanes.SetNum(NumChannels);
 
 	for (int Channel = 0; Channel < NumChannels; Channel++)
 	{
-		ChannelPlanes[Channel].Init(PlaneSizeX, PlaneSizeY);
+		ChannelPlanes[Channel].Init(PlaneSize);
 	}
 
-	WeightPlane.Init(PlaneSizeX, PlaneSizeY);
+	WeightPlane.Init(PlaneSize);
 
 	TileMaskData.Empty();
-	TileMaskSizeX = 0;
-	TileMaskSizeY = 0;
+	TileMaskSize.X = 0;
+	TileMaskSize.Y = 0;
 }
 
 void FImageOverlappedAccumulator::ZeroPlanes()
@@ -314,8 +313,8 @@ void FImageOverlappedAccumulator::ZeroPlanes()
 
 void FImageOverlappedAccumulator::Reset()
 {
-	PlaneSizeX = 0;
-	PlaneSizeY = 0;
+	PlaneSize.X = 0;
+	PlaneSize.Y = 0;
 	NumChannels = 0;
 
 	// Let the desctructor clean up
@@ -323,55 +322,55 @@ void FImageOverlappedAccumulator::Reset()
 	WeightPlane.Reset();
 
 	TileMaskData.Empty();
-	TileMaskSizeX = 0;
-	TileMaskSizeY = 0;
+	TileMaskSize.X = 0;
+	TileMaskSize.Y = 0;
 }
 
 
-void FImageOverlappedAccumulator::GenerateTileWeight(TArray64<float>& Weights, int32 SizeX, int32 SizeY)
+void FImageOverlappedAccumulator::GenerateTileWeight(TArray64<float>& Weights, FIntPoint Size)
 {
-	Weights.SetNum(SizeX * SizeY);
+	Weights.SetNum(Size.X * Size.Y);
 
 	// we'll use a simple triangle filter, which goes from 1.0 at the center to 0.0 at 3/4 of the way to the tile edge.
 
-	float ScaleX = 1.0f / (float(SizeX / 2) * .75f);
-	float ScaleY = 1.0f / (float(SizeY / 2) * .75f);
+	float ScaleX = 1.0f / (float(Size.X / 2) * .75f);
+	float ScaleY = 1.0f / (float(Size.Y / 2) * .75f);
 
-	for (int PixY = 0; PixY < SizeY; PixY++)
+	for (int PixY = 0; PixY < Size.Y; PixY++)
 	{
-		for (int PixX = 0; PixX < SizeX; PixX++)
+		for (int PixX = 0; PixX < Size.X; PixX++)
 		{
 			float Y = float(PixY) + .5f;
 			float X = float(PixX) + .5f;
 
-			float DistX = FMath::Abs(SizeX / 2 - X);
-			float DistY = FMath::Abs(SizeY / 2 - Y);
+			float DistX = FMath::Abs(Size.X / 2 - X);
+			float DistY = FMath::Abs(Size.Y / 2 - Y);
 
 			float WeightX = FMath::Clamp(1.0f - DistX * ScaleX, 0.0f, 1.0f);
 			float WeightY = FMath::Clamp(1.0f - DistY * ScaleY, 0.0f, 1.0f);
 
-			Weights[PixY * SizeX + PixX] = WeightX * WeightY;
+			Weights[PixY * Size.X + PixX] = WeightX * WeightY;
 		}
 	}
 }
 
 // Any changes to GenerateTileWeight() need to propagate to this function. The subrect should be large enough to incorporate all nonzero values
-void FImageOverlappedAccumulator::GetTileSubRect(int32 & SubRectOffsetX, int32 & SubRectOffsetY, int32 & SubRectSizeX, int32 & SubRectSizeY, const TArray64<float>& Weights, const int32 SizeX, const int32 SizeY)
+void FImageOverlappedAccumulator::GetTileSubRect(FIntPoint & SubRectOffset, FIntPoint & SubRectSize, const TArray64<float>& Weights, const FIntPoint Size)
 {
 	static bool bBruteForce = false;
 	if (bBruteForce)
 	{
-		int32 MinX = SizeX-1;
-		int32 MinY = SizeY-1;
+		int32 MinX = Size.X-1;
+		int32 MinY = Size.Y-1;
 
 		int32 MaxX = 0;
 		int32 MaxY = 0;
 
-		for (int32 Y = 0; Y < SizeY; Y++)
+		for (int32 Y = 0; Y < Size.Y; Y++)
 		{
-			for (int32 X = 0; X < SizeX; X++)
+			for (int32 X = 0; X < Size.X; X++)
 			{
-				const float Val = Weights[Y*SizeX + X];
+				const float Val = Weights[Y*Size.X + X];
 				if (Val > 0.0f)
 				{
 					MinX = FMath::Min<int32>(MinX, X);
@@ -383,46 +382,46 @@ void FImageOverlappedAccumulator::GetTileSubRect(int32 & SubRectOffsetX, int32 &
 			}
 		}
 
-		SubRectOffsetX = MinX;
-		SubRectOffsetY = MinY;
-		SubRectSizeX = (MaxX - MinX) + 1;
-		SubRectSizeY = (MaxY - MinY) + 1;
+		SubRectOffset.X = MinX;
+		SubRectOffset.Y = MinY;
+		SubRectSize.X = (MaxX - MinX) + 1;
+		SubRectSize.Y = (MaxY - MinY) + 1;
 	}
 	else
 	{
 		// We're using a traingle filter using the middle 75% of the screen. The filter starts at 12.5%, and ends at 87.5%. Then add one pixel
 		// just in case for roundoff/precision error.
-		int32 MinX = FMath::FloorToInt(0.125f * SizeX) - 1;
-		int32 MinY = FMath::FloorToInt(0.125f * SizeY) - 1;
-		int32 MaxX = FMath::CeilToInt(0.875f * SizeX) + 1;
-		int32 MaxY = FMath::CeilToInt(0.875f * SizeY) + 1;
-
-		SubRectOffsetX = MinX;
-		SubRectOffsetY = MinY;
-		SubRectSizeX = (MaxX - MinX) + 1;
-		SubRectSizeY = (MaxY - MinY) + 1;
+		int32 MinX = FMath::FloorToInt(0.125f * Size.X) - 1;
+		int32 MinY = FMath::FloorToInt(0.125f * Size.Y) - 1;
+		int32 MaxX = FMath::CeilToInt(0.875f * Size.X) + 1;
+		int32 MaxY = FMath::CeilToInt(0.875f * Size.Y) + 1;
+
+		SubRectOffset.X = MinX;
+		SubRectOffset.Y = MinY;
+		SubRectSize.X = (MaxX - MinX) + 1;
+		SubRectSize.Y = (MaxY - MinY) + 1;
 	}
 }
 
-void FImageOverlappedAccumulator::CheckTileSubRect(const TArray64<float>& Weights, const int32 SizeX, const int32 SizeY, int32 SubRectOffsetX, int32 SubRectOffsetY, int32 SubRectSizeX, int32 SubRectSizeY)
+void FImageOverlappedAccumulator::CheckTileSubRect(const TArray64<float>& Weights, const FIntPoint Size, FIntPoint SubRectOffset, FIntPoint SubRectSize)
 {
-	for (int32 Y = 0; Y < SizeY; Y++)
+	for (int32 Y = 0; Y < Size.Y; Y++)
 	{
-		for (int32 X = 0; X < SizeX; X++)
+		for (int32 X = 0; X < Size.X; X++)
 		{
-			const float Val = Weights[Y*SizeX + X];
+			const float Val = Weights[Y*Size.X + X];
 			if (Val > 0.0f)
 			{
-				check(SubRectOffsetX <= X);
-				check(SubRectOffsetY <= Y);
-				check(X < SubRectOffsetX + SubRectSizeX);
-				check(Y < SubRectOffsetY + SubRectSizeY);
+				check(SubRectOffset.X <= X);
+				check(SubRectOffset.Y <= Y);
+				check(X < SubRectOffset.X + SubRectSize.X);
+				check(Y < SubRectOffset.Y + SubRectSize.Y);
 			}
 		}
 	}
 }
 
-void FImageOverlappedAccumulator::AccumulatePixelData(const FImagePixelData& InPixelData, int32 InTileOffsetX, int32 InTileOffsetY, FVector2D InSubpixelOffset)
+void FImageOverlappedAccumulator::AccumulatePixelData(const FImagePixelData& InPixelData, FIntPoint InTileOffset, FVector2D InSubpixelOffset)
 {
 	EImagePixelType Fmt = InPixelData.GetType();
 
@@ -581,26 +580,24 @@ void FImageOverlappedAccumulator::AccumulatePixelData(const FImagePixelData& InP
 
 		// Calculate weight data for this tile. If it is the same, use the cached size. In theory it is allowed
 		// to have a bunch of tiles of different sizes, but so far we are only using the same size for all tiles.
-		if (RawSize.X != TileMaskSizeX ||
-			RawSize.Y != TileMaskSizeY ||
-			TileMaskSizeX * TileMaskSizeY != TileMaskData.Num())
+		if (RawSize.X != TileMaskSize.X ||
+			RawSize.Y != TileMaskSize.Y ||
+			TileMaskSize.X * TileMaskSize.Y != TileMaskData.Num())
 		{
-			TileMaskSizeX = RawSize.X;
-			TileMaskSizeY = RawSize.Y;
+			TileMaskSize.X = RawSize.X;
+			TileMaskSize.Y = RawSize.Y;
 
-			GenerateTileWeight(TileMaskData, RawSize.X, RawSize.Y);
+			GenerateTileWeight(TileMaskData, RawSize);
 		}
 		
-		int SubRectOffsetX = 0;
-		int SubRectOffsetY = 0;
-		int SubRectSizeX = RawSize.X;
-		int SubRectSizeY = RawSize.Y;
-		GetTileSubRect(SubRectOffsetX, SubRectOffsetY, SubRectSizeX, SubRectSizeY, TileMaskData, RawSize.X, RawSize.Y);
+		FIntPoint SubRectOffset(0, 0);
+		FIntPoint SubRectSize(0, 0);
+		GetTileSubRect(SubRectOffset, SubRectSize, TileMaskData, RawSize);
 		
 		const bool bVerifySubRect = false;
 		if (bVerifySubRect)
 		{
-			CheckTileSubRect(TileMaskData, RawSize.X, RawSize.Y, SubRectOffsetX, SubRectOffsetY, SubRectSizeX, SubRectSizeY);
+			CheckTileSubRect(TileMaskData, RawSize, SubRectOffset, SubRectSize);
 		}
 
 		const double TileEndTime = FPlatformTime::Seconds();
@@ -613,8 +610,8 @@ void FImageOverlappedAccumulator::AccumulatePixelData(const FImagePixelData& InP
 
 		for (int32 ChanIter = 0; ChanIter < NumChannels; ChanIter++)
 		{
-			ChannelPlanes[ChanIter].AccumulateSinglePlane(RawData[ChanIter], TileMaskData, RawSize.X, RawSize.Y, InTileOffsetX, InTileOffsetY, InSubpixelOffset.X, InSubpixelOffset.Y,
-				SubRectOffsetX, SubRectOffsetY, SubRectSizeX, SubRectSizeY);
+			ChannelPlanes[ChanIter].AccumulateSinglePlane(RawData[ChanIter], TileMaskData, RawSize, InTileOffset, InSubpixelOffset.X, InSubpixelOffset.Y,
+				SubRectOffset, SubRectSize);
 		}
 
 		// Accumulate weights as well.
@@ -626,8 +623,8 @@ void FImageOverlappedAccumulator::AccumulatePixelData(const FImagePixelData& InP
 				VecOne[Index] = 1.0f;
 			}
 
-			WeightPlane.AccumulateSinglePlane(VecOne, TileMaskData, RawSize.X, RawSize.Y, InTileOffsetX, InTileOffsetY, InSubpixelOffset.X, InSubpixelOffset.Y,
-				SubRectOffsetX, SubRectOffsetY, SubRectSizeX, SubRectSizeY);
+			WeightPlane.AccumulateSinglePlane(VecOne, TileMaskData, RawSize, InTileOffset, InSubpixelOffset.X, InSubpixelOffset.Y,
+				SubRectOffset, SubRectSize);
 		}
 
 		const double PlaneEndTime = FPlatformTime::Seconds();
@@ -647,13 +644,13 @@ void FImageOverlappedAccumulator::FetchFullImageValue(float Rgba[4], int32 FullX
 	Rgba[2] = 0.0f;
 	Rgba[3] = 1.0f;
 
-	float RawWeight = WeightPlane.ChannelData[FullY * PlaneSizeX + FullX];
+	float RawWeight = WeightPlane.ChannelData[FullY * PlaneSize.X + FullX];
 
 	float Scale = 1.0f / FMath::Max<float>(RawWeight, 0.0001f);
 
 	for (int64 Chan = 0; Chan < NumChannels; Chan++)
 	{
-		float Val = ChannelPlanes[Chan].ChannelData[FullY * PlaneSizeX + FullX];
+		float Val = ChannelPlanes[Chan].ChannelData[FullY * PlaneSize.X + FullX];
 
 		Rgba[Chan] = Val * Scale;
 	}
@@ -670,8 +667,8 @@ void FImageOverlappedAccumulator::FetchFullImageValue(float Rgba[4], int32 FullX
 
 void FImageOverlappedAccumulator::FetchFinalPixelDataByte(TArray<FColor> & OutPixelData) const
 {
-	int32 FullSizeX = PlaneSizeX;
-	int32 FullSizeY = PlaneSizeY;
+	int32 FullSizeX = PlaneSize.X;
+	int32 FullSizeY = PlaneSize.Y;
 	OutPixelData.SetNumUninitialized(FullSizeX * FullSizeY);
 
 	for (int32 FullY = 0L; FullY < FullSizeY; FullY++)
@@ -692,8 +689,8 @@ void FImageOverlappedAccumulator::FetchFinalPixelDataByte(TArray<FColor> & OutPi
 
 void FImageOverlappedAccumulator::FetchFinalPixelDataLinearColor(TArray<FLinearColor> & OutPixelData) const
 {
-	int32 FullSizeX = PlaneSizeX;
-	int32 FullSizeY = PlaneSizeY;
+	int32 FullSizeX = PlaneSize.X;
+	int32 FullSizeY = PlaneSize.Y;
 	OutPixelData.SetNumUninitialized(FullSizeX * FullSizeY);
 
 	for (int32 FullY = 0L; FullY < FullSizeY; FullY++)
diff --git a/Engine/Plugins/MovieScene/MovieRenderPipeline/Source/MovieRenderPipelineCore/Public/MovieRenderOverlappedImage.h b/Engine/Plugins/MovieScene/MovieRenderPipeline/Source/MovieRenderPipelineCore/Public/MovieRenderOverlappedImage.h
--- a/Engine/Plugins/MovieScene/MovieRenderPipeline/Source/MovieRenderPipelineCore/Public/MovieRenderOverlappedImage.h
+++ b/Engine/Plugins/MovieScene/MovieRenderPipeline/Source/MovieRenderPipelineCore/Public/MovieRenderOverlappedImage.h
@@ -19,8 +19,7 @@ public:
 	 * Default constructor.
 	 */
 	FImageOverlappedPlane()
-		: SizeX(0)
-		, SizeY(0)
+		: Size(0,0)
 	{
 	}
 
@@ -30,7 +29,7 @@ public:
 	 * @param InSizeX - Horizontal size of the tile.
 	 * @param InSizeY - Vertical size of the tile.
 	 */
-	void Init(int32 InSizeX, int32 InSizeY);
+	void Init(FIntPoint InSize);
 
 	/**
 	 * Zeroes the accumulation. Assumes the data is already allocated.
@@ -53,19 +52,16 @@ public:
 	 * @param InSampleOffsetX - Pixel offset of the tile (in X)
 	 * @param InSampleOffsetY - Pixel offset of the tile (in Y)
 	 */
-	void AccumulateSinglePlane(const TArray64<float>& InRawData, const TArray64<float>& InWeightData, int32 InSizeX, int32 InSizeY, int32 InOffsetX, int32 InOffsetY,
+	void AccumulateSinglePlane(const TArray64<float>& InRawData, const TArray64<float>& InWeightData, FIntPoint InSize, FIntPoint InOffset,
 								float SubpixelOffsetX, float SubpixelOffsetY,
-								int32 SubRectOffsetX, int32 SubRectOffsetY,
-								int32 SubRectSizeX, int32 SubRectSizeY);
+								FIntPoint SubRectOffset,
+								FIntPoint SubRectSize);
 
 	/** Actual channel data*/
 	TArray64<float> ChannelData;
 
-	/** Width of the image. */
-	int32 SizeX;
-
-	/** Height of the image. */
-	int32 SizeY;
+	/** Width and Height of the image. */
+	FIntPoint Size;
 };
 
 /**
@@ -78,12 +74,10 @@ public:
 	 * Default constructor.
 	 */
 	FImageOverlappedAccumulator()
-		: PlaneSizeX(0)
-		, PlaneSizeY(0)
+		: PlaneSize(0,0)
 		, NumChannels(0)
 		, AccumulationGamma(1.0f)
-		, TileMaskSizeX(0)
-		, TileMaskSizeY(0)
+		, TileMaskSize(0,0)
 	{
 	}
 
@@ -96,7 +90,7 @@ public:
 	 * @param InNumTilesY - Num vertical tiles.
 	 * @param InNumChannels - Num Channels.
 	 */
-	void InitMemory(int InPlaneSizeX, int InPlaneSizeY, int InNumChannels);
+	void InitMemory(FIntPoint InPlaneSize, int InNumChannels);
 
 	/**
 	 * Initializes memory.
@@ -111,11 +105,11 @@ public:
 	void Reset();
 
 
-	static void GenerateTileWeight(TArray64<float>& Weights, int32 SizeX, int32 SizeY);
+	static void GenerateTileWeight(TArray64<float>& Weights, FIntPoint Size);
 
-	static void GetTileSubRect(int32 & SubRectOffsetX, int32 & SubRectOffsetY, int32 & SubRectSizeX, int32 & SubRectSizeY, const TArray64<float>& Weights, const int32 SizeX, const int32 SizeY);
+	static void GetTileSubRect(FIntPoint & SubRectOffset, FIntPoint & SubRectSize, const TArray64<float>& Weights, const FIntPoint Size);
 
-	static void CheckTileSubRect(const TArray64<float>& Weights, const int32 SizeX, const int32 SizeY, int32 SubRectOffsetX, int32 SubRectOffsetY, int32 SubRectSizeX, int32 SubRectSizeY);
+	static void CheckTileSubRect(const TArray64<float>& Weights, const FIntPoint Size, FIntPoint SubRectOffset, FIntPoint SubRectSize);
 
 	/**
 	 * Given a rendered tile, accumulate the data to the full size image.
@@ -125,7 +119,7 @@ public:
 	 * @param InTileY - Tile index in Y.
 	 * @param InSubPixelOffset - SubPixel offset, should be in the range [0,1)
 	 */
-	void AccumulatePixelData(const FImagePixelData& InPixelData, int32 InTileOffsetX, int32 InTileOffsetY, FVector2D InSubpixelOffset);
+	void AccumulatePixelData(const FImagePixelData& InPixelData, FIntPoint InTileOffset, FVector2D InSubpixelOffset);
 
 	/**
 	 * After accumulation is finished, fetch the final image as bytes. In theory we don't need this, because we could
@@ -155,11 +149,8 @@ public:
 	void FetchFullImageValue(float Rgba[4], int32 FullX, int32 FullY) const;
 
 public:
-	/** Width of each tile in pixels */
-	int32 PlaneSizeX;
-
-	/** Height of each tile in pixels */
-	int32 PlaneSizeY;
+	/** Width and height of each tile in pixels */
+	FIntPoint PlaneSize;
 
 	/** Number of channels in the tiles. Typical will be 3 (RGB). */
 	int32 NumChannels;
@@ -174,11 +165,8 @@ public:
 	/** Weights image. Recalculate if it changes. */
 	TArray64<float> TileMaskData;
 
-	/** Width of weight image. */
-	int32 TileMaskSizeX;
-
-	/** Height of weight image. */
-	int32 TileMaskSizeY;
+	/** Width and height of weight mask image. */
+	FIntPoint TileMaskSize;
 };
 
 
diff --git a/Engine/Plugins/MovieScene/MovieRenderPipeline/Source/MovieRenderPipelineRenderPasses/Private/MoviePipelineImageSequenceContainer.cpp b/Engine/Plugins/MovieScene/MovieRenderPipeline/Source/MovieRenderPipelineRenderPasses/Private/MoviePipelineImageSequenceContainer.cpp
--- a/Engine/Plugins/MovieScene/MovieRenderPipeline/Source/MovieRenderPipelineRenderPasses/Private/MoviePipelineImageSequenceContainer.cpp
+++ b/Engine/Plugins/MovieScene/MovieRenderPipeline/Source/MovieRenderPipelineRenderPasses/Private/MoviePipelineImageSequenceContainer.cpp
@@ -75,7 +75,7 @@ void UMoviePipelineImageSequenceContainerBase::OnInitializedForPipelineImpl(UMov
 		{
 			if (bIsFirstTile && bIsFirstTemporalSample)
 			{
-				OverlappedAccumulator->InitMemory(TilePayload->OverlappedSizeX * TilePayload->NumTilesX, TilePayload->OverlappedSizeY * TilePayload->NumTilesY, 3);
+				OverlappedAccumulator->InitMemory(FIntPoint(TilePayload->OverlappedSizeX * TilePayload->NumTilesX, TilePayload->OverlappedSizeY * TilePayload->NumTilesY), 3);
 				OverlappedAccumulator->ZeroPlanes();
 				OverlappedAccumulator->AccumulationGamma = TilePayload->AccumulationGamma;
 			}
@@ -84,7 +84,7 @@ void UMoviePipelineImageSequenceContainerBase::OnInitializedForPipelineImpl(UMov
 
 			check(TilePayload->OverlappedSizeX + 2 * TilePayload->OverlappedPadX == RawSize.X);
 			check(TilePayload->OverlappedSizeY + 2 * TilePayload->OverlappedPadY == RawSize.Y);
-			OverlappedAccumulator->AccumulatePixelData(*InOwnedImage.Get(), TilePayload->OverlappedOffsetX, TilePayload->OverlappedOffsetY, TilePayload->OverlappedSubpixelShift);
+			OverlappedAccumulator->AccumulatePixelData(*InOwnedImage.Get(), FIntPoint(TilePayload->OverlappedOffsetX, TilePayload->OverlappedOffsetY), TilePayload->OverlappedSubpixelShift);
 		}
 
 		//UE_LOG(LogTemp, Log, TEXT("[%8.2f] Accumulation time."), ElapsedMs);
@@ -113,8 +113,8 @@ void UMoviePipelineImageSequenceContainerBase::OnInitializedForPipelineImpl(UMov
 		if (bIsLastTile && bIsLastTemporalSample)
 		{
 
-			int32 FullSizeX = TilePayload->bIsUsingOverlappedTiles ? OverlappedAccumulator->PlaneSizeX : TileAccumulator->TileSizeX * TileAccumulator->NumTilesX;
-			int32 FullSizeY = TilePayload->bIsUsingOverlappedTiles ? OverlappedAccumulator->PlaneSizeY : TileAccumulator->TileSizeY * TileAccumulator->NumTilesY;
+			int32 FullSizeX = TilePayload->bIsUsingOverlappedTiles ? OverlappedAccumulator->PlaneSize.X : TileAccumulator->TileSizeX * TileAccumulator->NumTilesX;
+			int32 FullSizeY = TilePayload->bIsUsingOverlappedTiles ? OverlappedAccumulator->PlaneSize.Y : TileAccumulator->TileSizeY * TileAccumulator->NumTilesY;
 
 			{
 				TUniquePtr<FImageWriteTask> ImageTask = MakeUnique<FImageWriteTask>();
