Fixed comments on overlapped accumulator classes.

diff --git a/Engine/Plugins/MovieScene/MovieRenderPipeline/Source/MovieRenderPipelineCore/Private/MovieRenderOverlappedImage.cpp b/Engine/Plugins/MovieScene/MovieRenderPipeline/Source/MovieRenderPipelineCore/Private/MovieRenderOverlappedImage.cpp
--- a/Engine/Plugins/MovieScene/MovieRenderPipeline/Source/MovieRenderPipelineCore/Private/MovieRenderOverlappedImage.cpp
+++ b/Engine/Plugins/MovieScene/MovieRenderPipeline/Source/MovieRenderPipelineCore/Private/MovieRenderOverlappedImage.cpp
@@ -327,9 +327,9 @@ void FImageOverlappedAccumulator::Reset()
 }
 
 
-void FImageOverlappedAccumulator::GenerateTileWeight(TArray64<float>& Weights, FIntPoint Size)
+void FImageOverlappedAccumulator::GenerateTileWeight(TArray64<float>& OutWeights, FIntPoint Size)
 {
-	Weights.SetNum(Size.X * Size.Y);
+	OutWeights.SetNum(Size.X * Size.Y);
 
 	// we'll use a simple triangle filter, which goes from 1.0 at the center to 0.0 at 3/4 of the way to the tile edge.
 
@@ -349,13 +349,13 @@ void FImageOverlappedAccumulator::GenerateTileWeight(TArray64<float>& Weights, F
 			float WeightX = FMath::Clamp(1.0f - DistX * ScaleX, 0.0f, 1.0f);
 			float WeightY = FMath::Clamp(1.0f - DistY * ScaleY, 0.0f, 1.0f);
 
-			Weights[PixY * Size.X + PixX] = WeightX * WeightY;
+			OutWeights[PixY * Size.X + PixX] = WeightX * WeightY;
 		}
 	}
 }
 
 // Any changes to GenerateTileWeight() need to propagate to this function. The subrect should be large enough to incorporate all nonzero values
-void FImageOverlappedAccumulator::GetTileSubRect(FIntPoint & SubRectOffset, FIntPoint & SubRectSize, const TArray64<float>& Weights, const FIntPoint Size)
+void FImageOverlappedAccumulator::GetTileSubRect(FIntPoint & OutSubRectOffset, FIntPoint & OutSubRectSize, const TArray64<float>& Weights, const FIntPoint Size)
 {
 	static bool bBruteForce = false;
 	if (bBruteForce)
@@ -382,10 +382,10 @@ void FImageOverlappedAccumulator::GetTileSubRect(FIntPoint & SubRectOffset, FInt
 			}
 		}
 
-		SubRectOffset.X = MinX;
-		SubRectOffset.Y = MinY;
-		SubRectSize.X = (MaxX - MinX) + 1;
-		SubRectSize.Y = (MaxY - MinY) + 1;
+		OutSubRectOffset.X = MinX;
+		OutSubRectOffset.Y = MinY;
+		OutSubRectSize.X = (MaxX - MinX) + 1;
+		OutSubRectSize.Y = (MaxY - MinY) + 1;
 	}
 	else
 	{
@@ -396,20 +396,20 @@ void FImageOverlappedAccumulator::GetTileSubRect(FIntPoint & SubRectOffset, FInt
 		int32 MaxX = FMath::CeilToInt(0.875f * Size.X) + 1;
 		int32 MaxY = FMath::CeilToInt(0.875f * Size.Y) + 1;
 
-		SubRectOffset.X = MinX;
-		SubRectOffset.Y = MinY;
-		SubRectSize.X = (MaxX - MinX) + 1;
-		SubRectSize.Y = (MaxY - MinY) + 1;
+		OutSubRectOffset.X = MinX;
+		OutSubRectOffset.Y = MinY;
+		OutSubRectSize.X = (MaxX - MinX) + 1;
+		OutSubRectSize.Y = (MaxY - MinY) + 1;
 	}
 }
 
-void FImageOverlappedAccumulator::CheckTileSubRect(const TArray64<float>& Weights, const FIntPoint Size, FIntPoint SubRectOffset, FIntPoint SubRectSize)
+void FImageOverlappedAccumulator::CheckTileSubRect(const TArray64<float>& OutWeights, const FIntPoint Size, FIntPoint SubRectOffset, FIntPoint SubRectSize)
 {
 	for (int32 Y = 0; Y < Size.Y; Y++)
 	{
 		for (int32 X = 0; X < Size.X; X++)
 		{
-			const float Val = Weights[Y*Size.X + X];
+			const float Val = OutWeights[Y*Size.X + X];
 			if (Val > 0.0f)
 			{
 				check(SubRectOffset.X <= X);
diff --git a/Engine/Plugins/MovieScene/MovieRenderPipeline/Source/MovieRenderPipelineCore/Public/MovieRenderOverlappedImage.h b/Engine/Plugins/MovieScene/MovieRenderPipeline/Source/MovieRenderPipelineCore/Public/MovieRenderOverlappedImage.h
--- a/Engine/Plugins/MovieScene/MovieRenderPipeline/Source/MovieRenderPipelineCore/Public/MovieRenderOverlappedImage.h
+++ b/Engine/Plugins/MovieScene/MovieRenderPipeline/Source/MovieRenderPipelineCore/Public/MovieRenderOverlappedImage.h
@@ -26,8 +26,7 @@ public:
 	/**
 	 * Initialize the memory. Before using the memory, we also will need a call to ZeroPlane.
 	 *
-	 * @param InSizeX - Horizontal size of the tile.
-	 * @param InSizeY - Vertical size of the tile.
+	 * @param InSize - Size of the tile in pixels.
 	 */
 	void Init(FIntPoint InSize);
 
@@ -47,10 +46,13 @@ public:
 	 * have unused border areas. The SubRect parameters describe the area inside the source plane.
 	 *
 	 * @param InRawData - Raw data to accumulate
-	 * @param InRawSizeX - Width of the tile. Must exactly match.
-	 * @param InRawSizeY - Height of the tile. Must exactly match.
-	 * @param InSampleOffsetX - Pixel offset of the tile (in X)
-	 * @param InSampleOffsetY - Pixel offset of the tile (in Y)
+	 * @param InWeightData - Mask to apply to the raw data for blending the overlapped areas.
+	 * @param InSize - Size of the tile. Must exactly match. InSize.X
+	 * @param InOffset - The (x,y) offset to the overlapped image.
+	 * @param SubpixelOffsetX - Subpixeoffset of the tile (in X), goes from [0,1] with 0.5 meaning it is right in the center.
+	 * @param SubpixelOffsetY - Subixel offset of the tile (in Y), goes from [0,1] with 0.5 meaning it is right in the center
+	 * @param SubRectOffset - The offset of the SubRect inside the raw data (InRawData) that actually has a weight > 0.0.
+	 * @param SubRectSize - The size of the SubRect inside the raw data (InRawData) that actually has a weight > 0.0.
 	 */
 	void AccumulateSinglePlane(const TArray64<float>& InRawData, const TArray64<float>& InWeightData, FIntPoint InSize, FIntPoint InOffset,
 								float SubpixelOffsetX, float SubpixelOffsetY,
@@ -104,19 +106,39 @@ public:
 	 */
 	void Reset();
 
+	/**
+	 * Given the size of a tile, calculate its weight mask.
+	 * 
+	 * @param OutWeights - Raw pixel data.
+	 * @param Size - Size of the Tile.
+	 */
+	static void GenerateTileWeight(TArray64<float>& OutWeights, FIntPoint Size);
 
-	static void GenerateTileWeight(TArray64<float>& Weights, FIntPoint Size);
-
-	static void GetTileSubRect(FIntPoint & SubRectOffset, FIntPoint & SubRectSize, const TArray64<float>& Weights, const FIntPoint Size);
+	/**
+	 * Given a tile size and weights, calculate the SubRect that contains all nonzere weights.
+	 * 
+	 * @param OutSubRectOffset - Offset of the found subrect.
+	 * @param OutSubRectSize - Size of the found subrect.
+	 * @param Weights - Input weights mask data.
+	 * @param Size - Input weights mask dimensions.
+	 */
+	static void GetTileSubRect(FIntPoint & OutSubRectOffset, FIntPoint & OutSubRectSize, const TArray64<float>& Weights, const FIntPoint Size);
 
+	/**
+	 * Given a tile size and weights, check() that every non-zero weight is in the subrect for debugging.
+	 * 
+	 * @param Weights - Input weights mask data.
+	 * @param Size - Input weights mask dimensions.
+	 * @param SubRectOffset - Offset of the found subrect.
+	 * @param SubRectSize - Size of the found subrect.
+	 */
 	static void CheckTileSubRect(const TArray64<float>& Weights, const FIntPoint Size, FIntPoint SubRectOffset, FIntPoint SubRectSize);
 
 	/**
 	 * Given a rendered tile, accumulate the data to the full size image.
 	 * 
 	 * @param InPixelData - Raw pixel data.
-	 * @param InTileX - Tile index in X.
-	 * @param InTileY - Tile index in Y.
+	 * @param InTileOffset - Tile offset in pixels.
 	 * @param InSubPixelOffset - SubPixel offset, should be in the range [0,1)
 	 */
 	void AccumulatePixelData(const FImagePixelData& InPixelData, FIntPoint InTileOffset, FVector2D InSubpixelOffset);
