MoviePipeline: Refactored shot configuration to make the "per-shot" vs "master" configurations more explicit to help make resource lifecycles more obvious as well.

#rnx

diff --git a/Engine/Plugins/MovieScene/MovieRenderPipeline/Source/MovieRenderPipelineCore/Private/MoviePipeline.cpp b/Engine/Plugins/MovieScene/MovieRenderPipeline/Source/MovieRenderPipelineCore/Private/MoviePipeline.cpp
--- a/Engine/Plugins/MovieScene/MovieRenderPipeline/Source/MovieRenderPipelineCore/Private/MoviePipeline.cpp
+++ b/Engine/Plugins/MovieScene/MovieRenderPipeline/Source/MovieRenderPipelineCore/Private/MoviePipeline.cpp
@@ -16,16 +16,22 @@
 #include "Engine/World.h"
 #include "GameFramework/PlayerController.h"
 #include "MovieRenderDebugWidget.h"
-#include "MovieRenderPipelineConfig.h"
 #include "MoviePipelineShotConfig.h"
 #include "MovieRenderPipelineDataTypes.h"
 #include "MoviePipelineRenderPass.h"
-#include "MoviePipelineOutput.h"
+#include "MoviePipelineOutputBase.h"
 #include "ShaderCompiler.h"
 #include "ImageWriteStream.h"
 #include "MoviePipelineAccumulationSetting.h"
 #include "MoviePipelineOutputBuilder.h"
 #include "DistanceFieldAtlas.h"
+#include "UObject/SoftObjectPath.h"
+#include "MoviePipelineOutputSetting.h"
+#include "MoviePipelineMasterConfig.h"
+#include "MoviePipelineOutputSetting.h"
+#include "MoviePipelineBlueprintLibrary.h"
+
+#define LOCTEXT_NAMESPACE "MoviePipeline"
 
 static TAutoConsoleVariable<int32> CVarMovieRenderPipelineFrameStepper(
 	TEXT("MovieRenderPipeline.FrameStepDebug"),
@@ -39,76 +45,95 @@ static TAutoConsoleVariable<int32> CVarMovieRenderPipelineFrameStepper(
 
 UMoviePipeline::UMoviePipeline()
 	: CustomTimeStep(nullptr)
-	, CachedCustomTimeStep(nullptr)
+	, CachedPrevCustomTimeStep(nullptr)
 	, TargetSequence(nullptr)
 	, LevelSequenceActor(nullptr)
 	, PipelineState(EMovieRenderPipelineState::Uninitialized)
 	, CurrentShotIndex(-1)
 	, bHasRunBeginFrameOnce(false)
+	, bPauseAtEndOfFrame(false)
+	, AccumulatedTickSubFrameDeltas(0.f)
 {
 	CustomTimeStep = CreateDefaultSubobject<UMoviePipelineCustomTimeStep>("MoviePipelineCustomTimeStep");
 	CustomSequenceTimeController = MakeShared<FMoviePipelineTimeController>();
-}
+	OutputBuilder = MakeShared<FMoviePipelineOutputMerger, ESPMode::ThreadSafe>();
 
-ULevelSequence* PostProcessSequence(ULevelSequence* InSequence)
-{
-	ULevelSequence* DestinationSequence = InSequence;
 
-	return DestinationSequence;
+	// Temp
+	OutputPipe = MakeShared<FImagePixelPipe, ESPMode::ThreadSafe>();
 }
 
+
 void ValidateSequence(ULevelSequence* InSequence)
 {
+	// ToDo: 
+	// Warn for Blueprint Streaming Levels
+	// Warn for sections that aren't extended far enough (once handle frames are taken into account)
+	// Warn for not whole frame aligned sections
 
 }
 
-void UMoviePipeline::Initialize(UMovieRenderPipelineConfig* InConfig)
+void UMoviePipeline::Initialize(const FMoviePipelineExecutorJob& InJob)
 {
 	// This function is called after the PIE world has finished initializing, but before
 	// the PIE world is ticked for the first time. We'll end up waiting for the next tick
 	// for FCoreDelegateS::OnBeginFrame to get called to actually start processing.
 	UE_LOG(LogMovieRenderPipeline, Log, TEXT("[%d] Initializing overall Movie Pipeline"), GFrameCounter);
 	
-	check(InConfig);
-	check(PipelineState == EMovieRenderPipelineState::Uninitialized);
+	if (!ensureAlwaysMsgf(InJob.Configuration, TEXT("MoviePipeline cannot be initialized with null configuration. Aborting.")))
+	{
+		OnMoviePipelineErroredDelegate.Broadcast(this, true, LOCTEXT("MissingConfiguration", "Job did not specify a configuration, movie render is aborting."));
+
+		Shutdown();
+		return;
+	}
+
+	if (!ensureAlwaysMsgf(PipelineState == EMovieRenderPipelineState::Uninitialized, TEXT("Pipeline cannot be reused. Create a new pipeline to execute a job.")))
+	{
+		OnMoviePipelineErroredDelegate.Broadcast(this, true, LOCTEXT("DontReusePipeline", "Attempted to reuse an existing Movie Pipeline. Initialize a new pipeline instead of reusing an existing one."));
+
+		Shutdown();
+		return;
+	}
 
 	// Ensure this object has the World as part of its Outer (so that it has context to spawn things)
-	check(GetWorld());
+	if (!ensureAlwaysMsgf(GetWorld(), TEXT("Pipeline does not contain the world as an outer.")))
+	{
+		OnMoviePipelineErroredDelegate.Broadcast(this, true, LOCTEXT("MissingWorld", "Could not find World in the Outer Path for Pipeline. The world must be an outer to give the Pipeline enough context to spawn things."));
 
-	Config = InConfig;
+		Shutdown();
+		return;
+	}
 
-	// Add the Config to root to ensure it doesn't go out of scope until we've shut down.
-	Config->AddToRoot();
+	CurrentJob = InJob;
 
 	// When start up we want to override the engine's Custom Timestep with our own.
 	// This gives us the ability to completely control the engine tick/delta time before the frame
 	// is started so that we don't have to always be thinking of delta times one frame ahead.
-	CachedCustomTimeStep = GEngine->GetCustomTimeStep();
+	CachedPrevCustomTimeStep = GEngine->GetCustomTimeStep();
 	GEngine->SetCustomTimeStep(CustomTimeStep);
-
-	// Duplicate the target level sequence. This reduces complexity by preventing the need to undo all
-	// changes applied to the sequence when running from the Editor. This is a medium-depth copy, we
-	// copy all sub-sequences and pointed to sequences but don't copy any actual assets.
-	ULevelSequence* SequenceAsset = LoadObject<ULevelSequence>(this, *Config->Sequence.GetAssetPathString());
-	if (!ensureAlwaysMsgf(SequenceAsset, TEXT("Failed to load Sequence Asset from specified path, aborting movie render! Path: %s"), *Config->Sequence.GetAssetPathString()))
+	
+	ULevelSequence* OriginalSequence = Cast<ULevelSequence>(InJob.Sequence.TryLoad());
+	if (!ensureAlwaysMsgf(OriginalSequence, TEXT("Failed to load Sequence Asset from specified path, aborting movie render! Attempted to load Path: %s"), *InJob.Sequence.ToString()))
 	{
+		OnMoviePipelineErroredDelegate.Broadcast(this, true, LOCTEXT("MissingSequence", "Could not load sequence asset, movie render is aborting. Check logs for additional details."));
+
 		Shutdown();
-		// ToDo: Promote this to a error that gets broadcast out.
 		return;
 	}
 
-	TargetSequence = CreateCopyOfSequence(SequenceAsset);
+	// Duplicate the target level sequence. This way modifications don't need to be undone.
+	TargetSequence = Cast<ULevelSequence>(UMoviePipelineBlueprintLibrary::DuplicateSequence(this, OriginalSequence));
 
-	// ToDo: Allow extensions to modify the sequence
+	// Initialize all of our master config settings.
+	for (UMoviePipelineSetting* Setting : GetPipelineMasterConfig()->GetSettings())
 	{
-
+		Setting->SetupForPipeline(this);
 	}
 
-	// Once the sequence has been modified by any extensions (which may have added new sub-sequences, added
-	// or removed tracks, etc.) we want to do some validation on the sequence such as adding Camera Cut tracks
-	// or other needed things we can do to smooth over the user experience. This can modify the root sequence
-	// such as adding a Cinematic Shot Section, so we re-assign our Target Sequence in case it does.
-	TargetSequence = PostProcessSequence(TargetSequence);
+	// Allow master settings to modify the sequence. This can be useful when working with dynamic content, you might
+	// want to modify things in the sequence, or modify things in the world before rendering.
+	ModifySequenceViaExtensions(TargetSequence);
 
 	// Now that we've post-processed it, we're going to run a validation pass on it. This will produce warnings
 	// for anything we can't fix that might be an issue - extending sections, etc. This should be const as this
@@ -122,7 +147,7 @@ void UMoviePipeline::Initialize(UMovieRenderPipelineConfig* InConfig)
 	// Finally, we're going to create a Level Sequence Actor in the world that has its settings configured by us.
 	// Because this callback is at the end of PIE startup (and before tick) we should be able to spawn the actor
 	// and give it a chance to tick once (where it should do nothing) before we start manually manipulating it.
-	InitializeLevelSequenceActor(SequenceAsset, TargetSequence);
+	InitializeLevelSequenceActor(OriginalSequence, TargetSequence);
 
 	// Register any additional engine callbacks needed.
 	{
@@ -137,25 +162,16 @@ void UMoviePipeline::Initialize(UMovieRenderPipelineConfig* InConfig)
 	// Construct a debug UI and bind it to this instance.
 	LoadDebugWidget();
 
-	OutputBuilder = MakeShared<FMoviePipelineOutputMerger, ESPMode::ThreadSafe>();
-
-	OutputPipe = MakeShared<FImagePixelPipe, ESPMode::ThreadSafe>();
-
-	// Loop through our Output Containers and let them do any first-frame initialization. It's important that
-	// this happens on the first frame (before we start rendering movies).
-	for (UMoviePipelineOutput* OutputContainer : Config->OutputContainers)
-	{
-		OutputContainer->OnInitializedForPipeline(this);
-	}
-
 	// Initialization is complete. This engine frame is a wash (because the tick started with a 
 	// delta time not generated by us) so we'll wait until the next engine frame to start rendering.
-	UE_LOG(LogMovieRenderPipeline, Log, TEXT("[%d] Initialized Movie Pipeline with [%d Shots]. Expected total output is %d frames."), GFrameCounter, ShotList.Num(), 0);
+	TotalExpectedWork = CalculateExpectedOutputMetrics();
 
-	// We'll get started on the next frame
-	CurrentShotIndex = 0;
 	PipelineState = EMovieRenderPipelineState::ProducingFrames;
 	InitializationTime = FDateTime::UtcNow();
+	CurrentShotIndex = 0;
+
+	UE_LOG(LogMovieRenderPipeline, Log, TEXT("[%d] Initialized Movie Pipeline with [%d Camera Cuts]. Expected Total Work: %s"),
+		GFrameCounter, TotalExpectedWork.NumCameraCuts, *TotalExpectedWork.ToDisplayString());
 }
 
 void UMoviePipeline::Shutdown()
@@ -164,8 +180,14 @@ void UMoviePipeline::Shutdown()
 	// loop where we shutdown because finish and you responded by requesting we shutdown.
 	ensure(PipelineState != EMovieRenderPipelineState::Shutdown);
 
+	// Uninitialize our master config settings.
+	for (UMoviePipelineSetting* Setting : GetPipelineMasterConfig()->GetSettings())
+	{
+		Setting->TeardownForPipeline(this);
+	}
+
 	// Restore any custom Time Step that may have been set before.
-	GEngine->SetCustomTimeStep(CachedCustomTimeStep);
+	GEngine->SetCustomTimeStep(CachedPrevCustomTimeStep);
 
 	// Ensure our delegates don't get called anymore as we're going to become null soon.
 	FCoreDelegates::OnBeginFrame.RemoveAll(this);
@@ -186,7 +208,7 @@ void UMoviePipeline::Shutdown()
 		DebugWidget = nullptr;
 	}
 
-	for (UMoviePipelineOutput* Setting : Config->OutputContainers)
+	for (UMoviePipelineOutputBase* Setting : GetPipelineMasterConfig()->GetOutputContainers())
 	{
 		Setting->OnPipelineFinished();
 	}
@@ -273,7 +295,7 @@ void UMoviePipeline::BeginFinalize()
 {
 	// Notify all of our output containers that we have finished producing and
 	// submitting all frames to them and that they should start any async flushes.
-	for (UMoviePipelineOutput* Container : Config->OutputContainers)
+	for (UMoviePipelineOutputBase* Container : GetPipelineMasterConfig()->GetOutputContainers())
 	{
 		Container->BeginFinalize();
 	}
@@ -289,7 +311,7 @@ void UMoviePipeline::TickFinalizeOutputContainers(const bool bInForceFinish)
 		bAllContainsFinishedProcessing = true;
 
 		// Ask the containers if they're all done processing.
-		for (UMoviePipelineOutput* Container : Config->OutputContainers)
+		for (UMoviePipelineOutputBase* Container : GetPipelineMasterConfig()->GetOutputContainers())
 		{
 			bAllContainsFinishedProcessing &= Container->HasFinishedProcessing();
 		}
@@ -314,7 +336,7 @@ void UMoviePipeline::TickFinalizeOutputContainers(const bool bInForceFinish)
 		return;
 	}
 
-	for (UMoviePipelineOutput* Container : Config->OutputContainers)
+	for (UMoviePipelineOutputBase* Container : GetPipelineMasterConfig()->GetOutputContainers())
 	{
 		// All containers have finished processing, final shutdown.
 		Container->Finalize();
@@ -324,14 +346,6 @@ void UMoviePipeline::TickFinalizeOutputContainers(const bool bInForceFinish)
 	// instantaneous and will complete here, but on the next tick we will check them all anyways
 	// to allow running long-running export processes (such as encoding).
 	PipelineState = EMovieRenderPipelineState::Export;
-	PostFinalizeExport();
-}
-
-void UMoviePipeline::PostFinalizeExport()
-{
-	check(PipelineState == EMovieRenderPipelineState::Export);
-	
-	// This is called once
 }
 
 void UMoviePipeline::TickPostFinalizeExport()
@@ -388,33 +402,40 @@ bool UMoviePipelineCustomTimeStep::UpdateTimeStep(UEngine* /*InEngine*/)
 	return false;
 }
 
-ULevelSequence* UMoviePipeline::CreateCopyOfSequence(ULevelSequence* InSequence)
+void UMoviePipeline::ModifySequenceViaExtensions(ULevelSequence* InSequence)
 {
-	// We perform a medium-depth copy here. The standard Duplicate Object
-	// function only performs a shallow copy which will duplicate all of the
-	// tracks and sections that belong to this sequence and do reference fix
-	// up. We don't want a deep copy (because we don't want to copy all assets)
-	// but because we are modifying the Playback Start/End on sub-sequences
-	// we need to manually copy those and fix up the references ourself.
-
-	// Duplicate the sequence and change the outer to be ourself.
-	FObjectDuplicationParameters DuplicationParams(InSequence, this);
-	DuplicationParams.DestName = FName(*(InSequence->GetName() + TEXT("_Instanced")));
-	ULevelSequence* DuplicatedSequence = (ULevelSequence*)StaticDuplicateObjectEx(DuplicationParams);
-
-	// // Now that we've copied this sequence, we'll look for any camera cut sections and duplicate those.
-	// TArray<UMovieSceneCameraCutTrack*> SubTracks = DuplicatedSequence->GetMovieScene()->GetMasterTracks<UMovieSceneCameraCutTrack>();
-	// 
-	// for (const UMovieSceneCameraCutTrack* CameraCutTrack : SubTracks)
-	// {
-	// 	for (const UMovieSceneSection* Section : CameraCutTrack->GetAllSections())
-	// 	{
-	// 		UMovieSceneCameraCutSection* CameraCutSection = CastChecked<UMovieSceneCameraCutSection>(Section);
-	// 		CameraCutSection
-	// 	}
-	// }
-
-	return DuplicatedSequence;
+}
+
+FMoviePipelineWorkInfo UMoviePipeline::CalculateExpectedOutputMetrics()
+{
+	UE_LOG(LogMovieRenderPipeline, Log, TEXT("[%d] Calculating expected output metrics..."), GFrameCounter);
+
+	// We can calculate the expected amount of work that we will be doing, unless slow-mo/timescale is used. 
+	FMoviePipelineWorkInfo ExpectedWork;
+
+	for (const FMoviePipelineShotInfo& Shot : GetShotList())
+	{
+		ExpectedWork.NumCameraCuts += Shot.CameraCuts.Num();
+
+		int32 NumTiles = 0;
+		UMoviePipelineAccumulationSetting* TilingSettings = Shot.ShotConfig->FindSetting<UMoviePipelineAccumulationSetting>();
+		if (TilingSettings)
+		{
+			NumTiles = TilingSettings->TileCount * TilingSettings->TileCount;
+		}
+
+		for (const FMoviePipelineCameraCutInfo& CameraCut : Shot.CameraCuts)
+		{
+			FMoviePipelineWorkInfo CameraCutWorkInfo = CameraCut.TotalWorkInfo;
+
+			ExpectedWork += CameraCut.TotalWorkInfo;
+
+			UE_LOG(LogMovieRenderPipeline, Log, TEXT("[%d] CameraCut: %d Expected Work: %s"), 
+				GFrameCounter, ExpectedWork.NumCameraCuts, *CameraCutWorkInfo.ToDisplayString());
+		}
+	}
+
+	return ExpectedWork;
 }
 
 void UMoviePipeline::InitializeLevelSequenceActor(ULevelSequence* OriginalLevelSequence, ULevelSequence* InSequenceToApply)
@@ -459,7 +480,7 @@ void UMoviePipeline::InitializeLevelSequenceActor(ULevelSequence* OriginalLevelS
 }
 
 
-FMoviePipelineShotCache CreateShotFromMovieScene(const UMovieScene* InMovieScene, const TRange<FFrameNumber>& InIntersectionRange)
+FMoviePipelineShotInfo CreateShotFromMovieScene(const UMovieScene* InMovieScene, const TRange<FFrameNumber>& InIntersectionRange)
 {
 	check(InMovieScene);
 
@@ -498,14 +519,14 @@ FMoviePipelineShotCache CreateShotFromMovieScene(const UMovieScene* InMovieScene
 		NewRange.Section = nullptr;
 	}
 
-	FMoviePipelineShotCache NewShot;
+	FMoviePipelineShotInfo NewShot;
 	NewShot.OriginalRange = InIntersectionRange;
 	NewShot.TotalOutputRange = NewShot.OriginalRange;
 
 	for (const FCameraCutRange& Range : IntersectedRanges)
 	{
 		// Generate a CameraCut for each range.
-		FMoviePipelineShotCutCache& CameraCut = NewShot.CameraCuts.AddDefaulted_GetRef();
+		FMoviePipelineCameraCutInfo& CameraCut = NewShot.CameraCuts.AddDefaulted_GetRef();
 
 		CameraCut.CameraCutSection = Range.Section; // May be nullptr.
 		CameraCut.OriginalRange = Range.Range;
@@ -515,9 +536,9 @@ FMoviePipelineShotCache CreateShotFromMovieScene(const UMovieScene* InMovieScene
 	return NewShot;
 }
 
-TArray<FMoviePipelineShotCache> UMoviePipeline::BuildShotListFromSequence(const ULevelSequence* InSequence)
+TArray<FMoviePipelineShotInfo> UMoviePipeline::BuildShotListFromSequence(const ULevelSequence* InSequence)
 {
-	TArray<FMoviePipelineShotCache> NewShotList;
+	TArray<FMoviePipelineShotInfo> NewShotList;
 	
 	// Shot Tracks take precedent over camera cuts, as settings can only be applied as granular as a shot.
 	UMovieSceneCinematicShotTrack* CinematicShotTrack = InSequence->GetMovieScene()->FindMasterTrack<UMovieSceneCinematicShotTrack>();
@@ -543,15 +564,24 @@ TArray<FMoviePipelineShotCache> UMoviePipeline::BuildShotListFromSequence(const
 				continue;
 			}
 
+			if (CurrentJob.ShotRenderMask.Num() > 0)
+			{
+				if (!CurrentJob.ShotRenderMask.Contains(ShotSection->GetShotDisplayName()))
+				{
+					UE_LOG(LogMovieRenderPipeline, Verbose, TEXT("Skipped adding Shot %s to Shot List due to a shot render mask being active, and this shot not being on the list."), *ShotSection->GetShotDisplayName());
+					continue;
+				}
+			}
+
 			// The Shot Section may extend past our Sequence's Playback Bounds. We intersect the two bounds to ensure that
 			// the Playback Start/Playback End of the overall sequence is respected.
 			TRange<FFrameNumber> CinematicShotSectionRange = TRange<FFrameNumber>::Intersection(ShotSection->GetRange(), InSequence->GetMovieScene()->GetPlaybackRange());
 
-			FMoviePipelineShotCache NewShot = CreateShotFromMovieScene(ShotSection->GetSequence()->GetMovieScene(), CinematicShotSectionRange);
+			FMoviePipelineShotInfo NewShot = CreateShotFromMovieScene(ShotSection->GetSequence()->GetMovieScene(), CinematicShotSectionRange);
 
 			// The first thing we do is find the appropriate configuration from the settings. Each shot can have its own config
 			// or they fall back to a default one specified for the whole pipeline.
-			NewShot.ShotConfig = Config->GetConfigForShot(ShotSection->GetShotDisplayName());
+			NewShot.ShotConfig = GetPipelineMasterConfig()->GetConfigForShot(ShotSection->GetShotDisplayName());
 			NewShot.CinematicShotSection = ShotSection;
 
 			// There should always be a shot config as the Pipeline default is returned in the event they didn't customize.
@@ -563,8 +593,8 @@ TArray<FMoviePipelineShotCache> UMoviePipeline::BuildShotListFromSequence(const
 	else
 	{
 		// They don't have a cinematic shot track. We'll slice them up by camera cuts instead.
-		FMoviePipelineShotCache NewShot = CreateShotFromMovieScene(TargetSequence->GetMovieScene(), TargetSequence->GetMovieScene()->GetPlaybackRange());
-		NewShot.ShotConfig = Config->DefaultShotConfig;
+		FMoviePipelineShotInfo NewShot = CreateShotFromMovieScene(TargetSequence->GetMovieScene(), TargetSequence->GetMovieScene()->GetPlaybackRange());
+		NewShot.ShotConfig = GetPipelineMasterConfig()->DefaultShotConfig;
 		
 		NewShotList.Add(MoveTemp(NewShot));
 	}
@@ -574,23 +604,23 @@ TArray<FMoviePipelineShotCache> UMoviePipeline::BuildShotListFromSequence(const
 	if (NewShotList.Num() == 0)
 	{
 		UE_LOG(LogMovieRenderPipeline, Warning, TEXT("No Cinematic Shot Tracks found, and no Camera Cut Tracks found. Playback Range will be used but camera will render from Pawns perspective."));
-		FMoviePipelineShotCache NewShot;
-		NewShot.ShotConfig = Config->DefaultShotConfig;
+		FMoviePipelineShotInfo NewShot;
+		NewShot.ShotConfig = GetPipelineMasterConfig()->DefaultShotConfig;
 
-		FMoviePipelineShotCutCache& CameraCut = NewShot.CameraCuts.AddDefaulted_GetRef();
+		FMoviePipelineCameraCutInfo& CameraCut = NewShot.CameraCuts.AddDefaulted_GetRef();
 		CameraCut.OriginalRange = TargetSequence->GetMovieScene()->GetPlaybackRange();
 	}
 
 	// Now that we've gathered at least one or more shots with one or more cuts, we can apply settings. It's easier to
 	// debug when all of the shots are calculated up front and debug info is printed in a block instead of as it is reached.
 	int32 ShotIndex = 1;
-	for (FMoviePipelineShotCache& Shot : NewShotList)
+	for (FMoviePipelineShotInfo& Shot : NewShotList)
 	{
 		UE_LOG(LogMovieRenderPipeline, Log, TEXT("Shot %d/%d has %d Camera Cuts."), ShotIndex, NewShotList.Num(), Shot.CameraCuts.Num());
 		ShotIndex++;
 
-		Shot.NumHandleFrames = 0; // ToDo: Pull this from a config.
 		UMoviePipelineAccumulationSetting* AccumulationSettings = Shot.ShotConfig->FindOrAddSetting<UMoviePipelineAccumulationSetting>();
+		Shot.NumHandleFrames = AccumulationSettings->HandleFrameCount;
 
 		// Expand the shot to encompass handle frames. This will modify our Camera Cuts bounds.
 		ExpandShot(Shot);
@@ -601,15 +631,16 @@ TArray<FMoviePipelineShotCache> UMoviePipeline::BuildShotListFromSequence(const
 			ShotName = Shot.CinematicShotSection->GetShotDisplayName();
 		}
 
-		for (FMoviePipelineShotCutCache& CameraCut : Shot.CameraCuts)
+		for (FMoviePipelineCameraCutInfo& CameraCut : Shot.CameraCuts)
 		{
-			CameraCut.NumWarmUpFramesRemaining = CameraCut.NumWarmUpFrames = 4; // ToDo: pull from config
-			CameraCut.bAccurateFirstFrameHistory = true; // ToDo: pull from config
+			CameraCut.NumWarmUpFramesRemaining = CameraCut.NumWarmUpFrames = AccumulationSettings->WarmUpFrameCount;
+			CameraCut.bAccurateFirstFrameHistory = true;
 			CameraCut.NumTemporalSamples = AccumulationSettings->TemporalSampleCount;
 			CameraCut.NumSpatialSamples = AccumulationSettings->SpatialSampleCount;
-			CameraCut.CachedFrameRate = Config->GetEffectiveFrameRate();
+			CameraCut.CachedFrameRate = GetEffectiveFrameRate();
 			CameraCut.CachedTickResolution = TargetSequence->GetMovieScene()->GetTickResolution();
-			
+			CameraCut.TotalWorkInfo = CameraCut.GetTotalWorkEstimate();
+
 			FString CameraName;
 			if (CameraCut.CameraCutSection.IsValid())
 			{
@@ -622,23 +653,54 @@ TArray<FMoviePipelineShotCache> UMoviePipeline::BuildShotListFromSequence(const
 			// We want to start rendering from the first handle frame. Shutter Timing is a fixed offset from this number.
 			CameraCut.CurrentTick = CameraCut.TotalOutputRange.GetLowerBoundValue();
 		}
+
+		// We duplicate the Passes in each shot to be unique. This allows the passes to have state
+		// that persists while working on the next shot (in the event of slow-running tasks).
+		for (UMoviePipelineRenderPass* RenderPass : Shot.ShotConfig->GetRenderPasses())
+		{
+			FObjectDuplicationParameters DupParms(RenderPass, this);
+
+			UMoviePipelineRenderPass* PassCopy = (UMoviePipelineRenderPass*)StaticDuplicateObjectEx(DupParms);
+			Shot.RenderPasses.Add(PassCopy);
+		}
 	}
 
 	return NewShotList;
 }
 
-void UMoviePipeline::SetSoloCameraCut(FMoviePipelineShotCutCache& InCameraCut)
+/*void SetupForSHot(...)
 {
 
 }
 
-void UMoviePipeline::InitializeShot(FMoviePipelineShotCache& InShot)
+void TeardownForShot(...)
+{
+	// Initialize any settings specific to this shot.
+	for (UMoviePipelineSetting* Setting : InShot.ShotConfig->GetSettings())
+	{
+		// Skip any render passes. These have been manually duplicated to avoid shared state
+		// and thus have to be initialized separately.
+		if (Setting->IsA(UMoviePipelineRenderPass::StaticClass()))
+		{
+			continue;
+		}
+
+		Setting->SetupForPipeline(this);
+	}
+
+	for (UMoviePipelineRenderPass* RenderPass : InShot.RenderPasses)
+	{
+		RenderPass->SetupForPipeline(this);
+	}
+}*/
+
+void UMoviePipeline::InitializeShot(FMoviePipelineShotInfo& InShot)
 {
 	// We handle both tearing down the previous shot (if there is one) and
 	// standing up the specified shot here. This gives us a good chance to
 	// notify output containers/extensions that we're switching contexts
 	// and they can decide what to do with it (ie: split into multiple files)
-	TOptional<FMoviePipelineShotCache> PrevShot;
+	TOptional<FMoviePipelineShotInfo> PrevShot;
 	int32 NewShotIndex = ShotList.IndexOfByKey(InShot);
 	if (NewShotIndex > 0)
 	{
@@ -648,15 +710,22 @@ void UMoviePipeline::InitializeShot(FMoviePipelineShotCache& InShot)
 	// Set the new shot as the active shot. This enables the specified shot section and disables all other shot sections.
 	SetSoloShot(InShot);
 	
-	// ToDo
+	// Initialize any settings specific to this shot.
 	for (UMoviePipelineSetting* Setting : InShot.ShotConfig->GetSettings())
 	{
-		Setting->OnInitializedForPipeline(this);
+		// Skip any render passes. These have been manually duplicated to avoid shared state
+		// and thus have to be initialized separately.
+		if (Setting->IsA(UMoviePipelineRenderPass::StaticClass()))
+		{
+			continue;
+		}
+
+		Setting->SetupForPipeline(this);
 	}
 
-	for (UMoviePipelineSetting* Setting : InShot.ShotConfig->InputBuffers)
+	for (UMoviePipelineRenderPass* RenderPass : InShot.RenderPasses)
 	{
-		Setting->OnInitializedForPipeline(this);
+		RenderPass->SetupForPipeline(this);
 	}
 
 
@@ -664,10 +733,10 @@ void UMoviePipeline::InitializeShot(FMoviePipelineShotCache& InShot)
 	SetupRenderingPipelineForShot(InShot);
 }
 
-void UMoviePipeline::SetSoloShot(const FMoviePipelineShotCache& InShot)
+void UMoviePipeline::SetSoloShot(const FMoviePipelineShotInfo& InShot)
 {
 	// Iterate through the shot list and ensure all shots have been set to inactive.
-	for (FMoviePipelineShotCache& Shot : ShotList)
+	for (FMoviePipelineShotInfo& Shot : ShotList)
 	{
 		UMovieSceneCinematicShotSection* Section = Shot.CinematicShotSection.Get();
 
@@ -688,42 +757,18 @@ void UMoviePipeline::SetSoloShot(const FMoviePipelineShotCache& InShot)
 	}
 	else
 	{
-		UE_LOG(LogMovieRenderPipeline, Verbose, TEXT("Skipped enabling a shot track due to not shot section associated with the provided shot."));
+		UE_LOG(LogMovieRenderPipeline, Verbose, TEXT("Disabled all shot tracks and skipped enabling a shot track due to no shot section associated with the provided shot."));
 	}
 }
 
-FFrameNumber UMoviePipeline::GetMotionBlurDuration(int32 InShotIndex) const
-{
-	const FMoviePipelineShotCache& CurrentShot = ShotList[InShotIndex];
-	UMoviePipelineAccumulationSetting* AccumulationSettings = CurrentShot.ShotConfig->FindOrAddSetting<UMoviePipelineAccumulationSetting>();
-
-	const MoviePipeline::FFrameConstantMetrics FrameMetrics = CalculateShotFrameMetrics(CurrentShot);
-	
-	// We don't have data from before the first frame, at least from the common user perspective. To solve this,
-	// we can take the sample from the next frame, and then evaluate the first frame again and we will get a
-	// motion vector which is a close-enough approximation of the real motion (had there been any). While only
-	// an approximation, it is worth the trade off to be significantly more user friendly in regards to data prep.
-	return FrameMetrics.TicksPerSample.FloorToFrame();
-}
-
-FFrameNumber UMoviePipeline::CalculateHandleFrameDuration(const FMoviePipelineShotCache& InShot) const
-{
-	// Convert to Tick Resolution
-	FFrameNumber HandleFrameDuration = FFrameRate::TransformTime(FFrameTime(FFrameNumber(InShot.NumHandleFrames)),
-		TargetSequence->GetMovieScene()->GetDisplayRate(),
-		TargetSequence->GetMovieScene()->GetTickResolution()).FloorToFrame();
-
-	return HandleFrameDuration;
-}
-
-void UMoviePipeline::ExpandShot(FMoviePipelineShotCache& InShot)
+void UMoviePipeline::ExpandShot(FMoviePipelineShotInfo& InShot)
 {
 	const MoviePipeline::FFrameConstantMetrics FrameMetrics = CalculateShotFrameMetrics(InShot);
 
 	// Handle Frames will be added onto our original shot size. We track the ranges separately for counting purposes
 	// later - the actual rendering code is unaware of the handle frames. Handle Frames only apply to the shot
 	// and expand the first/last inner cut to cover this area.
-	FFrameNumber HandleFrameTicks = CalculateHandleFrameDuration(InShot);
+	FFrameNumber HandleFrameTicks = FrameMetrics.TicksPerOutputFrame.FloorToFrame().Value * InShot.NumHandleFrames;
 	InShot.TotalOutputRange = MovieScene::DilateRange(InShot.OriginalRange, -HandleFrameTicks, HandleFrameTicks);
 
 	// We'll expand the overall sequence by at least this much to ensure we don't get clamped. This is +1 frame for shutter timing overlap.
@@ -741,8 +786,8 @@ void UMoviePipeline::ExpandShot(FMoviePipelineShotCache& InShot)
 	if (InShot.CameraCuts.Num() > 0)
 	{
 		// These may be the same or they may be different.
-		FMoviePipelineShotCutCache& FirstCut = InShot.CameraCuts[0];
-		FMoviePipelineShotCutCache& LastCut = InShot.CameraCuts[InShot.CameraCuts.Num() - 1];
+		FMoviePipelineCameraCutInfo& FirstCut = InShot.CameraCuts[0];
+		FMoviePipelineCameraCutInfo& LastCut = InShot.CameraCuts[InShot.CameraCuts.Num() - 1];
 
 		FirstCut.TotalOutputRange = MovieScene::DilateRange(FirstCut.TotalOutputRange, -HandleFrameTicks, FFrameNumber(0));
 		LastCut.TotalOutputRange = MovieScene::DilateRange(LastCut.TotalOutputRange, FFrameNumber(0), HandleFrameTicks);
@@ -833,33 +878,6 @@ bool UMoviePipeline::DebugFrameStepPreTick()
 	return false;
 }
 
-FMoviePipelineShotCache UMoviePipeline::GetCurrentShotSnapshot() const
-{
-	if (CurrentShotIndex >= 0 && CurrentShotIndex < ShotList.Num())
-	{
-		return ShotList[CurrentShotIndex];
-	}
-
-	return FMoviePipelineShotCache();
-}
-
-FMoviePipelineShotCutCache UMoviePipeline::GetCurrentCameraCutSnapshot() const
-{
-	FMoviePipelineShotCache CurrentShot = GetCurrentShotSnapshot();
-	if (CurrentShot.CurrentCameraCutIndex >= 0 && CurrentShot.CurrentCameraCutIndex < CurrentShot.CameraCuts.Num())
-	{
-		return CurrentShot.GetCurrentCameraCut();
-	}
-
-	return FMoviePipelineShotCutCache();
-}
-
-
-FMoviePipelineFrameOutputState UMoviePipeline::GetOutputStateSnapshot() const
-{
-	return CachedOutputState;
-}
-
 void UMoviePipeline::LoadDebugWidget()
 {
 	FSoftClassPath DebugWidgetClassRef(TEXT("/MovieRenderPipeline/Blueprints/UI_MovieRenderPipelineScreenOverlay.UI_MovieRenderPipelineScreenOverlay_C"));
@@ -883,7 +901,7 @@ void UMoviePipeline::LoadDebugWidget()
 }
 
 
-void UMoviePipeline::FlushAsyncSystems()
+void UMoviePipeline::FlushAsyncEngineSystems()
 {
 	// Flush Level Streaming. This solves the problem where levels that are not controlled
 	// by the Sequencer Level Visibility track are marked for Async Load by a gameplay system.
@@ -896,19 +914,42 @@ void UMoviePipeline::FlushAsyncSystems()
 	}
 
 	// Now we can flush the shader compiler. ToDo: This should probably happen right before SendAllEndOfFrameUpdates() is normally called
-	if (GShaderCompilingManager && GShaderCompilingManager->GetNumRemainingJobs() > 0)
+	if (GShaderCompilingManager)
 	{
-		UE_LOG(LogMovieRenderPipeline, Log, TEXT("[%d] Waiting for %d shaders to finish compiling..."), GFrameCounter, GShaderCompilingManager->GetNumRemainingJobs());
-		GShaderCompilingManager->FinishAllCompilation();
-		UE_LOG(LogMovieRenderPipeline, Log, TEXT("[%d] Done waiting for shaders to finish."), GFrameCounter);
+		bool bDidWork = false;
+		while (GShaderCompilingManager->GetNumRemainingJobs() > 0)
+		{
+			UE_LOG(LogMovieRenderPipeline, Log, TEXT("[%d] Waiting for %d shaders to finish compiling..."), GFrameCounter, GShaderCompilingManager->GetNumRemainingJobs());
+
+			// Sleep for 1 second and then check again. This way we get an indication of progress as this works.
+			FPlatformProcess::Sleep(1.f);
+			bDidWork = true;
+		}
+
+		if (bDidWork)
+		{
+			UE_LOG(LogMovieRenderPipeline, Log, TEXT("[%d] Done waiting for shaders to finish."), GFrameCounter);
+		}
 	}
 
 	// Flush the Mesh Distance Field builder as well.
-	if (GDistanceFieldAsyncQueue && GDistanceFieldAsyncQueue->GetNumOutstandingTasks() > 0)
+	if (GDistanceFieldAsyncQueue)
 	{
-		UE_LOG(LogMovieRenderPipeline, Log, TEXT("[%d] Waiting for %d Mesh Distance Fields to finish building..."), GFrameCounter, GDistanceFieldAsyncQueue->GetNumOutstandingTasks());
-		GDistanceFieldAsyncQueue->BlockUntilAllBuildsComplete();
-		UE_LOG(LogMovieRenderPipeline, Log, TEXT("[%d] Done waiting for Mesh Distance Fields to build."), GFrameCounter);
+		bool bDidWork = false;
+		while (GDistanceFieldAsyncQueue->GetNumOutstandingTasks() > 0)
+		{
+			UE_LOG(LogMovieRenderPipeline, Log, TEXT("[%d] Waiting for %d Mesh Distance Fields to finish building..."), GFrameCounter, GDistanceFieldAsyncQueue->GetNumOutstandingTasks());
+
+			// Sleep for 1 second and then check again. This way we get an indication of progress as this works.
+			FPlatformProcess::Sleep(1.f);
+			bDidWork = true;
+		}
+
+		if (bDidWork)
+		{
+			UE_LOG(LogMovieRenderPipeline, Log, TEXT("[%d] Done waiting for Mesh Distance Fields to build."), GFrameCounter);
+		}
+
 	}
 }
 
@@ -920,11 +961,11 @@ FFrameTime FMoviePipelineTimeController::OnRequestCurrentTime(const FQualifiedFr
 	return RequestTime;
 }
 
-MoviePipeline::FFrameConstantMetrics UMoviePipeline::CalculateShotFrameMetrics(const FMoviePipelineShotCache& InShot) const
+MoviePipeline::FFrameConstantMetrics UMoviePipeline::CalculateShotFrameMetrics(const FMoviePipelineShotInfo& InShot) const
 {
 	MoviePipeline::FFrameConstantMetrics Output;
 	Output.TickResolution = TargetSequence->GetMovieScene()->GetTickResolution();
-	Output.FrameRate = Config->GetEffectiveFrameRate();
+	Output.FrameRate = GetEffectiveFrameRate();
 	Output.TicksPerOutputFrame = FFrameRate::TransformTime(FFrameTime(FFrameNumber(1)), Output.FrameRate, Output.TickResolution);
 
 	UMoviePipelineAccumulationSetting* AccumulationSettings = InShot.ShotConfig->FindOrAddSetting<UMoviePipelineAccumulationSetting>();
@@ -979,54 +1020,24 @@ MoviePipeline::FFrameConstantMetrics UMoviePipeline::CalculateShotFrameMetrics(c
 	return Output;
 }
 
-FFrameNumber UMoviePipeline::GetTotalOutputFrameCountEstimate() const
+FFrameRate UMoviePipeline::GetEffectiveFrameRate() const
 {
-	FFrameNumber EstimatedFrameCount = FFrameNumber(0);
-	
-	for (const FMoviePipelineShotCache& Shot : ShotList)
-	{
-		for (const FMoviePipelineShotCutCache& CameraCut : Shot.CameraCuts)
-		{
-			EstimatedFrameCount += CameraCut.GetOutputFrameCountEstimate();
-		}
-	}
-
-	return EstimatedFrameCount;
-}
+	UMoviePipelineOutputSetting* OutputSettings = GetPipelineMasterConfig()->FindSetting<UMoviePipelineOutputSetting>();
+	check(OutputSettings);
 
-bool UMoviePipeline::GetRemainingTimeEstimate(FTimespan& OutTimespan) const
-{
-	// If they haven't produced a single frame yet, we can't give an estimate.
-	if (CachedOutputState.TotalSamplesRendered <= 0)
+	// Check to see if they overrode the frame rate.
+	if (OutputSettings->bUseCustomFrameRate)
 	{
-		OutTimespan = FTimespan();
-		return false;
+		return OutputSettings->OutputFrameRate;
 	}
 
-	// Look at how many total samples we expect across all shots. This includes
-	// samples produced for warm-ups, motion blur fixes, and temporal/spatial samples.
-	FFrameNumber TotalExpectedSamples = FFrameNumber(0);
-
-	for (const FMoviePipelineShotCache& Shot : ShotList)
+	// Pull it from the sequence if they didn't.
+	if (TargetSequence)
 	{
-		for (const FMoviePipelineShotCutCache& CameraCut : Shot.CameraCuts)
-		{
-			TotalExpectedSamples += CameraCut.GetSampleCountEstimate(true, true);
-		}
+		return TargetSequence->GetMovieScene()->GetDisplayRate();
 	}
 
-	// Check to see how many frames we've rendered vs. our estimate.
-	int32 RenderedFrames = CachedOutputState.TotalSamplesRendered;
-	int32 TotalFrames = TotalExpectedSamples.Value;
-
-	double CompletionPercentage = FMath::Clamp(RenderedFrames / (double)TotalFrames, 0.0, 1.0);
-	FTimespan CurrentDuration = FDateTime::UtcNow() - InitializationTime;
-
-	// If it has taken us CurrentDuration to process CompletionPercentage frames, then we can get a total duration
-	// estimate by taking (CurrentDuration/CompletionPercentage) and then take that total estimate minus elapsed
-	// to get remaining. 
-	FTimespan EstimatedTotalDuration = CurrentDuration / CompletionPercentage;
-	OutTimespan = EstimatedTotalDuration - CurrentDuration;
+	return FFrameRate();
+}
 
-	return true;
-}
\ No newline at end of file
+#undef LOCTEXT_NAMESPACE // "MoviePipeline"
diff --git a/Engine/Plugins/MovieScene/MovieRenderPipeline/Source/MovieRenderPipelineCore/Private/MoviePipelineAudioOutput.cpp b/Engine/Plugins/MovieScene/MovieRenderPipeline/Source/MovieRenderPipelineCore/Private/MoviePipelineAudioOutput.cpp
new file mode 100644
--- /dev/null
+++ b/Engine/Plugins/MovieScene/MovieRenderPipeline/Source/MovieRenderPipelineCore/Private/MoviePipelineAudioOutput.cpp
@@ -0,0 +1,128 @@
+// Copyright 1998-2019 Epic Games, Inc. All Rights Reserved.
+
+#include "MoviePipelineAudioOutput.h"
+#include "MoviePipeline.h"
+#include "HAL/IConsoleManager.h"
+#include "AudioDeviceManager.h"
+#include "MovieRenderPipelineCoreModule.h"
+#include "Engine/Engine.h"
+#include "AudioMixerBlueprintLibrary.h"
+#include "LevelSequence.h"
+#include "AudioDevice.h"
+#include "AudioMixerDevice.h"
+// #include "AudioMixerPlatformNonRealtime.h"
+
+void UMoviePipelineAudioOutput::SetupForPipelineImpl(UMoviePipeline* InPipeline)
+{
+	/*
+	// Swap to the non-real-time audio renderer at runtime. This has to come before
+	// we try to retrieve the cvar, otherwise the module won't have been loaded to create the cvar.
+	check(GEngine);
+
+	// Ensure that we try to play audio at full volume, even if we're unfocused.
+	PrevUnfocusedAudioMultiplier = FApp::GetUnfocusedVolumeMultiplier();
+	FApp::SetUnfocusedVolumeMultiplier(1.f);
+
+	if (FAudioDeviceManager* AudioDeviceManager = GEngine->GetAudioDeviceManager())
+	{
+		PrevAudioDevicePlatform = AudioDeviceManager->GetCurrentAudioDeviceModuleName();
+		bool bInitSuccessful = AudioDeviceManager->SwitchAudioDevicePlatform(TEXT("NonRealtimeAudioRenderer"));
+		if (!bInitSuccessful)
+		{
+			UE_LOG(LogMovieRenderPipeline, Warning, TEXT("Failed to initialize non-real-time Audio Device. Only supported on Windows. No audio output will be generated."));
+			return;
+		}
+
+		if (bInitSuccessful)
+		{
+			IConsoleVariable* AudioRenderEveryTickCvar = IConsoleManager::Get().FindConsoleVariable(TEXT("au.nrt.RenderEveryTick"));
+			check(AudioRenderEveryTickCvar);
+
+			// Cache the value of their current setting
+			PrevRenderEveryTickValue = AudioRenderEveryTickCvar->GetInt();
+
+			// Override it to prevent it from automatically ticking, we'll control this below.
+			AudioRenderEveryTickCvar->Set(0, ECVF_SetByConstructor);
+
+			int32 PostRenderEveryTickValue = AudioRenderEveryTickCvar->GetInt();
+
+			UE_LOG(LogMovieRenderPipeline, Log, TEXT("Successfully initialized Non-Real Time audio backend."));
+		}
+	}
+	
+	// Grab a rough estimate of how long the expected output is to avoid unnecessarily resizing arrays.
+	// This may not be perfect (due to handle frames, slow-mo, etc.) but it's better than guessing.
+	ULevelSequence* TargetSequence = InPipeline->GetTargetSequence();
+	FFrameNumber StartFrame = TargetSequence->GetMovieScene()->GetPlaybackRange().GetLowerBoundValue();
+	FFrameNumber EndFrame = TargetSequence->GetMovieScene()->GetPlaybackRange().GetUpperBoundValue();
+	FFrameNumber SequenceDuration = (EndFrame - StartFrame);
+	CachedExpectedDuration = TargetSequence->GetMovieScene()->GetTickResolution().AsInterval() * SequenceDuration.Value;
+
+	// Start recording, but pause until the first shot starts producing frames.
+	UAudioMixerBlueprintLibrary::StartRecordingOutput(GetWorld(), CachedExpectedDuration);
+	UAudioMixerBlueprintLibrary::PauseRecordingOutput(GetWorld());*/
+}
+
+void UMoviePipelineAudioOutput::OnPipelineFinishedImpl()
+{
+	/*// Restore our cached CVar value.
+	IConsoleVariable* AudioRenderEveryTickCvar = IConsoleManager::Get().FindConsoleVariable(TEXT("au.nrt.RenderEveryTick"));
+	
+	// This will be null if initialization failed.
+	if (AudioRenderEveryTickCvar)
+	{
+		AudioRenderEveryTickCvar->Set(PrevRenderEveryTickValue, ECVF_SetByConstructor);
+		
+		// Convert it to absolute as the Audio Recorder will save relative paths to a different
+		// directory than we consider to be our relative root.
+		FString FormattedFileName = TEXT("Audio.wav");
+		FString AbsoluteDirectory = FPaths::ConvertRelativePathToFull(GetPipeline()->GetPipelineConfig()->OutputDirectory.Path);
+		UAudioMixerBlueprintLibrary::StopRecordingOutput(GetWorld(), EAudioRecordingExportType::WavFile, FormattedFileName, AbsoluteDirectory);
+	}
+
+	// Attempt to restore the previous Audio Platform.
+	check(GEngine);
+	if (FAudioDeviceManager* AudioDeviceManager = GEngine->GetAudioDeviceManager())
+	{
+		AudioDeviceManager->SwitchAudioDevicePlatform(PrevAudioDevicePlatform);
+	}
+
+	FApp::SetUnfocusedVolumeMultiplier(PrevUnfocusedAudioMultiplier);*/
+}
+
+void UMoviePipelineAudioOutput::OnPostTickImpl()
+{
+	/*check(GEngine);
+	
+	if (FAudioDevice* AudioDevice = GEngine->GetActiveAudioDevice())
+	{
+		// Handle any game logic that changed Audio State.
+		AudioDevice->Update(true);
+
+		// Now render some samples
+		Audio::FMixerDevice* MixerDevice = static_cast<Audio::FMixerDevice*>(AudioDevice);
+		Audio::FMixerPlatformNonRealtime* Platform = static_cast<Audio::FMixerPlatformNonRealtime*>(MixerDevice->GetAudioMixerPlatform());
+		Platform->RenderAudio(FApp::GetDeltaTime());
+
+	}
+
+	// IAudioDeviceModule* AudioDeviceModule = FModuleManager::LoadModulePtr<IAudioDeviceModule>(TEXT("NonRealtimeAudioRenderer"));
+	// Audio::FMixerPlatformNonRealtime NonRealtimePlatform = static_cast<Audio::FMixerPlatformNonRealtime>(AudioDeviceModule);
+
+	// Now process the same amount of time as this frame for processing audio samples. The
+	// App DeltaTime should be correct since we're controlling it each frame.
+	// Audio::FMixerPlatformNonRealtime::RenderAudio(FApp::GetDeltaTime());*/
+}
+
+/* void UMoviePipelineAudioOutput::OnFrameProductionStartImpl()
+{
+	// When we start producing frames for a shot, unpause our recording.
+	UAudioMixerBlueprintLibrary::ResumeRecordingOutput(GetWorld());
+}*/
+
+void UMoviePipelineAudioOutput::OnShotFinishedImpl(const FMoviePipelineShotInfo& Shot)
+{
+	// Pause our recording when the shot finishes. This lets us skip over any audio samples that
+	// are generated between shots (ie: warm up frames) where we're not recording the output.
+	// UAudioMixerBlueprintLibrary::PauseRecordingOutput(GetWorld());
+}
\ No newline at end of file
diff --git a/Engine/Plugins/MovieScene/MovieRenderPipeline/Source/MovieRenderPipelineCore/Private/MoviePipelineBlueprintLibrary.cpp b/Engine/Plugins/MovieScene/MovieRenderPipeline/Source/MovieRenderPipelineCore/Private/MoviePipelineBlueprintLibrary.cpp
--- a/Engine/Plugins/MovieScene/MovieRenderPipeline/Source/MovieRenderPipelineCore/Private/MoviePipelineBlueprintLibrary.cpp
+++ b/Engine/Plugins/MovieScene/MovieRenderPipeline/Source/MovieRenderPipelineCore/Private/MoviePipelineBlueprintLibrary.cpp
@@ -3,6 +3,10 @@
 #include "MoviePipelineBlueprintLibrary.h"
 #include "MoviePipeline.h"
 #include "MovieRenderPipelineDataTypes.h"
+#include "MovieSceneSequence.h"
+#include "Tracks/MovieSceneCinematicShotTrack.h"
+#include "Sections/MovieSceneCinematicShotSection.h"
+#include "LevelSequence.h"
 
 void UMoviePipelineBlueprintLibrary::GetCameraCutCounts(const UMoviePipeline* InPipeline, int32& OutTotalCuts, int32& OutCurrentCutIndex)
 {
@@ -14,7 +18,7 @@ void UMoviePipelineBlueprintLibrary::GetCameraCutCounts(const UMoviePipeline* In
 		return;
 	}
 
-	const TArray<FMoviePipelineShotCache>&  ShotList = InPipeline->GetShotList();
+	const TArray<FMoviePipelineShotInfo>&  ShotList = InPipeline->GetShotList();
 	for (int32 ShotIdx = 0; ShotIdx < ShotList.Num(); ShotIdx++)
 	{
 		OutTotalCuts += ShotList[ShotIdx].CameraCuts.Num();
@@ -28,22 +32,145 @@ void UMoviePipelineBlueprintLibrary::GetCameraCutCounts(const UMoviePipeline* In
 	}
 }
 
-int32 UMoviePipelineBlueprintLibrary::GetOutputFrameCountEstimate(const FMoviePipelineShotCutCache& InCameraCut)
+int32 UMoviePipelineBlueprintLibrary::GetOutputFrameCountEstimate(const FMoviePipelineCameraCutInfo& InCameraCut)
 {
 	return InCameraCut.GetOutputFrameCountEstimate().Value;
 }
 
-int32 UMoviePipelineBlueprintLibrary::GetTemporalFrameCountEstimate(const FMoviePipelineShotCutCache& InCameraCut)
+int32 UMoviePipelineBlueprintLibrary::GetTemporalFrameCountEstimate(const FMoviePipelineCameraCutInfo& InCameraCut)
 {
 	return InCameraCut.GetTemporalFrameCountEstimate().Value;
 }
 
-int32 UMoviePipelineBlueprintLibrary::GetUtilityFrameCountEstimate(const FMoviePipelineShotCutCache& InCameraCut)
+int32 UMoviePipelineBlueprintLibrary::GetUtilityFrameCountEstimate(const FMoviePipelineCameraCutInfo& InCameraCut)
 {
 	return InCameraCut.GetUtilityFrameCountEstimate().Value;
 }
 
-int32 UMoviePipelineBlueprintLibrary::GetSampleCountEstimate(const FMoviePipelineShotCutCache& InCameraCut, const bool bIncludeWarmup, const bool bIncludeMotionBlur)
+int32 UMoviePipelineBlueprintLibrary::GetSampleCountEstimate(const FMoviePipelineCameraCutInfo& InCameraCut, const bool bIncludeWarmup, const bool bIncludeMotionBlur)
 {
 	return InCameraCut.GetSampleCountEstimate(bIncludeWarmup, bIncludeMotionBlur).Value;
 }
+
+UMovieSceneSequence* UMoviePipelineBlueprintLibrary::DuplicateSequence(UObject* Outer, UMovieSceneSequence* InSequence)
+{
+	if (!InSequence)
+	{
+		return nullptr;
+	}
+
+	FObjectDuplicationParameters DuplicationParams(InSequence, Outer);
+	DuplicationParams.DestName = MakeUniqueObjectName(Outer, UMovieSceneSequence::StaticClass(), InSequence->GetFName());
+
+	// Duplicate the given sequence.
+	ULevelSequence* DuplicatedSequence = (ULevelSequence*)StaticDuplicateObjectEx(DuplicationParams);
+
+	// Now go through looking for Shot and SubSequence tracks. These currently point to the same (shared) sequence as the InSequence.
+	TArray<UMovieSceneSection*> AllSubSequenceSections;
+
+	UMovieSceneCinematicShotTrack* ShotTrack = DuplicatedSequence->GetMovieScene()->FindMasterTrack<UMovieSceneCinematicShotTrack>();
+	if (ShotTrack)
+	{
+		AllSubSequenceSections.Append(ShotTrack->GetAllSections());
+	}
+
+	UMovieSceneSubTrack* SubTrack = DuplicatedSequence->GetMovieScene()->FindMasterTrack<UMovieSceneSubTrack>();
+	if (SubTrack)
+	{
+		AllSubSequenceSections.Append(ShotTrack->GetAllSections());
+	}
+
+	for (UMovieSceneSection* Section : AllSubSequenceSections)
+	{
+		UMovieSceneSubSection* SubSequenceSection = CastChecked<UMovieSceneSubSection>(Section);
+		UMovieSceneSequence* OriginalSubSequence = SubSequenceSection->GetSequence();
+
+		// Recursively duplicate this sequence, and then update our sequence to point to the duplicated one.
+		UMovieSceneSequence* DuplicatedSubSequence = DuplicateSequence(Outer, OriginalSubSequence);
+		SubSequenceSection->SetSequence(DuplicatedSubSequence);
+	}
+
+	return DuplicatedSequence;
+}
+
+/*
+FMoviePipelineShotInfo UMoviePipeline::GetCurrentShotSnapshot() const
+{
+	if (CurrentShotIndex >= 0 && CurrentShotIndex < ShotList.Num())
+	{
+		return ShotList[CurrentShotIndex];
+	}
+
+	return FMoviePipelineShotInfo();
+}
+
+FMoviePipelineCameraCutInfo UMoviePipeline::GetCurrentCameraCutSnapshot() const
+{
+	FMoviePipelineShotInfo CurrentShot = GetCurrentShotSnapshot();
+	if (CurrentShot.CurrentCameraCutIndex >= 0 && CurrentShot.CurrentCameraCutIndex < CurrentShot.CameraCuts.Num())
+	{
+		return CurrentShot.GetCurrentCameraCut();
+	}
+
+	return FMoviePipelineCameraCutInfo();
+}
+
+
+FMoviePipelineFrameOutputState UMoviePipeline::GetOutputStateSnapshot() const
+{
+	return CachedOutputState;
+}*/
+
+/*
+FFrameNumber UMoviePipeline::GetTotalOutputFrameCountEstimate() const
+{
+	FFrameNumber EstimatedFrameCount = FFrameNumber(0);
+
+	for (const FMoviePipelineShotInfo& Shot : ShotList)
+	{
+		for (const FMoviePipelineCameraCutInfo& CameraCut : Shot.CameraCuts)
+		{
+			EstimatedFrameCount += CameraCut.GetOutputFrameCountEstimate();
+		}
+	}
+
+	return EstimatedFrameCount;
+}
+
+bool UMoviePipeline::GetRemainingTimeEstimate(FTimespan& OutTimespan) const
+{
+	// If they haven't produced a single frame yet, we can't give an estimate.
+	if (CachedOutputState.TotalSamplesRendered <= 0)
+	{
+		OutTimespan = FTimespan();
+		return false;
+	}
+
+	// Look at how many total samples we expect across all shots. This includes
+	// samples produced for warm-ups, motion blur fixes, and temporal/spatial samples.
+	FFrameNumber TotalExpectedSamples = FFrameNumber(0);
+
+	for (const FMoviePipelineShotInfo& Shot : ShotList)
+	{
+		for (const FMoviePipelineCameraCutInfo& CameraCut : Shot.CameraCuts)
+		{
+			TotalExpectedSamples += CameraCut.GetSampleCountEstimate(true, true);
+		}
+	}
+
+	// Check to see how many frames we've rendered vs. our estimate.
+	int32 RenderedFrames = CachedOutputState.TotalSamplesRendered;
+	int32 TotalFrames = TotalExpectedSamples.Value;
+
+	double CompletionPercentage = FMath::Clamp(RenderedFrames / (double)TotalFrames, 0.0, 1.0);
+	FTimespan CurrentDuration = FDateTime::UtcNow() - InitializationTime;
+
+	// If it has taken us CurrentDuration to process CompletionPercentage frames, then we can get a total duration
+	// estimate by taking (CurrentDuration/CompletionPercentage) and then take that total estimate minus elapsed
+	// to get remaining.
+	FTimespan EstimatedTotalDuration = CurrentDuration / CompletionPercentage;
+	OutTimespan = EstimatedTotalDuration - CurrentDuration;
+
+	return true;
+}
+*/
\ No newline at end of file
diff --git a/Engine/Plugins/MovieScene/MovieRenderPipeline/Source/MovieRenderPipelineCore/Private/MoviePipelineConfigBase.cpp b/Engine/Plugins/MovieScene/MovieRenderPipeline/Source/MovieRenderPipelineCore/Private/MoviePipelineConfigBase.cpp
new file mode 100644
--- /dev/null
+++ b/Engine/Plugins/MovieScene/MovieRenderPipeline/Source/MovieRenderPipelineCore/Private/MoviePipelineConfigBase.cpp
@@ -0,0 +1,15 @@
+// Copyright 1998-2019 Epic Games, Inc. All Rights Reserved.
+
+#include "MoviePipelineConfigBase.h"
+
+#define LOCTEXT_NAMESPACE "MoviePipelineConfigBase"
+
+void UMoviePipelineConfigBase::RemoveSetting(UMoviePipelineSetting* InSource)
+{
+	Settings.Remove(InSource);
+
+	// Update our cached serial number so the UI rebuilds the tree
+	++SettingsSerialNumber;
+}
+
+#undef LOCTEXT_NAMESPACE // "MoviePipelineConfigBase"
\ No newline at end of file
diff --git a/Engine/Plugins/MovieScene/MovieRenderPipeline/Source/MovieRenderPipelineCore/Private/MoviePipelineInProcessExecutor.cpp b/Engine/Plugins/MovieScene/MovieRenderPipeline/Source/MovieRenderPipelineCore/Private/MoviePipelineInProcessExecutor.cpp
--- a/Engine/Plugins/MovieScene/MovieRenderPipeline/Source/MovieRenderPipelineCore/Private/MoviePipelineInProcessExecutor.cpp
+++ b/Engine/Plugins/MovieScene/MovieRenderPipeline/Source/MovieRenderPipelineCore/Private/MoviePipelineInProcessExecutor.cpp
@@ -4,9 +4,8 @@
 #include "Engine/World.h"
 
 #define LOCTEXT_NAMESPACE "MoviePipelineInProcessExecutor"
-void UMoviePipelineInProcessExecutor::Start(UMovieRenderPipelineConfig* InConfig, const int32 InConfigIndex, const int32 InNumConfigs)
+void UMoviePipelineInProcessExecutor::Start(const FMoviePipelineExecutorJob& InJob)
 {
-	ActiveConfig = InConfig;
 	FCoreUObjectDelegates::PostLoadMapWithWorld.AddUObject(this, &UMoviePipelineInProcessExecutor::OnMapLoadFinished);
 }
 
@@ -20,9 +19,11 @@ void UMoviePipelineInProcessExecutor::OnMapLoadFinished(UWorld* NewWorld)
 
 	// Stop listening for map load until we're done and know we want to start the next config.
 	FCoreUObjectDelegates::PostLoadMapWithWorld.RemoveAll(this);
+	
+	FMoviePipelineExecutorJob& CurrentJob = ExecutorJobs[CurrentPipelineIndex];
 
 	ActiveMoviePipeline = NewObject<UMoviePipeline>(NewWorld, TargetPipelineClass);
-	ActiveMoviePipeline->Initialize(ActiveConfig);
+	ActiveMoviePipeline->Initialize(CurrentJob);
 
 	// Listen for when the pipeline thinks it has finished.
 	ActiveMoviePipeline->OnMoviePipelineFinished().AddUObject(this, &UMoviePipelineInProcessExecutor::OnMoviePipelineFinished);
@@ -41,7 +42,6 @@ void UMoviePipelineInProcessExecutor::OnMoviePipelineFinished(UMoviePipeline* In
 	// Null these out now since OnIndividualPipelineFinished might invoke something that causes a GC
 	// and we want them to go away with the GC.
 	ActiveMoviePipeline = nullptr;
-	ActiveConfig = nullptr;
 
 	// Now that another frame has passed and we should be OK to start another PIE session, notify our owner.
 	OnIndividualPipelineFinished(MoviePipeline);
diff --git a/Engine/Plugins/MovieScene/MovieRenderPipeline/Source/MovieRenderPipelineCore/Private/MoviePipelineLinearExecutor.cpp b/Engine/Plugins/MovieScene/MovieRenderPipeline/Source/MovieRenderPipelineCore/Private/MoviePipelineLinearExecutor.cpp
--- a/Engine/Plugins/MovieScene/MovieRenderPipeline/Source/MovieRenderPipelineCore/Private/MoviePipelineLinearExecutor.cpp
+++ b/Engine/Plugins/MovieScene/MovieRenderPipeline/Source/MovieRenderPipelineCore/Private/MoviePipelineLinearExecutor.cpp
@@ -4,46 +4,50 @@
 
 #define LOCTEXT_NAMESPACE "MoviePipelineLinearExecutorBase"
 
-void UMoviePipelineLinearExecutorBase::ExecuteImpl(TArray<UMovieRenderPipelineConfig*>& InPipelines)
+void UMoviePipelineLinearExecutorBase::ExecuteImpl(const TArray<FMoviePipelineExecutorJob>& InPipelines)
 {
 	if (InPipelines.Num() == 0)
 	{
 		UE_LOG(LogMovieRenderPipeline, Warning, TEXT("Executor asked to execute on empty list of pipelines."));
-		OnExecutorFinishedImpl(false);
+		OnExecutorErroredImpl(nullptr, true, LOCTEXT("EmptyPipelineError", "Executor asked to execute empty list of jobs. This was probably not intended!"));
+		OnExecutorFinishedImpl();
 		return;
 	}
 
 	// We'll process them in linear fashion and wait until each one is canceled or finishes on its own
 	// before moving onto the next one. This may be parallelizable in the future (either multiple PIE
 	// sessions, or multiple external processes) but ideally one render would maximize resource usage anyways...
-	Pipelines = InPipelines;
-
+	ExecutorJobs = InPipelines;
 	InitializationTime = FDateTime::UtcNow();
 
+	UE_LOG(LogMovieRenderPipeline, Log, TEXT("MoviePipelineLinearExecutorBase starting %d jobs."), InPipelines.Num());
+
 	StartPipelineByIndex(0);
 }
 
 void UMoviePipelineLinearExecutorBase::StartPipelineByIndex(int32 InPipelineIndex)
 {
-	check(InPipelineIndex >= 0 && InPipelineIndex < Pipelines.Num());
+	check(InPipelineIndex >= 0 && InPipelineIndex < ExecutorJobs.Num());
 	
 	CurrentPipelineIndex = InPipelineIndex;
 
-	if (!Pipelines[CurrentPipelineIndex])
+	if (!ExecutorJobs[CurrentPipelineIndex].Configuration)
 	{
 		UE_LOG(LogMovieRenderPipeline, Warning, TEXT("Found null config in list of configs to render. Aborting the pipeline processing!"));
-		OnExecutorFinishedImpl(false);
+		OnExecutorErroredImpl(nullptr, true, LOCTEXT("NullPipelineError", "Found null config in list of configs to render with. Does your config have the wrong outer?"));
+		OnExecutorFinishedImpl();
 		return;
 	}
 	
-	Start(Pipelines[CurrentPipelineIndex], CurrentPipelineIndex, Pipelines.Num());
+	UE_LOG(LogMovieRenderPipeline, Log, TEXT("MoviePipelineLinearExecutorBase starting jobs [%d/%d]"), CurrentPipelineIndex, ExecutorJobs.Num());
+	Start(ExecutorJobs[CurrentPipelineIndex]);
 }
 
 void UMoviePipelineLinearExecutorBase::OnIndividualPipelineFinished(UMoviePipeline* /* FinishedPipeline */)
 {
-	if (CurrentPipelineIndex == Pipelines.Num() - 1)
+	if (CurrentPipelineIndex == ExecutorJobs.Num() - 1)
 	{
-		OnExecutorFinishedImpl(true);
+		OnExecutorFinishedImpl();
 	}
 	else
 	{
@@ -52,11 +56,16 @@ void UMoviePipelineLinearExecutorBase::OnIndividualPipelineFinished(UMoviePipeli
 	}
 }
 
-void UMoviePipelineLinearExecutorBase::OnExecutorFinishedImpl(const bool bInSuccess)
+void UMoviePipelineLinearExecutorBase::OnPipelineErrored(UMoviePipeline* InPipeline, bool bIsFatal, FText ErrorText)
+{
+	OnExecutorErroredImpl(InPipeline, bIsFatal, ErrorText);
+}
+
+void UMoviePipelineLinearExecutorBase::OnExecutorFinishedImpl()
 {
-	UE_LOG(LogMovieRenderPipeline, Log, TEXT("MoviePipelineLinearExecutorBase produced %d Pipelines in %s."), Pipelines.Num(), *(FDateTime::UtcNow() - InitializationTime).ToString());
+	UE_LOG(LogMovieRenderPipeline, Log, TEXT("MoviePipelineLinearExecutorBase finished %d jobs in %s."), ExecutorJobs.Num(), *(FDateTime::UtcNow() - InitializationTime).ToString());
 
-	Super::OnExecutorFinishedImpl(bInSuccess);
+	Super::OnExecutorFinishedImpl();
 }
 
 #undef LOCTEXT_NAMESPACE // "MoviePipelineLinearExecutorBase"
diff --git a/Engine/Plugins/MovieScene/MovieRenderPipeline/Source/MovieRenderPipelineCore/Private/MoviePipelineMasterConfig.cpp b/Engine/Plugins/MovieScene/MovieRenderPipeline/Source/MovieRenderPipelineCore/Private/MoviePipelineMasterConfig.cpp
new file mode 100644
--- /dev/null
+++ b/Engine/Plugins/MovieScene/MovieRenderPipeline/Source/MovieRenderPipelineCore/Private/MoviePipelineMasterConfig.cpp
@@ -0,0 +1,145 @@
+// Copyright 1998-2019 Epic Games, Inc. All Rights Reserved.
+
+#include "MoviePipelineMasterConfig.h"
+#include "MovieScene.h"
+#include "LevelSequence.h"
+#include "MoviePipelineShotConfig.h"
+#include "MoviePipelineOutputBase.h"
+#include "MoviePipelineOutputSetting.h"
+
+#define LOCTEXT_NAMESPACE "MoviePipelineMasterConfig"
+
+/*FFrameRate UMoviePipelineMasterConfig::GetEffectiveFrameRate()
+{
+	LoadTargetSequence();
+
+	// Check to see if they overrode the frame rate.
+	if (bUseCustomFrameRate)
+	{
+		return OutputFrameRate;
+	}
+
+	// Pull it from the sequence if they didn't.
+	if (LoadedSequence)
+	{
+		return LoadedSequence->GetMovieScene()->GetDisplayRate();
+	}
+
+	return FFrameRate();
+}*/
+
+UMoviePipelineShotConfig* UMoviePipelineMasterConfig::GetConfigForShot(const FString& ShotName) const
+{
+	UMoviePipelineShotConfig* OutConfig = PerShotConfigMapping.FindRef(ShotName);
+
+	// They didn't customize this shot, return the global pipeline default
+	if (!OutConfig)
+	{
+		OutConfig = DefaultShotConfig;
+	}
+
+	return OutConfig;
+}
+
+/*void UMoviePipelineMasterConfig::LoadTargetSequence()
+{
+	if (LoadedSequence)
+	{
+		return;
+	}
+
+	LoadedSequence = LoadObject<ULevelSequence>(this, *Sequence.GetAssetPathString());
+	ensureMsgf(LoadedSequence, TEXT("Failed to load target sequence. Pipeline is not fully configured."));
+}*/
+
+/*bool UMoviePipelineShotConfig::ValidateConfig(TArray<FText>& OutValidationErrors) const
+{
+	bool bValidPipeline = true;
+	OutValidationErrors.Reset();
+	
+	// Check that we have a sequence to render
+	if(!Sequence.IsValid())
+	{
+		OutValidationErrors.Add(LOCTEXT("InvalidInputSequence", "No sequence to render has been specified."));
+		bValidPipeline = false;
+	}
+	
+	// Check that we have a valid output resolution
+	if(OutputResolution.X == 0 || OutputResolution.Y == 0)
+	{
+		OutValidationErrors.Add(LOCTEXT("InvalidOutputResolution", "Output Resolution must be greater than zero."));
+		bValidPipeline = false;
+	}
+	
+	// Check for a non-zero framerate
+	if(OutputFrameRate.AsDecimal() == 0.0)
+	{
+		OutValidationErrors.Add(LOCTEXT("InvalidOutputFrameRate", "Output FrameRate must be greater than zero."));
+		bValidPipeline = false;
+	}
+	
+	// Output Directory can probably be ignored. The folder path you output to may not exist yet anyways.
+	
+	// Give each setting a chance to validate the pipeline
+	for(const UMoviePipelineSetting* Setting : Settings)
+	{
+		bValidPipeline &= Setting->ValidatePipeline(OutValidationErrors);
+	}
+	
+	// ToDo-MoviePipeline: Reflect over all settings classes and make sure all required classes have been added.
+	// ToDo-MoviePipeline: Scan for any shots/sections that don't start on whole frames issue warning.
+	
+	// Give each Input Buffer a chance to validate the pipeline
+	for(const UMoviePipelineRenderPass* InputBuffer : InputBuffers)
+	{
+		bValidPipeline &= InputBuffer->ValidatePipeline(OutValidationErrors);
+	}
+	
+	// Give each output container a chance to validate the pipeline
+	for(const UMoviePipelineOutput* OutputContainer : OutputContainers)
+	{
+		bValidPipeline &= OutputContainer->ValidatePipeline(OutValidationErrors);
+	}
+
+	// Ensure we have at least one input buffer and one output buffer.
+	if (InputBuffers.Num() == 0)
+	{
+		OutValidationErrors.Add(LOCTEXT("NoInputBuffer", "Must specify at least one Input Buffer to capture."));
+		bValidPipeline = false;
+	}
+
+	if (OutputContainers.Num() == 0)
+	{
+		OutValidationErrors.Add(LOCTEXT("NoOutputContainer", "Must specify at least one Output Container to write to."));
+		bValidPipeline = false;
+	}
+	
+	return bValidPipeline;
+}*/
+
+UMoviePipelineMasterConfig::UMoviePipelineMasterConfig()
+{
+	DefaultShotConfig = NewObject<UMoviePipelineShotConfig>(this, "DefaultShotConfig");
+	
+	// Always add at least the output settings block since having a framerate/directory/etc. is crucial.
+	Settings.Add(NewObject<UMoviePipelineOutputSetting>(this, "OutputSetting"));
+
+}
+
+TArray<UMoviePipelineOutputBase*> UMoviePipelineMasterConfig::GetOutputContainers() const
+{
+	TArray<UMoviePipelineOutputBase*> OutputContainers;
+
+	for (UMoviePipelineSetting* Setting : GetSettings())
+	{
+		UMoviePipelineOutputBase* Output = Cast<UMoviePipelineOutputBase>(Setting);
+		if (Output)
+		{
+			OutputContainers.Add(Output);
+		}
+	}
+
+	return OutputContainers;
+}
+
+#undef LOCTEXT_NAMESPACE // "MovieRenderPipelineConfig"
\ No newline at end of file
diff --git a/Engine/Plugins/MovieScene/MovieRenderPipeline/Source/MovieRenderPipelineCore/Private/MoviePipelineOutput.cpp b/Engine/Plugins/MovieScene/MovieRenderPipeline/Source/MovieRenderPipelineCore/Private/MoviePipelineOutput.cpp
deleted file mode 100644
diff --git a/Engine/Plugins/MovieScene/MovieRenderPipeline/Source/MovieRenderPipelineCore/Private/MoviePipelineOutputBase.cpp b/Engine/Plugins/MovieScene/MovieRenderPipeline/Source/MovieRenderPipelineCore/Private/MoviePipelineOutputBase.cpp
new file mode 100644
--- /dev/null
+++ b/Engine/Plugins/MovieScene/MovieRenderPipeline/Source/MovieRenderPipelineCore/Private/MoviePipelineOutputBase.cpp
@@ -0,0 +1,3 @@
+// Copyright 1998-2019 Epic Games, Inc. All Rights Reserved.
+
+#include "MoviePipelineOutputBase.h"
diff --git a/Engine/Plugins/MovieScene/MovieRenderPipeline/Source/MovieRenderPipelineCore/Private/MoviePipelineOutputSetting.cpp b/Engine/Plugins/MovieScene/MovieRenderPipeline/Source/MovieRenderPipelineCore/Private/MoviePipelineOutputSetting.cpp
new file mode 100644
--- /dev/null
+++ b/Engine/Plugins/MovieScene/MovieRenderPipeline/Source/MovieRenderPipelineCore/Private/MoviePipelineOutputSetting.cpp
@@ -0,0 +1,18 @@
+// Copyright 1998-2019 Epic Games, Inc. All Rights Reserved.
+
+#include "MoviePipelineOutputSetting.h"
+#include "Misc/Paths.h"
+
+UMoviePipelineOutputSetting::UMoviePipelineOutputSetting()
+	: OutputResolution(FIntPoint(1920, 1080))
+	, bUseCustomFrameRate(false)
+	, OutputFrameRate(FFrameRate(24, 1))
+	, OutputFrameStep(FFrameNumber(1))
+	, bOverrideExistingOutput(true)
+	, bUseCustomPlaybackRange(false)
+	, CustomStartFrame(FFrameNumber(0))
+	, CustomEndFrame(FFrameNumber(0))
+	{
+		FileNameFormat = TEXT("{sequence_name}.{frame_number}");
+		OutputDirectory.Path = FPaths::ProjectSavedDir() / TEXT("/VideoCaptures/");
+	}
diff --git a/Engine/Plugins/MovieScene/MovieRenderPipeline/Source/MovieRenderPipelineCore/Private/MoviePipelineRendering.cpp b/Engine/Plugins/MovieScene/MovieRenderPipeline/Source/MovieRenderPipelineCore/Private/MoviePipelineRendering.cpp
--- a/Engine/Plugins/MovieScene/MovieRenderPipeline/Source/MovieRenderPipelineCore/Private/MoviePipelineRendering.cpp
+++ b/Engine/Plugins/MovieScene/MovieRenderPipeline/Source/MovieRenderPipelineCore/Private/MoviePipelineRendering.cpp
@@ -3,15 +3,18 @@
 #include "MoviePipeline.h"
 #include "MovieRenderPipelineCoreModule.h"
 #include "Misc/FrameRate.h"
-#include "MoviePipelineOutput.h"
-#include "MovieRenderPipelineConfig.h"
+#include "MoviePipelineOutputBase.h"
 #include "MoviePipelineAccumulationSetting.h"
 #include "MoviePipelineShotConfig.h"
 #include "MoviePipelineRenderPass.h"
 #include "MoviePipelineOutputBuilder.h"
 #include "RenderingThread.h"
+#include "MoviePipelineOutputSetting.h"
+#include "MoviePipelineConfigBase.h"
+#include "MoviePipelineMasterConfig.h"
+#include "Math/Halton.h"
 
-void UMoviePipeline::SetupRenderingPipelineForShot(FMoviePipelineShotCache& Shot)
+void UMoviePipeline::SetupRenderingPipelineForShot(FMoviePipelineShotInfo& Shot)
 {
 	/*
 	* To support tiled rendering we take the final effective resolution and divide
@@ -27,11 +30,13 @@ void UMoviePipeline::SetupRenderingPipelineForShot(FMoviePipelineShotCache& Shot
 	* RightOffset = (1925-1920-LeftOffset)
 	*/
 	UMoviePipelineAccumulationSetting* AccumulationSettings = Shot.ShotConfig->FindOrAddSetting<UMoviePipelineAccumulationSetting>();
+	UMoviePipelineOutputSetting* OutputSettings = GetPipelineMasterConfig()->FindSetting<UMoviePipelineOutputSetting>();
+	check(OutputSettings);
 
 	FMoviePipelineRenderPassInitSettings RenderPassInitSettings;
 	RenderPassInitSettings.TileResolution = FIntPoint(
-		FMath::CeilToInt(Config->OutputResolution.X / AccumulationSettings->TileCount),
-		FMath::CeilToInt(Config->OutputResolution.Y / AccumulationSettings->TileCount));
+		FMath::CeilToInt(OutputSettings->OutputResolution.X / AccumulationSettings->TileCount),
+		FMath::CeilToInt(OutputSettings->OutputResolution.Y / AccumulationSettings->TileCount));
 	RenderPassInitSettings.TileCount = AccumulationSettings->TileCount;
 	RenderPassInitSettings.ShotConfig = Shot.ShotConfig;
 
@@ -50,42 +55,28 @@ void UMoviePipeline::SetupRenderingPipelineForShot(FMoviePipelineShotCache& Shot
 	// Code expects at least a 1x1 tile.
 	ensure(RenderPassInitSettings.TileCount > 0);
 
-	for (UMoviePipelineRenderPass* Input : Shot.ShotConfig->InputBuffers)
+	for (UMoviePipelineRenderPass* Input : Shot.RenderPasses)
 	{
 		Input->Setup(RenderPassInitSettings, OutputPipe.ToSharedRef());
 	}
 }
 
-void UMoviePipeline::TeardownRenderingPipelineForShot(FMoviePipelineShotCache& Shot)
+void UMoviePipeline::TeardownRenderingPipelineForShot(FMoviePipelineShotInfo& Shot)
 {
-	for (UMoviePipelineRenderPass* Input : Shot.ShotConfig->InputBuffers)
-	{
-		Input->Teardown();
-	}
-}
-
-float Halton(int32 Index, int32 Base)
-{
-	float Result = 0.0f;
-	float InvBase = 1.0f / Base;
-	float Fraction = InvBase;
-	while (Index > 0)
-	{
-		Result += (Index % Base) * Fraction;
-		Index /= Base;
-		Fraction *= InvBase;
-	}
-	return Result;
+	//for (UMoviePipelineRenderPass* Input : Shot.ShotConfig->InputBuffers)
+	//{
+	//	Input->Teardown();
+	//}
 }
 
 void UMoviePipeline::RenderFrame()
 {
 	// Flush built in systems before we render anything. This maximizes the likelyhood that the data is prepared for when
 	// the render thread uses it.
-	FlushAsyncSystems();
+	FlushAsyncEngineSystems();
 
-	FMoviePipelineShotCache& CurrentShot = ShotList[CurrentShotIndex];
-	FMoviePipelineShotCutCache& CurrentCameraCut = CurrentShot.GetCurrentCameraCut();
+	FMoviePipelineShotInfo& CurrentShot = ShotList[CurrentShotIndex];
+	FMoviePipelineCameraCutInfo& CurrentCameraCut = CurrentShot.GetCurrentCameraCut();
 
 	// We render a frame during motion blur fixes (but don't try to submit it for output) to allow seeding
 	// the view histories with data. We want to do this all in one tick so that the motion is correct.
@@ -100,10 +91,10 @@ void UMoviePipeline::RenderFrame()
 
 	// Allow our containers to think
 	// ToDo: Is this necessaary?
-	for (UMoviePipelineOutput* Container : Config->OutputContainers)
-	{
-		Container->OnPostTick();
-	}
+	// for (UMoviePipelineOutputBase* Container : Config->OutputContainers)
+	// {
+	// 	Container->OnPostTick();
+	// }
 	
 	// To produce a frame from the movie pipeline we may render many frames over a period of time, additively collecting the results
 	// together before submitting it for writing on the last result - this is referred to as an "output frame". The 1 (or more) samples
@@ -118,15 +109,18 @@ void UMoviePipeline::RenderFrame()
 	// In short, for each output frame, for each accumulation frame, for each tile X/Y, for each jitter, we render a pass. This setup is
 	// designed to maximize the likely hood of deterministic rendering and that different passes line up with each other.
 	UMoviePipelineAccumulationSetting* AccumulationSettings = CurrentShot.ShotConfig->FindOrAddSetting<UMoviePipelineAccumulationSetting>();
+	UMoviePipelineOutputSetting* OutputSetting = GetPipelineMasterConfig()->FindSetting<UMoviePipelineOutputSetting>();
+	check(OutputSetting);
+
+
 	const int32 NumTilesX = AccumulationSettings->TileCount;
 	const int32 NumTilesY = NumTilesX;
 	const bool bIsUsingOverlappedTiles = AccumulationSettings->bIsUsingOverlappedTiles;
 
 	int32 NumSpatialSamples = AccumulationSettings->SpatialSampleCount;
 	ensure(NumTilesX > 0 && NumTilesY > 0 && NumSpatialSamples > 0);
-	bool bHasCountedSamples = false;
 
-	TArray<UMoviePipelineRenderPass*> InputBuffers = CurrentShot.ShotConfig->InputBuffers;
+	TArray<UMoviePipelineRenderPass*> InputBuffers = CurrentShot.RenderPasses;
 
 	// If this is the first sample for a new frame, we want to notify the output builder that it should expect data to accumulate for this frame.
 	if (CachedOutputState.TemporalSampleIndex == 0)
@@ -208,12 +202,6 @@ void UMoviePipeline::RenderFrame()
 				TilePixelSizeFractionY = 1.0f;
 			}
 
-			// Tell each pass to reset their spatial sampling accumulation.
-			for (UMoviePipelineRenderPass* Input : InputBuffers)
-			{
-				//Input->ResetSpatialSamplingAccumulation();
-			}
-
 			// If they're trying to seed the histories, then we'll just override the NumSpatialSamples since they're contributing towards
 			// our final target buffer which needs the history.
 			if (bIsHistoryOnlyFrame)
@@ -225,12 +213,8 @@ void UMoviePipeline::RenderFrame()
 			// spatial jitters which are accumulated together before contributing to the temporal accumulation.
 			for (int32 SpatialSample = 0; SpatialSample < NumSpatialSamples; SpatialSample++)
 			{
-				// We only want to count samples from the first tile
-				if (!bHasCountedSamples)
-				{
-					CurrentCameraCut.NumSamplesRendered++;
-					CachedOutputState.TotalSamplesRendered++;
-				}
+				// Count this as a sample rendered for the current work.
+				CurrentCameraCut.CurrentWorkInfo.NumSamples++;
 
 				float OffsetX = 0.5f;
 				float OffsetY = 0.5f; 
@@ -251,8 +235,8 @@ void UMoviePipeline::RenderFrame()
 				double FinalSubPixelShiftX = - (TileShiftX + SpatialShiftX - 0.5);
 				double FinalSubPixelShiftY = (TileShiftY + SpatialShiftY - 0.5);
 
-				int32 TileSizeX = FMath::CeilToInt(Config->OutputResolution.X / NumTilesX);
-				int32 TileSizeY = FMath::CeilToInt(Config->OutputResolution.Y / NumTilesY);
+				int32 TileSizeX = FMath::CeilToInt(OutputSetting->OutputResolution.X / NumTilesX);
+				int32 TileSizeY = FMath::CeilToInt(OutputSetting->OutputResolution.Y / NumTilesY);
 
 				// in the case of overlapped rendering, the render target size will be larger than tile size because of overlap
 				int32 PadX = bIsUsingOverlappedTiles ? int32(float(TileSizeX) * AccumulationSettings->PadRatioX) : 0;
@@ -319,32 +303,8 @@ void UMoviePipeline::RenderFrame()
 				}
 
 			}
-
-			// We only want to count samples from the first tile
-			bHasCountedSamples = true;
 		}
 	}
 	
-	// 	// Get our list of buffers that are supposed to be producing this frame
-
-	// 
-	// 	// Once all inputs have had a chance to render, we're going to flush the rendering and then
-	// 	// transfer all of the created data to the output containers in one go.
 	FlushRenderingCommands();
-	// 
-	// 	// Store all of the data generated this frame
-	TArray<MoviePipeline::FOutputFrameData> AllPassFrameData;
-	// 
-	for (UMoviePipelineRenderPass* Input : InputBuffers)
-	{
-		MoviePipeline::FOutputFrameData& OutFrameData = AllPassFrameData.Add_GetRef(MoviePipeline::FOutputFrameData());
-		Input->GetFrameData(OutFrameData);
-	}
-	// 	
-	TArray<UMoviePipelineOutput*> OutputContainers = Config->OutputContainers;
-	for (UMoviePipelineOutput* Output : OutputContainers)
-	{
-		Output->ProcessFrame(AllPassFrameData, CachedOutputState, Config->OutputDirectory);
-	}
-	// }
 }
\ No newline at end of file
diff --git a/Engine/Plugins/MovieScene/MovieRenderPipeline/Source/MovieRenderPipelineCore/Private/MoviePipelineSetting.cpp b/Engine/Plugins/MovieScene/MovieRenderPipeline/Source/MovieRenderPipelineCore/Private/MoviePipelineSetting.cpp
--- a/Engine/Plugins/MovieScene/MovieRenderPipeline/Source/MovieRenderPipelineCore/Private/MoviePipelineSetting.cpp
+++ b/Engine/Plugins/MovieScene/MovieRenderPipeline/Source/MovieRenderPipelineCore/Private/MoviePipelineSetting.cpp
@@ -3,10 +3,10 @@
 #include "MoviePipelineSetting.h"
 #include "MoviePipeline.h"
 
-void UMoviePipelineSetting::OnInitializedForPipeline(UMoviePipeline* InPipeline)
+void UMoviePipelineSetting::SetupForPipeline(UMoviePipeline* InPipeline)
 {
 	CachedPipeline = MakeWeakObjectPtr(InPipeline);
-	OnInitializedForPipelineImpl(InPipeline);
+	SetupForPipelineImpl(InPipeline);
 }
 
 UMoviePipeline* UMoviePipelineSetting::GetPipeline() const
diff --git a/Engine/Plugins/MovieScene/MovieRenderPipeline/Source/MovieRenderPipelineCore/Private/MoviePipelineShotConfig.cpp b/Engine/Plugins/MovieScene/MovieRenderPipeline/Source/MovieRenderPipelineCore/Private/MoviePipelineShotConfig.cpp
--- a/Engine/Plugins/MovieScene/MovieRenderPipeline/Source/MovieRenderPipelineCore/Private/MoviePipelineShotConfig.cpp
+++ b/Engine/Plugins/MovieScene/MovieRenderPipeline/Source/MovieRenderPipelineCore/Private/MoviePipelineShotConfig.cpp
@@ -1,84 +1,24 @@
 // Copyright 1998-2019 Epic Games, Inc. All Rights Reserved.
 
 #include "MoviePipelineShotConfig.h"
-#include "MoviePipelineSetting.h"
 #include "MoviePipelineRenderPass.h"
-#include "MoviePipelineOutput.h"
-#include "Misc/FrameRate.h"
 
 #define LOCTEXT_NAMESPACE "MovieRenderShotConfig"
 
-bool UMoviePipelineShotConfig::ValidateConfig(TArray<FText>& OutValidationErrors) const
+TArray<UMoviePipelineRenderPass*> UMoviePipelineShotConfig::GetRenderPasses() const
 {
-	bool bValidPipeline = true;
-	OutValidationErrors.Reset();
-	
-	// Check that we have a sequence to render
-	if(!Sequence.IsValid())
-	{
-		OutValidationErrors.Add(LOCTEXT("InvalidInputSequence", "No sequence to render has been specified."));
-		bValidPipeline = false;
-	}
-	
-	// Check that we have a valid output resolution
-	if(OutputResolution.X == 0 || OutputResolution.Y == 0)
-	{
-		OutValidationErrors.Add(LOCTEXT("InvalidOutputResolution", "Output Resolution must be greater than zero."));
-		bValidPipeline = false;
-	}
-	
-	// Check for a non-zero framerate
-	if(OutputFrameRate.AsDecimal() == 0.0)
-	{
-		OutValidationErrors.Add(LOCTEXT("InvalidOutputFrameRate", "Output FrameRate must be greater than zero."));
-		bValidPipeline = false;
-	}
-	
-	// Output Directory can probably be ignored. The folder path you output to may not exist yet anyways.
-	
-	// Give each setting a chance to validate the pipeline
-	for(const UMoviePipelineSetting* Setting : Settings)
-	{
-		bValidPipeline &= Setting->ValidatePipeline(OutValidationErrors);
-	}
-	
-	// ToDo-MoviePipeline: Reflect over all settings classes and make sure all required classes have been added.
-	// ToDo-MoviePipeline: Scan for any shots/sections that don't start on whole frames issue warning.
-	
-	// Give each Input Buffer a chance to validate the pipeline
-	for(const UMoviePipelineRenderPass* InputBuffer : InputBuffers)
-	{
-		bValidPipeline &= InputBuffer->ValidatePipeline(OutValidationErrors);
-	}
-	
-	// Give each output container a chance to validate the pipeline
-	for(const UMoviePipelineOutput* OutputContainer : OutputContainers)
-	{
-		bValidPipeline &= OutputContainer->ValidatePipeline(OutValidationErrors);
-	}
+	TArray<UMoviePipelineRenderPass*> RenderPasses;
 
-	// Ensure we have at least one input buffer and one output buffer.
-	if (InputBuffers.Num() == 0)
+	for (UMoviePipelineSetting* Setting : GetSettings())
 	{
-		OutValidationErrors.Add(LOCTEXT("NoInputBuffer", "Must specify at least one Input Buffer to capture."));
-		bValidPipeline = false;
+		UMoviePipelineRenderPass* RenderPass = Cast<UMoviePipelineRenderPass>(Setting);
+		if (RenderPass)
+		{
+			RenderPasses.Add(RenderPass);
+		}
 	}
 
-	if (OutputContainers.Num() == 0)
-	{
-		OutValidationErrors.Add(LOCTEXT("NoOutputContainer", "Must specify at least one Output Container to write to."));
-		bValidPipeline = false;
-	}
-	
-	return bValidPipeline;
+	return RenderPasses;
 }
 
-
-void UMoviePipelineShotConfig::RemoveSetting(UMoviePipelineSetting* InSource)
-{
-	Settings.Remove(InSource);
-
-	// Update our cached serial number so the UI rebuilds the tree
-	++SettingsSerialNumber;
-}
 #undef LOCTEXT_NAMESPACE // "MovieRenderShotConfig"
\ No newline at end of file
diff --git a/Engine/Plugins/MovieScene/MovieRenderPipeline/Source/MovieRenderPipelineCore/Private/MoviePipelineTiming.cpp b/Engine/Plugins/MovieScene/MovieRenderPipeline/Source/MovieRenderPipelineCore/Private/MoviePipelineTiming.cpp
--- a/Engine/Plugins/MovieScene/MovieRenderPipeline/Source/MovieRenderPipelineCore/Private/MoviePipelineTiming.cpp
+++ b/Engine/Plugins/MovieScene/MovieRenderPipeline/Source/MovieRenderPipelineCore/Private/MoviePipelineTiming.cpp
@@ -4,8 +4,7 @@
 #include "MovieRenderPipelineDataTypes.h"
 #include "MoviePipelineAccumulationSetting.h"
 #include "MoviePipelineShotConfig.h"
-#include "MoviePipelineOutput.h"
-#include "MovieRenderPipelineConfig.h"
+#include "MoviePipelineOutputBase.h"
 #include "CoreGlobals.h"
 #include "Containers/Array.h"
 #include "MovieRenderPipelineCoreModule.h"
@@ -15,6 +14,7 @@
 #include "Engine/World.h"
 #include "GameFramework/PlayerController.h"
 #include "GameFramework/WorldSettings.h"
+#include "MoviePipelineMasterConfig.h"
 
 void UMoviePipeline::TickProducingFrames()
 {
@@ -37,8 +37,8 @@ void UMoviePipeline::TickProducingFrames()
 	// on the last Tick.
 	GetWorld()->GetFirstPlayerController()->SetPause(false);
 
-	FMoviePipelineShotCache& CurrentShot = ShotList[CurrentShotIndex];
-	FMoviePipelineShotCutCache& CurrentCameraCut = CurrentShot.GetCurrentCameraCut();
+	FMoviePipelineShotInfo& CurrentShot = ShotList[CurrentShotIndex];
+	FMoviePipelineCameraCutInfo& CurrentCameraCut = CurrentShot.GetCurrentCameraCut();
 
 	// Check to see if we need to initialize a new shot. This potentially changes a lot of state.
 	if (CurrentCameraCut.State == EMovieRenderShotState::Uninitialized)
@@ -53,9 +53,6 @@ void UMoviePipeline::TickProducingFrames()
 		UE_LOG(LogMovieRenderPipeline, Log, TEXT("[%d] Initializing Camera Cut [%d/%d] in Shot (%s)."), GFrameCounter, 
 			CurrentShot.CurrentCameraCutIndex + 1, CurrentShot.CameraCuts.Num(), *CurrentShot.GetDisplayName());
 		
-		// Ensure our Camera Cut is the only one active so it will evaluate the right camera.
-		SetSoloCameraCut(CurrentCameraCut);
-
 		CurrentCameraCut.SetNextState(EMovieRenderShotState::Uninitialized);
 
 		// The shot has been initialized and the new state chosen. Jump the Sequence to the starting time.
@@ -95,7 +92,7 @@ void UMoviePipeline::TickProducingFrames()
 		CurrentCameraCut.NumWarmUpFramesRemaining--;
 
 		// Render the next frame at a reasonable frame rate to pass time in the world.
-		double FrameDeltaTime = Config->GetEffectiveFrameRate().AsInterval();
+		double FrameDeltaTime = GetEffectiveFrameRate().AsInterval();
 
 		UE_LOG(LogMovieRenderPipeline, VeryVerbose, TEXT("[%d] Shot WarmUp set engine DeltaTime to %f seconds."), GFrameCounter, FrameDeltaTime);
 		CustomTimeStep->SetCachedFrameTiming(MoviePipeline::FFrameTimeStepCache(FrameDeltaTime));
@@ -193,11 +190,7 @@ void UMoviePipeline::TickProducingFrames()
 	if (!CurrentCameraCut.bHasNotifiedFrameProductionStarted)
 	{
 		CurrentCameraCut.bHasNotifiedFrameProductionStarted = true;
-		for (UMoviePipelineOutput* Container : Config->OutputContainers)
-		{
-			// ToDo: We should name this OnRenderingStart/OnRenderingStop
-			Container->OnFrameProductionStart();
-		}
+
 	}
 
 	// Alright we're in the point where we finally want to start rendering out frames, we've finished warming up,
@@ -225,7 +218,6 @@ void UMoviePipeline::TickProducingFrames()
 		}
 
 		CachedOutputState.TemporalSampleIndex++;
-		CurrentCameraCut.CurrentFrameSamples = (CachedOutputState.TemporalSampleIndex + 1) * CurrentCameraCut.NumSpatialSamples;
 
 		if (AccumulationSettings->TemporalSampleCount == 1)
 		{
@@ -466,6 +458,8 @@ void UMoviePipeline::TickProducingFrames()
 		// The combination of shutter angle percentage, non-uniform render frame delta times and dividing by sample
 		// count produce the correct length for motion blur in all cases.
 		CachedOutputState.MotionBlurFraction = FrameMetrics.ShutterAnglePercentage / AccumulationSettings->TemporalSampleCount;
+
+		// Update our Output State with any data we need to derive based on our current time.
 		CalculateFrameNumbersForOutputState(FrameMetrics, CurrentCameraCut, CachedOutputState);
 
 		// Now that we know the delta time for the upcoming frame, we can see if this time would extend past the end of our
@@ -480,15 +474,10 @@ void UMoviePipeline::TickProducingFrames()
 			CurrentCameraCut.State = EMovieRenderShotState::Finished;
 			
 			// Compare our expected vs. actual results for logging.
+			if (CurrentCameraCut.CurrentWorkInfo != CurrentCameraCut.TotalWorkInfo)
 			{
-				FFrameNumber NumExpectedFrames = CurrentCameraCut.GetOutputFrameCountEstimate();
-				FFrameNumber NumExpectedSamples = CurrentCameraCut.GetSampleCountEstimate(/*bIncludeWarmup*/ true, /*bIncludeMotionBlur*/ true);
-
-				FFrameNumber NumFrames = CurrentCameraCut.NumFramesProduced;
-				FFrameNumber NumSamples = CurrentCameraCut.NumSamplesRendered;
-
-				UE_LOG(LogMovieRenderPipeline, Log, TEXT("\t Frames Produced: %d (Expected: %d), Samples Rendered: %d (Expected: %d)"),
-					NumFrames.Value, NumExpectedFrames.Value, NumSamples.Value, NumExpectedSamples.Value);
+				UE_LOG(LogMovieRenderPipeline, Warning, TEXT("Mismatch in work done vs. expected work done.\nExpected: %s\nTotal: %s"),
+					*CurrentCameraCut.TotalWorkInfo.ToDisplayString(), *CurrentCameraCut.CurrentWorkInfo.ToDisplayString());
 			}
 
 			// We pause at the end too, just so that frames during finalize don't continue to trigger Sequence Eval messages.
@@ -502,7 +491,7 @@ void UMoviePipeline::TickProducingFrames()
 				CurrentShotIndex++;
 
 				// Notify our containers that the current shot has ended.
-				for (UMoviePipelineOutput* Container : Config->OutputContainers)
+				for (UMoviePipelineOutputBase* Container : GetPipelineMasterConfig()->GetOutputContainers())
 				{
 					Container->OnShotFinished(CurrentShot);
 				}
@@ -520,7 +509,7 @@ void UMoviePipeline::TickProducingFrames()
 				FCoreDelegates::OnEndFrame.RemoveAll(this);
 
 				// Reset the Custom Timestep because we don't care how long the engine takes now
-				GEngine->SetCustomTimeStep(CachedCustomTimeStep);
+				GEngine->SetCustomTimeStep(CachedPrevCustomTimeStep);
 
 				BeginFinalize();
 				return;
@@ -540,7 +529,6 @@ void UMoviePipeline::TickProducingFrames()
 		if (CachedOutputState.TemporalSampleIndex == 0)
 		{
 			CachedOutputState.OutputFrameNumber++;
-			CurrentCameraCut.NumFramesProduced++;
 		}
 		
 		// Set our time step for the next frame
@@ -553,7 +541,7 @@ void UMoviePipeline::TickProducingFrames()
 	check(false);
 }
 
-void UMoviePipeline::CalculateFrameNumbersForOutputState(const MoviePipeline::FFrameConstantMetrics& InFrameMetrics, const FMoviePipelineShotCutCache& InCameraCut, FMoviePipelineFrameOutputState& InOutOutputState) const
+void UMoviePipeline::CalculateFrameNumbersForOutputState(const MoviePipeline::FFrameConstantMetrics& InFrameMetrics, const FMoviePipelineCameraCutInfo& InCameraCut, FMoviePipelineFrameOutputState& InOutOutputState) const
 {
 	// Find the closest number on the source Sequence. This may produce duplicates when using slow-mo tracks.
 	FFrameRate SourceDisplayRate = TargetSequence->GetMovieScene()->GetDisplayRate();
diff --git a/Engine/Plugins/MovieScene/MovieRenderPipeline/Source/MovieRenderPipelineCore/Private/MovieRenderPipelineCommandLine.cpp b/Engine/Plugins/MovieScene/MovieRenderPipeline/Source/MovieRenderPipelineCore/Private/MovieRenderPipelineCommandLine.cpp
--- a/Engine/Plugins/MovieScene/MovieRenderPipeline/Source/MovieRenderPipelineCore/Private/MovieRenderPipelineCommandLine.cpp
+++ b/Engine/Plugins/MovieScene/MovieRenderPipeline/Source/MovieRenderPipelineCore/Private/MovieRenderPipelineCommandLine.cpp
@@ -7,8 +7,9 @@
 #include "MoviePipelineExecutor.h"
 #include "MoviePipeline.h"
 #include "MoviePipelineInProcessExecutor.h"
-#include "MovieRenderPipelineConfig.h"
 #include "LevelSequence.h"
+#include "MoviePipelineMasterConfig.h"
+#include "MovieRenderPipelineDataTypes.h"
 
 #if WITH_EDITOR
 //#include "Editor.h"
@@ -34,13 +35,13 @@ void FMovieRenderPipelineCoreModule::InitializeCommandLineMovieRender()
 
 	// Attempt to convert their command line arguments into the required objects.
 	ULevelSequence* LevelSequence = nullptr;
-	UMovieRenderPipelineConfig* Config = nullptr;
+	UMoviePipelineMasterConfig* MasterConfig = nullptr;
 	UClass* ExecutorClass = nullptr;
 	UClass* PipelineClass = nullptr;
 
 	uint8 ReturnCode = GetMovieRenderClasses(SequenceAssetValue, SettingsAssetValue, MoviePipelineLocalExecutorClassType, MoviePipelineClassType, 
-											 /*Out*/ LevelSequence, /*Out*/ Config, /*Out*/ ExecutorClass, /*Out*/ PipelineClass);
-	if (!ensureMsgf(LevelSequence && Config && ExecutorClass && PipelineClass, TEXT("There was a failure parsing the command line and a movie render cannot be started. Check the log for more details.")))
+											 /*Out*/ LevelSequence, /*Out*/ MasterConfig, /*Out*/ ExecutorClass, /*Out*/ PipelineClass);
+	if (!ensureMsgf(LevelSequence && MasterConfig && ExecutorClass && PipelineClass, TEXT("There was a failure parsing the command line and a movie render cannot be started. Check the log for more details.")))
 	{
 		// Take the failure return code from the detection of our command line arguments.
 		FPlatformMisc::RequestExitWithStatus(/*Force*/ false, /*ReturnCode*/ ReturnCode);
@@ -49,7 +50,7 @@ void FMovieRenderPipelineCoreModule::InitializeCommandLineMovieRender()
 	else
 	{
 		UE_LOG(LogMovieRenderPipeline, Log, TEXT("Successfully detected and loaded required movie arguments. Rendering will begin once the map is loaded."));
-		UE_LOG(LogMovieRenderPipeline, Log, TEXT("LevelSequence: %s Config: %s ExecutorClass: %s PipelineClass: %s"), *LevelSequence->GetPathName(), *Config->GetPathName(), *ExecutorClass->GetPathName(), *PipelineClass->GetPathName());
+		UE_LOG(LogMovieRenderPipeline, Log, TEXT("LevelSequence: %s Config: %s ExecutorClass: %s PipelineClass: %s"), *LevelSequence->GetPathName(), *MasterConfig->GetPathName(), *ExecutorClass->GetPathName(), *PipelineClass->GetPathName());
 	}
 
 	// By this time, we know what assets you want to render, how to process the array of assets, and what runs an individual render. First we will create an executor.
@@ -59,16 +60,16 @@ void FMovieRenderPipelineCoreModule::InitializeCommandLineMovieRender()
 	// and means we only have to add/remove one thing from the root set, everything else uses normal outer ownership.
 	ExecutorBase->AddToRoot();
 	ExecutorBase->OnExecutorFinished().AddRaw(this, &FMovieRenderPipelineCoreModule::OnCommandLineMovieRenderCompleted);
+	ExecutorBase->OnExecutorErrored().AddRaw(this, &FMovieRenderPipelineCoreModule::OnCommandLineMovieRenderErrored);
 	ExecutorBase->SetMoviePipelineClass(PipelineClass);
 
-	TArray<UMovieRenderPipelineConfig*> Pipelines;
-	Pipelines.Add(Config);
-	Config->Sequence = FSoftObjectPath(SequenceAssetValue);
+	TArray<FMoviePipelineExecutorJob> Jobs;
+	Jobs.Add(FMoviePipelineExecutorJob(SequenceAssetValue, MasterConfig));
 
-	ExecutorBase->Execute(Pipelines);
+	ExecutorBase->Execute(Jobs);
 }
 
-void FMovieRenderPipelineCoreModule::OnCommandLineMovieRenderCompleted(UMoviePipelineExecutorBase* InExecutor)
+void FMovieRenderPipelineCoreModule::OnCommandLineMovieRenderCompleted(UMoviePipelineExecutorBase* InExecutor, bool bSuccess)
 {
 	if (InExecutor)
 	{
@@ -79,6 +80,12 @@ void FMovieRenderPipelineCoreModule::OnCommandLineMovieRenderCompleted(UMoviePip
 	FPlatformMisc::RequestExitWithStatus(/*Force*/ false, /*ReturnCode*/ 0);
 }
 
+void FMovieRenderPipelineCoreModule::OnCommandLineMovieRenderErrored(UMoviePipelineExecutorBase* InExecutor, UMoviePipeline* InPipelineWithError, bool bIsFatal, FText ErrorText)
+{
+
+}
+
+
 UClass* GetLocalExecutorClass(const FString& MoviePipelineLocalExecutorClassType, const FString ExecutorClassFormatString)
 {
 	if (MoviePipelineLocalExecutorClassType.Len() > 0)
@@ -158,7 +165,7 @@ bool FMovieRenderPipelineCoreModule::IsTryingToRenderMovieFromCommandLine(FStrin
 }
 
 uint8 FMovieRenderPipelineCoreModule::GetMovieRenderClasses(const FString& InSequenceAssetPath, const FString& InConfigAssetPath, const FString& InExecutorType, const FString& InPipelineType,
-	ULevelSequence*& OutSequence, UMovieRenderPipelineConfig*& OutConfig, UClass*& OutExecutorClass, UClass*& OutPipelineClass) const
+	ULevelSequence*& OutSequence, UMoviePipelineMasterConfig*& OutConfig, UClass*& OutExecutorClass, UClass*& OutPipelineClass) const
 {
 	if (!ensureMsgf(InSequenceAssetPath.Len() > 0, TEXT("GetMovieRenderClasses should not be called without an sequence asset path specified!")))
 	{
@@ -212,7 +219,7 @@ uint8 FMovieRenderPipelineCoreModule::GetMovieRenderClasses(const FString& InSeq
 	{
 		// Convert it to a soft object path and use that load to ensure it follows redirectors, etc.
 		FSoftObjectPath AssetPath = FSoftObjectPath(InConfigAssetPath);
-		OutConfig = CastChecked<UMovieRenderPipelineConfig>(AssetPath.TryLoad());
+		OutConfig = CastChecked<UMoviePipelineMasterConfig>(AssetPath.TryLoad());
 
 		if (!OutConfig)
 		{
diff --git a/Engine/Plugins/MovieScene/MovieRenderPipeline/Source/MovieRenderPipelineCore/Private/MovieRenderPipelineConfig.cpp b/Engine/Plugins/MovieScene/MovieRenderPipeline/Source/MovieRenderPipelineCore/Private/MovieRenderPipelineConfig.cpp
deleted file mode 100644
diff --git a/Engine/Plugins/MovieScene/MovieRenderPipeline/Source/MovieRenderPipelineCore/Public/MoviePipeline.h b/Engine/Plugins/MovieScene/MovieRenderPipeline/Source/MovieRenderPipelineCore/Public/MoviePipeline.h
--- a/Engine/Plugins/MovieScene/MovieRenderPipeline/Source/MovieRenderPipelineCore/Public/MoviePipeline.h
+++ b/Engine/Plugins/MovieScene/MovieRenderPipeline/Source/MovieRenderPipelineCore/Public/MoviePipeline.h
@@ -8,7 +8,7 @@
 #include "MoviePipeline.generated.h"
 
 // Forward Declares
-class UMovieRenderPipelineConfig;
+class UMoviePipelineMasterConfig;
 class ULevelSequence;
 class UMovieSceneSequencePlayer;
 class UMoviePipelineCustomTimeStep;
@@ -21,7 +21,10 @@ struct FImagePixelPipe;
 struct FMoviePipelineTimeController;
 class FMoviePipelineOutputMerger;
 
+
+
 DECLARE_MULTICAST_DELEGATE_OneParam(FMoviePipelineFinished, UMoviePipeline*);
+DECLARE_MULTICAST_DELEGATE_ThreeParams(FMoviePipelineErrored, UMoviePipeline* /*Pipeline*/, bool /*bIsFatal*/, FText /*ErrorText*/);
 
 UCLASS(Blueprintable)
 class MOVIERENDERPIPELINECORE_API UMoviePipeline : public UObject
@@ -34,14 +37,9 @@ public:
 	/** 
 	* Initialize the movie pipeline with the specified settings. This kicks off the rendering process. 
 	
-	* @param InInitSettings - Settings related to the context the pipeline is being executed in. Not configuration for the Pipeline.
-	* @param InConfig		- What settings should the Pipeline use to render with? This should be filled out completely
-	*						  before starting the pipeline. Any overrides based on command line or other should be done
-	*						  before the pipeline is started.
+	* @param InJob	- This contains settings and sequence to render this Movie Pipeline with.
 	*/
-	void Initialize(UMovieRenderPipelineConfig* InConfig);
-
-	void LoadDebugWidget();
+	void Initialize(const FMoviePipelineExecutorJob& InJob);
 
 	/** 
 	* Call to shut down the pipeline. This flushes any outstanding file writes and unregisters all delegates.
@@ -61,113 +59,147 @@ public:
 		return OnMoviePipelineFinishedDelegate;
 	}
 
-	UFUNCTION(BlueprintPure, Category = "Movie Render Pipeline")
-	FMoviePipelineShotCache GetCurrentShotSnapshot() const;
-
-	UFUNCTION(BlueprintPure, Category = "Movie Render Pipeline")
-	FMoviePipelineShotCutCache GetCurrentCameraCutSnapshot() const;
+	/**
+	* Called when there was an error during the rendering of this movie pipeline (such as missing sequence, i/o failure, etc.)
+	*/
+	FMoviePipelineErrored& OnMoviePipelineErrored()
+	{
+		return OnMoviePipelineErroredDelegate;
+	}
 
+	/**
+	* Get the Master Configuration used to render this shot. This contains the global settings for the shot, as well as per-shot
+	* configurations which can contain their own settings.
+	*/
 	UFUNCTION(BlueprintPure, Category = "Movie Render Pipeline")
-	FMoviePipelineFrameOutputState GetOutputStateSnapshot() const;
+	UMoviePipelineMasterConfig* GetPipelineMasterConfig() const { return CurrentJob.Configuration; }
+public:
 
 	UFUNCTION(BlueprintPure, Category = "Movie Render Pipeline")
 	ULevelSequence* GetTargetSequence() const { return TargetSequence; }
 
 	UFUNCTION(BlueprintPure, Category = "Movie Render Pipeline")
-	UMovieRenderPipelineConfig* GetPipelineConfig() const { return Config; }
+	const TArray<FMoviePipelineShotInfo>& GetShotList() const { return ShotList; }
 
 	UFUNCTION(BlueprintPure, Category = "Movie Render Pipeline")
-	FFrameNumber GetTotalOutputFrameCountEstimate() const;
+	int32 GetCurrentShotIndex() const { return CurrentShotIndex; }
 
 	/** 
-	* Returns an estimate based on the average time taken to render all previous frames.
-	* Will be incorrect if PlayRate tracks are in use. Will be inaccurate when different
-	* shots take significantly different amounts of time to render.
-	*
-	* @param OutTimespan: The returned estimated time left. Will be default initialized if there is no estimate.
-	* @return True if we can make a reasonable estimate, false otherwise (ie: no rendering has been done to estimate).
+	* Returns the frame rate override from the Master Configuration (if any) or the Sequence frame rate if no override is specified.
+	* This should be treated as the actual output framerate of the overall Pipeline.
 	*/
 	UFUNCTION(BlueprintPure, Category = "Movie Render Pipeline")
-	bool GetRemainingTimeEstimate(FTimespan& OutTimespan) const;
-
-	const TArray<FMoviePipelineShotCache>& GetShotList() const { return ShotList; }
-	int32 GetCurrentShotIndex() const { return CurrentShotIndex; }
-
-	TSharedPtr<FImagePixelPipe, ESPMode::ThreadSafe> GetOutputPipe() const { return OutputPipe; }
+	FFrameRate GetEffectiveFrameRate() const;
 private:
 
+	/** Instantiate our Debug UI Widget and initialize it to ourself. */
+	void LoadDebugWidget();
+	
+	/** Called before the Engine ticks for the given frame. We use this to calculate delta times that the frame should use. */
 	void OnEngineTickBeginFrame();
+	
+	/** Called after the Engine has ticked for a given frame. Everything in the world has been updated by now so we can submit things to render. */
+	void OnEngineTickEndFrame();
+
+
+	/** Runs the per-tick logic when doing the ProducingFrames state. */
 	void TickProducingFrames();
+	
+	/** Called once when first moving to the Finalize state. */
+	void BeginFinalize();
 
-	/** Return true if we should early out of the Tick function */
+	/** 
+	* Runs the per-tick logic when doing the Finalize state.
+	* @param bInForceFinish		If true, this function will not return until all Output Containers say they have finalized.
+	*/
+	void TickFinalizeOutputContainers(const bool bInForceFinish);
+	/** 
+	* Runs the per-tick logic when doing the Export state. This is spread over multiple ticks to allow non-blocking background 
+	* processes (such as extra encoding) 
+	*/
+	void TickPostFinalizeExport();
+
+	/** Return true if we should early out of the TickProducingFrames function. Decrements the remaining number of steps when false. */
 	bool DebugFrameStepPreTick();
 	/** Returns true if we are idling because of debug frame stepping. */
 	bool IsDebugFrameStepIdling() const;
 
+	/** Debugging/Information. Don't hinge any logic on this as it will get called multiple times per frame in some cases. */
 	void OnSequenceEvaluated(const UMovieSceneSequencePlayer& Player, FFrameTime CurrentTime, FFrameTime PreviousTime);
 
-	void OnEngineTickEndFrame();
-	void TickFinalizeOutputContainers(const bool bInForceFinish);
-	void PostFinalizeExport();
-	void TickPostFinalizeExport();
-	void InitializeShot(FMoviePipelineShotCache& InShot);
-	void SetupRenderingPipelineForShot(FMoviePipelineShotCache& InShot);
-	void TeardownRenderingPipelineForShot(FMoviePipelineShotCache& InShot);
+	/** Set up per-shot state for the specific shot, tearing down old state (if it exists) */
+	void InitializeShot(FMoviePipelineShotInfo& InShot);
+
+	/** Initialize the rendering pipeline for the given shot. This should not get called if rendering work is still in progress for a previous shot. */
+	void SetupRenderingPipelineForShot(FMoviePipelineShotInfo& InShot);
+	/** Deinitialize the rendering pipeline for the given shot. */
+	void TeardownRenderingPipelineForShot(FMoviePipelineShotInfo& InShot);
+
+	/** Flush any async resources in the engine that need to be finalized before submitting anything to the GPU, ie: Streaming Levels and Shaders */
+	void FlushAsyncEngineSystems();
 
-	void BeginFinalize();
-	void FlushAsyncSystems();
 	/** 
 	* Renders the next frame in the Pipeline. This updates/ticks all scene view render states
 	* and produces data. This may not result in an output frame due to multiple renders 
 	* accumulating together to produce an output. frame.
-	* Should not be called if we're idling (debug) or not initialized yet. 
+	* Should not be called if we're idling (debug), not initialized yet, or finalizing/exporting. 
 	*/
 	void RenderFrame();
+
+	/** Allow any Settings to modify the (already duplicated) sequence. This allows inserting automatic pre-roll, etc. */
+	void ModifySequenceViaExtensions(ULevelSequence* InSequence);
+
+	/** Calculate the expected amount of total work that this Movie Pipeline is expected to do. */
+	FMoviePipelineWorkInfo CalculateExpectedOutputMetrics();
 private:
-	ULevelSequence* CreateCopyOfSequence(ULevelSequence* InSequence);
+	/** Initialize a new Level Sequence Actor to evaluate our target sequence. Disables any existing Level Sequences pointed at our original sequence. */
 	void InitializeLevelSequenceActor(ULevelSequence* OriginalLevelSequence, ULevelSequence* InSequenceToApply);
 
-	/** Modifies the TargetSequence to ensure that only the specified Shot has it's associated Cinematic Shot Section enabled. */
-	void SetSoloShot(const FMoviePipelineShotCache& InShot);
-	void SetSoloCameraCut(FMoviePipelineShotCutCache& InCameraCut);
-
 	/** This converts the sequence into a Shot List and expands bounds.*/
-	TArray<FMoviePipelineShotCache> BuildShotListFromSequence(const ULevelSequence* InSequence);
+	TArray<FMoviePipelineShotInfo> BuildShotListFromSequence(const ULevelSequence* InSequence);
 
-	/* Tick Resolution Frames*/
-	FFrameNumber GetMotionBlurDuration(int32 InShotIndex) const;
-	/* Tick Resolution Frames*/
-	FFrameNumber CalculateHandleFrameDuration(const FMoviePipelineShotCache& InShot) const;
-	
-	MoviePipeline::FFrameConstantMetrics CalculateShotFrameMetrics(const FMoviePipelineShotCache& InShot) const;
+	/** 
+	* Modifies the TargetSequence to ensure that only the specified Shot has it's associated Cinematic Shot Section enabled.
+	* This way when Handle Frames are enabled and the sections are expanded, we don't end up evaluating the previous shot. 
+	*/
+	void SetSoloShot(const FMoviePipelineShotInfo& InShot);
+
+	/* Expands the specified shot (and contained camera cuts)'s ranges for the given settings. */
+	void ExpandShot(FMoviePipelineShotInfo& InShot);
 
-	void CalculateFrameNumbersForOutputState(const MoviePipeline::FFrameConstantMetrics& InFrameMetrics, const FMoviePipelineShotCutCache& InCameraCut, FMoviePipelineFrameOutputState& InOutOutputState) const;
+	/** Calculates lots of useful numbers used in timing based off of the current shot. These are constant for a given shot. */
+	MoviePipeline::FFrameConstantMetrics CalculateShotFrameMetrics(const FMoviePipelineShotInfo& InShot) const;
 
-	/* Expands the specified shot by the specified offsets. Offsets should be in TickResolution */
-	void ExpandShot(FMoviePipelineShotCache& InShot);
+	/** It can be useful to know where the data we're generating was relative to the original Timeline, so this calculates that. */
+	void CalculateFrameNumbersForOutputState(const MoviePipeline::FFrameConstantMetrics& InFrameMetrics, const FMoviePipelineCameraCutInfo& InCameraCut, FMoviePipelineFrameOutputState& InOutOutputState) const;
 
 private:
+	/** Custom TimeStep used to drive the engine while rendering. */
 	UPROPERTY(Transient, Instanced)
 	UMoviePipelineCustomTimeStep* CustomTimeStep;
 
+	/** Custom Time Controller for the Sequence Player, used to match Custom TimeStep without any floating point accumulation errors. */
 	TSharedPtr<FMoviePipelineTimeController> CustomSequenceTimeController;
 
 	/** Hold a reference to the existing custom time step (if any) so we can restore it after we're done using our custom one. */
 	UPROPERTY(Transient)
-	UEngineCustomTimeStep* CachedCustomTimeStep;
+	UEngineCustomTimeStep* CachedPrevCustomTimeStep;
 
-	UPROPERTY(Transient, BlueprintReadOnly, meta=(AllowPrivateAccess=true), Category = "Movie Render Pipeline")
+	/** This is our duplicated sequence that we're rendering. This will get modified throughout the rendering process. */
+	UPROPERTY(Transient)
 	ULevelSequence* TargetSequence;
 
+	/** The Level Sequence Actor we spawned to play our TargetSequence. */
 	UPROPERTY(Transient)
 	ALevelSequenceActor* LevelSequenceActor;
 
+	/** The Debug UI Widget that is spawned and placed on the player UI */
 	UPROPERTY(Transient)
 	UMovieRenderDebugWidget* DebugWidget;
 
 	/** A list of all of the shots we are going to render out from this sequence. */
 	UPROPERTY(Transient, BlueprintReadOnly, meta = (AllowPrivateAccess = true), Category = "Movie Render Pipeline")
-	TArray<FMoviePipelineShotCache> ShotList;
+	TArray<FMoviePipelineShotInfo> ShotList;
 
 	/** What state of the overall flow are we in? See enum for specifics. */
 	UPROPERTY(Transient, BlueprintReadOnly, meta = (AllowPrivateAccess = true), Category = "Movie Render Pipeline")
@@ -177,19 +209,24 @@ private:
 	UPROPERTY(Transient, BlueprintReadOnly, meta = (AllowPrivateAccess = true), Category = "Movie Render Pipeline")
 	int32 CurrentShotIndex;
 
-	UPROPERTY(Transient, BlueprintReadOnly, meta = (AllowPrivateAccess = true), Category = "Movie Render Pipeline")
-	FMoviePipelineFrameOutputState CachedOutputState;
-
 	/** The time (in UTC) that Initialize was called. Used to track elapsed time. */
 	UPROPERTY(Transient, BlueprintReadOnly, meta = (AllowPrivateAccess = true), Category = "Movie Render Pipeline")
 	FDateTime InitializationTime;
 
+	UPROPERTY(Transient, BlueprintReadOnly, meta = (AllowPrivateAccess = true), Category = "Movie Render Pipeline")
+	FMoviePipelineFrameOutputState CachedOutputState;
+
+	/** 
+	* Cache the overall expected work for the entire Movie Pipeline. Used to compare current progress against total for progress.
+	* Current Progress is tracked by individual shots.
+	*/
+	FMoviePipelineWorkInfo TotalExpectedWork;
+
 	/** 
 	* Have we hit the callback for the BeginFrame at least once? This solves an issue where the delegates
 	* get registered mid-frame so you end up calling EndFrame before BeginFrame which is undesirable.
 	*/
 	bool bHasRunBeginFrameOnce;
-
 	/** Should we pause the game at the end of the frame? Used to implement frame step debugger. */
 	bool bPauseAtEndOfFrame;
 
@@ -199,20 +236,21 @@ private:
 	/** Called when we have completely finished. This object will call Shutdown before this and stop ticking. */
 	FMoviePipelineFinished OnMoviePipelineFinishedDelegate;
 
+	/** Called when there is a warning/error that the user should pay attention to.*/
+	FMoviePipelineErrored OnMoviePipelineErroredDelegate;
+
+public:
 	/** This gathers all of the produced data for an output frame (which may come in async many frames later) before passing them onto the Output Containers. */
 	TSharedPtr<FMoviePipelineOutputMerger, ESPMode::ThreadSafe> OutputBuilder;
 
-	// Debug Rendering
-	TSharedPtr<FMovieRenderViewport> DummyViewport;
-	TSharedPtr<FMovieRenderViewportClient> ViewportClient;
-
-
+	// debug
+	TSharedPtr<FImagePixelPipe, ESPMode::ThreadSafe> GetOutputPipe() const { return OutputPipe; }
 	TSharedPtr<FImagePixelPipe, ESPMode::ThreadSafe> OutputPipe;
 
 private:
-	/** The overall configuration for how this Pipeline executes. Should be immutable. */
+	/** Keep track of which job we're working on. This holds our Configuration + which shots we're supposed to render from it. */
 	UPROPERTY(Transient)
-	UMovieRenderPipelineConfig* Config;
+	FMoviePipelineExecutorJob CurrentJob;
 };
 
 UCLASS()
diff --git a/Engine/Plugins/MovieScene/MovieRenderPipeline/Source/MovieRenderPipelineCore/Public/MoviePipelineAccumulationSetting.h b/Engine/Plugins/MovieScene/MovieRenderPipeline/Source/MovieRenderPipelineCore/Public/MoviePipelineAccumulationSetting.h
--- a/Engine/Plugins/MovieScene/MovieRenderPipeline/Source/MovieRenderPipelineCore/Public/MoviePipelineAccumulationSetting.h
+++ b/Engine/Plugins/MovieScene/MovieRenderPipeline/Source/MovieRenderPipelineCore/Public/MoviePipelineAccumulationSetting.h
@@ -73,6 +73,12 @@ public:
 	UPROPERTY(EditAnywhere, BlueprintReadWrite, meta = (UIMin = 1, ClampMin = 1), Category = "Movie Pipeline")
 	int32 SpatialSampleCount;
 
+	UPROPERTY(EditAnywhere, BlueprintReadWrite, meta = (UIMin = 1, ClampMin = 1), Category = "Movie Pipeline")
+	int32 HandleFrameCount;
+
+	UPROPERTY(EditAnywhere, BlueprintReadWrite, meta = (UIMin = 1, ClampMin = 1), Category = "Movie Pipeline")
+	int32 WarmUpFrameCount;
+
 	/**
 	* 
 	*/
diff --git a/Engine/Plugins/MovieScene/MovieRenderPipeline/Source/MovieRenderPipelineCore/Public/MoviePipelineAudioOutput.h b/Engine/Plugins/MovieScene/MovieRenderPipeline/Source/MovieRenderPipelineCore/Public/MoviePipelineAudioOutput.h
new file mode 100644
--- /dev/null
+++ b/Engine/Plugins/MovieScene/MovieRenderPipeline/Source/MovieRenderPipelineCore/Public/MoviePipelineAudioOutput.h
@@ -0,0 +1,34 @@
+// Copyright 1998-2019 Epic Games, Inc. All Rights Reserved.
+#pragma once
+
+#include "MoviePipelineOutputBase.h"
+#include "MoviePipelineAudioOutput.generated.h"
+
+UCLASS(BlueprintType, meta=(MovieRenderPipelineDisplayName="Audio"))
+class MOVIERENDERPIPELINECORE_API UMoviePipelineAudioOutput : public UMoviePipelineOutputBase
+{
+	GENERATED_BODY()
+public:
+	UMoviePipelineAudioOutput()
+		: bExportWavFile(true)
+		, CachedExpectedDuration(0.f)
+		, PrevRenderEveryTickValue(false)
+	{
+	}
+
+	// virtual FText GetDisplayText() const override { return NSLOCTEXT("MovieRenderPipeline", "AudioSettingDisplayName", "Audio"); }
+	virtual void SetupForPipelineImpl(UMoviePipeline* InPipeline) override;
+	virtual void OnPipelineFinishedImpl() override;
+	virtual void OnPostTickImpl() override;
+	virtual void OnShotFinishedImpl(const FMoviePipelineShotInfo& Shot) override;
+public:
+	/** If true audio will be written to a separate wave file in addition to being embedded into any supported containers */
+	UPROPERTY(EditAnywhere, BlueprintReadWrite, Category = "Movie Render Pipeline")
+	bool bExportWavFile;
+
+private:
+	float CachedExpectedDuration;
+	int32 PrevRenderEveryTickValue;
+	FString PrevAudioDevicePlatform;
+	float PrevUnfocusedAudioMultiplier;
+};
\ No newline at end of file
diff --git a/Engine/Plugins/MovieScene/MovieRenderPipeline/Source/MovieRenderPipelineCore/Public/MoviePipelineBlueprintLibrary.h b/Engine/Plugins/MovieScene/MovieRenderPipeline/Source/MovieRenderPipelineCore/Public/MoviePipelineBlueprintLibrary.h
--- a/Engine/Plugins/MovieScene/MovieRenderPipeline/Source/MovieRenderPipelineCore/Public/MoviePipelineBlueprintLibrary.h
+++ b/Engine/Plugins/MovieScene/MovieRenderPipeline/Source/MovieRenderPipelineCore/Public/MoviePipelineBlueprintLibrary.h
@@ -9,6 +9,7 @@
 
 // Forward Declare
 class UMoviePipeline;
+class UMovieSceneSequence;
 
 UCLASS()
 class MOVIERENDERPIPELINECORE_API UMoviePipelineBlueprintLibrary : public UBlueprintFunctionLibrary
@@ -28,25 +29,60 @@ public:
 	* This will be inaccurate when PlayRate tracks are involved.
 	*/
 	UFUNCTION(BlueprintPure, Category = "Movie Render Pipeline")
-	static int32 GetOutputFrameCountEstimate(const FMoviePipelineShotCutCache& InCameraCut);
+	static int32 GetOutputFrameCountEstimate(const FMoviePipelineCameraCutInfo& InCameraCut);
 	
 	/**
 	* Returns the expected number of frames different frames that will be submitted for rendering. This is
 	* the number of output frames * temporal samples. This will be inaccurate when PlayRate tracks are involved.
 	*/
 	UFUNCTION(BlueprintPure, Category = "Movie Render Pipeline")
-	static int32 GetTemporalFrameCountEstimate(const FMoviePipelineShotCutCache& InCameraCut);
+	static int32 GetTemporalFrameCountEstimate(const FMoviePipelineCameraCutInfo& InCameraCut);
 	
 	/** 
 	* Returns misc. utility frame counts (Warm Up + MotionBlur Fix) since these are outside the ranges.
 	*/
 	UFUNCTION(BlueprintPure, Category = "Movie Render Pipeline")
-	static int32 GetUtilityFrameCountEstimate(const FMoviePipelineShotCutCache& InCameraCut);
+	static int32 GetUtilityFrameCountEstimate(const FMoviePipelineCameraCutInfo& InCameraCut);
 	
 	/**
 	* Returns the number of samples submitted to the GPU, optionally counting samples for the various parts.
 	* This will be inaccurate when PlayRate tracks are involved.
 	*/
 	UFUNCTION(BlueprintPure, Category = "Movie Render Pipeline")
-	static int32 GetSampleCountEstimate(const FMoviePipelineShotCutCache& InCameraCut, const bool bIncludeWarmup = true, const bool bIncludeMotionBlur = true);
+	static int32 GetSampleCountEstimate(const FMoviePipelineCameraCutInfo& InCameraCut, const bool bIncludeWarmup = true, const bool bIncludeMotionBlur = true);
+	
+	/**
+	* Duplicates the specified sequence using a medium depth copy. Standard duplication will only duplicate
+	* the top level Sequence (since shots and sub-sequences are other standalone assets) so this function
+	* recursively duplicates the given sequence, shot and subsequence and then fixes up the references to
+	* point to newly duplicated sequences.
+	*
+	* This does not duplicate any assets that the sequence points to outside of Shots/Subsequences.
+	*
+	* @param	Outer		- The Outer of the newly duplicated object. Leave null for TransientPackage();
+	* @param	InSequence	- The sequence to recursively duplicate.
+	* @return				- The duplicated sequence, or null if no sequence was provided to duplicate.
+	*/
+	UFUNCTION(BlueprintCallable, Category = "Movie Render Pipeline")
+	static UMovieSceneSequence* DuplicateSequence(UObject* Outer, UMovieSceneSequence* InSequence);
+
+	// UFUNCTION(BlueprintPure, Category = "Movie Render Pipeline")
+	// 	FMoviePipelineShotInfo GetCurrentShotSnapshot() const;
+	// 
+	// UFUNCTION(BlueprintPure, Category = "Movie Render Pipeline")
+	// 	FMoviePipelineCameraCutInfo GetCurrentCameraCutSnapshot() const;
+	// 
+	// UFUNCTION(BlueprintPure, Category = "Movie Render Pipeline")
+	// 	FMoviePipelineFrameOutputState GetOutputStateSnapshot() const;\
+
+	/**
+	* Returns an estimate based on the average time taken to render all previous frames.
+	* Will be incorrect if PlayRate tracks are in use. Will be inaccurate when different
+	* shots take significantly different amounts of time to render.
+	*
+	* @param OutTimespan: The returned estimated time left. Will be default initialized if there is no estimate.
+	* @return True if we can make a reasonable estimate, false otherwise (ie: no rendering has been done to estimate).
+	*/
+	// UFUNCTION(BlueprintPure, Category = "Movie Render Pipeline")
+	// bool GetRemainingTimeEstimate(FTimespan& OutTimespan) const;
 };
\ No newline at end of file
diff --git a/Engine/Plugins/MovieScene/MovieRenderPipeline/Source/MovieRenderPipelineCore/Public/MoviePipelineConfigBase.h b/Engine/Plugins/MovieScene/MovieRenderPipeline/Source/MovieRenderPipelineCore/Public/MoviePipelineConfigBase.h
new file mode 100644
--- /dev/null
+++ b/Engine/Plugins/MovieScene/MovieRenderPipeline/Source/MovieRenderPipelineCore/Public/MoviePipelineConfigBase.h
@@ -0,0 +1,105 @@
+// Copyright 1998-2019 Epic Games, Inc. All Rights Reserved.
+#pragma once
+
+#include "UObject/Object.h"
+#include "Templates/SubclassOf.h"
+#include "MoviePipelineSetting.h"
+#include "MoviePipelineConfigBase.generated.h"
+
+
+UCLASS(Blueprintable)
+class MOVIERENDERPIPELINECORE_API UMoviePipelineConfigBase : public UObject
+{
+	GENERATED_BODY()
+	
+public:
+	UMoviePipelineConfigBase()
+	{
+		SettingsSerialNumber = -1;
+	}
+
+public:
+	/**
+	* Removes the specific instance from our Setting list. 
+	*/
+	UFUNCTION(BlueprintCallable, Category = "Movie Render Pipeline")
+	void RemoveSetting(UMoviePipelineSetting* InSetting);
+
+	int32 GetSettingsSerialNumber() const { return SettingsSerialNumber; }
+	
+	/**
+	* Returns an array of all settings in this config.
+	*/
+	const TArray<UMoviePipelineSetting*> GetSettings() const { return Settings; }
+
+public:
+	/**
+	* Find a setting of a particular type for this config.
+	* @param InClass - Class that you wish to find the setting object for.
+	* @return An instance of this class if it already exists as a setting on this config, otherwise null.
+	*/
+	UFUNCTION(BlueprintPure, Category = "Movie Render Pipeline")
+	UMoviePipelineSetting* FindSettingByClass(TSubclassOf<UMoviePipelineSetting> InClass) const
+	{
+		UMoviePipelineSetting* const* Found = Settings.FindByPredicate([InClass](UMoviePipelineSetting* In) { return In && In->GetClass() == InClass; });
+		return Found ? CastChecked<UMoviePipelineSetting>(*Found) : nullptr;
+	}
+
+	/** 
+	* Finds a setting of a particular type for this pipeline config, adding it if it doesn't already exist.
+	* @param InClass - Class you wish to find or create the setting object for.
+	* @return An instance of this class as a setting on this config.
+	*/
+	UFUNCTION(BlueprintCallable, Category = "Movie Render Pipeline")
+	UMoviePipelineSetting* FindOrAddSettingByClass(TSubclassOf<UMoviePipelineSetting> InClass)
+	{
+		UMoviePipelineSetting* Found = FindSettingByClass(InClass);
+		if (!Found)
+		{
+			Modify();
+
+			Found = NewObject<UMoviePipelineSetting>(this, InClass);
+			Settings.Add(Found);
+			++SettingsSerialNumber;
+		}
+
+		return Found;
+	}
+
+	/**
+	* Find a setting of a particular type for this config
+	*/
+	template<typename SettingType>
+	SettingType* FindSetting() const
+	{
+		UClass* PredicateClass = SettingType::StaticClass();
+		UMoviePipelineSetting* const* Found = Settings.FindByPredicate([PredicateClass](UMoviePipelineSetting* In) { return In && In->GetClass() == PredicateClass; });
+		return Found ? CastChecked<SettingType>(*Found) : nullptr;
+	}
+
+	/**
+	 * Find a setting of a particular type for this config instance, adding one if it was not found.
+	 */
+	template<typename SettingType>
+	SettingType* FindOrAddSetting()
+	{
+		SettingType* Found = FindSetting<SettingType>();
+		if (!Found)
+		{
+			Modify();
+
+			Found = NewObject<SettingType>(this);
+			Settings.Add(Found);
+			++SettingsSerialNumber;
+		}
+		return Found;
+	}
+
+protected:
+	/** Array of settings classes that affect various parts of the output pipeline. */
+	UPROPERTY(VisibleAnywhere, Category = "Movie Pipeline")
+	TArray<UMoviePipelineSetting*> Settings;
+
+private:
+	int32 SettingsSerialNumber;
+};
\ No newline at end of file
diff --git a/Engine/Plugins/MovieScene/MovieRenderPipeline/Source/MovieRenderPipelineCore/Public/MoviePipelineExecutor.h b/Engine/Plugins/MovieScene/MovieRenderPipeline/Source/MovieRenderPipelineCore/Public/MoviePipelineExecutor.h
--- a/Engine/Plugins/MovieScene/MovieRenderPipeline/Source/MovieRenderPipelineCore/Public/MoviePipelineExecutor.h
+++ b/Engine/Plugins/MovieScene/MovieRenderPipeline/Source/MovieRenderPipelineCore/Public/MoviePipelineExecutor.h
@@ -3,14 +3,20 @@
 
 #include "UObject/Object.h"
 #include "Templates/SubclassOf.h"
+#include "MovieRenderPipelineDataTypes.h"
 #include "MoviePipelineExecutor.generated.h"
 
-class UMovieRenderPipelineConfig;
+class UMoviePipelineMasterConfig;
 class UMoviePipelineExecutorBase;
 class UMoviePipeline;
 
-DECLARE_MULTICAST_DELEGATE_OneParam(FOnMoviePipelineExecutorFinishedNative, UMoviePipelineExecutorBase*);
-DECLARE_DYNAMIC_MULTICAST_DELEGATE_OneParam(FOnMoviePipelineExecutorFinished, UMoviePipelineExecutorBase*, PipelineExecutor);
+DECLARE_MULTICAST_DELEGATE_TwoParams(FOnMoviePipelineExecutorFinishedNative, UMoviePipelineExecutorBase* /*PipelineExecutor*/, bool /*bSuccess*/);
+DECLARE_DYNAMIC_MULTICAST_DELEGATE_TwoParams(FOnMoviePipelineExecutorFinished, UMoviePipelineExecutorBase*, PipelineExecutor, bool, bSuccess);
+
+DECLARE_MULTICAST_DELEGATE_FourParams(FOnMoviePipelineExecutorErroredNative, UMoviePipelineExecutorBase* /*PipelineExecutor*/, UMoviePipeline* /*PipelineWithError*/, bool /*bIsFatal*/, FText /*ErrorText*/);
+DECLARE_DYNAMIC_MULTICAST_DELEGATE_FourParams(FOnMoviePipelineExecutorErrored, UMoviePipelineExecutorBase*, PipelineExecutor, UMoviePipeline*, PipelineWithError, bool, bIsFatal, FText, ErrorText);
+
+
 
 /**
 * A Movie Pipeline Executor is responsible for executing an array of Movie Pipelines,
@@ -29,12 +35,13 @@ class MOVIERENDERPIPELINECORE_API UMoviePipelineExecutorBase : public UObject
 	GENERATED_BODY()
 public:
 	UMoviePipelineExecutorBase()
-		: TargetPipelineClass(nullptr)
+		: bAnyJobHadFatalError(false)
+		, TargetPipelineClass(nullptr)
 	{
 	}
 
 	UFUNCTION(BlueprintCallable, Category = "Movie Render Pipeline")
-	virtual void Execute(UPARAM(Ref) TArray<UMovieRenderPipelineConfig*>& InPipelines)
+	virtual void Execute(const TArray<FMoviePipelineExecutorJob>& InPipelines)
 	{
 		ExecuteImpl(InPipelines);
 	}
@@ -44,6 +51,11 @@ public:
 		return OnExecutorFinishedDelegateNative;
 	}
 
+	FOnMoviePipelineExecutorErroredNative& OnExecutorErrored()
+	{
+		return OnExecutorErroredDelegateNative;
+	}
+
 	void SetMoviePipelineClass(UClass* InPipelineClass)
 	{
 		TargetPipelineClass = InPipelineClass;
@@ -52,24 +64,63 @@ protected:
 	/** 
 	* This should be called when the Executor has finished executing all of the things
 	* it has been asked to execute. This should be called in the event of a failure as 
-	* well, and passing in false for success to allow the caller to know failure.
+	* well, and passing in false for success to allow the caller to know failure. Errors
+	* should be broadcast on the error delegate, so this is just a handy way to know at
+	* the end without having to track it yourself.
+	*
+	* @param bInSuccess	True if the pipeline successfully executed all jobs. False if there was an error. 
 	*/
-	virtual void OnExecutorFinishedImpl(const bool bInSuccess)
+	virtual void OnExecutorFinishedImpl()
+	{
+		// Broadcast to both Native and Python/BP
+		OnExecutorFinishedDelegateNative.Broadcast(this, !bAnyJobHadFatalError);
+		OnExecutorFinishedDelegate.Broadcast(this, !bAnyJobHadFatalError);
+	}
+
+	virtual void OnExecutorErroredImpl(UMoviePipeline* ErroredPipeline, bool bFatal, FText ErrorReason)
 	{
+		if (bFatal)
+		{
+			bAnyJobHadFatalError = true;
+		}
+
 		// Broadcast to both Native and Python/BP
-		OnExecutorFinishedDelegateNative.Broadcast(this);
-		OnExecutorFinishedDelegate.Broadcast(this);
+		OnExecutorFinishedDelegateNative.Broadcast(this, bFatal);
+		OnExecutorFinishedDelegate.Broadcast(this, bFatal);
 	}
-	virtual void ExecuteImpl(TArray<UMovieRenderPipelineConfig*>& InPipelines) PURE_VIRTUAL(UMoviePipelineExecutorBase::ExecuteImpl, );
+
+	virtual void ExecuteImpl(const TArray<FMoviePipelineExecutorJob>& InPipelines) PURE_VIRTUAL(UMoviePipelineExecutorBase::ExecuteImpl, );
 
 private:
-	/** Exposed for Blueprints/Python. Called at the same time as the native one. */
+	/** 
+	* Called when the Executor has finished all jobs. Reports success if no jobs
+	* had fatal errors. Subscribe to the error delegate for more information about
+	* any errors.
+	*
+	* Exposed for Blueprints/Python. Called at the same time as the native one.
+	*/
 	UPROPERTY(BlueprintAssignable, Category = "Movie Render Pipeline")
 	FOnMoviePipelineExecutorFinished OnExecutorFinishedDelegate;
 
 	/** For native C++ code. Called at the same time as the Blueprint/Python one. */
 	FOnMoviePipelineExecutorFinishedNative OnExecutorFinishedDelegateNative;
+
+	/**
+	* Called when an individual job reports a warning/error. Jobs are considered fatal
+	* if the severity was bad enough to abort the job (missing sequence, write failure, etc.)
+	*
+	* Exposed for Blueprints/Python. Called at the same time as the native one.
+	*/
+	UPROPERTY(BlueprintAssignable, Category = "Movie Render Pipeline")
+	FOnMoviePipelineExecutorErrored OnExecutorErroredDelegate;
+
+	/** For native C++ code. Called at the same time as the Blueprint/Python one. */
+	FOnMoviePipelineExecutorErroredNative OnExecutorErroredDelegateNative;
+
+	/** Set automatically when the error delegate gets broadcast (if fatal). */
+	bool bAnyJobHadFatalError;
 protected:
+	/** Which Pipeline Class should be created by this Executor. May be null. */
 	UPROPERTY()
 	TSubclassOf<UMoviePipeline> TargetPipelineClass;
 };
\ No newline at end of file
diff --git a/Engine/Plugins/MovieScene/MovieRenderPipeline/Source/MovieRenderPipelineCore/Public/MoviePipelineInProcessExecutor.h b/Engine/Plugins/MovieScene/MovieRenderPipeline/Source/MovieRenderPipelineCore/Public/MoviePipelineInProcessExecutor.h
--- a/Engine/Plugins/MovieScene/MovieRenderPipeline/Source/MovieRenderPipelineCore/Public/MoviePipelineInProcessExecutor.h
+++ b/Engine/Plugins/MovieScene/MovieRenderPipeline/Source/MovieRenderPipelineCore/Public/MoviePipelineInProcessExecutor.h
@@ -23,7 +23,7 @@ public:
 	}
 
 protected:
-	virtual void Start(UMovieRenderPipelineConfig* InConfig, const int32 InConfigIndex, const int32 InNumConfigs) override;
+	virtual void Start(const FMoviePipelineExecutorJob& InJob) override;
 
 private:
 	void OnMapLoadFinished(UWorld* NewWorld);
diff --git a/Engine/Plugins/MovieScene/MovieRenderPipeline/Source/MovieRenderPipelineCore/Public/MoviePipelineLinearExecutor.h b/Engine/Plugins/MovieScene/MovieRenderPipeline/Source/MovieRenderPipelineCore/Public/MoviePipelineLinearExecutor.h
--- a/Engine/Plugins/MovieScene/MovieRenderPipeline/Source/MovieRenderPipelineCore/Public/MoviePipelineLinearExecutor.h
+++ b/Engine/Plugins/MovieScene/MovieRenderPipeline/Source/MovieRenderPipelineCore/Public/MoviePipelineLinearExecutor.h
@@ -2,6 +2,7 @@
 #pragma once
 
 #include "MoviePipelineExecutor.h"
+#include "MovieRenderPipelineDataTypes.h"
 #include "MoviePipelineLinearExecutor.generated.h"
 
 class UMoviePipeline;
@@ -26,30 +27,27 @@ public:
 	}
 
 protected:
-	virtual void ExecuteImpl(TArray<UMovieRenderPipelineConfig*>& InPipelines) override;
+	virtual void ExecuteImpl(const TArray<FMoviePipelineExecutorJob>& InPipelines) override;
 
 	virtual void StartPipelineByIndex(int32 InPipelineIndex);
-	virtual void Start(UMovieRenderPipelineConfig* InConfig, const int32 InConfigIndex, const int32 InNumConfigs) { check(InConfig); }
+	virtual void Start(const FMoviePipelineExecutorJob& InJob) {}
 
 public:
 	virtual void OnIndividualPipelineFinished(UMoviePipeline* /* FinishedPipeline */);
-	virtual void OnExecutorFinishedImpl(const bool bInSuccess);
+	virtual void OnExecutorFinishedImpl();
+	virtual void OnPipelineErrored(UMoviePipeline* InPipeline, bool bIsFatal, FText ErrorText);
 protected:
 	
 	/** List of Pipeline Configs we've been asked to execute. */
 	UPROPERTY(Transient)
-	TArray<UMovieRenderPipelineConfig*> Pipelines;
+	TArray<FMoviePipelineExecutorJob> ExecutorJobs;
 
-	/** Which Pipeline Config are we currently working on. */
-	int32 CurrentPipelineIndex;
-
-	/** Instance of the Pipeline that exists in the world that is currently processing (if any) */
+	/** A Movie Pipeline that has been spawned and is running (if any) */
 	UPROPERTY(Transient)
 	UMoviePipeline* ActiveMoviePipeline;
 
-	/** Instance of the Config we're supposed to be working on */
-	UPROPERTY(Transient)
-	UMovieRenderPipelineConfig* ActiveConfig;
+	/** Which Pipeline Config are we currently working on. */
+	int32 CurrentPipelineIndex;
 
 private:
 	/** Used to track total processing duration. */
diff --git a/Engine/Plugins/MovieScene/MovieRenderPipeline/Source/MovieRenderPipelineCore/Public/MoviePipelineMasterConfig.h b/Engine/Plugins/MovieScene/MovieRenderPipeline/Source/MovieRenderPipelineCore/Public/MoviePipelineMasterConfig.h
new file mode 100644
--- /dev/null
+++ b/Engine/Plugins/MovieScene/MovieRenderPipeline/Source/MovieRenderPipelineCore/Public/MoviePipelineMasterConfig.h
@@ -0,0 +1,50 @@
+// Copyright 1998-2019 Epic Games, Inc. All Rights Reserved.
+#pragma once
+
+#include "MoviePipelineConfigBase.h"
+
+#include "MoviePipelineMasterConfig.generated.h"
+
+
+// Forward Declares
+class ULevelSequence;
+class UMoviePipelineSetting;
+class UMoviePipelineRenderPass;
+class UMoviePipelineOutputBase;
+class UMoviePipelineShotConfig;
+
+/**
+* This class describes the main configuration for a Movie Render Pipeline.
+* Only settings that apply to the entire output should be stored here,
+* anything that is changed on a per-shot basis should be stored inside of 
+* UMovieRenderShotConfig instead.
+*
+* THIS CLASS SHOULD BE IMMUTABLE ONCE PASSED TO THE PIPELINE FOR PROCESSING.
+* (Otherwise you will be modifying the instance that exists in the UI)
+*/
+UCLASS(Blueprintable)
+class MOVIERENDERPIPELINECORE_API UMoviePipelineMasterConfig : public UMoviePipelineConfigBase
+{
+	GENERATED_BODY()
+	
+public:
+	UMoviePipelineMasterConfig();
+
+public:
+	TArray<UMoviePipelineOutputBase*> GetOutputContainers() const;
+
+public:
+
+	/** Returns a pointer to the config specified for the shot, otherwise the default for this pipeline. */
+	UMoviePipelineShotConfig* GetConfigForShot(const FString& ShotName) const;
+
+public:
+	
+	/** The default shot-setup to use for any shot that doesn't a specific implementation. This is required! */
+	UPROPERTY(EditAnywhere, Instanced, BlueprintReadWrite, Category = "Movie Render Pipeline")
+	UMoviePipelineShotConfig* DefaultShotConfig;
+	
+	/** A mapping of Shot Name -> Shot Config to use for rendering specific shots with specific configs. */
+	UPROPERTY(VisibleAnywhere, Instanced, BlueprintReadOnly, Category = "Movie Render Pipeline")
+	TMap<FString, UMoviePipelineShotConfig*> PerShotConfigMapping;
+};
\ No newline at end of file
diff --git a/Engine/Plugins/MovieScene/MovieRenderPipeline/Source/MovieRenderPipelineCore/Public/MoviePipelineOutput.h b/Engine/Plugins/MovieScene/MovieRenderPipeline/Source/MovieRenderPipelineCore/Public/MoviePipelineOutput.h
deleted file mode 100644
diff --git a/Engine/Plugins/MovieScene/MovieRenderPipeline/Source/MovieRenderPipelineCore/Public/MoviePipelineOutputBase.h b/Engine/Plugins/MovieScene/MovieRenderPipeline/Source/MovieRenderPipelineCore/Public/MoviePipelineOutputBase.h
new file mode 100644
--- /dev/null
+++ b/Engine/Plugins/MovieScene/MovieRenderPipeline/Source/MovieRenderPipelineCore/Public/MoviePipelineOutputBase.h
@@ -0,0 +1,69 @@
+// Copyright 1998-2019 Epic Games, Inc. All Rights Reserved.
+#pragma once
+
+#include "UObject/Object.h"
+#include "MovieRenderPipelineDataTypes.h"
+#include "Engine/EngineTypes.h"
+#include "MoviePipelineSetting.h"
+#include "MoviePipelineOutputBase.generated.h"
+
+class UMoviePipeline;
+
+/**
+* A base class for different output containers for the data generated by the Movie Pipeline,
+* i.e: image sequences, video containers, etc.
+*/
+UCLASS(Blueprintable, Abstract)
+class MOVIERENDERPIPELINECORE_API UMoviePipelineOutputBase : public UMoviePipelineSetting
+{
+	GENERATED_BODY()
+public:
+
+	/** 
+	* Called once when all frames have been produced for the pipeline. Use this as an indicator to start flushing to disk. 
+	* Async processes can be started here and HasFinishedProcessing() will be called each frame until all all containers
+	* return true, at which point Finalize will be called.
+	*/
+	void BeginFinalize() { BeginFinalizeImpl(); }
+
+	/**
+	* Called each frame to check to see if this output container has finished processing. Finalize won't be called until
+	* all output containers return true for this. 
+	*/
+	bool HasFinishedProcessing() { return HasFinishedProcessingImpl(); }
+
+	/** 
+	* Called after all output containers have reported that they are done processing.
+	*/
+	void Finalize() { FinalizeImpl(); }
+
+	/**
+	* This is called when a new shot starts, before warm up happens. See OnFrameProductionStart.
+	*/
+	void OnShotInitialized(const TOptional<FMoviePipelineShotInfo> PrevShot, const FMoviePipelineShotInfo& NewShot) { OnShotInitializedImpl(PrevShot, NewShot); }
+	
+	/**
+	* This is called when a shot ends, right after the last frame is rendered.
+	*/
+	void OnShotFinished(const FMoviePipelineShotInfo& Shot) { OnShotFinishedImpl(Shot); }
+
+	/**
+	* This is called during the Shutdown process of the Pipeline. This is after finalization.
+	*/
+	void OnPipelineFinished() { OnPipelineFinishedImpl(); }
+
+	/**
+	* Called at the end of the frame, before rendering has happened for that frame.
+	*/
+	void OnPostTick() { OnPostTickImpl(); }
+
+protected:
+	virtual void OnShotInitializedImpl(const TOptional<FMoviePipelineShotInfo> PrevShot, const FMoviePipelineShotInfo& NewShot) {}
+	virtual void OnShotFinishedImpl(const FMoviePipelineShotInfo& Shot) {}
+
+	virtual void FinalizeImpl() {}
+	virtual void BeginFinalizeImpl() {}
+	virtual bool HasFinishedProcessingImpl() { return true; }
+	virtual void OnPipelineFinishedImpl() {}
+	virtual void OnPostTickImpl() {}
+};
\ No newline at end of file
diff --git a/Engine/Plugins/MovieScene/MovieRenderPipeline/Source/MovieRenderPipelineCore/Public/MoviePipelineOutputSetting.h b/Engine/Plugins/MovieScene/MovieRenderPipeline/Source/MovieRenderPipelineCore/Public/MoviePipelineOutputSetting.h
new file mode 100644
--- /dev/null
+++ b/Engine/Plugins/MovieScene/MovieRenderPipeline/Source/MovieRenderPipelineCore/Public/MoviePipelineOutputSetting.h
@@ -0,0 +1,60 @@
+// Copyright 1998-2019 Epic Games, Inc. All Rights Reserved.
+#pragma once
+
+#include "MoviePipelineSetting.h"
+#include "Engine/EngineTypes.h"
+#include "Misc/FrameRate.h"
+#include "MoviePipelineOutputSetting.generated.h"
+
+// Forward Declares
+class AGameModeBase;
+
+UCLASS(Blueprintable)
+class MOVIERENDERPIPELINECORE_API UMoviePipelineOutputSetting : public UMoviePipelineSetting
+{
+	GENERATED_BODY()
+public:
+	UMoviePipelineOutputSetting();
+	
+public:
+	/** What directory should all of our output files be relative to. */
+	UPROPERTY(EditAnywhere, BlueprintReadWrite, Category = "Movie Render Pipeline")
+	FDirectoryPath OutputDirectory;
+	
+	/** What format string should the final files use? Can include folder prefixes, and format string ({shot_name}, etc.) */
+	UPROPERTY(EditAnywhere, BlueprintReadWrite, Category = "Movie Render Pipeline")
+	FString FileNameFormat;
+
+	/** What resolution should our output files be exported at? */
+	UPROPERTY(EditAnywhere, BlueprintReadWrite, Category = "Movie Render Pipeline")
+	FIntPoint OutputResolution;
+	
+	/** Should we use the custom frame rate specified by OutputFrameRate? Otherwise defaults to Sequence frame rate. */
+	UPROPERTY(EditAnywhere, BlueprintReadWrite, Category = "Movie Render Pipeline")
+	bool bUseCustomFrameRate;
+	
+	/** What frame rate should the output files be exported at? This overrides the Display Rate of the target sequence. */
+	UPROPERTY(EditAnywhere, BlueprintReadWrite, meta=(EditCondition=bUseCustomFrameRate), Category = "Movie Render Pipeline")
+	FFrameRate OutputFrameRate;
+	
+	/** How many frames (in Display Rate) should we step at a time. Can be used to render every other frame for faster drafts. */
+	UPROPERTY(EditAnywhere, BlueprintReadWrite, Category = "Movie Render Pipeline")
+	FFrameNumber OutputFrameStep;
+	
+	/** If true, output containers should attempt to override any existing files with the same name. */
+	UPROPERTY(EditAnywhere, BlueprintReadWrite, Category = "Movie Render Pipeline")
+	bool bOverrideExistingOutput;
+	
+public:
+	/** If true, override the Playback Range start/end bounds with the bounds specified below.*/
+	UPROPERTY(EditAnywhere, BlueprintReadWrite, Category = "Movie Render Pipeline")
+	bool bUseCustomPlaybackRange;
+	
+	/** Used when overriding the playback range. In Display Rate frames. If bUseCustomPlaybackRange is false range will come from Sequence. */
+	UPROPERTY(EditAnywhere, BlueprintReadWrite, meta=(EditCondition=bUseCustomPlaybackRange), Category = "Movie Render Pipeline")
+	FFrameNumber CustomStartFrame;
+	
+	/** Used when overriding the playback range. In Display Rate frames. If bUseCustomPlaybackRange is false range will come from Sequence. */
+	UPROPERTY(EditAnywhere, BlueprintReadWrite, meta=(EditCondition=bUseCustomPlaybackRange), Category = "Movie Render Pipeline")
+	FFrameNumber CustomEndFrame;
+};
\ No newline at end of file
diff --git a/Engine/Plugins/MovieScene/MovieRenderPipeline/Source/MovieRenderPipelineCore/Public/MoviePipelineRenderPass.h b/Engine/Plugins/MovieScene/MovieRenderPipeline/Source/MovieRenderPipelineCore/Public/MoviePipelineRenderPass.h
--- a/Engine/Plugins/MovieScene/MovieRenderPipeline/Source/MovieRenderPipelineCore/Public/MoviePipelineRenderPass.h
+++ b/Engine/Plugins/MovieScene/MovieRenderPipeline/Source/MovieRenderPipelineCore/Public/MoviePipelineRenderPass.h
@@ -23,12 +23,6 @@ public:
 		CaptureFrameImpl(OutputFrameMetrics);
 	}
 
-
-	void GetFrameData(MoviePipeline::FOutputFrameData& OutFrameData)
-	{
-		GetFrameDataImpl(OutFrameData);
-	}
-
 	void Teardown()
 	{
 		TeardownImpl();
@@ -42,9 +36,7 @@ public:
 protected:
 	virtual void SetupImpl(const FMoviePipelineRenderPassInitSettings& InInitSettings, TSharedRef<FImagePixelPipe, ESPMode::ThreadSafe> InOutputPipe) {  }
 
-	virtual void CaptureFrameImpl(const FMoviePipelineRenderPassMetrics& OutputFrameMetrics) { }
-
-	virtual void GetFrameDataImpl(MoviePipeline::FOutputFrameData& OutFrameData) { }
+	virtual void CaptureFrameImpl(const FMoviePipelineRenderPassMetrics& OutputFrameMetrics) {}
 
 	virtual void TeardownImpl() {}
 
diff --git a/Engine/Plugins/MovieScene/MovieRenderPipeline/Source/MovieRenderPipelineCore/Public/MoviePipelineSetting.h b/Engine/Plugins/MovieScene/MovieRenderPipeline/Source/MovieRenderPipelineCore/Public/MoviePipelineSetting.h
--- a/Engine/Plugins/MovieScene/MovieRenderPipeline/Source/MovieRenderPipelineCore/Public/MoviePipelineSetting.h
+++ b/Engine/Plugins/MovieScene/MovieRenderPipeline/Source/MovieRenderPipelineCore/Public/MoviePipelineSetting.h
@@ -17,9 +17,16 @@ class MOVIERENDERPIPELINECORE_API UMoviePipelineSetting : public UObject
 		
 public:
 	/**
-	* Called when the Pipeline is initialized for the first time before the first frame is ticked.
+	* Called during setup in the relevant portion of the Movie Pipeline lifecycle.
+	* Can get called multiple times if a setting is shared between shots.
 	*/
-	void OnInitializedForPipeline(UMoviePipeline* InPipeline);
+	void SetupForPipeline(UMoviePipeline* InPipeline);
+
+	/**
+	* Called during shutdown in the relevant portion of the Movie Pipeline lifecycle.
+	* Can get called multiple times if a setting is shared between shots.
+	*/
+	void TeardownForPipeline(UMoviePipeline* InPipeline) { TeardownForPipelineImpl(InPipeline); }
 
 	// UObject Interface
 	virtual UWorld* GetWorld() const override;
@@ -28,8 +35,8 @@ public:
 protected:
 	UMoviePipeline* GetPipeline() const;
 
-
-	virtual void OnInitializedForPipelineImpl(UMoviePipeline* InPipeline) {}
+	virtual void SetupForPipelineImpl(UMoviePipeline* InPipeline) {}
+	virtual void TeardownForPipelineImpl(UMoviePipeline* InPipeline) {}
 
 public:
 #if WITH_EDITOR
diff --git a/Engine/Plugins/MovieScene/MovieRenderPipeline/Source/MovieRenderPipelineCore/Public/MoviePipelineShotConfig.h b/Engine/Plugins/MovieScene/MovieRenderPipeline/Source/MovieRenderPipelineCore/Public/MoviePipelineShotConfig.h
--- a/Engine/Plugins/MovieScene/MovieRenderPipeline/Source/MovieRenderPipelineCore/Public/MoviePipelineShotConfig.h
+++ b/Engine/Plugins/MovieScene/MovieRenderPipeline/Source/MovieRenderPipelineCore/Public/MoviePipelineShotConfig.h
@@ -1,11 +1,8 @@
 // Copyright 1998-2019 Epic Games, Inc. All Rights Reserved.
 #pragma once
 
-#include "UObject/Object.h"
-#include "UObject/SoftObjectPath.h"
-#include "MoviePipelineSetting.h"
-#include "Engine/EngineTypes.h"
-#include "Misc/FrameRate.h"
+#include "MoviePipelineConfigBase.h"
+#include "MoviePipelineRenderPass.h"
 
 #include "MoviePipelineShotConfig.generated.h"
 
@@ -14,118 +11,17 @@
 class ULevelSequence;
 class UMoviePipelineSetting;
 class UMoviePipelineRenderPass;
-class UMoviePipelineOutput;
+class UMoviePipelineOutputBase;
 
 UCLASS(Blueprintable)
-class MOVIERENDERPIPELINECORE_API UMoviePipelineShotConfig : public UObject
+class MOVIERENDERPIPELINECORE_API UMoviePipelineShotConfig : public UMoviePipelineConfigBase
 {
 	GENERATED_BODY()
 	
 public:
 	UMoviePipelineShotConfig()
 	{
-		SettingsSerialNumber = -1;
-	}
-
-	bool ValidateConfig(TArray<FText>& OutValidationErrors) const;
-	void RemoveSetting(UMoviePipelineSetting* InSetting);
-	int32 GetSettingsSerialNumber() const { return SettingsSerialNumber; }
-	TArray<UMoviePipelineSetting*> GetSettings() const { return Settings; }
-
-public:
-	/**
-	* Find a setting of a particular type for this config.
-	* @param InClass - Class that you wish to find the setting object for.
-	* @return An instance of this class if it already exists as a setting on this config, otherwise null.
-	*/
-	UFUNCTION(BlueprintPure, Category = "Movie Render Pipeline")
-	UMoviePipelineSetting* FindSettingByClass(TSubclassOf<UMoviePipelineSetting> InClass) const
-	{
-		UMoviePipelineSetting* const* Found = Settings.FindByPredicate([InClass](UMoviePipelineSetting* In) { return In && In->GetClass() == InClass; });
-		return Found ? CastChecked<UMoviePipelineSetting>(*Found) : nullptr;
-	}
-	
-	/** Finds a setting of a particular type for this pipeline config, adding it if it doesn't already exist.
-	* @param InClass - Class you wish to find or create the setting object for.
-	* @return An instance of this class as a setting on this config. 
-	*/
-	UFUNCTION(BlueprintCallable, Category = "Movie Render Pipeline")
-	UMoviePipelineSetting* FindOrAddSettingByClass(TSubclassOf<UMoviePipelineSetting> InClass)
-	{
-		UMoviePipelineSetting* Found = FindSettingByClass(InClass);
-		if (!Found)
-		{
-			Modify();
-
-			Found = NewObject<UMoviePipelineSetting>(this, InClass);
-			Settings.Add(Found);
-			++SettingsSerialNumber;
-		}
-	
-		return Found;
-	}
-	
-	/**
-	* Find a setting of a particular type for this config
-	*/ 
-	template<typename SettingType>
-	SettingType* FindSetting() const
-	{
-		UClass* PredicateClass = SettingType::StaticClass();
-		UMoviePipelineSetting* const* Found = Settings.FindByPredicate([PredicateClass](UMoviePipelineSetting* In) { return In && In->GetClass() == PredicateClass; });
-		return Found ? CastChecked<SettingType>(*Found) : nullptr;
 	}
 	
-	/**
-	 * Find a setting of a particular type for this config instance, adding one if it was not found.
-	 */
-	template<typename SettingType>
-	SettingType* FindOrAddSetting()
-	{
-		SettingType* Found = FindSetting<SettingType>();
-		if (!Found)
-		{
-			Modify();
-
-			Found = NewObject<SettingType>(this);
-			Settings.Add(Found);
-			++SettingsSerialNumber;
-		}
-		return Found;
-	}
-public:
-	/** What Sequence do we want to render out. Can be overridden by command line launch arguments. */
-	UPROPERTY(EditAnywhere, BlueprintReadWrite, Category="Movie Pipeline", meta=(MetaClass="LevelSequence"))
-	FSoftObjectPath Sequence;
-	
-	/** What resolution should our output files be exported at? */
-	UPROPERTY(EditAnywhere, BlueprintReadWrite, Category = "Movie Pipeline")
-	FVector2D OutputResolution;
-	
-	/** What frame rate should the output files be exported at? This overrides the Display Rate of the target sequence. */
-	UPROPERTY(EditAnywhere, BlueprintReadWrite, Category = "Movie Pipeline")
-	FFrameRate OutputFrameRate;
-	
-	/** What directory should all of our output files be relative to. */
-	UPROPERTY(EditAnywhere, BlueprintReadWrite, Category = "Movie Pipeline")
-	FDirectoryPath OutputDirectory;
-	
-	/** If true, output containers should attempt to override any existing files with the same name. */
-	UPROPERTY(EditAnywhere, BlueprintReadWrite, Category = "Movie Pipeline")
-	bool bOverrideExistingOutput;
-	
-	/** Array of settings classes that affect various parts of the output pipeline. */
-	UPROPERTY(EditAnywhere, BlueprintReadWrite, Category = "Movie Pipeline")
-	TArray<UMoviePipelineSetting*> Settings;
-	
-	/** Array of Input Buffers. These specify how many render passes should happen each frame to get the data for that frame. */
-	UPROPERTY(EditAnywhere, BlueprintReadWrite, Category = "Movie Pipeline")
-	TArray<UMoviePipelineRenderPass*> InputBuffers;
-	
-	/** Array of Output Containers. Each output container is passed data for each Input Buffer every frame. */
-	UPROPERTY(EditAnywhere, BlueprintReadWrite, Category = "Movie Pipeline")
-	TArray<UMoviePipelineOutput*> OutputContainers;
-
-private:
-	int32 SettingsSerialNumber;
+	TArray<UMoviePipelineRenderPass*> GetRenderPasses() const;
 };
\ No newline at end of file
diff --git a/Engine/Plugins/MovieScene/MovieRenderPipeline/Source/MovieRenderPipelineCore/Public/MovieRenderPipelineConfig.h b/Engine/Plugins/MovieScene/MovieRenderPipeline/Source/MovieRenderPipelineCore/Public/MovieRenderPipelineConfig.h
deleted file mode 100644
diff --git a/Engine/Plugins/MovieScene/MovieRenderPipeline/Source/MovieRenderPipelineCore/Public/MovieRenderPipelineCoreModule.h b/Engine/Plugins/MovieScene/MovieRenderPipeline/Source/MovieRenderPipelineCore/Public/MovieRenderPipelineCoreModule.h
--- a/Engine/Plugins/MovieScene/MovieRenderPipeline/Source/MovieRenderPipelineCore/Public/MovieRenderPipelineCoreModule.h
+++ b/Engine/Plugins/MovieScene/MovieRenderPipeline/Source/MovieRenderPipelineCore/Public/MovieRenderPipelineCoreModule.h
@@ -7,8 +7,9 @@
 
 // Forward Declare
 class UMoviePipelineExecutorBase;
-class UMovieRenderPipelineConfig;
+class UMoviePipelineMasterConfig;
 class ULevelSequence;
+class UMoviePipeline;
 
 namespace MoviePipelineErrorCodes
 {
@@ -32,9 +33,12 @@ public:
 private:
 	bool IsTryingToRenderMovieFromCommandLine(FString& OutSequenceAssetPath, FString& OutConfigAssetPath, FString& OutExecutorType, FString& OutPipelineType) const;
 	void InitializeCommandLineMovieRender();
-	void OnCommandLineMovieRenderCompleted(UMoviePipelineExecutorBase* InExecutor);
+
+	void OnCommandLineMovieRenderCompleted(UMoviePipelineExecutorBase* InExecutor, bool bSuccess);
+	void OnCommandLineMovieRenderErrored(UMoviePipelineExecutorBase* InExecutor, UMoviePipeline* InPipelineWithError, bool bIsFatal, FText ErrorText);
+
 	uint8 GetMovieRenderClasses(const FString& InAssetPath, const FString& InConfigAssetPath, const FString& InExecutorType, const FString& InPipelineType, 
-								ULevelSequence*& OutAsset, UMovieRenderPipelineConfig*& OutConfig, UClass*& OutExecutorClass, UClass*& OutPipelineClass) const;
+								ULevelSequence*& OutAsset, UMoviePipelineMasterConfig*& OutConfig, UClass*& OutExecutorClass, UClass*& OutPipelineClass) const;
 
 private:
 	FString MoviePipelineLocalExecutorClassType;
diff --git a/Engine/Plugins/MovieScene/MovieRenderPipeline/Source/MovieRenderPipelineCore/Public/MovieRenderPipelineDataTypes.cpp b/Engine/Plugins/MovieScene/MovieRenderPipeline/Source/MovieRenderPipelineCore/Public/MovieRenderPipelineDataTypes.cpp
--- a/Engine/Plugins/MovieScene/MovieRenderPipeline/Source/MovieRenderPipelineCore/Public/MovieRenderPipelineDataTypes.cpp
+++ b/Engine/Plugins/MovieScene/MovieRenderPipeline/Source/MovieRenderPipelineCore/Public/MovieRenderPipelineDataTypes.cpp
@@ -5,7 +5,7 @@
 #include "Sections/MovieSceneCinematicShotSection.h"
 #include "MoviePipelineAccumulationSetting.h"
 
-FString FMoviePipelineShotCache::GetDisplayName() const
+FString FMoviePipelineShotInfo::GetDisplayName() const
 {
 	if (CinematicShotSection.Get())
 	{
@@ -15,7 +15,7 @@ FString FMoviePipelineShotCache::GetDisplayName() const
 	return TEXT("Unnamed");
 }
 
-FFrameNumber FMoviePipelineShotCutCache::GetOutputFrameCountEstimate() const
+FFrameNumber FMoviePipelineCameraCutInfo::GetOutputFrameCountEstimate() const
 {
 	// TotalRange is stored in Tick Resolution, so we convert 1 frame of Frame Rate to the number of ticks.
 	FFrameNumber OneFrameInTicks = FFrameRate::TransformTime(FFrameTime(FFrameNumber(1)), CachedFrameRate, CachedTickResolution).FloorToFrame();
@@ -27,7 +27,7 @@ FFrameNumber FMoviePipelineShotCutCache::GetOutputFrameCountEstimate() const
 	return FFrameNumber(NumFrames);
 }
 
-FFrameNumber FMoviePipelineShotCutCache::GetTemporalFrameCountEstimate() const
+FFrameNumber FMoviePipelineCameraCutInfo::GetTemporalFrameCountEstimate() const
 {
 	FFrameNumber OutputEstimate = GetOutputFrameCountEstimate();
 	FFrameNumber NumRenderedSamples = OutputEstimate * NumTemporalSamples;
@@ -35,12 +35,12 @@ FFrameNumber FMoviePipelineShotCutCache::GetTemporalFrameCountEstimate() const
 	return NumRenderedSamples;
 }
 
-FFrameNumber FMoviePipelineShotCutCache::GetUtilityFrameCountEstimate() const
+FFrameNumber FMoviePipelineCameraCutInfo::GetUtilityFrameCountEstimate() const
 {
 	return FFrameNumber(NumWarmUpFrames + (bAccurateFirstFrameHistory ? 1 : 0));
 }
 
-FFrameNumber FMoviePipelineShotCutCache::GetSampleCountEstimate(const bool bIncludeWarmup, const bool bIncludeMotionBlur) const
+FFrameNumber FMoviePipelineCameraCutInfo::GetSampleCountEstimate(const bool bIncludeWarmup, const bool bIncludeMotionBlur) const
 {
 	int32 TotalSamples = 0;
 	
@@ -50,7 +50,7 @@ FFrameNumber FMoviePipelineShotCutCache::GetSampleCountEstimate(const bool bIncl
 		TotalSamples += 0;
 	}
 
-	// Motion blur only adds one frame
+	// Motion blur only adds one sample (overrides spatial samples, only executes once before first temporal sample)
 	if (bIncludeMotionBlur)
 	{
 		TotalSamples += 1;
@@ -62,7 +62,7 @@ FFrameNumber FMoviePipelineShotCutCache::GetSampleCountEstimate(const bool bIncl
 	return FFrameNumber(TotalSamples);
 }
 
-void FMoviePipelineShotCutCache::SetNextState(const EMovieRenderShotState InCurrentState)
+void FMoviePipelineCameraCutInfo::SetNextState(const EMovieRenderShotState InCurrentState)
 {
 	switch (InCurrentState)
 	{
@@ -105,4 +105,30 @@ void FMoviePipelineShotCutCache::SetNextState(const EMovieRenderShotState InCurr
 		break;
 	}
 
+}
+
+FMoviePipelineWorkInfo FMoviePipelineCameraCutInfo::GetTotalWorkEstimate() const
+{
+	FMoviePipelineWorkInfo WorkEstimate;
+
+	// Initial Range + Handle Frames
+	FFrameNumber OutputFrameCount = GetOutputFrameCountEstimate();
+	WorkEstimate.NumOutputFrames += OutputFrameCount.Value;
+
+	// Warm Up (no submission to gpu), Motion Blur
+	FFrameNumber UtilityFrameCount = GetUtilityFrameCountEstimate();
+	WorkEstimate.NumUtilityFrames += UtilityFrameCount.Value;
+
+	// OutputFrameCount * TemporalSampleCount (since each output frame gets time sliced)
+	FFrameNumber SubsampledOutputFrameCount = GetTemporalFrameCountEstimate();
+	WorkEstimate.NumOutputFramesWithSubsampling += SubsampledOutputFrameCount.Value;
+
+	// How many different samples are we submitting to the GPU?
+	FFrameNumber TotalSampleCount = GetSampleCountEstimate(true, true) * NumTiles;
+	WorkEstimate.NumSamples += TotalSampleCount.Value;
+
+	// How many images will we produce to create the high-res tile.
+	WorkEstimate.NumTiles = NumTiles;
+
+	return WorkEstimate;
 }
\ No newline at end of file
diff --git a/Engine/Plugins/MovieScene/MovieRenderPipeline/Source/MovieRenderPipelineCore/Public/MovieRenderPipelineDataTypes.h b/Engine/Plugins/MovieScene/MovieRenderPipeline/Source/MovieRenderPipelineCore/Public/MovieRenderPipelineDataTypes.h
--- a/Engine/Plugins/MovieScene/MovieRenderPipeline/Source/MovieRenderPipelineCore/Public/MovieRenderPipelineDataTypes.h
+++ b/Engine/Plugins/MovieScene/MovieRenderPipeline/Source/MovieRenderPipelineCore/Public/MovieRenderPipelineDataTypes.h
@@ -13,9 +13,10 @@ template<class T, class TWeakObjectPtrBase> struct TWeakObjectPtr;
 class UWorld;
 class ULevelSequence;
 class UMovieSceneCinematicShotSection;
-class UMovieRenderPipelineConfig;
+class UMoviePipelineMasterConfig;
 class UMoviePipelineShotConfig;
 class UMovieSceneCameraCutSection;
+class UMoviePipelineRenderPass;
 
 
 /**
@@ -151,6 +152,80 @@ namespace MoviePipeline
 	};
 }
 
+USTRUCT(BlueprintType)
+struct FMoviePipelineWorkInfo
+{
+	GENERATED_BODY()
+
+public:
+	FMoviePipelineWorkInfo()
+		: NumCameraCuts(0)
+		, NumOutputFrames(0)
+		, NumUtilityFrames(0)
+		, NumOutputFramesWithSubsampling(0)
+		, NumSamples(0)
+		, NumTiles(0)
+	{
+	}
+
+	FString ToDisplayString()
+	{
+		return FString::Printf(TEXT("Camera Cuts: %d Output Frames: %d SubsampledOutputFrames: %d TotalSamples: %d UtilityFrames: %d TileCount: %d"),
+			NumCameraCuts, NumOutputFramesWithSubsampling, NumOutputFramesWithSubsampling, NumSamples, NumUtilityFrames, NumTiles);
+	}
+
+	/** Either the current camera cut number (not index) or the total number of camera cuts. */
+	int32 NumCameraCuts;
+	/** How many frames do we expect to see written to disk. */
+	int32 NumOutputFrames;
+	/** How many extra engine ticks are run for utility reasons that don't produce any output. */
+	int32 NumUtilityFrames;
+	/** OutputFrames * TemporalSamples. This many engine frames are passed during the production of all Output Frames. */
+	int32 NumOutputFramesWithSubsampling;
+	/** NumOutputFrames * NumTemporalSamples * SpatialSamples * NumTiles */
+	int32 NumSamples;
+	/** Either the current tile number (not index) or the total number of tiles */
+	int32 NumTiles;
+
+	FMoviePipelineWorkInfo& operator += (const FMoviePipelineWorkInfo& InRHS)
+	{
+		NumCameraCuts += InRHS.NumCameraCuts;
+		NumOutputFrames += InRHS.NumOutputFrames;
+		NumUtilityFrames += InRHS.NumUtilityFrames;
+		NumOutputFramesWithSubsampling += InRHS.NumOutputFramesWithSubsampling;
+		NumSamples += InRHS.NumSamples;
+		NumTiles += InRHS.NumTiles;
+
+		return *this;
+	}
+
+	friend FMoviePipelineWorkInfo operator+ (const FMoviePipelineWorkInfo& InLHS, const FMoviePipelineWorkInfo& InRHS)
+	{
+		FMoviePipelineWorkInfo Sum = InLHS;
+		Sum += InRHS;
+		return Sum;
+	}
+
+	bool operator == (const FMoviePipelineWorkInfo& InRHS) const
+	{
+		return	NumCameraCuts == InRHS.NumCameraCuts
+				&& NumOutputFrames == InRHS.NumOutputFrames
+				&& NumUtilityFrames == InRHS.NumUtilityFrames
+				&& NumOutputFramesWithSubsampling == InRHS.NumOutputFramesWithSubsampling
+				&& NumSamples == InRHS.NumSamples
+				&& NumTiles == InRHS.NumTiles;
+	}
+
+	bool operator != (const FMoviePipelineWorkInfo& InRHS) const
+	{
+		return !(*this == InRHS);
+	}
+
+
+};
+
+
+
 UENUM(BlueprintType)
 enum class EMoviePipelineShutterTiming : uint8
 {
@@ -160,22 +235,20 @@ enum class EMoviePipelineShutterTiming : uint8
 };
 
 USTRUCT(BlueprintType)
-struct FMoviePipelineShotCutCache
+struct FMoviePipelineCameraCutInfo
 {
 	GENERATED_BODY()
 public:
-	FMoviePipelineShotCutCache()
+	FMoviePipelineCameraCutInfo()
 		: OriginalRange(TRange<FFrameNumber>::Empty())
 		, TotalOutputRange(TRange<FFrameNumber>::Empty())
 		, NumWarmUpFrames(0)
 		, bAccurateFirstFrameHistory(true)
 		, NumTemporalSamples(0)
 		, NumSpatialSamples(0)
+		, NumTiles(0)
 		, State(EMovieRenderShotState::Uninitialized)
 		, CurrentTick(FFrameNumber(0))
-		, NumFramesProduced(FFrameNumber(0))
-		, NumSamplesRendered(FFrameNumber(0))
-		, CurrentFrameSamples(0)
 		, bHasEvaluatedMotionBlurFrame(false)
 		, bHasNotifiedFrameProductionStarted(false)
 		, NumWarmUpFramesRemaining(0)
@@ -207,6 +280,11 @@ public:
 	*/
 	FFrameNumber GetSampleCountEstimate(const bool bIncludeWarmup = true, const bool bIncludeMotionBlur = true) const;
 	
+	/**
+	* Summarizes the total expected amount of work into one easy to access struct.
+	*/
+	FMoviePipelineWorkInfo GetTotalWorkEstimate() const;
+
 public:
 	/** The original non-modified range for this shot that will be rendered. */
 	TRange<FFrameNumber> OriginalRange;
@@ -233,6 +311,10 @@ public:
 	UPROPERTY(BlueprintReadOnly, Category = "Movie Render Pipeline")
 	int32 NumSpatialSamples;
 
+	/** How many image tiles are going to be rendered per temporal frame. */
+	UPROPERTY(BlueprintReadOnly, Category = "Movie Render Pipeline")
+	int32 NumTiles;
+
 	/** Display name for UI Purposes. */
 	UPROPERTY(BlueprintReadOnly, Category = "Movie Render Pipeline")
 	FString DisplayName;
@@ -250,21 +332,18 @@ public:
 	UPROPERTY(BlueprintReadOnly, Category = "Movie Render Pipeline")
 	FFrameNumber CurrentTick;
 
-	/** Metrics - How many frames have been rendered for this shot so far? Can be more than the estimated total due to PlayRate. */
-	UPROPERTY(BlueprintReadOnly, Category = "Movie Render Pipeline")
-	FFrameNumber NumFramesProduced;
-
-	/** Metrics - How many samples have been rendered for this shot so far? Can be more than the estimated total due to PlayRate. */
+	/** Metrics - How much work has been done for this particular shot. */
 	UPROPERTY(BlueprintReadOnly, Category = "Movie Render Pipeline")
-	FFrameNumber NumSamplesRendered;
+	FMoviePipelineWorkInfo CurrentWorkInfo;
 
 	/** 
-	* Metrics - How many samples of the total have been generated for this frame. Only as granular as SpatialSamples since
-	* we can't refresh the screen between individual spatial samples.
+	* Metrics - How much work is there estimated to do for this shot? Can be more than estimated total due to PlayRate tracks. 
+	* Cached after setup so that Blueprints can read the variable directly due to no support for UFunctions on Structs.
 	*/
 	UPROPERTY(BlueprintReadOnly, Category = "Movie Render Pipeline")
-	int32 CurrentFrameSamples;
+	FMoviePipelineWorkInfo TotalWorkInfo;
 
+public:
 	/** Have we evaluated the motion blur frame? Only used if bEvaluateMotionBlurOnFirstFrame is set */
 	UPROPERTY(BlueprintReadOnly, Category = "Movie Render Pipeline")
 	bool bHasEvaluatedMotionBlurFrame;
@@ -277,7 +356,7 @@ public:
 	UPROPERTY(BlueprintReadOnly, Category = "Movie Render Pipeline")
 	int32 NumWarmUpFramesRemaining;
 
-	bool operator == (const FMoviePipelineShotCutCache& InRHS) const
+	bool operator == (const FMoviePipelineCameraCutInfo& InRHS) const
 	{
 		return
 			OriginalRange == InRHS.OriginalRange &&
@@ -291,7 +370,7 @@ public:
 			CameraCutSection == InRHS.CameraCutSection;
 	}
 
-	bool operator != (const FMoviePipelineShotCutCache& InRHS) const
+	bool operator != (const FMoviePipelineCameraCutInfo& InRHS) const
 	{
 		return !(*this == InRHS);
 	}
@@ -303,11 +382,11 @@ public:
 * all information in advanced aids in debugging and visualization of progress.
 */
 USTRUCT(BlueprintType)
-struct FMoviePipelineShotCache
+struct FMoviePipelineShotInfo
 {
 	GENERATED_BODY()
 public:
-	FMoviePipelineShotCache()
+	FMoviePipelineShotInfo()
 		: NumHandleFrames(0)
 		, OriginalRange(TRange<FFrameNumber>::Empty())
 		, TotalOutputRange(TRange<FFrameNumber>::Empty())
@@ -341,12 +420,20 @@ public:
 	TWeakObjectPtr<UMovieSceneCinematicShotSection> CinematicShotSection;
 
 	UPROPERTY(BlueprintReadOnly, Category = "Movie Render Pipeline")
-	TArray<FMoviePipelineShotCutCache> CameraCuts;
+	TArray<FMoviePipelineCameraCutInfo> CameraCuts;
 
 	UPROPERTY(BlueprintReadOnly, Category = "Movie Render Pipeline")
 	int32 CurrentCameraCutIndex;
 
-	FMoviePipelineShotCutCache& GetCurrentCameraCut()
+	/** 
+	* An array of the Render Passes for this particular shot. These are duplicated out of the settings
+	* so that they can have their own state per shot that overlaps with other in progress shots. This
+	* is useful for long running async tasks.
+	*/
+	UPROPERTY(BlueprintReadOnly, Category = "Movie Render Pipeline")
+	TArray<UMoviePipelineRenderPass*> RenderPasses;
+
+	FMoviePipelineCameraCutInfo& GetCurrentCameraCut()
 	{
 		check(CurrentCameraCutIndex >= 0 && CurrentCameraCutIndex < CameraCuts.Num());
 		return CameraCuts[CurrentCameraCutIndex];
@@ -364,7 +451,7 @@ public:
 		return true;
 	}
 
-	bool operator == (const FMoviePipelineShotCache& InRHS) const
+	bool operator == (const FMoviePipelineShotInfo& InRHS) const
 	{ 
 		return
 			NumHandleFrames == InRHS.NumHandleFrames &&
@@ -377,7 +464,7 @@ public:
 			CameraCuts == InRHS.CameraCuts;
 	}
 
-	bool operator != (const FMoviePipelineShotCache& InRHS) const
+	bool operator != (const FMoviePipelineShotInfo& InRHS) const
 	{
 		return !(*this == InRHS);
 	}
@@ -582,4 +669,37 @@ struct FImagePixelDataPayload : IImagePixelDataPayload
 
 
 	FString PassName;
+};
+
+/**
+* Describes an individual job executed by a Movie Pipeline Executor.
+*/
+USTRUCT(BlueprintType)
+struct FMoviePipelineExecutorJob
+{
+	GENERATED_BODY()
+public:
+	FMoviePipelineExecutorJob()
+		: Sequence()
+		, Configuration(nullptr)
+	{
+	}
+
+	FMoviePipelineExecutorJob(const FSoftObjectPath& InSequenceAsset, UMoviePipelineMasterConfig* InConfig)
+		: Sequence(InSequenceAsset)
+		, Configuration(InConfig)
+	{
+	}
+
+	/** Which sequence should this job render? */
+	UPROPERTY(BlueprintReadWrite, Category = "Movie Render Pipeline")
+	FSoftObjectPath Sequence;
+
+	/** What master configuration is used to render this sequence? This specifies output resolution/path, etc. */
+	UPROPERTY(BlueprintReadWrite, Category = "Movie Render Pipeline")
+	UMoviePipelineMasterConfig* Configuration;
+
+	/** Which shots (by name) should this job render out of the given sequence? Leave empty to render whole sequence. */
+	UPROPERTY(BlueprintReadWrite, Category = "Movie Render Pipeline")
+	TArray<FString> ShotRenderMask;
 };
\ No newline at end of file
diff --git a/Engine/Plugins/MovieScene/MovieRenderPipeline/Source/MovieRenderPipelineEditor/Private/MoviePipelineEditorBlueprintLibrary.cpp b/Engine/Plugins/MovieScene/MovieRenderPipeline/Source/MovieRenderPipelineEditor/Private/MoviePipelineEditorBlueprintLibrary.cpp
--- a/Engine/Plugins/MovieScene/MovieRenderPipeline/Source/MovieRenderPipelineEditor/Private/MoviePipelineEditorBlueprintLibrary.cpp
+++ b/Engine/Plugins/MovieScene/MovieRenderPipeline/Source/MovieRenderPipelineEditor/Private/MoviePipelineEditorBlueprintLibrary.cpp
@@ -3,7 +3,6 @@
 #include "MoviePipelineEditorBlueprintLibrary.h"
 #include "UObject/Package.h"
 #include "Misc/PackageName.h"
-#include "MovieRenderPipelineConfig.h"
 #include "AssetRegistryModule.h"
 #include "FileHelpers.h"
 #include "ObjectTools.h"
@@ -11,7 +10,7 @@
 
 #define LOCTEXT_NAMESPACE "MoviePipelineEditorBlueprintLibrary"
 
-bool UMoviePipelineEditorBlueprintLibrary::ExportConfigToAsset(const UMovieRenderPipelineConfig* InConfig, const FString& InPackagePath, const FString& InFileName, const bool bInSaveAsset, UMovieRenderPipelineConfig*& OutAsset, FText& OutErrorReason)
+bool UMoviePipelineEditorBlueprintLibrary::ExportConfigToAsset(const UMoviePipelineMasterConfig* InConfig, const FString& InPackagePath, const FString& InFileName, const bool bInSaveAsset, UMoviePipelineMasterConfig*& OutAsset, FText& OutErrorReason)
 {
 	if(!InConfig)
 	{
@@ -31,7 +30,7 @@ bool UMoviePipelineEditorBlueprintLibrary::ExportConfigToAsset(const UMovieRende
 	NewPackage->AddToRoot();
 	
 	// Duplicate the provided config into this package.
-	UMovieRenderPipelineConfig* NewConfig = Cast<UMovieRenderPipelineConfig>(StaticDuplicateObject(InConfig, NewPackage, FName(*InFileName), RF_NoFlags));
+	UMoviePipelineMasterConfig* NewConfig = Cast<UMoviePipelineMasterConfig>(StaticDuplicateObject(InConfig, NewPackage, FName(*InFileName), RF_NoFlags));
 	NewConfig->SetFlags(RF_Public | RF_Transactional | RF_Standalone);
 	NewConfig->MarkPackageDirty();
 
diff --git a/Engine/Plugins/MovieScene/MovieRenderPipeline/Source/MovieRenderPipelineEditor/Private/MoviePipelinePIEExecutor.cpp b/Engine/Plugins/MovieScene/MovieRenderPipeline/Source/MovieRenderPipelineEditor/Private/MoviePipelinePIEExecutor.cpp
--- a/Engine/Plugins/MovieScene/MovieRenderPipeline/Source/MovieRenderPipelineEditor/Private/MoviePipelinePIEExecutor.cpp
+++ b/Engine/Plugins/MovieScene/MovieRenderPipeline/Source/MovieRenderPipelineEditor/Private/MoviePipelinePIEExecutor.cpp
@@ -1,17 +1,16 @@
 // Copyright 1998-2019 Epic Games, Inc. All Rights Reserved.
 #include "MoviePipelinePIEExecutor.h"
-#include "Settings/LevelEditorPlaySettings.h"
-#include "MovieRenderPipelineCoreModule.h"
-#include "MovieRenderPipelineConfig.h"
+#include "MoviePipelineMasterConfig.h"
 #include "MoviePipelineShotConfig.h"
-#include "MoviePipelineSeparateProcessSetting.h"
-#include "Framework/Application/SlateApplication.h"
-#include "Editor/EditorEngine.h"
-#include "Editor.h"
+#include "MoviePipeline.h"
 #include "MoviePipelineGameOverrideSetting.h"
+#include "MovieRenderPipelineCoreModule.h"
+#include "Framework/Application/SlateApplication.h"
 #include "MovieRenderPipelineDataTypes.h"
-#include "MoviePipeline.h"
 #include "MovieRenderPipelineSettings.h"
+#include "Editor/EditorEngine.h"
+#include "Editor.h"
+#include "Settings/LevelEditorPlaySettings.h"
 
 #define LOCTEXT_NAMESPACE "MoviePipelinePIEExecutor"
 
@@ -26,15 +25,13 @@ FText UMoviePipelinePIEExecutor::GetWindowTitle(const int32 InConfigIndex, const
 	return WindowTitle;
 }
 
-void UMoviePipelinePIEExecutor::Start(UMovieRenderPipelineConfig* InConfig, const int32 InConfigIndex, const int32 InNumConfigs)
+void UMoviePipelinePIEExecutor::Start(const FMoviePipelineExecutorJob& InJob)
 {
-	Super::Start(InConfig, InConfigIndex, InNumConfigs);
-
-	ActiveConfig = InConfig;
+	Super::Start(InJob);
 
 	// Create a Slate window to hold our preview.
 	TSharedRef<SWindow> CustomWindow = SNew(SWindow)
-		.Title_UObject(this, &UMoviePipelinePIEExecutor::GetWindowTitle, InConfigIndex, InNumConfigs)
+		.Title_UObject(this, &UMoviePipelinePIEExecutor::GetWindowTitle, CurrentPipelineIndex, ExecutorJobs.Num())
 		.ClientSize(FVector2D(1280, 720))
 		.AutoCenter(EAutoCenter::PrimaryWorkArea)
 		.UseOSWindowBorder(true)
@@ -57,7 +54,7 @@ void UMoviePipelinePIEExecutor::Start(UMovieRenderPipelineConfig* InConfig, cons
 	Params.EditorPlaySettings = PlayInEditorSettings;
 	Params.CustomPIEWindow = CustomWindow;
 
-	UMoviePipelineGameOverrideSetting* GameOverrides = InConfig->DefaultShotConfig->FindSetting<UMoviePipelineGameOverrideSetting>();
+	UMoviePipelineGameOverrideSetting* GameOverrides = InJob.Configuration->DefaultShotConfig->FindSetting<UMoviePipelineGameOverrideSetting>();
 	if (GameOverrides)
 	{
 		Params.GameModeOverride = GameOverrides->GameModeOverride;
@@ -97,12 +94,15 @@ void UMoviePipelinePIEExecutor::OnPIEStartupFinished(bool)
 	// This Pipeline belongs to the world being created so that they have context for things they execute.
 	ActiveMoviePipeline = NewObject<UMoviePipeline>(ExecutingWorld, PipelineClass);
 	
-	ActiveMoviePipeline->Initialize(ActiveConfig);
+	ActiveMoviePipeline->Initialize(ExecutorJobs[CurrentPipelineIndex]);
 
 	// Listen for when the pipeline thinks it has finished.
 	ActiveMoviePipeline->OnMoviePipelineFinished().AddUObject(this, &UMoviePipelinePIEExecutor::OnPIEMoviePipelineFinished);
+	ActiveMoviePipeline->OnMoviePipelineErrored().AddUObject(this, &UMoviePipelinePIEExecutor::OnPipelineErrored);
 }
 
+
+
 void UMoviePipelinePIEExecutor::OnPIEMoviePipelineFinished(UMoviePipeline* InMoviePipeline)
 {
 	if (ActiveMoviePipeline)
@@ -126,7 +126,6 @@ void UMoviePipelinePIEExecutor::DelayedFinishNotification()
 	// Null these out now since OnIndividualPipelineFinished might invoke something that causes a GC
 	// and we want them to go away with the GC.
 	ActiveMoviePipeline = nullptr;
-	ActiveConfig = nullptr;
 	
 	// Now that another frame has passed and we should be OK to start another PIE session, notify our owner.
 	OnIndividualPipelineFinished(MoviePipeline);
diff --git a/Engine/Plugins/MovieScene/MovieRenderPipeline/Source/MovieRenderPipelineEditor/Private/MovieRenderPipelineEditorModule.cpp b/Engine/Plugins/MovieScene/MovieRenderPipeline/Source/MovieRenderPipelineEditor/Private/MovieRenderPipelineEditorModule.cpp
--- a/Engine/Plugins/MovieScene/MovieRenderPipeline/Source/MovieRenderPipelineEditor/Private/MovieRenderPipelineEditorModule.cpp
+++ b/Engine/Plugins/MovieScene/MovieRenderPipeline/Source/MovieRenderPipelineEditor/Private/MovieRenderPipelineEditorModule.cpp
@@ -10,7 +10,6 @@
 #include "Editor.h"
 #include "MoviePipeline.h"
 #include "MovieRenderPipelineDataTypes.h"
-#include "MovieRenderPipelineConfig.h"
 #include "MoviePipelineBackbufferPass.h"
 #include "MoviePipelineImageSequenceContainer.h"
 #include "MoviePipelineShotConfig.h"
@@ -76,11 +75,18 @@ void FMovieRenderPipelineEditorModule::PerformTestPipelineRender(const TArray<FS
 	Executor->AddToRoot();
 	Executor->OnExecutorFinished().AddRaw(this, &FMovieRenderPipelineEditorModule::OnTestPipelineExecutorFinished);
 
-	TArray<UMovieRenderPipelineConfig*> Pipelines = GenerateTestPipelineConfigs(SequencePath);
-	Executor->Execute(Pipelines);
+	TArray<UMoviePipelineMasterConfig*> Pipelines = GenerateTestPipelineConfigs();
+
+	TArray<FMoviePipelineExecutorJob> Jobs;
+	for(UMoviePipelineMasterConfig* Pipeline : Pipelines)
+	{
+		Jobs.Add(FMoviePipelineExecutorJob(SequencePath, Pipeline));
+	}
+
+	Executor->Execute(Jobs);
 }
 
-void FMovieRenderPipelineEditorModule::OnTestPipelineExecutorFinished(UMoviePipelineExecutorBase* InExecutor)
+void FMovieRenderPipelineEditorModule::OnTestPipelineExecutorFinished(UMoviePipelineExecutorBase* InExecutor, bool bSuccess)
 {
 	if (Executor)
 	{
diff --git a/Engine/Plugins/MovieScene/MovieRenderPipeline/Source/MovieRenderPipelineEditor/Private/MovieRenderPipelineEditorModule.h b/Engine/Plugins/MovieScene/MovieRenderPipeline/Source/MovieRenderPipelineEditor/Private/MovieRenderPipelineEditorModule.h
--- a/Engine/Plugins/MovieScene/MovieRenderPipeline/Source/MovieRenderPipelineEditor/Private/MovieRenderPipelineEditorModule.h
+++ b/Engine/Plugins/MovieScene/MovieRenderPipeline/Source/MovieRenderPipelineEditor/Private/MovieRenderPipelineEditorModule.h
@@ -4,7 +4,7 @@
 
 #include "IMovieRenderPipelineEditorModule.h"
 class UMoviePipelineShotConfig;
-class UMovieRenderPipelineConfig;
+class UMoviePipelineMasterConfig;
 class UMoviePipelineExecutorBase;
 
 class FMovieRenderPipelineEditorModule : public IMovieRenderPipelineEditorModule
@@ -18,12 +18,12 @@ private:
 	void RegisterSettings();
 	void UnregisterSettings();
 	void PerformTestPipelineRender(const TArray<FString>& Args);
-	void OnTestPipelineExecutorFinished(UMoviePipelineExecutorBase* InExecutor);
+	void OnTestPipelineExecutorFinished(UMoviePipelineExecutorBase* InExecutor, bool bSuccess);
 
 private:
-	TArray<UMovieRenderPipelineConfig*> GenerateTestPipelineConfigs(FSoftObjectPath InSequence);
+	TArray<UMoviePipelineMasterConfig*> GenerateTestPipelineConfigs();
 private:
 	class UMoviePipeline* ActiveMoviePipeline;
-	class UMovieRenderPipelineConfig* PipelineConfig;
+	class UMoviePipelineMasterConfig* PipelineConfig;
 	UMoviePipelineExecutorBase* Executor;
 };
diff --git a/Engine/Plugins/MovieScene/MovieRenderPipeline/Source/MovieRenderPipelineEditor/Private/MovieRenderPipelineEditorTest.cpp b/Engine/Plugins/MovieScene/MovieRenderPipeline/Source/MovieRenderPipelineEditor/Private/MovieRenderPipelineEditorTest.cpp
--- a/Engine/Plugins/MovieScene/MovieRenderPipeline/Source/MovieRenderPipelineEditor/Private/MovieRenderPipelineEditorTest.cpp
+++ b/Engine/Plugins/MovieScene/MovieRenderPipeline/Source/MovieRenderPipelineEditor/Private/MovieRenderPipelineEditorTest.cpp
@@ -5,18 +5,18 @@
 #include "MoviePipelineGameOverrideSetting.h"
 #include "GameFramework/GameMode.h"
 #include "MoviePipelineAccumulationSetting.h"
-#include "MovieRenderPipelineConfig.h"
 #include "MoviePipelineImageSequenceContainer.h"
 #include "MoviePipelineAudioOutput.h"
 #include "Misc/Paths.h"
+#include "MoviePipelineMasterConfig.h"
+#include "MoviePipelineOutputSetting.h"
 
 UMoviePipelineShotConfig* GenerateTestShotConfig(UObject* InOwner, int32 InSampleCount,
 	int32 InShutterAngle, EMoviePipelineShutterTiming InFrameTiming, int32 InTileCount,
 	int32 InSpatialSampleCount, bool bIsUsingOverlappedTiles, float PadRatioX, float PadRatioY, float AccumulationGamma)
 {
 	UMoviePipelineShotConfig* OutConfig = NewObject<UMoviePipelineShotConfig>(InOwner);
-	UMoviePipelineBackbufferPass* DefaultPass = NewObject<UMoviePipelineBackbufferPass>(InOwner);
-	OutConfig->InputBuffers.Add(DefaultPass);
+	OutConfig->FindOrAddSetting<UMoviePipelineBackbufferPass>();
 
 	UMoviePipelineGameOverrideSetting* GamemodeOverride = OutConfig->FindOrAddSetting<UMoviePipelineGameOverrideSetting>();
 	GamemodeOverride->GameModeOverride = AGameMode::StaticClass();
@@ -35,7 +35,7 @@ UMoviePipelineShotConfig* GenerateTestShotConfig(UObject* InOwner, int32 InSampl
 	return OutConfig;
 }
 
-TArray<UMovieRenderPipelineConfig*> FMovieRenderPipelineEditorModule::GenerateTestPipelineConfigs(FSoftObjectPath InSequence)
+TArray<UMoviePipelineMasterConfig*> FMovieRenderPipelineEditorModule::GenerateTestPipelineConfigs()
 {
 	// int32 ShutterAngles[] = { 180, 360 };
 	// EMoviePipelineShutterTiming ShutterTimings[] = { EMoviePipelineShutterTiming::FrameOpen, EMoviePipelineShutterTiming::FrameCenter, EMoviePipelineShutterTiming::FrameClose };
@@ -47,7 +47,7 @@ TArray<UMovieRenderPipelineConfig*> FMovieRenderPipelineEditorModule::GenerateTe
 	FString ShutterTimingNames[] = { TEXT("FCenter") };
 	int32 TemporalSampleCounts[] = { 1 };
 
-	TArray<UMovieRenderPipelineConfig*> Configs;
+	TArray<UMoviePipelineMasterConfig*> Configs;
 
 	for (int32 ShutterAngleIndex = 0; ShutterAngleIndex < UE_ARRAY_COUNT(ShutterAngles); ShutterAngleIndex++)
 	{
@@ -58,8 +58,8 @@ TArray<UMovieRenderPipelineConfig*> FMovieRenderPipelineEditorModule::GenerateTe
 				// Build a folder string to place this particular test in.
 				FString DirectoryName = FString::Printf(TEXT("SA_%d_ST_%s_SC_%d"), ShutterAngles[ShutterAngleIndex], *ShutterTimingNames[ShutterTimingIndex], TemporalSampleCounts[TemporalSampleCountIndex]);
 
-				UMovieRenderPipelineConfig* OutPipeline = NewObject<UMovieRenderPipelineConfig>(GetTransientPackage());
-				OutPipeline->Sequence = InSequence;
+				UMoviePipelineMasterConfig* OutPipeline = NewObject<UMoviePipelineMasterConfig>(GetTransientPackage());
+				UMoviePipelineOutputSetting* OutputSetting = OutPipeline->FindOrAddSetting<UMoviePipelineOutputSetting>();
 
 				int32 SizeX = 1080;
 				int32 SizeY = 1920;
@@ -70,8 +70,8 @@ TArray<UMovieRenderPipelineConfig*> FMovieRenderPipelineEditorModule::GenerateTe
 				float PadRatioY = 0.5f;
 				float AccumulationGamma = 1.0f;
 
-				OutPipeline->OutputResolution = FIntPoint(SizeX*TileX, SizeY*TileY);
-				OutPipeline->OutputDirectory.Path = FPaths::ProjectSavedDir() / TEXT("/VideoCaptures/") / *DirectoryName;
+				OutputSetting->OutputResolution = FIntPoint(SizeX*TileX, SizeY*TileY);
+				OutputSetting->OutputDirectory.Path = FPaths::ProjectSavedDir() / TEXT("/VideoCaptures/") / *DirectoryName;
 
 				const int32 NumSpatialSamples = TestNumSamples;
 				const int32 NumTiles = TileX;
@@ -86,11 +86,7 @@ TArray<UMovieRenderPipelineConfig*> FMovieRenderPipelineEditorModule::GenerateTe
 				OutPipeline->DefaultShotConfig = DefaultConfig;
 
 				// Add some Outputs
-				UMoviePipelineImageSequenceContainerBase* ImageSequenceOutput = NewObject<UMoviePipelineImageSequenceContainerBase>(OutPipeline);
-				OutPipeline->OutputContainers.Add(ImageSequenceOutput);
-
-				UMoviePipelineAudioOutput* AudioOutput = NewObject<UMoviePipelineAudioOutput>(OutPipeline);
-				OutPipeline->OutputContainers.Add(AudioOutput);
+				OutPipeline->FindOrAddSetting<UMoviePipelineImageSequenceContainerBase>();
 
 				// Add it to our list to be processed.
 				Configs.Add(OutPipeline);
diff --git a/Engine/Plugins/MovieScene/MovieRenderPipeline/Source/MovieRenderPipelineEditor/Private/Widgets/SMoviePipelineEditor.cpp b/Engine/Plugins/MovieScene/MovieRenderPipeline/Source/MovieRenderPipelineEditor/Private/Widgets/SMoviePipelineEditor.cpp
--- a/Engine/Plugins/MovieScene/MovieRenderPipeline/Source/MovieRenderPipelineEditor/Private/Widgets/SMoviePipelineEditor.cpp
+++ b/Engine/Plugins/MovieScene/MovieRenderPipeline/Source/MovieRenderPipelineEditor/Private/Widgets/SMoviePipelineEditor.cpp
@@ -44,7 +44,7 @@
 #define LOCTEXT_NAMESPACE "SMoviePipelineEditor"
 
 
-TArray<UClass*> FindRecordingSourceClasses()
+TArray<UClass*> FindMoviePipelineSettingClasses()
 {
 	FAssetRegistryModule& AssetRegistryModule = FModuleManager::LoadModuleChecked<FAssetRegistryModule>(TEXT("AssetRegistry"));
 
@@ -168,7 +168,7 @@ PRAGMA_ENABLE_OPTIMIZATION
 
 void SMoviePipelineEditor::CheckForNewMoviePipeline()
 {
-	UMoviePipelineShotConfig* NewMoviePipeline = MoviePipelineAttribute.Get();
+	UMoviePipelineConfigBase* NewMoviePipeline = MoviePipelineAttribute.Get();
 	if (CachedMoviePipeline != NewMoviePipeline)
 	{
 		CachedMoviePipeline = NewMoviePipeline;
@@ -206,7 +206,7 @@ TSharedRef<SWidget> SMoviePipelineEditor::OnGenerateSettingsMenu()
 
 	MenuBuilder.BeginSection("Settings", LOCTEXT("SettingsMenuSection", "Available Settings"));
 	{
-		TArray<UClass*> SourceClasses = FindRecordingSourceClasses();
+		TArray<UClass*> SourceClasses = FindMoviePipelineSettingClasses();
 		Algo::SortBy(SourceClasses, &UClass::GetDisplayNameText, FText::FSortPredicate());
 
 		for (UClass* Class : SourceClasses)
@@ -281,7 +281,7 @@ void SMoviePipelineEditor::UpdateDetails()
 
 void SMoviePipelineEditor::AddSettingFromClass(TSubclassOf<UMoviePipelineSetting> SettingClass)
 {
-	UMoviePipelineShotConfig* Pipeline = MoviePipelineAttribute.Get();
+	UMoviePipelineConfigBase* Pipeline = MoviePipelineAttribute.Get();
 
 	if (*SettingClass && Pipeline)
 	{
diff --git a/Engine/Plugins/MovieScene/MovieRenderPipeline/Source/MovieRenderPipelineEditor/Private/Widgets/SMoviePipelineEditor.h b/Engine/Plugins/MovieScene/MovieRenderPipeline/Source/MovieRenderPipelineEditor/Private/Widgets/SMoviePipelineEditor.h
--- a/Engine/Plugins/MovieScene/MovieRenderPipeline/Source/MovieRenderPipelineEditor/Private/Widgets/SMoviePipelineEditor.h
+++ b/Engine/Plugins/MovieScene/MovieRenderPipeline/Source/MovieRenderPipelineEditor/Private/Widgets/SMoviePipelineEditor.h
@@ -12,7 +12,7 @@ struct IMoviePipelineSettingTreeItem;
 
 class SScrollBox;
 class SMoviePipelineSettings;
-class UMoviePipelineShotConfig;
+class UMoviePipelineConfigBase;
 class IDetailsView;
 class UMoviePipelineSetting;
 
@@ -29,7 +29,7 @@ public:
 		: _MoviePipeline(nullptr)
 		{}
 
-		SLATE_ATTRIBUTE(UMoviePipelineShotConfig*, MoviePipeline)
+		SLATE_ATTRIBUTE(UMoviePipelineConfigBase*, MoviePipeline)
 
 	SLATE_END_ARGS()
 
@@ -68,8 +68,8 @@ private:
 private:
 
 	bool bRequestDetailsRefresh;
-	TAttribute<UMoviePipelineShotConfig*> MoviePipelineAttribute;
-	TWeakObjectPtr<UMoviePipelineShotConfig> CachedMoviePipeline;
+	TAttribute<UMoviePipelineConfigBase*> MoviePipelineAttribute;
+	TWeakObjectPtr<UMoviePipelineConfigBase> CachedMoviePipeline;
     
 	TSharedPtr<SMoviePipelineSettings> SettingsWidget;
 	TSharedPtr<SScrollBox> DetailsBox;
diff --git a/Engine/Plugins/MovieScene/MovieRenderPipeline/Source/MovieRenderPipelineEditor/Private/Widgets/SMoviePipelinePanel.cpp b/Engine/Plugins/MovieScene/MovieRenderPipeline/Source/MovieRenderPipelineEditor/Private/Widgets/SMoviePipelinePanel.cpp
--- a/Engine/Plugins/MovieScene/MovieRenderPipeline/Source/MovieRenderPipelineEditor/Private/Widgets/SMoviePipelinePanel.cpp
+++ b/Engine/Plugins/MovieScene/MovieRenderPipeline/Source/MovieRenderPipelineEditor/Private/Widgets/SMoviePipelinePanel.cpp
@@ -94,7 +94,7 @@ void SMoviePipelinePanel::Construct(const FArguments& InArgs)
 						.AutoWidth()
 						[
 							SNew(SImage)
-							.Image(FSlateIconFinder::FindIconBrushForClass(UMoviePipelineShotConfig::StaticClass()))
+							.Image(FSlateIconFinder::FindIconBrushForClass(UMoviePipelineConfigBase::StaticClass()))
 						]
 
 						+ SHorizontalBox::Slot()
@@ -183,12 +183,12 @@ void SMoviePipelinePanel::Construct(const FArguments& InArgs)
 
 PRAGMA_ENABLE_OPTIMIZATION
 
-UMoviePipelineShotConfig* SMoviePipelinePanel::AllocateTransientPreset()
+UMoviePipelineConfigBase* SMoviePipelinePanel::AllocateTransientPreset()
 {
 	static const TCHAR* PackageName = TEXT("/Temp/MoviePipeline/PendingPipeline");
 
 	// Return a cached transient if it exists
-	UMoviePipelineShotConfig* ExistingPreset = FindObject<UMoviePipelineShotConfig>(nullptr, TEXT("/Temp/MoviePipeline/PendingPipeline.PendingPipeline"));
+	UMoviePipelineConfigBase* ExistingPreset = FindObject<UMoviePipelineConfigBase>(nullptr, TEXT("/Temp/MoviePipeline/PendingPipeline.PendingPipeline"));
 	if (ExistingPreset)
 	{
 		return ExistingPreset;
@@ -200,7 +200,7 @@ UMoviePipelineShotConfig* SMoviePipelinePanel::AllocateTransientPreset()
 	NewPackage->SetFlags(RF_Transient);
 	NewPackage->AddToRoot();
 
-	UMoviePipelineShotConfig* NewPreset = NewObject<UMoviePipelineShotConfig>(NewPackage, DesiredName, RF_Transient | RF_Transactional | RF_Standalone);
+	UMoviePipelineConfigBase* NewPreset = NewObject<UMoviePipelineConfigBase>(NewPackage, DesiredName, RF_Transient | RF_Transactional | RF_Standalone);
 
 	return NewPreset;
 }
@@ -235,7 +235,7 @@ TSharedRef<SWidget> SMoviePipelinePanel::OnGeneratePresetsMenu()
 		AssetPickerConfig.bSortByPathInColumnView = false;
 
 		AssetPickerConfig.AssetShowWarningText = LOCTEXT("NoPresets_Warning", "No Presets Found");
-		AssetPickerConfig.Filter.ClassNames.Add(UMoviePipelineShotConfig::StaticClass()->GetFName());
+		AssetPickerConfig.Filter.ClassNames.Add(UMoviePipelineConfigBase::StaticClass()->GetFName());
 		// AssetPickerConfig.OnAssetSelected = FOnAssetSelected::CreateSP(this, &SMoviePipelinePanel::OnImportPreset);
 	}
 
@@ -280,7 +280,7 @@ void SMoviePipelinePanel::AddReferencedObjects(FReferenceCollector& Collector)
 	// Collector.AddReferencedObject(RecordingLevelSequence);
 }
 
-UMoviePipelineShotConfig* SMoviePipelinePanel::GetMoviePipeline() const
+UMoviePipelineConfigBase* SMoviePipelinePanel::GetMoviePipeline() const
 {
 	return TransientPreset;
 }
diff --git a/Engine/Plugins/MovieScene/MovieRenderPipeline/Source/MovieRenderPipelineEditor/Private/Widgets/SMoviePipelinePanel.h b/Engine/Plugins/MovieScene/MovieRenderPipeline/Source/MovieRenderPipelineEditor/Private/Widgets/SMoviePipelinePanel.h
--- a/Engine/Plugins/MovieScene/MovieRenderPipeline/Source/MovieRenderPipelineEditor/Private/Widgets/SMoviePipelinePanel.h
+++ b/Engine/Plugins/MovieScene/MovieRenderPipeline/Source/MovieRenderPipelineEditor/Private/Widgets/SMoviePipelinePanel.h
@@ -6,7 +6,7 @@
 #include "Widgets/SCompoundWidget.h"
 #include "UObject/GCObject.h"
 
-class UMoviePipelineShotConfig;
+class UMoviePipelineConfigBase;
 class SMoviePipelineEditor;
 
 /**
@@ -26,12 +26,12 @@ public:
 		/*~ All following arguments are mutually-exclusive */
 		/*-------------------------------------------------*/
 		/** A preset asset to base the pipeline off */
-		SLATE_ARGUMENT(UMoviePipelineShotConfig*, BasePreset)
+		SLATE_ARGUMENT(UMoviePipelineConfigBase*, BasePreset)
 
 	SLATE_END_ARGS()
 
 	void Construct(const FArguments& InArgs);
-	UMoviePipelineShotConfig* GetMoviePipeline() const;
+	UMoviePipelineConfigBase* GetMoviePipeline() const;
 
 private:
 	virtual void AddReferencedObjects(FReferenceCollector& Collector) override;
@@ -44,12 +44,12 @@ private:
 	FReply OnRevertChanges();
 
 	/** Allocates a transient preset so that the user can use the pipeline without saving it to an asset first. */
-	UMoviePipelineShotConfig* AllocateTransientPreset();
+	UMoviePipelineConfigBase* AllocateTransientPreset();
 
 
 private:
 	/** The transient preset that we use - kept alive by AddReferencedObjects */
-	UMoviePipelineShotConfig* TransientPreset;
+	UMoviePipelineConfigBase* TransientPreset;
 
 	/** The main movie pipeline editor widget */
 	TSharedPtr<SMoviePipelineEditor> MoviePipelineEditorWidget;
diff --git a/Engine/Plugins/MovieScene/MovieRenderPipeline/Source/MovieRenderPipelineEditor/Private/Widgets/SMoviePipelineSettings.cpp b/Engine/Plugins/MovieScene/MovieRenderPipeline/Source/MovieRenderPipelineEditor/Private/Widgets/SMoviePipelineSettings.cpp
--- a/Engine/Plugins/MovieScene/MovieRenderPipeline/Source/MovieRenderPipelineEditor/Private/Widgets/SMoviePipelineSettings.cpp
+++ b/Engine/Plugins/MovieScene/MovieRenderPipeline/Source/MovieRenderPipelineEditor/Private/Widgets/SMoviePipelineSettings.cpp
@@ -38,7 +38,7 @@ struct IMoviePipelineSettingTreeItem : TSharedFromThis<IMoviePipelineSettingTree
 {
 	virtual ~IMoviePipelineSettingTreeItem() {}
 
-	virtual void Delete(UMoviePipelineShotConfig* Owner) = 0;
+	virtual void Delete(UMoviePipelineConfigBase* Owner) = 0;
 
 	virtual TSharedPtr<FMoviePipelineSettingCategory> AsCategory() { return nullptr; }
 	virtual TSharedPtr<FMoviePipelineSettingTreeItem> AsSetting()   { return nullptr; }
@@ -61,7 +61,7 @@ struct FMoviePipelineSettingTreeItem : IMoviePipelineSettingTreeItem
 		return Setting ? Setting->GetDisplayText() : FText();
 	}
 
-	virtual void Delete(UMoviePipelineShotConfig* Owner) override
+	virtual void Delete(UMoviePipelineConfigBase* Owner) override
 	{
 		UMoviePipelineSetting* Setting = WeakSetting.Get();
 		if (Setting)
@@ -251,7 +251,7 @@ struct FMoviePipelineSettingCategory : IMoviePipelineSettingTreeItem
 		: Category(FText::FromString(InCategory))
 	{}
 
-	virtual void Delete(UMoviePipelineShotConfig* Owner) override
+	virtual void Delete(UMoviePipelineConfigBase* Owner) override
 	{
 		for (TSharedPtr<FMoviePipelineSettingTreeItem> Child : Children)
 		{
@@ -313,7 +313,7 @@ void SMoviePipelineSettings::Construct(const FArguments& InArgs)
 
 void SMoviePipelineSettings::Tick(const FGeometry& AllottedGeometry, const double InCurrentTime, const float InDeltaTime)
 {
-	UMoviePipelineShotConfig* ShotConfig = WeakShotConfig.Get();
+	UMoviePipelineConfigBase* ShotConfig = WeakShotConfig.Get();
 
 	// If we have a settings ptr, we expect its serial number to match our cached one, if not, we rebuild the tree
 	if (ShotConfig)
@@ -356,7 +356,7 @@ void SMoviePipelineSettings::GetSelectedSettings(TArray<UMoviePipelineSetting*>&
 	}
 }
 
-void SMoviePipelineSettings::SetShotConfigObject(UMoviePipelineShotConfig* InShotConfig)
+void SMoviePipelineSettings::SetShotConfigObject(UMoviePipelineConfigBase* InShotConfig)
 {
 	WeakShotConfig = InShotConfig;
 	ReconstructTree();
@@ -364,7 +364,7 @@ void SMoviePipelineSettings::SetShotConfigObject(UMoviePipelineShotConfig* InSho
 
 void SMoviePipelineSettings::ReconstructTree()
 {
-	UMoviePipelineShotConfig* ShotConfig = WeakShotConfig.Get();
+	UMoviePipelineConfigBase* ShotConfig = WeakShotConfig.Get();
 	if (!ShotConfig)
 	{
 		CachedSettingsSerialNumber = uint32(-1);
@@ -501,7 +501,7 @@ void SMoviePipelineSettings::OnGetChildren(TSharedPtr<IMoviePipelineSettingTreeI
 
 void SMoviePipelineSettings::OnDeleteSelected()
 {
-	UMoviePipelineShotConfig* ShotConfig = WeakShotConfig.Get();
+	UMoviePipelineConfigBase* ShotConfig = WeakShotConfig.Get();
 	if (ShotConfig)
 	{
 		TArray<TSharedPtr<IMoviePipelineSettingTreeItem>> Items = TreeView->GetSelectedItems();
diff --git a/Engine/Plugins/MovieScene/MovieRenderPipeline/Source/MovieRenderPipelineEditor/Private/Widgets/SMoviePipelineSettings.h b/Engine/Plugins/MovieScene/MovieRenderPipeline/Source/MovieRenderPipelineEditor/Private/Widgets/SMoviePipelineSettings.h
--- a/Engine/Plugins/MovieScene/MovieRenderPipeline/Source/MovieRenderPipelineEditor/Private/Widgets/SMoviePipelineSettings.h
+++ b/Engine/Plugins/MovieScene/MovieRenderPipeline/Source/MovieRenderPipelineEditor/Private/Widgets/SMoviePipelineSettings.h
@@ -13,7 +13,7 @@ class ITableRow;
 class FUICommandList;
 class STableViewBase;
 class UMoviePipelineSetting;
-class UMoviePipelineShotConfig;
+class UMoviePipelineConfigBase;
 struct IMoviePipelineSettingTreeItem;
 struct FMoviePipelineSettingTreeItem;
 template<typename> class STreeView;
@@ -35,7 +35,7 @@ public:
 
 	void Construct(const FArguments& InArgs);
 
-	void SetShotConfigObject(UMoviePipelineShotConfig* InShotConfig);
+	void SetShotConfigObject(UMoviePipelineConfigBase* InShotConfig);
 
 	void GetSelectedSettings(TArray<UMoviePipelineSetting*>& OutSettings) const;
 
@@ -57,7 +57,7 @@ private:
 private:
 
 	uint32 CachedSettingsSerialNumber;
-	TWeakObjectPtr<UMoviePipelineShotConfig> WeakShotConfig;
+	TWeakObjectPtr<UMoviePipelineConfigBase> WeakShotConfig;
 
 	TArray<TSharedPtr<IMoviePipelineSettingTreeItem>> RootNodes;
 	TMap<FObjectKey, TSharedPtr<FMoviePipelineSettingTreeItem>> SettingToTreeItem;
diff --git a/Engine/Plugins/MovieScene/MovieRenderPipeline/Source/MovieRenderPipelineEditor/Private/Widgets/SMoviePipelineTabContent.cpp b/Engine/Plugins/MovieScene/MovieRenderPipeline/Source/MovieRenderPipelineEditor/Private/Widgets/SMoviePipelineTabContent.cpp
--- a/Engine/Plugins/MovieScene/MovieRenderPipeline/Source/MovieRenderPipelineEditor/Private/Widgets/SMoviePipelineTabContent.cpp
+++ b/Engine/Plugins/MovieScene/MovieRenderPipeline/Source/MovieRenderPipelineEditor/Private/Widgets/SMoviePipelineTabContent.cpp
@@ -25,11 +25,11 @@ void SMoviePipelineTabContent::Construct(const FArguments& InArgs)
 
 EActiveTimerReturnType SMoviePipelineTabContent::OnActiveTimer(double InCurrentTime, float InDeltaTime)
 {
-	SetupForPipeline((UMoviePipelineShotConfig*)nullptr);
+	SetupForPipeline((UMoviePipelineConfigBase*)nullptr);
 	return EActiveTimerReturnType::Stop;
 }
 
-void SMoviePipelineTabContent::SetupForPipeline(UMoviePipelineShotConfig* BasePreset)
+void SMoviePipelineTabContent::SetupForPipeline(UMoviePipelineConfigBase* BasePreset)
 {
 	// Null out the tab content to ensure that all references have been cleaned up before constructing the new one
 	ChildSlot [ SNullWidget::NullWidget ];
diff --git a/Engine/Plugins/MovieScene/MovieRenderPipeline/Source/MovieRenderPipelineEditor/Private/Widgets/SMoviePipelineTabContent.h b/Engine/Plugins/MovieScene/MovieRenderPipeline/Source/MovieRenderPipelineEditor/Private/Widgets/SMoviePipelineTabContent.h
--- a/Engine/Plugins/MovieScene/MovieRenderPipeline/Source/MovieRenderPipelineEditor/Private/Widgets/SMoviePipelineTabContent.h
+++ b/Engine/Plugins/MovieScene/MovieRenderPipeline/Source/MovieRenderPipelineEditor/Private/Widgets/SMoviePipelineTabContent.h
@@ -5,7 +5,7 @@
 #include "Widgets/DeclarativeSyntaxSupport.h"
 #include "Widgets/SCompoundWidget.h"
 
-class UMoviePipelineShotConfig;
+class UMoviePipelineConfigBase;
 class SMoviePipelinePanel;
 
 class SMoviePipelineTabContent : public SCompoundWidget
@@ -18,7 +18,7 @@ public:
 	void Construct(const FArguments& InArgs);
 
 	/** Set this tab up for a specific Pipeline preset. Passing nullptr will initialize a default pipeline. */
-	void SetupForPipeline(UMoviePipelineShotConfig* BasePreset);
+	void SetupForPipeline(UMoviePipelineConfigBase* BasePreset);
 
 private:
 	EActiveTimerReturnType OnActiveTimer(double InCurrentTime, float InDeltaTime);
diff --git a/Engine/Plugins/MovieScene/MovieRenderPipeline/Source/MovieRenderPipelineEditor/Public/MoviePipelineEditorBlueprintLibrary.h b/Engine/Plugins/MovieScene/MovieRenderPipeline/Source/MovieRenderPipelineEditor/Public/MoviePipelineEditorBlueprintLibrary.h
--- a/Engine/Plugins/MovieScene/MovieRenderPipeline/Source/MovieRenderPipelineEditor/Public/MoviePipelineEditorBlueprintLibrary.h
+++ b/Engine/Plugins/MovieScene/MovieRenderPipeline/Source/MovieRenderPipelineEditor/Public/MoviePipelineEditorBlueprintLibrary.h
@@ -7,12 +7,12 @@
 #include "UObject/TextProperty.h"
 #include "Internationalization/Text.h"
 #include "Containers/UnrealString.h"
-#include "MovieRenderPipelineConfig.h"
+#include "MoviePipelineMasterConfig.h"
 
 #include "MoviePipelineEditorBlueprintLibrary.generated.h"
 
 // Forward Declare
-class UMovieRenderPipelineConfig;
+class UMoviePipelineMasterConfig;
 
 UCLASS()
 class MOVIERENDERPIPELINEEDITOR_API UMoviePipelineEditorBlueprintLibrary : public UBlueprintFunctionLibrary
@@ -21,5 +21,5 @@ class MOVIERENDERPIPELINEEDITOR_API UMoviePipelineEditorBlueprintLibrary : publi
 	
 public:
 	UFUNCTION(BlueprintCallable, Category = "Movie Render Pipeline")
-	static bool ExportConfigToAsset(const UMovieRenderPipelineConfig* InConfig, const FString& InPackagePath, const FString& InFileName, const bool bInSaveAsset, UMovieRenderPipelineConfig*& OutAsset, FText& OutErrorReason);
+	static bool ExportConfigToAsset(const UMoviePipelineMasterConfig* InConfig, const FString& InPackagePath, const FString& InFileName, const bool bInSaveAsset, UMoviePipelineMasterConfig*& OutAsset, FText& OutErrorReason);
 };
\ No newline at end of file
diff --git a/Engine/Plugins/MovieScene/MovieRenderPipeline/Source/MovieRenderPipelineEditor/Public/MoviePipelineNewProcessExecutor.h b/Engine/Plugins/MovieScene/MovieRenderPipeline/Source/MovieRenderPipelineEditor/Public/MoviePipelineNewProcessExecutor.h
--- a/Engine/Plugins/MovieScene/MovieRenderPipeline/Source/MovieRenderPipelineEditor/Public/MoviePipelineNewProcessExecutor.h
+++ b/Engine/Plugins/MovieScene/MovieRenderPipeline/Source/MovieRenderPipelineEditor/Public/MoviePipelineNewProcessExecutor.h
@@ -21,4 +21,16 @@ public:
 		: UMoviePipelineLinearExecutorBase()
 	{
 	}
+
+	/** If enabled the editor will close itself when a new process is started. This can be used to gain some performance. */
+	UPROPERTY(EditAnywhere, BlueprintReadWrite, Category = "Movie Pipeline")
+	bool bCloseEditor;
+
+	/** A list of additional command line arguments to be appended to the new process startup. Can be useful if your game requires certain arguments to start such as disabling log-in screens. */
+	UPROPERTY(EditAnywhere, BlueprintReadWrite, Category = "Movie Pipeline")
+	FString AdditionalCommandLineArguments;
+
+	/** A list of command line arguments which are inherited from the currently running Editor instance that will be automatically appended to the new process. */
+	UPROPERTY(VisibleAnywhere, BlueprintReadOnly, Category = "Movie Pipeline")
+	FString InheritedCommandLineArguments;
 };
\ No newline at end of file
diff --git a/Engine/Plugins/MovieScene/MovieRenderPipeline/Source/MovieRenderPipelineEditor/Public/MoviePipelinePIEExecutor.h b/Engine/Plugins/MovieScene/MovieRenderPipeline/Source/MovieRenderPipelineEditor/Public/MoviePipelinePIEExecutor.h
--- a/Engine/Plugins/MovieScene/MovieRenderPipeline/Source/MovieRenderPipelineEditor/Public/MoviePipelinePIEExecutor.h
+++ b/Engine/Plugins/MovieScene/MovieRenderPipeline/Source/MovieRenderPipelineEditor/Public/MoviePipelinePIEExecutor.h
@@ -24,7 +24,7 @@ public:
 	}
 
 protected:
-	virtual void Start(UMovieRenderPipelineConfig* InConfig, const int32 InConfigIndex, const int32 InNumConfigs) override;
+	virtual void Start(const FMoviePipelineExecutorJob& InJob) override;
 
 private:
 	/** Called when PIE finishes booting up and it is safe for us to spawn an object into that world. */
diff --git a/Engine/Plugins/MovieScene/MovieRenderPipeline/Source/MovieRenderPipelineRenderPasses/Private/MoviePipelineBackbufferPass.cpp b/Engine/Plugins/MovieScene/MovieRenderPipeline/Source/MovieRenderPipelineRenderPasses/Private/MoviePipelineBackbufferPass.cpp
--- a/Engine/Plugins/MovieScene/MovieRenderPipeline/Source/MovieRenderPipelineRenderPasses/Private/MoviePipelineBackbufferPass.cpp
+++ b/Engine/Plugins/MovieScene/MovieRenderPipeline/Source/MovieRenderPipelineRenderPasses/Private/MoviePipelineBackbufferPass.cpp
@@ -16,7 +16,6 @@
 #include "ImageWriteStream.h"
 #include "GameFramework/PlayerController.h"
 #include "Math/Vector.h"
-#include "MovieRenderPipelineConfig.h"
 #include "Misc/FrameRate.h"
 #include "MoviePipelineAccumulationSetting.h"
 #include "MoviePipelineRenderPass.h"
@@ -88,7 +87,7 @@ void UMoviePipelineBackbufferPass::CaptureFrameImpl(const FMoviePipelineRenderPa
 	
 	// Override the Motion Blur settings since these are controlled by the movie pipeline.
 	{
-		View->FinalPostProcessSettings.MotionBlurTargetFPS = FMath::RoundToInt(GetPipeline()->GetPipelineConfig()->GetEffectiveFrameRate().AsDecimal());
+		View->FinalPostProcessSettings.MotionBlurTargetFPS = FMath::RoundToInt(GetPipeline()->GetEffectiveFrameRate().AsDecimal());
 		View->FinalPostProcessSettings.MotionBlurAmount = InPassMetrics.OutputState.MotionBlurFraction;
 		View->FinalPostProcessSettings.MotionBlurMax = 100.f;
 		View->FinalPostProcessSettings.bOverride_MotionBlurAmount = true;
@@ -197,11 +196,6 @@ void UMoviePipelineBackbufferPass::CaptureFrameImpl(const FMoviePipelineRenderPa
 	);
 }
 
-void UMoviePipelineBackbufferPass::GetFrameDataImpl(MoviePipeline::FOutputFrameData& OutFrameData)
-{
-
-}
-
 FSceneView* UMoviePipelineBackbufferPass::CalcSceneView(FSceneViewFamily* ViewFamily, const FMoviePipelineRenderPassMetrics& InPassMetrics)
 {
 	FVector ViewLocation;
@@ -210,8 +204,6 @@ FSceneView* UMoviePipelineBackbufferPass::CalcSceneView(FSceneViewFamily* ViewFa
 	APlayerController* LocalPlayerController = GetWorld()->GetFirstPlayerController();
 	LocalPlayerController->GetPlayerViewPoint(ViewLocation, ViewRotation);
 
-	UE_LOG(LogTemp, Log, TEXT("%d,%s"), GFrameCounter, *ViewLocation.ToString());
-
 	int32 TileSizeX = InitSettings.TileResolution.X;
 	int32 TileSizeY = InitSettings.TileResolution.Y;
 
diff --git a/Engine/Plugins/MovieScene/MovieRenderPipeline/Source/MovieRenderPipelineRenderPasses/Private/MoviePipelineImageSequenceContainer.cpp b/Engine/Plugins/MovieScene/MovieRenderPipeline/Source/MovieRenderPipelineRenderPasses/Private/MoviePipelineImageSequenceContainer.cpp
--- a/Engine/Plugins/MovieScene/MovieRenderPipeline/Source/MovieRenderPipelineRenderPasses/Private/MoviePipelineImageSequenceContainer.cpp
+++ b/Engine/Plugins/MovieScene/MovieRenderPipeline/Source/MovieRenderPipelineRenderPasses/Private/MoviePipelineImageSequenceContainer.cpp
@@ -7,11 +7,12 @@
 #include "ImageWriteQueue.h"
 #include "MoviePipeline.h"
 #include "ImageWriteStream.h"
-#include "MovieRenderPipelineConfig.h"
+#include "MoviePipelineMasterConfig.h"
 #include "MovieRenderTileImage.h"
 #include "MovieRenderOverlappedImage.h"
 #include "MovieRenderPipelineCoreModule.h"
 #include "Misc/FrameRate.h"
+#include "MoviePipelineOutputSetting.h"
 
 
 UMoviePipelineImageSequenceContainerBase::UMoviePipelineImageSequenceContainerBase()
@@ -34,9 +35,12 @@ bool UMoviePipelineImageSequenceContainerBase::HasFinishedProcessingImpl()
 static FImageTileAccumulator TestAccumulator;
 static FImageOverlappedAccumulator TestOverlappedAccumulator;
 
-void UMoviePipelineImageSequenceContainerBase::OnInitializedForPipelineImpl(UMoviePipeline* InPipeline)
+void UMoviePipelineImageSequenceContainerBase::SetupForPipelineImpl(UMoviePipeline* InPipeline)
 {
-	FString OutputDirectory = GetPipeline()->GetPipelineConfig()->OutputDirectory.Path;
+	UMoviePipelineOutputSetting* OutputSettings = GetPipeline()->GetPipelineMasterConfig()->FindSetting<UMoviePipelineOutputSetting>();
+	check(OutputSettings);
+
+	FString OutputDirectory = OutputSettings->OutputDirectory.Path;
 	IImageWriteQueue* WriteQueue = ImageWriteQueue;
 
 	FImageTileAccumulator* TileAccumulator = &TestAccumulator;
@@ -175,37 +179,4 @@ void UMoviePipelineImageSequenceContainerBase::OnInitializedForPipelineImpl(UMov
 	};
 
 	InPipeline->GetOutputPipe()->AddEndpoint(OnImageReceived);
-}
-
-void UMoviePipelineImageSequenceContainerBase::ProcessFrameImpl(TArray<MoviePipeline::FOutputFrameData> FrameData, FMoviePipelineFrameOutputState CachedOutputState, FDirectoryPath OutputDirectory)
-{
-	if (FrameData.Num() == 0)
-	{
-		return;
-	}
-
-	for (MoviePipeline::FOutputFrameData& Frame : FrameData)
-	{
-		if (Frame.Resolution.X <= 0 || Frame.Resolution.Y <= 0)
-		{
-			continue;
-		}
-
-		TUniquePtr<FImageWriteTask> ImageTask = MakeUnique<FImageWriteTask>();
-
-		// Move the color buffer into a raw image data container that we can pass to the write queue
-		ImageTask->PixelData = MakeUnique<TImagePixelData<FColor>>(Frame.Resolution, MoveTemp(Frame.ColorBuffer));
-	
-		// Always write full alpha for PNGs
-		ImageTask->PixelPreProcessors.Add(TAsyncAlphaWrite<FColor>(255));
-		ImageTask->Format = EImageFormat::PNG;
-	
-		// ImageTask->CompressionQuality = GetCompressionQuality();
-	
-		FString OutputName = FString::Printf(TEXT("/Frame_%s_%d.png"), *Frame.PassName, CachedOutputState.OutputFrameNumber);
-		FString OutputPath = OutputDirectory.Path + OutputName;
-		ImageTask->Filename = OutputPath;
-
-		ImageWriteQueue->Enqueue(MoveTemp(ImageTask));
-	}
-}
+}
\ No newline at end of file
diff --git a/Engine/Plugins/MovieScene/MovieRenderPipeline/Source/MovieRenderPipelineRenderPasses/Public/MoviePipelineBackbufferPass.h b/Engine/Plugins/MovieScene/MovieRenderPipeline/Source/MovieRenderPipelineRenderPasses/Public/MoviePipelineBackbufferPass.h
--- a/Engine/Plugins/MovieScene/MovieRenderPipeline/Source/MovieRenderPipelineRenderPasses/Public/MoviePipelineBackbufferPass.h
+++ b/Engine/Plugins/MovieScene/MovieRenderPipeline/Source/MovieRenderPipelineRenderPasses/Public/MoviePipelineBackbufferPass.h
@@ -22,9 +22,8 @@ class MOVIERENDERPIPELINERENDERPASSES_API UMoviePipelineBackbufferPass : public
 protected:
 
 	virtual void SetupImpl(const FMoviePipelineRenderPassInitSettings& InInitSettings, TSharedRef<FImagePixelPipe, ESPMode::ThreadSafe> InOutputPipe) override;
-	virtual void CaptureFrameImpl(const FMoviePipelineRenderPassMetrics& OutputFrameMetrics) override;
-	virtual void GetFrameDataImpl(MoviePipeline::FOutputFrameData& OutFrameData) override;
 	virtual void TeardownImpl() override;
+	virtual void CaptureFrameImpl(const FMoviePipelineRenderPassMetrics& OutputFrameMetrics) override;
 	virtual void GatherOutputPassesImpl(TArray<FMoviePipelinePassIdentifier>& ExpectedRenderPasses) override;
 
 	FSceneView* CalcSceneView(FSceneViewFamily* ViewFamily, const FMoviePipelineRenderPassMetrics& InPassMetrics);
diff --git a/Engine/Plugins/MovieScene/MovieRenderPipeline/Source/MovieRenderPipelineRenderPasses/Public/MoviePipelineImageSequenceContainer.h b/Engine/Plugins/MovieScene/MovieRenderPipeline/Source/MovieRenderPipelineRenderPasses/Public/MoviePipelineImageSequenceContainer.h
--- a/Engine/Plugins/MovieScene/MovieRenderPipeline/Source/MovieRenderPipelineRenderPasses/Public/MoviePipelineImageSequenceContainer.h
+++ b/Engine/Plugins/MovieScene/MovieRenderPipeline/Source/MovieRenderPipelineRenderPasses/Public/MoviePipelineImageSequenceContainer.h
@@ -1,7 +1,7 @@
 // Copyright 1998-2019 Epic Games, Inc. All Rights Reserved.
 #pragma once
 
-#include "MoviePipelineOutput.h"
+#include "MoviePipelineOutputBase.h"
 #include "Async/Future.h"
 #include "MoviePipelineImageSequenceContainer.generated.h"
 
@@ -9,7 +9,7 @@
 class IImageWriteQueue;
 
 UCLASS(Blueprintable)
-class MOVIERENDERPIPELINERENDERPASSES_API UMoviePipelineImageSequenceContainerBase : public UMoviePipelineOutput
+class MOVIERENDERPIPELINERENDERPASSES_API UMoviePipelineImageSequenceContainerBase : public UMoviePipelineOutputBase
 {
 	GENERATED_BODY()
 public:
@@ -17,8 +17,7 @@ public:
 
 protected:
 	// UMovieRenderPipelineOutputContainer interface
-	virtual void OnInitializedForPipelineImpl(UMoviePipeline* InPipeline);
-	virtual void ProcessFrameImpl(TArray<MoviePipeline::FOutputFrameData> FrameData, FMoviePipelineFrameOutputState CachedOutputState, FDirectoryPath OutputDirectory) override;
+	virtual void SetupForPipelineImpl(UMoviePipeline* InPipeline);
 	virtual void BeginFinalizeImpl() override;
 	virtual bool HasFinishedProcessingImpl() override;
 	// ~UMovieRenderPipelineOutputContainer interface
diff --git a/Engine/Plugins/MovieScene/MovieRenderPipeline/Source/MovieRenderPipelineSettings/Private/MoviePipelineAudioOutput.cpp b/Engine/Plugins/MovieScene/MovieRenderPipeline/Source/MovieRenderPipelineSettings/Private/MoviePipelineAudioOutput.cpp
deleted file mode 100644
diff --git a/Engine/Plugins/MovieScene/MovieRenderPipeline/Source/MovieRenderPipelineSettings/Public/MoviePipelineAudioOutput.h b/Engine/Plugins/MovieScene/MovieRenderPipeline/Source/MovieRenderPipelineSettings/Public/MoviePipelineAudioOutput.h
deleted file mode 100644
diff --git a/Engine/Plugins/MovieScene/MovieRenderPipeline/Source/MovieRenderPipelineSettings/Public/MoviePipelineSeparateProcessSetting.h b/Engine/Plugins/MovieScene/MovieRenderPipeline/Source/MovieRenderPipelineSettings/Public/MoviePipelineSeparateProcessSetting.h
deleted file mode 100644
