CIS: Removed errant line continuation at the end of a comment.

#rnx

diff --git a/Engine/Plugins/MovieScene/MovieRenderPipeline/Source/MovieRenderPipelineCore/Public/MoviePipelineBlueprintLibrary.h b/Engine/Plugins/MovieScene/MovieRenderPipeline/Source/MovieRenderPipelineCore/Public/MoviePipelineBlueprintLibrary.h
--- a/Engine/Plugins/MovieScene/MovieRenderPipeline/Source/MovieRenderPipelineCore/Public/MoviePipelineBlueprintLibrary.h
+++ b/Engine/Plugins/MovieScene/MovieRenderPipeline/Source/MovieRenderPipelineCore/Public/MoviePipelineBlueprintLibrary.h
@@ -73,7 +73,7 @@ public:
 	// 	FMoviePipelineCameraCutInfo GetCurrentCameraCutSnapshot() const;
 	// 
 	// UFUNCTION(BlueprintPure, Category = "Movie Render Pipeline")
-	// 	FMoviePipelineFrameOutputState GetOutputStateSnapshot() const;\
+	// 	FMoviePipelineFrameOutputState GetOutputStateSnapshot() const;
 
 	/**
 	* Returns an estimate based on the average time taken to render all previous frames.
