MoviePipeline: Cleaned up the render passes to support an accumulator per image type, and synchronization of all data before being passed onto the output stage.

#rnx

diff --git a/Engine/Plugins/MovieScene/MovieRenderPipeline/Source/MovieRenderPipelineCore/Private/MoviePipeline.cpp b/Engine/Plugins/MovieScene/MovieRenderPipeline/Source/MovieRenderPipelineCore/Private/MoviePipeline.cpp
--- a/Engine/Plugins/MovieScene/MovieRenderPipeline/Source/MovieRenderPipelineCore/Private/MoviePipeline.cpp
+++ b/Engine/Plugins/MovieScene/MovieRenderPipeline/Source/MovieRenderPipelineCore/Private/MoviePipeline.cpp
@@ -30,6 +30,7 @@
 #include "MoviePipelineMasterConfig.h"
 #include "MoviePipelineOutputSetting.h"
 #include "MoviePipelineBlueprintLibrary.h"
+#include "ImageWriteQueue.h"
 
 #define LOCTEXT_NAMESPACE "MoviePipeline"
 
@@ -56,8 +57,9 @@ UMoviePipeline::UMoviePipeline()
 {
 	CustomTimeStep = CreateDefaultSubobject<UMoviePipelineCustomTimeStep>("MoviePipelineCustomTimeStep");
 	CustomSequenceTimeController = MakeShared<FMoviePipelineTimeController>();
-	OutputBuilder = MakeShared<FMoviePipelineOutputMerger, ESPMode::ThreadSafe>();
+	OutputBuilder = MakeShared<FMoviePipelineOutputMerger, ESPMode::ThreadSafe>(this);
 
+	ImageWriteQueue = &FModuleManager::Get().LoadModuleChecked<IImageWriteQueueModule>("ImageWriteQueue").GetWriteQueue();
 
 	// Temp
 	OutputPipe = MakeShared<FImagePixelPipe, ESPMode::ThreadSafe>();
@@ -653,16 +655,6 @@ TArray<FMoviePipelineShotInfo> UMoviePipeline::BuildShotListFromSequence(const U
 			// We want to start rendering from the first handle frame. Shutter Timing is a fixed offset from this number.
 			CameraCut.CurrentTick = CameraCut.TotalOutputRange.GetLowerBoundValue();
 		}
-
-		// We duplicate the Passes in each shot to be unique. This allows the passes to have state
-		// that persists while working on the next shot (in the event of slow-running tasks).
-		for (UMoviePipelineRenderPass* RenderPass : Shot.ShotConfig->GetRenderPasses())
-		{
-			FObjectDuplicationParameters DupParms(RenderPass, this);
-
-			UMoviePipelineRenderPass* PassCopy = (UMoviePipelineRenderPass*)StaticDuplicateObjectEx(DupParms);
-			Shot.RenderPasses.Add(PassCopy);
-		}
 	}
 
 	return NewShotList;
@@ -713,22 +705,9 @@ void UMoviePipeline::InitializeShot(FMoviePipelineShotInfo& InShot)
 	// Initialize any settings specific to this shot.
 	for (UMoviePipelineSetting* Setting : InShot.ShotConfig->GetSettings())
 	{
-		// Skip any render passes. These have been manually duplicated to avoid shared state
-		// and thus have to be initialized separately.
-		if (Setting->IsA(UMoviePipelineRenderPass::StaticClass()))
-		{
-			continue;
-		}
-
 		Setting->SetupForPipeline(this);
 	}
 
-	for (UMoviePipelineRenderPass* RenderPass : InShot.RenderPasses)
-	{
-		RenderPass->SetupForPipeline(this);
-	}
-
-
 	// Setup required rendering architecture for all passes in this shot.
 	SetupRenderingPipelineForShot(InShot);
 }
diff --git a/Engine/Plugins/MovieScene/MovieRenderPipeline/Source/MovieRenderPipelineCore/Private/MoviePipelineOutputBuilder.cpp b/Engine/Plugins/MovieScene/MovieRenderPipeline/Source/MovieRenderPipelineCore/Private/MoviePipelineOutputBuilder.cpp
--- a/Engine/Plugins/MovieScene/MovieRenderPipeline/Source/MovieRenderPipelineCore/Private/MoviePipelineOutputBuilder.cpp
+++ b/Engine/Plugins/MovieScene/MovieRenderPipeline/Source/MovieRenderPipelineCore/Private/MoviePipelineOutputBuilder.cpp
@@ -1,6 +1,8 @@
 // Copyright 1998-2019 Epic Games, Inc. All Rights Reserved.
 
 #include "MoviePipelineOutputBuilder.h"
+#include "MoviePipeline.h"
+#include "Async/Async.h"
 
 
 // When the Movie Pipeline Render goes to produce a frame, it will push an expected output frame into
@@ -16,27 +18,78 @@ FMoviePipelineMergerOutputFrame& FMoviePipelineOutputMerger::QueueOutputFrame_Ga
 	// Lock the ActiveData while we get a new output frame.
 	FScopeLock ScopeLock(&ActiveDataMutex);
 
-	FMoviePipelineMergerOutputFrame& NewFrame = ActiveData.FindOrAdd(CachedOutputState);
-	NewFrame.FrameOutput = CachedOutputState;
+	// Ensure this frame hasn't already been entered somehow.
+	check(!ActiveData.Find(CachedOutputState));
+
+	FMoviePipelineMergerOutputFrame& NewFrame = ActiveData.Add(CachedOutputState);
+	NewFrame.FrameOutputState = CachedOutputState;
 
 	return NewFrame;
 }
 
-void FMoviePipelineOutputMerger::OnDataAvailable_AnyThread(TUniquePtr<FImagePixelData> InData)
+void FMoviePipelineOutputMerger::OnSingleSampleDataAvailable_AnyThread(TUniquePtr<FImagePixelData>&& InData)
+{
+	// This is to support outputting individual samples (skipping accumulation) for debug reasons,
+	// or because you want to post-process them yourself. We just forward this directly on for output to disk.
+
+	TWeakObjectPtr<UMoviePipeline> LocalWeakPipeline = WeakMoviePipeline;
+
+	AsyncTask(ENamedThreads::GameThread, [LocalData = MoveTemp(InData), LocalWeakPipeline]() mutable
+	{
+		if (ensureAlwaysMsgf(LocalWeakPipeline.IsValid(), TEXT("A memory lifespan issue has left an output builder alive without an owning Movie Pipeline.")))
+		{
+			LocalWeakPipeline->OnSampleRendered(MoveTemp(LocalData));
+		}
+	}
+	);
+}
+
+void FMoviePipelineOutputMerger::OnCompleteRenderPassDataAvailable_AnyThread(TUniquePtr<FImagePixelData>&& InData)
 {
 	// Lock the ActiveData when we're updating what data has been gathered.
 	FScopeLock ScopeLock(&ActiveDataMutex);
 
-	
-}
+	FImagePixelDataPayload* Payload = InData->GetPayload<FImagePixelDataPayload>();
+	check(Payload);
+
+	// See if we can find the frame this data is for. This should always be valid, if it's not
+	// valid it means they either forgot to declare they were going to produce it, or this is
+	// coming in after the system already thinks it's finished that frame.
+	FMoviePipelineMergerOutputFrame* OutputFrame = ActiveData.Find(Payload->OutputState);
+	if (!ensureAlwaysMsgf(OutputFrame, TEXT("Recieved data for unknown frame. Frame was either already processed or not queued yet!")))
+	{
+		return;
+	}
+
+	// Ensure this pass is expected as well...
+	if (!ensureAlwaysMsgf(OutputFrame->ExpectedRenderPasses.Contains(Payload->PassIdentifier), TEXT("Recieved data for unexpected render pass: %s"), *Payload->PassIdentifier.Name))
+	{
+		return;
+	}
+
+	// If this data was expected and this frame is still in progress, pass the data to the frame.
+	OutputFrame->ImageOutputData.FindOrAdd(Payload->PassIdentifier) = MoveTemp(InData);
+
+	// Check to see if this was the last piece of data needed for this frame.
+	int32 TotalPasses = OutputFrame->ExpectedRenderPasses.Num();
+	int32 SucceededPasses = OutputFrame->ImageOutputData.Num();
+
+	if (SucceededPasses == TotalPasses)
+	{
+		// Transfer ownership from the map to here;
+		FMoviePipelineMergerOutputFrame FinalFrame;
+		ActiveData.RemoveAndCopyValue(Payload->OutputState, FinalFrame);
+
+		// Notify the Movie Pipeline that this frame has been completed and can be forwarded to the output containers.
+		TWeakObjectPtr<UMoviePipeline> LocalWeakPipeline = WeakMoviePipeline;
 
-// 
-// void OnDataAvailable(ImageData, Frame)
-// {
-// 	QueuedData[FrameNumber][Pass].Ready();
-// 
-// 	if (QueuedData[FrameNumber].AllReady())
-// 	{
-// 
-// 	}
-// }
\ No newline at end of file
+		AsyncTask(ENamedThreads::GameThread, [LocalFinalFrame = MoveTemp(FinalFrame), LocalWeakPipeline]() mutable
+		{
+			if (ensureAlwaysMsgf(LocalWeakPipeline.IsValid(), TEXT("A memory lifespan issue has left an output builder alive without an owning Movie Pipeline.")))
+			{
+				LocalWeakPipeline->OnFrameCompletelyRendered(MoveTemp(LocalFinalFrame));
+			}
+		}
+		);
+	}
+}
\ No newline at end of file
diff --git a/Engine/Plugins/MovieScene/MovieRenderPipeline/Source/MovieRenderPipelineCore/Private/MoviePipelineOutputBuilder.h b/Engine/Plugins/MovieScene/MovieRenderPipeline/Source/MovieRenderPipelineCore/Private/MoviePipelineOutputBuilder.h
deleted file mode 100644
diff --git a/Engine/Plugins/MovieScene/MovieRenderPipeline/Source/MovieRenderPipelineCore/Private/MoviePipelineRendering.cpp b/Engine/Plugins/MovieScene/MovieRenderPipeline/Source/MovieRenderPipelineCore/Private/MoviePipelineRendering.cpp
--- a/Engine/Plugins/MovieScene/MovieRenderPipeline/Source/MovieRenderPipelineCore/Private/MoviePipelineRendering.cpp
+++ b/Engine/Plugins/MovieScene/MovieRenderPipeline/Source/MovieRenderPipelineCore/Private/MoviePipelineRendering.cpp
@@ -13,6 +13,8 @@
 #include "MoviePipelineConfigBase.h"
 #include "MoviePipelineMasterConfig.h"
 #include "Math/Halton.h"
+#include "ImageWriteTask.h"
+#include "ImageWriteQueue.h"
 
 void UMoviePipeline::SetupRenderingPipelineForShot(FMoviePipelineShotInfo& Shot)
 {
@@ -55,18 +57,18 @@ void UMoviePipeline::SetupRenderingPipelineForShot(FMoviePipelineShotInfo& Shot)
 	// Code expects at least a 1x1 tile.
 	ensure(RenderPassInitSettings.TileCount > 0);
 
-	for (UMoviePipelineRenderPass* Input : Shot.RenderPasses)
+	for (UMoviePipelineRenderPass* Input : Shot.ShotConfig->GetRenderPasses())
 	{
-		Input->Setup(RenderPassInitSettings, OutputPipe.ToSharedRef());
+		Input->Setup(RenderPassInitSettings);
 	}
 }
 
 void UMoviePipeline::TeardownRenderingPipelineForShot(FMoviePipelineShotInfo& Shot)
 {
-	//for (UMoviePipelineRenderPass* Input : Shot.ShotConfig->InputBuffers)
-	//{
-	//	Input->Teardown();
-	//}
+	for (UMoviePipelineRenderPass* Input : Shot.ShotConfig->GetRenderPasses())
+	{
+		Input->Teardown();
+	}
 }
 
 void UMoviePipeline::RenderFrame()
@@ -120,7 +122,7 @@ void UMoviePipeline::RenderFrame()
 	int32 NumSpatialSamples = AccumulationSettings->SpatialSampleCount;
 	ensure(NumTilesX > 0 && NumTilesY > 0 && NumSpatialSamples > 0);
 
-	TArray<UMoviePipelineRenderPass*> InputBuffers = CurrentShot.RenderPasses;
+	TArray<UMoviePipelineRenderPass*> InputBuffers = CurrentShot.ShotConfig->GetRenderPasses();
 
 	// If this is the first sample for a new frame, we want to notify the output builder that it should expect data to accumulate for this frame.
 	if (CachedOutputState.TemporalSampleIndex == 0)
@@ -306,5 +308,48 @@ void UMoviePipeline::RenderFrame()
 		}
 	}
 	
+	// UE_LOG(LogMovieRenderPipeline, Warning, TEXT("[%d] Pre-FlushRenderingCommands"), GFrameCounter);
 	FlushRenderingCommands();
-}
\ No newline at end of file
+	// UE_LOG(LogMovieRenderPipeline, Warning, TEXT("[%d] Post-FlushRenderingCommands"), GFrameCounter);
+
+}
+
+void UMoviePipeline::OnFrameCompletelyRendered(FMoviePipelineMergerOutputFrame&& OutputFrame)
+{
+	UE_LOG(LogMovieRenderPipeline, Warning, TEXT("[%d] Data required for output available! Frame: %d"), GFrameCounter, OutputFrame.FrameOutputState.OutputFrameNumber);
+
+	for (UMoviePipelineOutputBase* OutputContainer : GetPipelineMasterConfig()->GetOutputContainers())
+	{
+		OutputContainer->OnRecieveImageDataImpl(&OutputFrame);
+	}
+}
+
+void UMoviePipeline::OnSampleRendered(TUniquePtr<FImagePixelData>&& OutputSample)
+{
+	FImagePixelDataPayload* Payload = OutputSample->GetPayload< FImagePixelDataPayload>();
+	check(Payload);
+
+	UMoviePipelineOutputSetting* OutputSettings = GetPipelineMasterConfig()->FindSetting<UMoviePipelineOutputSetting>();
+	check(OutputSettings);
+
+	TUniquePtr<FImageWriteTask> TileImageTask = MakeUnique<FImageWriteTask>();
+
+	// Fill alpha for now 
+	TileImageTask->PixelPreProcessors.Add(TAsyncAlphaWrite<FColor>(255));
+
+	// JPEG output
+	TileImageTask->Format = EImageFormat::JPEG;
+	TileImageTask->CompressionQuality = 100;
+
+	FString OutputName = FString::Printf(TEXT("/%s_SS_%d_TS_%d_TileX_%d_TileY_%d.%d.jpeg"),
+		*Payload->PassIdentifier.Name, Payload->SpatialJitterIndex, Payload->TemporalJitterIndex,
+		Payload->TileIndexX, Payload->TileIndexY, Payload->OutputState.OutputFrameNumber);
+
+	FString OutputDirectory = OutputSettings->OutputDirectory.Path;
+	FString OutputPath = OutputDirectory + OutputName;
+	TileImageTask->Filename = OutputPath;
+
+	// Duplicate the data so that the Image Task can own it.
+	TileImageTask->PixelData = MoveTemp(OutputSample);
+	ImageWriteQueue->Enqueue(MoveTemp(TileImageTask));
+}
diff --git a/Engine/Plugins/MovieScene/MovieRenderPipeline/Source/MovieRenderPipelineCore/Public/MoviePipeline.h b/Engine/Plugins/MovieScene/MovieRenderPipeline/Source/MovieRenderPipelineCore/Public/MoviePipeline.h
--- a/Engine/Plugins/MovieScene/MovieRenderPipeline/Source/MovieRenderPipelineCore/Public/MoviePipeline.h
+++ b/Engine/Plugins/MovieScene/MovieRenderPipeline/Source/MovieRenderPipelineCore/Public/MoviePipeline.h
@@ -20,6 +20,7 @@ class FMovieRenderViewportClient;
 struct FImagePixelPipe;
 struct FMoviePipelineTimeController;
 class FMoviePipelineOutputMerger;
+class IImageWriteQueue;
 
 
 
@@ -90,6 +91,10 @@ public:
 	*/
 	UFUNCTION(BlueprintPure, Category = "Movie Render Pipeline")
 	FFrameRate GetEffectiveFrameRate() const;
+
+public:
+	void OnFrameCompletelyRendered(FMoviePipelineMergerOutputFrame&& OutputFrame);
+	void OnSampleRendered(TUniquePtr<FImagePixelData>&& OutputSample);
 private:
 
 	/** Instantiate our Debug UI Widget and initialize it to ourself. */
@@ -243,6 +248,9 @@ public:
 	/** This gathers all of the produced data for an output frame (which may come in async many frames later) before passing them onto the Output Containers. */
 	TSharedPtr<FMoviePipelineOutputMerger, ESPMode::ThreadSafe> OutputBuilder;
 
+	/** A debug image sequence writer in the event they want to dump every sample generated on its own. */
+	IImageWriteQueue* ImageWriteQueue;
+
 	// debug
 	TSharedPtr<FImagePixelPipe, ESPMode::ThreadSafe> GetOutputPipe() const { return OutputPipe; }
 	TSharedPtr<FImagePixelPipe, ESPMode::ThreadSafe> OutputPipe;
diff --git a/Engine/Plugins/MovieScene/MovieRenderPipeline/Source/MovieRenderPipelineCore/Public/MoviePipelineOutputBase.h b/Engine/Plugins/MovieScene/MovieRenderPipeline/Source/MovieRenderPipelineCore/Public/MoviePipelineOutputBase.h
--- a/Engine/Plugins/MovieScene/MovieRenderPipeline/Source/MovieRenderPipelineCore/Public/MoviePipelineOutputBase.h
+++ b/Engine/Plugins/MovieScene/MovieRenderPipeline/Source/MovieRenderPipelineCore/Public/MoviePipelineOutputBase.h
@@ -57,6 +57,7 @@ public:
 	*/
 	void OnPostTick() { OnPostTickImpl(); }
 
+	virtual void OnRecieveImageDataImpl(FMoviePipelineMergerOutputFrame* InMergedOutputFrame) {}
 protected:
 	virtual void OnShotInitializedImpl(const TOptional<FMoviePipelineShotInfo> PrevShot, const FMoviePipelineShotInfo& NewShot) {}
 	virtual void OnShotFinishedImpl(const FMoviePipelineShotInfo& Shot) {}
diff --git a/Engine/Plugins/MovieScene/MovieRenderPipeline/Source/MovieRenderPipelineCore/Public/MoviePipelineOutputBuilder.h b/Engine/Plugins/MovieScene/MovieRenderPipeline/Source/MovieRenderPipelineCore/Public/MoviePipelineOutputBuilder.h
new file mode 100644
--- /dev/null
+++ b/Engine/Plugins/MovieScene/MovieRenderPipeline/Source/MovieRenderPipelineCore/Public/MoviePipelineOutputBuilder.h
@@ -0,0 +1,36 @@
+// Copyright 1998-2019 Epic Games, Inc. All Rights Reserved.
+
+#pragma once
+#include "MovieRenderPipelineDataTypes.h"
+
+// Forward Declares
+struct FImagePixelData;
+class UMoviePipeline;
+
+
+
+class MOVIERENDERPIPELINECORE_API FMoviePipelineOutputMerger
+{
+public:
+	FMoviePipelineOutputMerger(UMoviePipeline* InOwningMoviePipeline)
+		: WeakMoviePipeline(MakeWeakObjectPtr(InOwningMoviePipeline))
+	{
+	}
+
+public:
+	FMoviePipelineMergerOutputFrame& QueueOutputFrame_GameThread(const FMoviePipelineFrameOutputState& CachedOutputState);
+	void OnCompleteRenderPassDataAvailable_AnyThread(TUniquePtr<FImagePixelData>&& InData);
+	void OnSingleSampleDataAvailable_AnyThread(TUniquePtr<FImagePixelData>&& InData);
+
+private:
+	/** The Movie Pipeline that owns us. */
+	TWeakObjectPtr<UMoviePipeline> WeakMoviePipeline;
+
+	/** Data that is expected but not fully available yet. */
+	TMap<FMoviePipelineFrameOutputState, FMoviePipelineMergerOutputFrame> ActiveData;
+
+	/** Mutex that protects adding/updating/removing from ActiveData */
+	FCriticalSection ActiveDataMutex;
+};
+
+
diff --git a/Engine/Plugins/MovieScene/MovieRenderPipeline/Source/MovieRenderPipelineCore/Public/MoviePipelineRenderPass.h b/Engine/Plugins/MovieScene/MovieRenderPipeline/Source/MovieRenderPipelineCore/Public/MoviePipelineRenderPass.h
--- a/Engine/Plugins/MovieScene/MovieRenderPipeline/Source/MovieRenderPipelineCore/Public/MoviePipelineRenderPass.h
+++ b/Engine/Plugins/MovieScene/MovieRenderPipeline/Source/MovieRenderPipelineCore/Public/MoviePipelineRenderPass.h
@@ -12,10 +12,10 @@ class MOVIERENDERPIPELINECORE_API UMoviePipelineRenderPass : public UMoviePipeli
 {
 	GENERATED_BODY()
 public:
-	void Setup(const FMoviePipelineRenderPassInitSettings& InInitSettings, TSharedRef<FImagePixelPipe, ESPMode::ThreadSafe> InOutputPipe)
+	void Setup(const FMoviePipelineRenderPassInitSettings& InInitSettings)
 	{
 		InitSettings = InInitSettings;
-		SetupImpl(InInitSettings, InOutputPipe);
+		SetupImpl(InInitSettings);
 	}
 
 	void CaptureFrame(const FMoviePipelineRenderPassMetrics& OutputFrameMetrics)
@@ -34,7 +34,7 @@ public:
 	}
 
 protected:
-	virtual void SetupImpl(const FMoviePipelineRenderPassInitSettings& InInitSettings, TSharedRef<FImagePixelPipe, ESPMode::ThreadSafe> InOutputPipe) {  }
+	virtual void SetupImpl(const FMoviePipelineRenderPassInitSettings& InInitSettings) {}
 
 	virtual void CaptureFrameImpl(const FMoviePipelineRenderPassMetrics& OutputFrameMetrics) {}
 
diff --git a/Engine/Plugins/MovieScene/MovieRenderPipeline/Source/MovieRenderPipelineCore/Public/MovieRenderPipelineDataTypes.h b/Engine/Plugins/MovieScene/MovieRenderPipeline/Source/MovieRenderPipelineCore/Public/MovieRenderPipelineDataTypes.h
--- a/Engine/Plugins/MovieScene/MovieRenderPipeline/Source/MovieRenderPipelineCore/Public/MovieRenderPipelineDataTypes.h
+++ b/Engine/Plugins/MovieScene/MovieRenderPipeline/Source/MovieRenderPipelineCore/Public/MovieRenderPipelineDataTypes.h
@@ -17,6 +17,7 @@ class UMoviePipelineMasterConfig;
 class UMoviePipelineShotConfig;
 class UMovieSceneCameraCutSection;
 class UMoviePipelineRenderPass;
+struct FImagePixelData;
 
 
 /**
@@ -72,12 +73,31 @@ enum class EMovieRenderShotState : uint8
 
 struct FMoviePipelinePassIdentifier
 {
-	FMoviePipelinePassIdentifier(const FName& InPassName)
+	FMoviePipelinePassIdentifier()
+	{}
+
+	FMoviePipelinePassIdentifier(const FString& InPassName)
 		: Name(InPassName)
 	{
 	}
 
-	FName Name;
+	bool operator == (const FMoviePipelinePassIdentifier& InRHS) const
+	{
+		return Name == InRHS.Name;
+	}
+
+	bool operator != (const FMoviePipelinePassIdentifier& InRHS) const
+	{
+		return !(*this == InRHS);
+	}
+
+	friend uint32 GetTypeHash(FMoviePipelinePassIdentifier OutputState)
+	{
+		return GetTypeHash(OutputState.Name);
+	}
+
+public:
+	FString Name;
 };
 
 namespace MoviePipeline
@@ -410,7 +430,10 @@ public:
 	/** The range for this shot including handle frames that will be rendered. */
 	TRange<FFrameNumber> TotalOutputRange;
 
+	/** The range of time that represents the handle frames (if any) before the shot. */
 	TRange<FFrameNumber> HandleFrameRangeStart;
+
+	/** The range of time that represents the handle frames (if any) after the shot. */
 	TRange<FFrameNumber> HandleFrameRangeEnd;
 	
 	UPROPERTY(Transient, BlueprintReadOnly, Category = "Movie Render Pipeline")
@@ -425,14 +448,6 @@ public:
 	UPROPERTY(BlueprintReadOnly, Category = "Movie Render Pipeline")
 	int32 CurrentCameraCutIndex;
 
-	/** 
-	* An array of the Render Passes for this particular shot. These are duplicated out of the settings
-	* so that they can have their own state per shot that overlaps with other in progress shots. This
-	* is useful for long running async tasks.
-	*/
-	UPROPERTY(BlueprintReadOnly, Category = "Movie Render Pipeline")
-	TArray<UMoviePipelineRenderPass*> RenderPasses;
-
 	FMoviePipelineCameraCutInfo& GetCurrentCameraCut()
 	{
 		check(CurrentCameraCutIndex >= 0 && CurrentCameraCutIndex < CameraCuts.Num());
@@ -668,7 +683,31 @@ struct FImagePixelDataPayload : IImagePixelDataPayload
 	FVector2D OverlappedSubpixelShift;
 
 
-	FString PassName;
+	FMoviePipelinePassIdentifier PassIdentifier;
+
+	/** Is this the first tile of an image and we should start accumulating? */
+	FORCEINLINE bool IsFirstTile() const
+	{
+		return TileIndexX == 0 && TileIndexY == 0 && SpatialJitterIndex == 0;
+	}
+
+	/** Is this the last tile of an image and we should finish accumulating? */
+	FORCEINLINE bool IsLastTile() const
+	{
+		return TileIndexX == NumTilesX - 1 &&
+			   TileIndexY == NumTilesY - 1 &&
+			   SpatialJitterIndex == NumSpatialJitters - 1;
+	}
+
+	FORCEINLINE bool IsFirstTemporalSample() const
+	{
+		return TemporalJitterIndex == 0;
+	}
+
+	FORCEINLINE bool IsLastTemporalSample() const
+	{
+		return TemporalJitterIndex == NumTemporalJitters - 1;
+	}
 };
 
 /**
@@ -702,4 +741,39 @@ public:
 	/** Which shots (by name) should this job render out of the given sequence? Leave empty to render whole sequence. */
 	UPROPERTY(BlueprintReadWrite, Category = "Movie Render Pipeline")
 	TArray<FString> ShotRenderMask;
+};
+
+
+struct MOVIERENDERPIPELINECORE_API FMoviePipelineMergerOutputFrame
+{
+public:
+	FMoviePipelineMergerOutputFrame() {}
+	virtual ~FMoviePipelineMergerOutputFrame() {};
+
+	FMoviePipelineMergerOutputFrame& operator=(FMoviePipelineMergerOutputFrame&& InOther)
+	{
+		FrameOutputState = InOther.FrameOutputState;
+		ExpectedRenderPasses = InOther.ExpectedRenderPasses;
+		ImageOutputData = MoveTemp(InOther.ImageOutputData);
+	
+		return *this;
+	}
+	FMoviePipelineMergerOutputFrame(FMoviePipelineMergerOutputFrame&& InOther)
+		: FrameOutputState(InOther.FrameOutputState)
+		, ExpectedRenderPasses(InOther.ExpectedRenderPasses)
+		, ImageOutputData(MoveTemp(InOther.ImageOutputData))
+	{
+	}
+
+private:
+	// Explicitly delete copy operators since we own unique data.
+	// void operator=(const FMoviePipelineMergerOutputFrame&);
+	// FMoviePipelineMergerOutputFrame(const FMoviePipelineMergerOutputFrame&);
+
+public:
+	FMoviePipelineFrameOutputState FrameOutputState;
+
+	TArray<FMoviePipelinePassIdentifier> ExpectedRenderPasses;
+
+	TMap<FMoviePipelinePassIdentifier, TUniquePtr<FImagePixelData>> ImageOutputData;
 };
\ No newline at end of file
diff --git a/Engine/Plugins/MovieScene/MovieRenderPipeline/Source/MovieRenderPipelineEditor/Private/MovieRenderPipelineEditorModule.cpp b/Engine/Plugins/MovieScene/MovieRenderPipeline/Source/MovieRenderPipelineEditor/Private/MovieRenderPipelineEditorModule.cpp
--- a/Engine/Plugins/MovieScene/MovieRenderPipeline/Source/MovieRenderPipelineEditor/Private/MovieRenderPipelineEditorModule.cpp
+++ b/Engine/Plugins/MovieScene/MovieRenderPipeline/Source/MovieRenderPipelineEditor/Private/MovieRenderPipelineEditorModule.cpp
@@ -65,7 +65,7 @@ void FMovieRenderPipelineEditorModule::PerformTestPipelineRender(const TArray<FS
 		return;
 	}
 
-	FSoftObjectPath SequencePath = FSoftObjectPath(Args[0]);
+	FSoftObjectPath SequencePath = FSoftObjectPath(FPackageName::ExportTextPathToObjectPath(Args[0]));
 
 	const UMovieRenderPipelineProjectSettings* ProjectSettings = GetDefault<UMovieRenderPipelineProjectSettings>();
 	TSubclassOf<UMoviePipelineExecutorBase> ExecutorClass = ProjectSettings->DefaultLocalExecutor;
diff --git a/Engine/Plugins/MovieScene/MovieRenderPipeline/Source/MovieRenderPipelineEditor/Private/MovieRenderPipelineEditorTest.cpp b/Engine/Plugins/MovieScene/MovieRenderPipeline/Source/MovieRenderPipelineEditor/Private/MovieRenderPipelineEditorTest.cpp
--- a/Engine/Plugins/MovieScene/MovieRenderPipeline/Source/MovieRenderPipelineEditor/Private/MovieRenderPipelineEditorTest.cpp
+++ b/Engine/Plugins/MovieScene/MovieRenderPipeline/Source/MovieRenderPipelineEditor/Private/MovieRenderPipelineEditorTest.cpp
@@ -61,8 +61,8 @@ TArray<UMoviePipelineMasterConfig*> FMovieRenderPipelineEditorModule::GenerateTe
 				UMoviePipelineMasterConfig* OutPipeline = NewObject<UMoviePipelineMasterConfig>(GetTransientPackage());
 				UMoviePipelineOutputSetting* OutputSetting = OutPipeline->FindOrAddSetting<UMoviePipelineOutputSetting>();
 
-				int32 SizeX = 1080;
-				int32 SizeY = 1920;
+				int32 SizeX = 1920;
+				int32 SizeY = 1080;
 				int32 TileX = 2;
 				int32 TileY = 2;
 				int32 TestNumSamples = 2;
diff --git a/Engine/Plugins/MovieScene/MovieRenderPipeline/Source/MovieRenderPipelineRenderPasses/Private/MoviePipelineBackbufferPass.cpp b/Engine/Plugins/MovieScene/MovieRenderPipeline/Source/MovieRenderPipelineRenderPasses/Private/MoviePipelineBackbufferPass.cpp
--- a/Engine/Plugins/MovieScene/MovieRenderPipeline/Source/MovieRenderPipelineRenderPasses/Private/MoviePipelineBackbufferPass.cpp
+++ b/Engine/Plugins/MovieScene/MovieRenderPipeline/Source/MovieRenderPipelineRenderPasses/Private/MoviePipelineBackbufferPass.cpp
@@ -22,28 +22,44 @@
 #include "MoviePipelineShotConfig.h"
 #include "MoviePipelineSetting.h"
 #include "MoviePipeline.h"
+#include "MoviePipelineOutputBuilder.h"
+#include "MovieRenderTileImage.h"
+#include "Async/Async.h"
+#include "MovieRenderPipelineCoreModule.h"
+#include "MovieRenderOverlappedImage.h"
 
-
-void UMoviePipelineBackbufferPass::SetupImpl(const FMoviePipelineRenderPassInitSettings& InInitSettings, TSharedRef<FImagePixelPipe, ESPMode::ThreadSafe> InOutputPipe)
+namespace MoviePipeline
 {
-	OutputPipe = InOutputPipe;
+	struct FSampleRenderThreadParams
+	{
+		FMoviePipelineRenderPassMetrics PassMetrics;
+		FRenderTarget* CanvasRenderTarget;
+		TSharedPtr<FImageOverlappedAccumulator, ESPMode::ThreadSafe> ImageAccumulator;
+		TSharedPtr<FMoviePipelineOutputMerger, ESPMode::ThreadSafe> OutputMerger;
+		bool bWriteTiles;
+	};
+}
+
+// Forward Declares
+static void ReadAndAccumulateSample_RenderThread(FRHICommandListImmediate &RHICmdList, const MoviePipeline::FSampleRenderThreadParams& InParams);
 
+void UMoviePipelineBackbufferPass::SetupImpl(const FMoviePipelineRenderPassInitSettings& InInitSettings)
+{
 	TileRenderTarget = NewObject<UTextureRenderTarget2D>(this);
 	TileRenderTarget->ClearColor = FLinearColor(0.0f, 0.0f, 0.0f, 0.0f);
 
-	// Initialize to the tile size (not final size) and use a 16 bit back buffer to avoid precision issues
-			// when accumulating later.
+	// Initialize to the tile size (not final size) and use a 16 bit back buffer to avoid precision issues when accumulating later
 
 	int32 TargetSizeX = InInitSettings.TileResolution.X;
 	int32 TargetSizeY = InInitSettings.TileResolution.Y;
 /*
 	*/
 	TileRenderTarget->InitCustomFormat(TargetSizeX, TargetSizeY, EPixelFormat::PF_A16B16G16R16, false);
-
 	// Create a Render Target to hold our back buffer.
 	ViewState.Allocate();
-}
 
+	ImageTileAccumulator = MakeShared<FImageOverlappedAccumulator, ESPMode::ThreadSafe>();
+}
 
 void UMoviePipelineBackbufferPass::TeardownImpl()
 {
@@ -52,20 +68,17 @@ void UMoviePipelineBackbufferPass::TeardownImpl()
 
 void UMoviePipelineBackbufferPass::GatherOutputPassesImpl(TArray<FMoviePipelinePassIdentifier>& ExpectedRenderPasses)
 {
-	ExpectedRenderPasses.Add(FName("Backbuffer"));
+	ExpectedRenderPasses.Add(FMoviePipelinePassIdentifier(TEXT("Backbuffer")));
 }
 
-
 void UMoviePipelineBackbufferPass::CaptureFrameImpl(const FMoviePipelineRenderPassMetrics& InPassMetrics)
 {
 	float TimeSeconds = InPassMetrics.OutputState.WorldSeconds;
 	float RealTimeSeconds = InPassMetrics.OutputState.WorldSeconds;
 	float DeltaTimeSeconds = InPassMetrics.OutputState.FrameDeltaTime;
 
-	// Reuse the common RenderTarget
-	UTextureRenderTarget2D* RenderTarget = TileRenderTarget;
+	FEngineShowFlags ShowFlags = FEngineShowFlags(EShowFlagInitMode::ESFIM_Game);
 
-	FEngineShowFlags ShowFlags(EShowFlagInitMode::ESFIM_Game);
 	if (InPassMetrics.bIsUsingOverlappedTiles)
 	{
 		// Disable these for now
@@ -74,7 +87,7 @@ void UMoviePipelineBackbufferPass::CaptureFrameImpl(const FMoviePipelineRenderPa
 	}
 
 	FSceneViewFamilyContext ViewFamily(FSceneViewFamily::ConstructionValues(
-		RenderTarget->GameThread_GetRenderTargetResource(),
+		TileRenderTarget->GameThread_GetRenderTargetResource(),
 		GetWorld()->Scene,
 		ShowFlags)
 		.SetWorldTimes(TimeSeconds, DeltaTimeSeconds, RealTimeSeconds)
@@ -124,76 +137,36 @@ void UMoviePipelineBackbufferPass::CaptureFrameImpl(const FMoviePipelineRenderPa
 	{
 		View->bCameraCut = true;
 	}
-	
+
 	// Pause the view family on renders after the first one. This prevents motion blur/view matrices from being
 	// updated
 	ViewFamily.bWorldIsPaused = InPassMetrics.SpatialJitterIndex != InPassMetrics.NumSpatialJitters - 1;
 
-	FCanvas Canvas = FCanvas(RenderTarget->GameThread_GetRenderTargetResource(), nullptr, GetWorld(), ERHIFeatureLevel::SM5, FCanvas::CDM_DeferDrawing, 1.0f);
+	FCanvas Canvas = FCanvas(TileRenderTarget->GameThread_GetRenderTargetResource(), nullptr, GetWorld(), ERHIFeatureLevel::SM5, FCanvas::CDM_DeferDrawing, 1.0f);
 
 
 	// Draw the world into this View Family
 	GetRendererModule().BeginRenderingViewFamily(&Canvas, &ViewFamily);
 
 	FRenderTarget* CanvasRenderTarget = Canvas.GetRenderTarget();
-	TSharedPtr<FImagePixelPipe, ESPMode::ThreadSafe> LocalOutputPipe = OutputPipe;
-	
-	ENQUEUE_RENDER_COMMAND(CanvasRenderTargetResolveCommand)(
-		[CanvasRenderTarget, InPassMetrics, LocalOutputPipe](FRHICommandListImmediate& RHICmdList)
-	{
-		// If they only needed this frame for the history, we don't actually care about trying to output it,
-		// so we can skip the copy and pushing to output.
-		if (InPassMetrics.bIsHistoryOnlyFrame)
-		{
-			return;
-		}
+	TSharedPtr<FMoviePipelineOutputMerger, ESPMode::ThreadSafe> OutputMerger = GetPipeline()->OutputBuilder;
+	TSharedPtr<FImageOverlappedAccumulator, ESPMode::ThreadSafe> LocalImageTileAccumulator = ImageTileAccumulator;
 
-		FReadSurfaceDataFlags ReadDataFlags;
-		ReadDataFlags.SetLinearToGamma(false);
-
-		FIntRect SourceRect = FIntRect(0, 0, CanvasRenderTarget->GetSizeXY().X, CanvasRenderTarget->GetSizeXY().Y);
-
-		// Build our metadata to associate with this frame.
-		TSharedRef<FImagePixelDataPayload, ESPMode::ThreadSafe> FrameData = MakeShared<FImagePixelDataPayload, ESPMode::ThreadSafe>();
-		FrameData->OutputState = InPassMetrics.OutputState;
-		FrameData->PassName = TEXT("Backbuffer");
-		FrameData->TileIndexX = InPassMetrics.TileIndexX;
-		FrameData->TileIndexY = InPassMetrics.TileIndexY;
-		FrameData->TileSizeX = SourceRect.Width();
-		FrameData->TileSizeY = SourceRect.Height();
-		FrameData->NumTilesX = InPassMetrics.NumTilesX;
-		FrameData->NumTilesY = InPassMetrics.NumTilesY;
-		FrameData->SpatialShift = InPassMetrics.SpatialShift;
-		FrameData->SpatialJitterIndex = InPassMetrics.SpatialJitterIndex;
-		FrameData->NumSpatialJitters = InPassMetrics.NumSpatialJitters;
-		FrameData->TemporalJitterIndex = InPassMetrics.TemporalJitterIndex;
-		FrameData->NumTemporalJitters = InPassMetrics.NumTemporalJitters;
-		FrameData->JitterOffsetX = InPassMetrics.JitterOffsetX;
-		FrameData->JitterOffsetY = InPassMetrics.JitterOffsetY;
-		FrameData->bIsUsingOverlappedTiles = InPassMetrics.bIsUsingOverlappedTiles;
-
-		FrameData->AccumulationGamma = InPassMetrics.AccumulationGamma;
-
-		FrameData->OverlappedOffsetX = InPassMetrics.OverlappedOffsetX;
-		FrameData->OverlappedOffsetY = InPassMetrics.OverlappedOffsetY;
-		FrameData->OverlappedSizeX = InPassMetrics.OverlappedSizeX;
-		FrameData->OverlappedSizeY = InPassMetrics.OverlappedSizeY;
-		FrameData->OverlappedPadX = InPassMetrics.OverlappedPadX;
-		FrameData->OverlappedPadY = InPassMetrics.OverlappedPadY;
-		FrameData->OverlappedSubpixelShift = InPassMetrics.OverlappedSubpixelShift;
-
-		TUniquePtr<TImagePixelData<FColor>> PixelData = MakeUnique<TImagePixelData<FColor>>(SourceRect.Size(), FrameData);
-		RHICmdList.ReadSurfaceData(CanvasRenderTarget->GetRenderTargetTexture(), SourceRect, PixelData->Pixels, ReadDataFlags);
-
-		check(PixelData->IsDataWellFormed());
-
-		const bool bIsOutputFrame = true;
-		if (bIsOutputFrame)
-		{
-			LocalOutputPipe->Push(MoveTemp(PixelData));
-		}
+	MoviePipeline::FSampleRenderThreadParams AccumulationParams;
+	{
+		AccumulationParams.bWriteTiles = true;
+		AccumulationParams.CanvasRenderTarget = CanvasRenderTarget;
+		AccumulationParams.OutputMerger = GetPipeline()->OutputBuilder;
+		AccumulationParams.ImageAccumulator = ImageTileAccumulator;
+		AccumulationParams.PassMetrics = InPassMetrics;
 	}
-	);
+
+	ENQUEUE_RENDER_COMMAND(CanvasRenderTargetResolveCommand)(
+		[AccumulationParams](FRHICommandListImmediate& RHICmdList)
+	{
+		UE_LOG(LogMovieRenderPipeline, Log, TEXT("Rendered Tile: %d Sample: %d"), AccumulationParams.PassMetrics.TileIndex, AccumulationParams.PassMetrics.SpatialJitterIndex);
+		ReadAndAccumulateSample_RenderThread(RHICmdList, AccumulationParams);
+	});
 }
 
 FSceneView* UMoviePipelineBackbufferPass::CalcSceneView(FSceneViewFamily* ViewFamily, const FMoviePipelineRenderPassMetrics& InPassMetrics)
@@ -324,4 +297,125 @@ FSceneView* UMoviePipelineBackbufferPass::CalcSceneView(FSceneViewFamily* ViewFa
 	View->EndFinalPostprocessSettings(ViewInitOptions);
 
 	return View;
+}
+
+void ReadAndAccumulateSample_RenderThread(FRHICommandListImmediate &RHICmdList, const MoviePipeline::FSampleRenderThreadParams& InParams)
+{
+	check(IsInRenderingThread());
+
+	// If they only needed this frame for the history, we don't actually care about trying to output it,
+	// so we can skip the copy and pushing to output.
+	if (InParams.PassMetrics.bIsHistoryOnlyFrame)
+	{
+		return;
+	}
+
+	FReadSurfaceDataFlags ReadDataFlags;
+	ReadDataFlags.SetLinearToGamma(false);
+
+	FIntRect SourceRect = FIntRect(0, 0, InParams.CanvasRenderTarget->GetSizeXY().X, InParams.CanvasRenderTarget->GetSizeXY().Y);
+
+	// Build our metadata to associate with this frame.
+	TSharedRef<FImagePixelDataPayload, ESPMode::ThreadSafe> FrameData = MakeShared<FImagePixelDataPayload, ESPMode::ThreadSafe>();
+	FrameData->OutputState = InParams.PassMetrics.OutputState;
+	FrameData->PassIdentifier = FMoviePipelinePassIdentifier(TEXT("Backbuffer"));
+	FrameData->TileIndexX = InParams.PassMetrics.TileIndexX;
+	FrameData->TileIndexY = InParams.PassMetrics.TileIndexY;
+	FrameData->TileSizeX = SourceRect.Width();
+	FrameData->TileSizeY = SourceRect.Height();
+	FrameData->NumTilesX = InParams.PassMetrics.NumTilesX;
+	FrameData->NumTilesY = InParams.PassMetrics.NumTilesY;
+	FrameData->SpatialShift = InParams.PassMetrics.SpatialShift;
+	FrameData->SpatialJitterIndex = InParams.PassMetrics.SpatialJitterIndex;
+	FrameData->NumSpatialJitters = InParams.PassMetrics.NumSpatialJitters;
+	FrameData->TemporalJitterIndex = InParams.PassMetrics.TemporalJitterIndex;
+	FrameData->NumTemporalJitters = InParams.PassMetrics.NumTemporalJitters;
+	FrameData->JitterOffsetX = InParams.PassMetrics.JitterOffsetX;
+	FrameData->JitterOffsetY = InParams.PassMetrics.JitterOffsetY;
+	FrameData->bIsUsingOverlappedTiles = InParams.PassMetrics.bIsUsingOverlappedTiles;
+
+	FrameData->AccumulationGamma = InParams.PassMetrics.AccumulationGamma;
+
+	FrameData->OverlappedOffsetX = InParams.PassMetrics.OverlappedOffsetX;
+	FrameData->OverlappedOffsetY = InParams.PassMetrics.OverlappedOffsetY;
+	FrameData->OverlappedSizeX = InParams.PassMetrics.OverlappedSizeX;
+	FrameData->OverlappedSizeY = InParams.PassMetrics.OverlappedSizeY;
+	FrameData->OverlappedPadX = InParams.PassMetrics.OverlappedPadX;
+	FrameData->OverlappedPadY = InParams.PassMetrics.OverlappedPadY;
+	FrameData->OverlappedSubpixelShift = InParams.PassMetrics.OverlappedSubpixelShift;
+
+	// Read the data back to the CPU
+	TUniquePtr<TImagePixelData<FColor>> PixelData = MakeUnique<TImagePixelData<FColor>>(SourceRect.Size(), FrameData);
+	RHICmdList.ReadSurfaceData(InParams.CanvasRenderTarget->GetRenderTargetTexture(), SourceRect, PixelData->Pixels, ReadDataFlags);
+
+	check(PixelData->IsDataWellFormed());
+
+	const bool bWriteTiles = true;
+	
+	// Writing tiles can be useful for debug reasons. These get passed onto the output every frame.
+	if (bWriteTiles)
+	{
+		// Send the data to the Output Builder. This has to be a copy of the pixel data from the GPU, since
+		// it enqueues it onto the game thread and won't be read/sent to write to disk for another frame. 
+		// The extra copy is unfortunate, but is only the size of a single sample (ie: 1920x1080 -> 17mb)
+		TUniquePtr<FImagePixelData> SampleData = PixelData->Copy();
+		InParams.OutputMerger->OnSingleSampleDataAvailable_AnyThread(MoveTemp(SampleData));
+	}
+
+	// For the first sample in a new output, we allocate memory
+	if (FrameData->IsFirstTile() && FrameData->IsFirstTemporalSample())
+	{
+		InParams.ImageAccumulator->InitMemory(FIntPoint(FrameData->OverlappedSizeX * FrameData->NumTilesX, FrameData->OverlappedSizeY * FrameData->NumTilesY), 3);
+		InParams.ImageAccumulator->ZeroPlanes();
+		InParams.ImageAccumulator->AccumulationGamma = FrameData->AccumulationGamma;
+	}
+
+	// Accumulate the new sample to our target
+	{
+		const double AccumulateBeginTime = FPlatformTime::Seconds();
+
+		FIntPoint RawSize = PixelData->GetSize();
+
+		check(FrameData->OverlappedSizeX + 2 * FrameData->OverlappedPadX == RawSize.X);
+		check(FrameData->OverlappedSizeY + 2 * FrameData->OverlappedPadY == RawSize.Y);
+
+		InParams.ImageAccumulator->AccumulatePixelData(*PixelData.Get(), FIntPoint(FrameData->OverlappedOffsetX, FrameData->OverlappedOffsetY), FrameData->OverlappedSubpixelShift);
+
+		const double AccumulateEndTime = FPlatformTime::Seconds();
+		const float ElapsedMs = float((AccumulateEndTime - AccumulateBeginTime)*1000.0f);
+	
+		UE_LOG(LogMovieRenderPipeline, Log, TEXT("Accumulation time: %8.2fms"), ElapsedMs);
+
+	}
+
+	if (FrameData->IsLastTile() && FrameData->IsLastTemporalSample())
+	{
+		int32 FullSizeX = InParams.ImageAccumulator->PlaneSize.X;
+		int32 FullSizeY = InParams.ImageAccumulator->PlaneSize.Y;
+
+		// Now that a tile is fully built and accumulated we can notify the output builder that the
+		// data is ready so it can pass that onto the output containers (if needed).
+		static bool bFullLinearColor = false; // Make this an option. For now, if it's linear, write EXR. Otherwise PNGs for 8bit.
+		if (bFullLinearColor)
+		{
+			// 32 bit FLinearColor
+			TUniquePtr<TImagePixelData<FLinearColor> > FinalPixelData = MakeUnique<TImagePixelData<FLinearColor>>(FIntPoint(FullSizeX, FullSizeY), FrameData);
+			InParams.ImageAccumulator->FetchFinalPixelDataLinearColor(FinalPixelData->Pixels);
+
+			// Send the data to the Output Builder
+			InParams.OutputMerger->OnCompleteRenderPassDataAvailable_AnyThread(MoveTemp(FinalPixelData));
+		}
+		else
+		{
+			// 8bit FColors
+			TUniquePtr<TImagePixelData<FColor>> FinalPixelData = MakeUnique<TImagePixelData<FColor>>(FIntPoint(FullSizeX, FullSizeY), FrameData);
+			InParams.ImageAccumulator->FetchFinalPixelDataByte(FinalPixelData->Pixels);
+
+			// Send the data to the Output Builder
+			InParams.OutputMerger->OnCompleteRenderPassDataAvailable_AnyThread(MoveTemp(FinalPixelData));
+		}
+
+		// Free the memory in the accumulator.
+		InParams.ImageAccumulator->Reset();
+	}
 }
\ No newline at end of file
diff --git a/Engine/Plugins/MovieScene/MovieRenderPipeline/Source/MovieRenderPipelineRenderPasses/Private/MoviePipelineImageSequenceContainer.cpp b/Engine/Plugins/MovieScene/MovieRenderPipeline/Source/MovieRenderPipelineRenderPasses/Private/MoviePipelineImageSequenceContainer.cpp
--- a/Engine/Plugins/MovieScene/MovieRenderPipeline/Source/MovieRenderPipelineRenderPasses/Private/MoviePipelineImageSequenceContainer.cpp
+++ b/Engine/Plugins/MovieScene/MovieRenderPipeline/Source/MovieRenderPipelineRenderPasses/Private/MoviePipelineImageSequenceContainer.cpp
@@ -32,151 +32,37 @@ bool UMoviePipelineImageSequenceContainerBase::HasFinishedProcessingImpl()
 	return Super::HasFinishedProcessingImpl() && (!FinalizeFence.IsValid() || FinalizeFence.WaitFor(0));
 }
 
-static FImageTileAccumulator TestAccumulator;
-static FImageOverlappedAccumulator TestOverlappedAccumulator;
-
 void UMoviePipelineImageSequenceContainerBase::SetupForPipelineImpl(UMoviePipeline* InPipeline)
 {
 	UMoviePipelineOutputSetting* OutputSettings = GetPipeline()->GetPipelineMasterConfig()->FindSetting<UMoviePipelineOutputSetting>();
 	check(OutputSettings);
 
-	FString OutputDirectory = OutputSettings->OutputDirectory.Path;
-	IImageWriteQueue* WriteQueue = ImageWriteQueue;
-
-	FImageTileAccumulator* TileAccumulator = &TestAccumulator;
-	FImageOverlappedAccumulator* OverlappedAccumulator = &TestOverlappedAccumulator;
+	OutputDirectory = OutputSettings->OutputDirectory.Path;
+}
 
+void UMoviePipelineImageSequenceContainerBase::OnRecieveImageDataImpl(FMoviePipelineMergerOutputFrame* InMergedOutputFrame)
+{
+	check(InMergedOutputFrame);
 
-	auto OnImageReceived = [OutputDirectory, WriteQueue, TileAccumulator, OverlappedAccumulator](TUniquePtr<FImagePixelData>&& InOwnedImage)
+	for (TPair<FMoviePipelinePassIdentifier, TUniquePtr<FImagePixelData>>& RenderPassData : InMergedOutputFrame->ImageOutputData)
 	{
-		FImagePixelDataPayload* TilePayload = InOwnedImage->GetPayload<FImagePixelDataPayload>();
-
-		// Init if first one, shut down if last one
-		bool bIsFirstTile = (TilePayload->TileIndexX == 0 && TilePayload->TileIndexY == 0 && TilePayload->SpatialJitterIndex == 0);
-		bool bIsLastTile =  (TilePayload->TileIndexX == TilePayload->NumTilesX-1  &&
-							TilePayload->TileIndexY == TilePayload->NumTilesY-1 &&
-							TilePayload->SpatialJitterIndex == TilePayload->NumSpatialJitters-1);
-		bool bIsLastTemporalSample = TilePayload->TemporalJitterIndex == TilePayload->NumTemporalJitters - 1;
-		bool bIsFirstTemporalSample = TilePayload->TemporalJitterIndex == 0;
-
-
-		if (!TilePayload->bIsUsingOverlappedTiles)
-		{
-			if (bIsFirstTile && bIsFirstTemporalSample)
-			{
-				TileAccumulator->InitMemory(TilePayload->TileSizeX, TilePayload->TileSizeY, TilePayload->NumTilesX, TilePayload->NumTilesY, 3);
-				TileAccumulator->ZeroPlanes();
-				TileAccumulator->AccumulationGamma = TilePayload->AccumulationGamma;
-			}
-
-			const double AccumulateBeginTime = FPlatformTime::Seconds();
-			TileAccumulator->AccumulatePixelData(*InOwnedImage.Get(), TilePayload->TileIndexX, TilePayload->TileIndexY, FVector2D(TilePayload->JitterOffsetX,TilePayload->JitterOffsetY));
-			const double AccumulateEndTime = FPlatformTime::Seconds();
-			const float ElapsedMs = float((AccumulateEndTime - AccumulateBeginTime)*1000.0f);
-
-		}
-		else
-		{
-			if (bIsFirstTile && bIsFirstTemporalSample)
-			{
-				OverlappedAccumulator->InitMemory(FIntPoint(TilePayload->OverlappedSizeX * TilePayload->NumTilesX, TilePayload->OverlappedSizeY * TilePayload->NumTilesY), 3);
-				OverlappedAccumulator->ZeroPlanes();
-				OverlappedAccumulator->AccumulationGamma = TilePayload->AccumulationGamma;
-			}
-
-			FIntPoint RawSize = InOwnedImage.Get()->GetSize();
-
-			check(TilePayload->OverlappedSizeX + 2 * TilePayload->OverlappedPadX == RawSize.X);
-			check(TilePayload->OverlappedSizeY + 2 * TilePayload->OverlappedPadY == RawSize.Y);
-			OverlappedAccumulator->AccumulatePixelData(*InOwnedImage.Get(), FIntPoint(TilePayload->OverlappedOffsetX, TilePayload->OverlappedOffsetY), TilePayload->OverlappedSubpixelShift);
-		}
-
-		//UE_LOG(LogTemp, Log, TEXT("[%8.2f] Accumulation time."), ElapsedMs);
-
-		static bool bWriteTiles = true;
-		if (bWriteTiles)
-		{
-			TUniquePtr<FImageWriteTask> TileImageTask = MakeUnique<FImageWriteTask>();
-			TileImageTask->PixelPreProcessors.Add(TAsyncAlphaWrite<FColor>(255));
-			// ImageTask->PixelPreProcessors.Add(TAsyncGammaCorrect<FColor>(2.2f));
-			TileImageTask->Format = EImageFormat::JPEG;
-			TileImageTask->CompressionQuality = 100;
-
-			// ImageTask->CompressionQuality = GetCompressionQuality();
-			FString OutputName = FString::Printf(TEXT("/%s_%d_SS_%d_TS_%d_TileX_%d_TileY_%d.jpeg"),
-				*TilePayload->PassName, TilePayload->OutputState.OutputFrameNumber, TilePayload->SpatialJitterIndex, TilePayload->OutputState.TemporalSampleIndex,
-				TilePayload->TileIndexX, TilePayload->TileIndexY);
-			FString OutputPath = OutputDirectory + OutputName;
-			TileImageTask->Filename = OutputPath;
-
-			TileImageTask->PixelData = MoveTemp(InOwnedImage);
-
-			WriteQueue->Enqueue(MoveTemp(TileImageTask));
-		}
-
-		if (bIsLastTile && bIsLastTemporalSample)
-		{
-
-			int32 FullSizeX = TilePayload->bIsUsingOverlappedTiles ? OverlappedAccumulator->PlaneSize.X : TileAccumulator->TileSizeX * TileAccumulator->NumTilesX;
-			int32 FullSizeY = TilePayload->bIsUsingOverlappedTiles ? OverlappedAccumulator->PlaneSize.Y : TileAccumulator->TileSizeY * TileAccumulator->NumTilesY;
-
-			{
-				TUniquePtr<FImageWriteTask> ImageTask = MakeUnique<FImageWriteTask>();
-
-				static bool bFullLinearColor = false; // Make this an option. For now, if it's linear, write EXR. Otherwise PNGs for 8bit.
-				if (bFullLinearColor)
-				{
-					ImageTask->Format = EImageFormat::EXR;
-
-					TUniquePtr<TImagePixelData<FLinearColor> > PixelData = MakeUnique<TImagePixelData<FLinearColor> >(FIntPoint(FullSizeX,FullSizeY));
-					if (!TilePayload->bIsUsingOverlappedTiles)
-					{
-						TileAccumulator->FetchFinalPixelDataLinearColor(PixelData->Pixels);
-					}
-					else
-					{
-						OverlappedAccumulator->FetchFinalPixelDataLinearColor(PixelData->Pixels);
-					}
-
-					ImageTask->PixelPreProcessors.Add(TAsyncAlphaWrite<FLinearColor>(1.0f));
-					ImageTask->PixelData = MoveTemp(PixelData);
-
-					FString TempOutputName = FString::Printf(TEXT("/FINAL_%s_SS.%d.exr"),
-						*TilePayload->PassName, TilePayload->OutputState.OutputFrameNumber);
-					FString TempOutputPath = OutputDirectory + TempOutputName;
-					ImageTask->Filename = TempOutputPath;
-				}
-				else
-				{
-					ImageTask->Format = EImageFormat::JPEG;
-
-					// 8bit FColors
-					TUniquePtr<TImagePixelData<FColor> > PixelData = MakeUnique<TImagePixelData<FColor> >(FIntPoint(FullSizeX,FullSizeY));
-					if (!TilePayload->bIsUsingOverlappedTiles)
-					{
-						TileAccumulator->FetchFinalPixelDataByte(PixelData->Pixels);
-					}
-					else
-					{
-						OverlappedAccumulator->FetchFinalPixelDataByte(PixelData->Pixels);
-					}
+		TUniquePtr<FImageWriteTask> TileImageTask = MakeUnique<FImageWriteTask>();
 
-					ImageTask->PixelPreProcessors.Add(TAsyncAlphaWrite<FColor>(255));
-					ImageTask->PixelData = MoveTemp(PixelData);
-					ImageTask->CompressionQuality = 100;
+		// Fill alpha for now 
+		TileImageTask->PixelPreProcessors.Add(TAsyncAlphaWrite<FColor>(255));
 
-					FString TempOutputName = FString::Printf(TEXT("/FINAL_%s_SS.%d.jpeg"),
-						*TilePayload->PassName, TilePayload->OutputState.OutputFrameNumber);
-					FString TempOutputPath = OutputDirectory + TempOutputName;
-					ImageTask->Filename = TempOutputPath;
-				}
+		// JPEG output
+		TileImageTask->Format = EImageFormat::JPEG;
+		TileImageTask->CompressionQuality = 100;
 
-				WriteQueue->Enqueue(MoveTemp(ImageTask));
-			}
+		FString OutputName = FString::Printf(TEXT("/%s.%d.jpeg"),
+			*RenderPassData.Key.Name, InMergedOutputFrame->FrameOutputState.OutputFrameNumber);
 
-			TileAccumulator->Reset();
-		}
-	};
+		FString OutputPath = OutputDirectory + OutputName;
+		TileImageTask->Filename = OutputPath;
 
-	InPipeline->GetOutputPipe()->AddEndpoint(OnImageReceived);
+		// Duplicate the data so that the Image Task can own it.
+		TileImageTask->PixelData = MoveTemp(RenderPassData.Value);
+		ImageWriteQueue->Enqueue(MoveTemp(TileImageTask));
+	}
 }
\ No newline at end of file
diff --git a/Engine/Plugins/MovieScene/MovieRenderPipeline/Source/MovieRenderPipelineRenderPasses/Public/MoviePipelineBackbufferPass.h b/Engine/Plugins/MovieScene/MovieRenderPipeline/Source/MovieRenderPipelineRenderPasses/Public/MoviePipelineBackbufferPass.h
--- a/Engine/Plugins/MovieScene/MovieRenderPipeline/Source/MovieRenderPipelineRenderPasses/Public/MoviePipelineBackbufferPass.h
+++ b/Engine/Plugins/MovieScene/MovieRenderPipeline/Source/MovieRenderPipelineRenderPasses/Public/MoviePipelineBackbufferPass.h
@@ -13,6 +13,7 @@ class FSceneViewFamily;
 class FSceneView;
 class UWorld;
 struct FImagePixelPipe;
+struct FImageOverlappedAccumulator;
 
 UCLASS(Blueprintable)
 class MOVIERENDERPIPELINERENDERPASSES_API UMoviePipelineBackbufferPass : public UMoviePipelineRenderPass
@@ -21,7 +22,7 @@ class MOVIERENDERPIPELINERENDERPASSES_API UMoviePipelineBackbufferPass : public
 
 protected:
 
-	virtual void SetupImpl(const FMoviePipelineRenderPassInitSettings& InInitSettings, TSharedRef<FImagePixelPipe, ESPMode::ThreadSafe> InOutputPipe) override;
+	virtual void SetupImpl(const FMoviePipelineRenderPassInitSettings& InInitSettings) override;
 	virtual void TeardownImpl() override;
 	virtual void CaptureFrameImpl(const FMoviePipelineRenderPassMetrics& OutputFrameMetrics) override;
 	virtual void GatherOutputPassesImpl(TArray<FMoviePipelinePassIdentifier>& ExpectedRenderPasses) override;
@@ -35,5 +36,5 @@ private:
 	UPROPERTY(Transient)
 	UTextureRenderTarget2D* TileRenderTarget;
 
-	TSharedPtr<FImagePixelPipe, ESPMode::ThreadSafe> OutputPipe;
+	TSharedPtr<FImageOverlappedAccumulator, ESPMode::ThreadSafe> ImageTileAccumulator;
 };
\ No newline at end of file
diff --git a/Engine/Plugins/MovieScene/MovieRenderPipeline/Source/MovieRenderPipelineRenderPasses/Public/MoviePipelineImageSequenceContainer.h b/Engine/Plugins/MovieScene/MovieRenderPipeline/Source/MovieRenderPipelineRenderPasses/Public/MoviePipelineImageSequenceContainer.h
--- a/Engine/Plugins/MovieScene/MovieRenderPipeline/Source/MovieRenderPipelineRenderPasses/Public/MoviePipelineImageSequenceContainer.h
+++ b/Engine/Plugins/MovieScene/MovieRenderPipeline/Source/MovieRenderPipelineRenderPasses/Public/MoviePipelineImageSequenceContainer.h
@@ -15,6 +15,8 @@ class MOVIERENDERPIPELINERENDERPASSES_API UMoviePipelineImageSequenceContainerBa
 public:
 	UMoviePipelineImageSequenceContainerBase();
 
+	virtual void OnRecieveImageDataImpl(FMoviePipelineMergerOutputFrame* InMergedOutputFrame) override;
+
 protected:
 	// UMovieRenderPipelineOutputContainer interface
 	virtual void SetupForPipelineImpl(UMoviePipeline* InPipeline);
@@ -27,4 +29,6 @@ private:
 
 	/** A fence to keep track of when the Image Write queue has fully flushed. */
 	TFuture<void> FinalizeFence;
+
+	FString OutputDirectory;
 };
\ No newline at end of file
