Enabling motion blur with overlapped tile rendering.

diff --git a/Engine/Plugins/MovieScene/MovieRenderPipeline/Source/MovieRenderPipelineCore/Private/MoviePipelineRendering.cpp b/Engine/Plugins/MovieScene/MovieRenderPipeline/Source/MovieRenderPipelineCore/Private/MoviePipelineRendering.cpp
--- a/Engine/Plugins/MovieScene/MovieRenderPipeline/Source/MovieRenderPipelineCore/Private/MoviePipelineRendering.cpp
+++ b/Engine/Plugins/MovieScene/MovieRenderPipeline/Source/MovieRenderPipelineCore/Private/MoviePipelineRendering.cpp
@@ -15,6 +15,8 @@
 #include "Math/Halton.h"
 #include "ImageWriteTask.h"
 #include "ImageWriteQueue.h"
+#include "Engine/World.h"
+#include "GameFramework/PlayerController.h"
 
 void UMoviePipeline::SetupRenderingPipelineForShot(FMoviePipelineShotInfo& Shot)
 {
@@ -122,6 +124,19 @@ void UMoviePipeline::RenderFrame()
 	int32 NumSpatialSamples = AccumulationSettings->SpatialSampleCount;
 	ensure(NumTilesX > 0 && NumTilesY > 0 && NumSpatialSamples > 0);
 
+	FrameInfo.PrevViewLocation = FrameInfo.CurrViewLocation;
+	FrameInfo.PrevViewRotation = FrameInfo.CurrViewRotation;
+
+	APlayerController* LocalPlayerController = GetWorld()->GetFirstPlayerController();
+	LocalPlayerController->GetPlayerViewPoint(FrameInfo.CurrViewLocation, FrameInfo.CurrViewRotation);
+
+	// if it is the first frame, then motion blur is disabled, and we just use the current frame camera position/rotation
+	if (bIsHistoryOnlyFrame)
+	{
+		FrameInfo.PrevViewLocation = FrameInfo.CurrViewLocation;
+		FrameInfo.PrevViewRotation = FrameInfo.CurrViewRotation;
+	}
+
 	TArray<UMoviePipelineRenderPass*> InputBuffers = CurrentShot.ShotConfig->GetRenderPasses();
 
 	// If this is the first sample for a new frame, we want to notify the output builder that it should expect data to accumulate for this frame.
@@ -285,6 +300,9 @@ void UMoviePipeline::RenderFrame()
 					FrameMetrics.OverlappedPadY = 0;
 					FrameMetrics.OverlappedSubpixelShift = FVector2D(0.0f,0.0f);
 				}
+
+				FrameMetrics.FrameInfo = FrameInfo;
+
 				// Finally, we can ask each pass to render. Passes have already been initialized knowing how many tiles they're broken into
 				for (UMoviePipelineRenderPass* Input : InputBuffers)
 				{
diff --git a/Engine/Plugins/MovieScene/MovieRenderPipeline/Source/MovieRenderPipelineCore/Public/MoviePipeline.h b/Engine/Plugins/MovieScene/MovieRenderPipeline/Source/MovieRenderPipelineCore/Public/MoviePipeline.h
--- a/Engine/Plugins/MovieScene/MovieRenderPipeline/Source/MovieRenderPipelineCore/Public/MoviePipeline.h
+++ b/Engine/Plugins/MovieScene/MovieRenderPipeline/Source/MovieRenderPipelineCore/Public/MoviePipeline.h
@@ -244,6 +244,12 @@ private:
 	/** Called when there is a warning/error that the user should pay attention to.*/
 	FMoviePipelineErrored OnMoviePipelineErroredDelegate;
 
+	/**
+	 * We have to apply camera motion vectors manually. So we keep the current and previous fram'es camera view and rotation.
+	 * Then we render a sequence of the same movement, and update after running the game sim.
+	 **/
+	MoviePipeline::FMoviePipelineFrameInfo FrameInfo;
+
 public:
 	/** This gathers all of the produced data for an output frame (which may come in async many frames later) before passing them onto the Output Containers. */
 	TSharedPtr<FMoviePipelineOutputMerger, ESPMode::ThreadSafe> OutputBuilder;
diff --git a/Engine/Plugins/MovieScene/MovieRenderPipeline/Source/MovieRenderPipelineCore/Public/MovieRenderPipelineDataTypes.h b/Engine/Plugins/MovieScene/MovieRenderPipeline/Source/MovieRenderPipelineCore/Public/MovieRenderPipelineDataTypes.h
--- a/Engine/Plugins/MovieScene/MovieRenderPipeline/Source/MovieRenderPipelineCore/Public/MovieRenderPipelineDataTypes.h
+++ b/Engine/Plugins/MovieScene/MovieRenderPipeline/Source/MovieRenderPipelineCore/Public/MovieRenderPipelineDataTypes.h
@@ -102,6 +102,33 @@ public:
 
 namespace MoviePipeline
 {
+	/**
+	* Frame info needed for creating the FSceneView for each tile.
+	*/
+	struct FMoviePipelineFrameInfo
+	{
+	public:
+		FMoviePipelineFrameInfo() {}
+		virtual ~FMoviePipelineFrameInfo() {};
+
+		/**
+		 * We have to apply camera motion vectors manually. So we keep the current and previous fram'es camera view and rotation.
+		 * Then we render a sequence of the same movement, and update after running the game sim.
+		 **/
+
+		/** Current frame camera view location **/
+		FVector CurrViewLocation;
+
+		/** Current frame camera view rotation **/
+		FRotator CurrViewRotation;
+
+		/** Previous frame camera view location **/
+		FVector PrevViewLocation;
+
+		/** Previous frame camera view rotation **/
+		FRotator PrevViewRotation;
+	};
+
 	/**
 	* Utility function for the Pipeline to calculate the time metrics for the current
 	* frame that can then be used by the custom time step to avoid custom time logic.
@@ -650,6 +677,9 @@ public:
 	int32 OverlappedPadX;
 	int32 OverlappedPadY;
 	FVector2D OverlappedSubpixelShift;
+
+
+	MoviePipeline::FMoviePipelineFrameInfo FrameInfo;
 };
 
 struct FImagePixelDataPayload : IImagePixelDataPayload
@@ -743,7 +773,6 @@ public:
 	TArray<FString> ShotRenderMask;
 };
 
-
 struct MOVIERENDERPIPELINECORE_API FMoviePipelineMergerOutputFrame
 {
 public:
diff --git a/Engine/Plugins/MovieScene/MovieRenderPipeline/Source/MovieRenderPipelineRenderPasses/Private/MoviePipelineBackbufferPass.cpp b/Engine/Plugins/MovieScene/MovieRenderPipeline/Source/MovieRenderPipelineRenderPasses/Private/MoviePipelineBackbufferPass.cpp
--- a/Engine/Plugins/MovieScene/MovieRenderPipeline/Source/MovieRenderPipelineRenderPasses/Private/MoviePipelineBackbufferPass.cpp
+++ b/Engine/Plugins/MovieScene/MovieRenderPipeline/Source/MovieRenderPipelineRenderPasses/Private/MoviePipelineBackbufferPass.cpp
@@ -81,9 +81,9 @@ void UMoviePipelineBackbufferPass::CaptureFrameImpl(const FMoviePipelineRenderPa
 
 	if (InPassMetrics.bIsUsingOverlappedTiles)
 	{
-		// Disable these for now
+		// We can actually keep these enabled for now
 		//ShowFlags.SetDepthOfField(false);
-		ShowFlags.SetMotionBlur(false);
+		//ShowFlags.SetMotionBlur(false);
 	}
 
 	FSceneViewFamilyContext ViewFamily(FSceneViewFamily::ConstructionValues(
@@ -94,6 +94,14 @@ void UMoviePipelineBackbufferPass::CaptureFrameImpl(const FMoviePipelineRenderPa
 		.SetRealtimeUpdate(true));
 
 	ViewFamily.SetScreenPercentageInterface(new FLegacyScreenPercentageDriver(ViewFamily, 1.0f, false));
+
+
+	FVector ViewLocation;
+	FRotator ViewRotation;
+
+	//APlayerController* LocalPlayerController = GetWorld()->GetFirstPlayerController();
+	//LocalPlayerController->GetPlayerViewPoint(ViewLocation, ViewRotation);
+
 	// View is added as a child of the ViewFamily
 	FSceneView* View = CalcSceneView(&ViewFamily, InPassMetrics);
 	ViewFamily.FrameNumber += InPassMetrics.SpatialJitterIndex;
@@ -138,9 +146,8 @@ void UMoviePipelineBackbufferPass::CaptureFrameImpl(const FMoviePipelineRenderPa
 		View->bCameraCut = true;
 	}
 
-	// Pause the view family on renders after the first one. This prevents motion blur/view matrices from being
-	// updated
-	ViewFamily.bWorldIsPaused = InPassMetrics.SpatialJitterIndex != InPassMetrics.NumSpatialJitters - 1;
+	// The world is always paused, but we are setting the motion blur and view matrices manually
+	ViewFamily.bWorldIsPaused = true;
 
 	FCanvas Canvas = FCanvas(TileRenderTarget->GameThread_GetRenderTargetResource(), nullptr, GetWorld(), ERHIFeatureLevel::SM5, FCanvas::CDM_DeferDrawing, 1.0f);
 
@@ -171,20 +178,16 @@ void UMoviePipelineBackbufferPass::CaptureFrameImpl(const FMoviePipelineRenderPa
 
 FSceneView* UMoviePipelineBackbufferPass::CalcSceneView(FSceneViewFamily* ViewFamily, const FMoviePipelineRenderPassMetrics& InPassMetrics)
 {
-	FVector ViewLocation;
-	FRotator ViewRotation;
-
 	APlayerController* LocalPlayerController = GetWorld()->GetFirstPlayerController();
-	LocalPlayerController->GetPlayerViewPoint(ViewLocation, ViewRotation);
 
 	int32 TileSizeX = InitSettings.TileResolution.X;
 	int32 TileSizeY = InitSettings.TileResolution.Y;
 
 	FSceneViewInitOptions ViewInitOptions;
 	ViewInitOptions.ViewFamily = ViewFamily;
-	ViewInitOptions.ViewOrigin = ViewLocation;
+	ViewInitOptions.ViewOrigin = InPassMetrics.FrameInfo.CurrViewLocation;
 	ViewInitOptions.SetViewRectangle(FIntRect(FIntPoint(0, 0), FIntPoint(TileSizeX, TileSizeY)));
-	ViewInitOptions.ViewRotationMatrix = FInverseRotationMatrix(ViewRotation);
+	ViewInitOptions.ViewRotationMatrix = FInverseRotationMatrix(InPassMetrics.FrameInfo.CurrViewRotation);
 
 	// Rotate the view 90 degrees (reason: unknown)
 	ViewInitOptions.ViewRotationMatrix = ViewInitOptions.ViewRotationMatrix * FMatrix(
@@ -273,8 +276,9 @@ FSceneView* UMoviePipelineBackbufferPass::CalcSceneView(FSceneViewFamily* ViewFa
 
 	FSceneView* View = new FSceneView(ViewInitOptions);
 	ViewFamily->Views.Add(View);
-	View->ViewLocation = ViewLocation;
-	View->ViewRotation = ViewRotation;
+	View->ViewLocation = InPassMetrics.FrameInfo.CurrViewLocation;
+	View->ViewRotation = InPassMetrics.FrameInfo.CurrViewRotation;
+	View->PreviousViewTransform = FTransform(InPassMetrics.FrameInfo.PrevViewRotation,InPassMetrics.FrameInfo.PrevViewLocation);
 	View->StartFinalPostprocessSettings(View->ViewLocation);
 
 	// CameraAnim override
@@ -291,7 +295,6 @@ FSceneView* UMoviePipelineBackbufferPass::CalcSceneView(FSceneViewFamily* ViewFa
 		}
 
 		View->OverridePostProcessSettings(ViewInfo.PostProcessSettings, ViewInfo.PostProcessBlendWeight);
-
 	}
 
 	View->EndFinalPostprocessSettings(ViewInitOptions);
