Disabled auto-exposure for tiled high res rendering.

diff --git a/Engine/Plugins/MovieScene/MovieRenderPipeline/Source/MovieRenderPipelineRenderPasses/Private/MoviePipelineBackbufferPass.cpp b/Engine/Plugins/MovieScene/MovieRenderPipeline/Source/MovieRenderPipelineRenderPasses/Private/MoviePipelineBackbufferPass.cpp
--- a/Engine/Plugins/MovieScene/MovieRenderPipeline/Source/MovieRenderPipelineRenderPasses/Private/MoviePipelineBackbufferPass.cpp
+++ b/Engine/Plugins/MovieScene/MovieRenderPipeline/Source/MovieRenderPipelineRenderPasses/Private/MoviePipelineBackbufferPass.cpp
@@ -116,6 +116,14 @@ void UMoviePipelineBackbufferPass::CaptureFrameImpl(const FMoviePipelineRenderPa
 		View->FinalPostProcessSettings.bOverride_MotionBlurMax = true;
 	}
 
+	// Autoexposure is not allowed
+	if (View->FinalPostProcessSettings.AutoExposureMethod != EAutoExposureMethod::AEM_Manual)
+	{
+		UE_LOG(LogMovieRenderPipeline, Log, TEXT("AutoExposure Method should always be Manual."));
+
+		View->FinalPostProcessSettings.AutoExposureMethod = EAutoExposureMethod::AEM_Manual;
+	}
+
 	// Disable TAA while we're working on pixel accuracy testing
 	{
 		View->AntiAliasingMethod = EAntiAliasingMethod::AAM_None;
