Adding determinism to spatial jitter in overlapped rendering.

diff --git a/Engine/Plugins/MovieScene/MovieRenderPipeline/Source/MovieRenderPipelineRenderPasses/Private/MoviePipelineBackbufferPass.cpp b/Engine/Plugins/MovieScene/MovieRenderPipeline/Source/MovieRenderPipelineRenderPasses/Private/MoviePipelineBackbufferPass.cpp
--- a/Engine/Plugins/MovieScene/MovieRenderPipeline/Source/MovieRenderPipelineRenderPasses/Private/MoviePipelineBackbufferPass.cpp
+++ b/Engine/Plugins/MovieScene/MovieRenderPipeline/Source/MovieRenderPipelineRenderPasses/Private/MoviePipelineBackbufferPass.cpp
@@ -288,6 +288,7 @@ FSceneView* UMoviePipelineBackbufferPass::CalcSceneView(FSceneViewFamily* ViewFa
 	View->ViewRotation = InPassMetrics.FrameInfo.CurrViewRotation;
 	View->PreviousViewTransform = FTransform(InPassMetrics.FrameInfo.PrevViewRotation,InPassMetrics.FrameInfo.PrevViewLocation);
 	View->StartFinalPostprocessSettings(View->ViewLocation);
+	View->OverrideFrameIndexValue = InPassMetrics.SpatialJitterIndex;
 
 	// CameraAnim override
 	if (LocalPlayerController->PlayerCameraManager)
