Addnig half float accumulation.

diff --git a/Engine/Plugins/MovieScene/MovieRenderPipeline/Source/MovieRenderPipelineCore/Private/MovieRenderOverlappedImage.cpp b/Engine/Plugins/MovieScene/MovieRenderPipeline/Source/MovieRenderPipelineCore/Private/MovieRenderOverlappedImage.cpp
--- a/Engine/Plugins/MovieScene/MovieRenderPipeline/Source/MovieRenderPipelineCore/Private/MovieRenderOverlappedImage.cpp
+++ b/Engine/Plugins/MovieScene/MovieRenderPipeline/Source/MovieRenderPipelineCore/Private/MovieRenderOverlappedImage.cpp
@@ -713,6 +713,28 @@ void FImageOverlappedAccumulator::FetchFinalPixelDataByte(TArray64<FColor> & Out
 	}
 }
 
+void FImageOverlappedAccumulator::FetchFinalPixelDataHalfFloat(TArray64<FFloat16Color>& OutPixelData) const
+{
+	int32 FullSizeX = PlaneSize.X;
+	int32 FullSizeY = PlaneSize.Y;
+	OutPixelData.SetNumUninitialized(FullSizeX * FullSizeY);
+
+	for (int32 FullY = 0L; FullY < FullSizeY; FullY++)
+	{
+		for (int32 FullX = 0L; FullX < FullSizeX; FullX++)
+		{
+			float Rgba[4] = { 0.0f, 0.0f, 0.0f, 1.0f };
+
+			FetchFullImageValue(Rgba, FullX, FullY);
+			FLinearColor Color = FLinearColor(Rgba[0],Rgba[1],Rgba[2],Rgba[3]);
+
+			// be careful with this index, make sure to use 64bit math, not 32bit
+			int64 DstIndex = int64(FullY) * int64(FullSizeX) + int64(FullX);
+			OutPixelData[DstIndex] = FFloat16Color(Color);
+		}
+	}
+}
+
 void FImageOverlappedAccumulator::FetchFinalPixelDataLinearColor(TArray64<FLinearColor> & OutPixelData) const
 {
 	int32 FullSizeX = PlaneSize.X;
diff --git a/Engine/Plugins/MovieScene/MovieRenderPipeline/Source/MovieRenderPipelineCore/Public/MovieRenderOverlappedImage.h b/Engine/Plugins/MovieScene/MovieRenderPipeline/Source/MovieRenderPipelineCore/Public/MovieRenderOverlappedImage.h
--- a/Engine/Plugins/MovieScene/MovieRenderPipeline/Source/MovieRenderPipelineCore/Public/MovieRenderOverlappedImage.h
+++ b/Engine/Plugins/MovieScene/MovieRenderPipeline/Source/MovieRenderPipelineCore/Public/MovieRenderOverlappedImage.h
@@ -149,14 +149,21 @@ public:
 	 * So fetching as LinearColor would create a 16GB intermediary image, so it's worth having an option to fetch
 	 * straight to FColors.
 	 * 
-	 * @param FImagePixelData - Finished pixel data.
+	 * @param OutPixelData - Finished pixel data.
 	 */
 	void FetchFinalPixelDataByte(TArray64<FColor>& OutPixelData) const;
 
 	/**
 	 * After accumulation is finished, fetch the final image as linear colors
 	 * 
-	 * @param FImagePixelData - Finished pixel data.
+	 * @param OutPixelData - Finished pixel data.
+	 */
+	void FetchFinalPixelDataHalfFloat(TArray64<FFloat16Color>& OutPixelData) const;
+
+	/**
+	 * After accumulation is finished, fetch the final image as linear colors
+	 * 
+	 * @param OutPixelData - Finished pixel data.
 	 */
 	void FetchFinalPixelDataLinearColor(TArray64<FLinearColor>& OutPixelData) const;
 
