MoviePipeline: Separated user-defined passes (final pixels, ambient occlusion, etc.) from engine passes to allow each user-defined pass to have an accumulator and reduce the need for multiple engine renders when content could be shared.

diff --git a/Engine/Plugins/MovieScene/MovieRenderPipeline/Source/MovieRenderPipelineCore/Private/MoviePipeline.cpp b/Engine/Plugins/MovieScene/MovieRenderPipeline/Source/MovieRenderPipelineCore/Private/MoviePipeline.cpp
--- a/Engine/Plugins/MovieScene/MovieRenderPipeline/Source/MovieRenderPipelineCore/Private/MoviePipeline.cpp
+++ b/Engine/Plugins/MovieScene/MovieRenderPipeline/Source/MovieRenderPipelineCore/Private/MoviePipeline.cpp
@@ -31,6 +31,7 @@
 #include "MoviePipelineOutputSetting.h"
 #include "MoviePipelineBlueprintLibrary.h"
 #include "ImageWriteQueue.h"
+#include "MoviePipelineHighResSetting.h"
 
 #define LOCTEXT_NAMESPACE "MoviePipeline"
 
@@ -60,9 +61,6 @@ UMoviePipeline::UMoviePipeline()
 	OutputBuilder = MakeShared<FMoviePipelineOutputMerger, ESPMode::ThreadSafe>(this);
 
 	ImageWriteQueue = &FModuleManager::Get().LoadModuleChecked<IImageWriteQueueModule>("ImageWriteQueue").GetWriteQueue();
-
-	// Temp
-	OutputPipe = MakeShared<FImagePixelPipe, ESPMode::ThreadSafe>();
 }
 
 
@@ -420,10 +418,10 @@ FMoviePipelineWorkInfo UMoviePipeline::CalculateExpectedOutputMetrics()
 		ExpectedWork.NumCameraCuts += Shot.CameraCuts.Num();
 
 		int32 NumTiles = 0;
-		UMoviePipelineAccumulationSetting* TilingSettings = Shot.ShotConfig->FindSetting<UMoviePipelineAccumulationSetting>();
-		if (TilingSettings)
+		UMoviePipelineHighResSetting* HighResSettings = Shot.ShotConfig->FindSetting<UMoviePipelineHighResSetting>();
+		if (HighResSettings)
 		{
-			NumTiles = TilingSettings->TileCount * TilingSettings->TileCount;
+			NumTiles = HighResSettings->TileCount * HighResSettings->TileCount;
 		}
 
 		for (const FMoviePipelineCameraCutInfo& CameraCut : Shot.CameraCuts)
@@ -1019,4 +1017,9 @@ FFrameRate UMoviePipeline::GetEffectiveFrameRate() const
 	return FFrameRate();
 }
 
+UMoviePipelineShotConfig* UMoviePipeline::GetPipelineCurrentShotConfig() const
+{
+	return ShotList[CurrentShotIndex].ShotConfig;
+}
+
 #undef LOCTEXT_NAMESPACE // "MoviePipeline"
diff --git a/Engine/Plugins/MovieScene/MovieRenderPipeline/Source/MovieRenderPipelineCore/Private/MoviePipelineOutputBuilder.cpp b/Engine/Plugins/MovieScene/MovieRenderPipeline/Source/MovieRenderPipelineCore/Private/MoviePipelineOutputBuilder.cpp
--- a/Engine/Plugins/MovieScene/MovieRenderPipeline/Source/MovieRenderPipelineCore/Private/MoviePipelineOutputBuilder.cpp
+++ b/Engine/Plugins/MovieScene/MovieRenderPipeline/Source/MovieRenderPipelineCore/Private/MoviePipelineOutputBuilder.cpp
@@ -27,48 +27,45 @@ FMoviePipelineMergerOutputFrame& FMoviePipelineOutputMerger::QueueOutputFrame_Ga
 	return NewFrame;
 }
 
-void FMoviePipelineOutputMerger::OnSingleSampleDataAvailable_AnyThread(TUniquePtr<FImagePixelData>&& InData)
+void FMoviePipelineOutputMerger::OnSingleSampleDataAvailable_AnyThread(TUniquePtr<FImagePixelData>&& InData, const TSharedRef<FImagePixelDataPayload, ESPMode::ThreadSafe> InFrameData)
 {
 	// This is to support outputting individual samples (skipping accumulation) for debug reasons,
 	// or because you want to post-process them yourself. We just forward this directly on for output to disk.
 
 	TWeakObjectPtr<UMoviePipeline> LocalWeakPipeline = WeakMoviePipeline;
 
-	AsyncTask(ENamedThreads::GameThread, [LocalData = MoveTemp(InData), LocalWeakPipeline]() mutable
+	AsyncTask(ENamedThreads::GameThread, [LocalData = MoveTemp(InData), InFrameData, LocalWeakPipeline]() mutable
 	{
 		if (ensureAlwaysMsgf(LocalWeakPipeline.IsValid(), TEXT("A memory lifespan issue has left an output builder alive without an owning Movie Pipeline.")))
 		{
-			LocalWeakPipeline->OnSampleRendered(MoveTemp(LocalData));
+			LocalWeakPipeline->OnSampleRendered(MoveTemp(LocalData), InFrameData);
 		}
 	}
 	);
 }
 
-void FMoviePipelineOutputMerger::OnCompleteRenderPassDataAvailable_AnyThread(TUniquePtr<FImagePixelData>&& InData)
+void FMoviePipelineOutputMerger::OnCompleteRenderPassDataAvailable_AnyThread(TUniquePtr<FImagePixelData>&& InData, const TSharedRef<FImagePixelDataPayload, ESPMode::ThreadSafe> InFrameData)
 {
 	// Lock the ActiveData when we're updating what data has been gathered.
 	FScopeLock ScopeLock(&ActiveDataMutex);
 
-	FImagePixelDataPayload* Payload = InData->GetPayload<FImagePixelDataPayload>();
-	check(Payload);
-
 	// See if we can find the frame this data is for. This should always be valid, if it's not
 	// valid it means they either forgot to declare they were going to produce it, or this is
 	// coming in after the system already thinks it's finished that frame.
-	FMoviePipelineMergerOutputFrame* OutputFrame = ActiveData.Find(Payload->OutputState);
+	FMoviePipelineMergerOutputFrame* OutputFrame = ActiveData.Find(InFrameData->OutputState);
 	if (!ensureAlwaysMsgf(OutputFrame, TEXT("Recieved data for unknown frame. Frame was either already processed or not queued yet!")))
 	{
 		return;
 	}
 
 	// Ensure this pass is expected as well...
-	if (!ensureAlwaysMsgf(OutputFrame->ExpectedRenderPasses.Contains(Payload->PassIdentifier), TEXT("Recieved data for unexpected render pass: %s"), *Payload->PassIdentifier.Name))
+	if (!ensureAlwaysMsgf(OutputFrame->ExpectedRenderPasses.Contains(InFrameData->PassIdentifier), TEXT("Recieved data for unexpected render pass: %s"), *InFrameData->PassIdentifier.Name))
 	{
 		return;
 	}
 
 	// If this data was expected and this frame is still in progress, pass the data to the frame.
-	OutputFrame->ImageOutputData.FindOrAdd(Payload->PassIdentifier) = MoveTemp(InData);
+	OutputFrame->ImageOutputData.FindOrAdd(InFrameData->PassIdentifier) = MoveTemp(InData);
 
 	// Check to see if this was the last piece of data needed for this frame.
 	int32 TotalPasses = OutputFrame->ExpectedRenderPasses.Num();
@@ -78,16 +75,16 @@ void FMoviePipelineOutputMerger::OnCompleteRenderPassDataAvailable_AnyThread(TUn
 	{
 		// Transfer ownership from the map to here;
 		FMoviePipelineMergerOutputFrame FinalFrame;
-		ActiveData.RemoveAndCopyValue(Payload->OutputState, FinalFrame);
+		ActiveData.RemoveAndCopyValue(InFrameData->OutputState, FinalFrame);
 
 		// Notify the Movie Pipeline that this frame has been completed and can be forwarded to the output containers.
 		TWeakObjectPtr<UMoviePipeline> LocalWeakPipeline = WeakMoviePipeline;
 
-		AsyncTask(ENamedThreads::GameThread, [LocalFinalFrame = MoveTemp(FinalFrame), LocalWeakPipeline]() mutable
+		AsyncTask(ENamedThreads::GameThread, [LocalFinalFrame = MoveTemp(FinalFrame), InFrameData, LocalWeakPipeline]() mutable
 		{
 			if (ensureAlwaysMsgf(LocalWeakPipeline.IsValid(), TEXT("A memory lifespan issue has left an output builder alive without an owning Movie Pipeline.")))
 			{
-				LocalWeakPipeline->OnFrameCompletelyRendered(MoveTemp(LocalFinalFrame));
+				LocalWeakPipeline->OnFrameCompletelyRendered(MoveTemp(LocalFinalFrame), InFrameData);
 			}
 		}
 		);
diff --git a/Engine/Plugins/MovieScene/MovieRenderPipeline/Source/MovieRenderPipelineCore/Private/MoviePipelineRendering.cpp b/Engine/Plugins/MovieScene/MovieRenderPipeline/Source/MovieRenderPipelineCore/Private/MoviePipelineRendering.cpp
--- a/Engine/Plugins/MovieScene/MovieRenderPipeline/Source/MovieRenderPipelineCore/Private/MoviePipelineRendering.cpp
+++ b/Engine/Plugins/MovieScene/MovieRenderPipeline/Source/MovieRenderPipelineCore/Private/MoviePipelineRendering.cpp
@@ -17,6 +17,10 @@
 #include "ImageWriteQueue.h"
 #include "Engine/World.h"
 #include "GameFramework/PlayerController.h"
+#include "MoviePipelineHighResSetting.h"
+#include "Modules/ModuleManager.h"
+
+#define LOCTEXT_NAMESPACE "MoviePipeline"
 
 void UMoviePipeline::SetupRenderingPipelineForShot(FMoviePipelineShotInfo& Shot)
 {
@@ -37,32 +41,82 @@ void UMoviePipeline::SetupRenderingPipelineForShot(FMoviePipelineShotInfo& Shot)
 	UMoviePipelineOutputSetting* OutputSettings = GetPipelineMasterConfig()->FindSetting<UMoviePipelineOutputSetting>();
 	check(OutputSettings);
 
-	FMoviePipelineRenderPassInitSettings RenderPassInitSettings;
-	RenderPassInitSettings.TileResolution = FIntPoint(
-		FMath::CeilToInt(OutputSettings->OutputResolution.X / AccumulationSettings->TileCount),
-		FMath::CeilToInt(OutputSettings->OutputResolution.Y / AccumulationSettings->TileCount));
-	RenderPassInitSettings.TileCount = AccumulationSettings->TileCount;
-	RenderPassInitSettings.ShotConfig = Shot.ShotConfig;
-
-	if (AccumulationSettings->bIsUsingOverlappedTiles)
+	FIntPoint BackbufferResolution = OutputSettings->OutputResolution;
+	FIntPoint BackbufferTileCount = FIntPoint(1, 1);
+	
+	UMoviePipelineHighResSetting* HighResSettings = Shot.ShotConfig->FindSetting<UMoviePipelineHighResSetting>();
+	if (HighResSettings)
 	{
-		// this code is duplicated in a few places, we should clean this up
+		// Figure out how big each sub-region (tile) is.
+		BackbufferResolution = FIntPoint(
+			FMath::CeilToInt(BackbufferResolution.X / HighResSettings->TileCount),
+			FMath::CeilToInt(BackbufferResolution.Y / HighResSettings->TileCount));
 
-		int32 OverlappedPadX = int32(float(RenderPassInitSettings.TileResolution.X) * AccumulationSettings->PadRatioX);
-		int32 OverlappedPadY = int32(float(RenderPassInitSettings.TileResolution.Y) * AccumulationSettings->PadRatioY);
-		RenderPassInitSettings.TileResolution.X += 2 * OverlappedPadX;
-		RenderPassInitSettings.TileResolution.Y += 2 * OverlappedPadY;
-	}
+		// Then increase each sub-region by the overlap amount.
+		if (HighResSettings->bIsUsingOverlappedTiles)
+		{
+			BackbufferResolution = HighResSettings->CalculatePaddedBackbufferSize(BackbufferResolution);
+		}
 
+		// Note how many tiles we wish to render with.
+		BackbufferTileCount = FIntPoint(HighResSettings->TileCount, HighResSettings->TileCount);
+	}
 
+	// Initialize our render pass. This is a copy of the settings to make this less coupled to the Settings UI.
+	MoviePipeline::FMoviePipelineRenderPassInitSettings RenderPassInitSettings;
+	RenderPassInitSettings.BackbufferResolution = BackbufferResolution;
+	RenderPassInitSettings.TileCount = BackbufferTileCount;
 
 	// Code expects at least a 1x1 tile.
-	ensure(RenderPassInitSettings.TileCount > 0);
+	ensure(RenderPassInitSettings.TileCount.X > 0 && RenderPassInitSettings.TileCount.Y > 0);
+
+	// Now we need to look at all of the desired passes and find a unique set of actual engine passes that need
+	// to be rendered. This allows us to have multiple passes that re-use one render from the engine for efficiency.
+	TSet<FMoviePipelinePassIdentifier> RequiredEnginePasses;
 
 	for (UMoviePipelineRenderPass* Input : Shot.ShotConfig->GetRenderPasses())
 	{
-		Input->Setup(RenderPassInitSettings);
+		Input->GetRequiredEnginePasses(RequiredEnginePasses);
+	}
+
+	// There shouldn't be any render passes active from previous shots by now. The system should have flushed/stalled between
+	// them to complete using resources before switching passes.
+	check(ActiveRenderPasses.Num() == 0);
+
+
+	// Instantiate a new instance of every engine render pass we know how to use.
+	TArray<TSharedPtr<MoviePipeline::FMoviePipelineEnginePass>> RenderPasses;
+
+	FMovieRenderPipelineCoreModule& CoreModule = FModuleManager::Get().LoadModuleChecked<FMovieRenderPipelineCoreModule>("MovieRenderPipelineCore");
+	for (const FOnCreateEngineRenderPass& PassCreationDelegate : CoreModule.GetEngineRenderPasses())
+	{
+		TSharedRef<MoviePipeline::FMoviePipelineEnginePass> PassInstance = PassCreationDelegate.Execute();
+		if (RequiredEnginePasses.Contains(PassInstance->PassIdentifier))
+		{
+			ActiveRenderPasses.Add(PassInstance);
+			RequiredEnginePasses.Remove(PassInstance->PassIdentifier);
+		}
+	}
+
+	for(const FMoviePipelinePassIdentifier& RemainingPass : RequiredEnginePasses)
+	{
+		UE_LOG(LogMovieRenderPipeline, Warning, TEXT("Pass \"%d\" was listed as a required engine render pass but was not found. Did you forget to register it with the module?"), *RemainingPass.Name);
+		OnMoviePipelineErrored().Broadcast(this, true, LOCTEXT("MissingEnginePass", "A Render Pass specified an invalid Pass Identifier. Aborting Render. Check the log for more information."));
+	}
+
+	// Initialize each of the engine render passes.
+	for (TSharedPtr<MoviePipeline::FMoviePipelineEnginePass> EnginePass : ActiveRenderPasses)
+	{
+		EnginePass->Setup(MakeWeakObjectPtr(this), RenderPassInitSettings);
+	}
+
+	// We can now initialize the output passes and provide them a reference to the engine passes to get data from.
+	for (UMoviePipelineRenderPass* RenderPass : Shot.ShotConfig->GetRenderPasses())
+	{
+		RenderPass->Setup(ActiveRenderPasses);
 	}
+
+	UE_LOG(LogMovieRenderPipeline, Log, TEXT("Finished setting up rendering for shot. Shot has %d Engine Passes and %d Output Passes."), ActiveRenderPasses.Num(), Shot.ShotConfig->GetRenderPasses().Num());
 }
 
 void UMoviePipeline::TeardownRenderingPipelineForShot(FMoviePipelineShotInfo& Shot)
@@ -92,13 +146,6 @@ void UMoviePipeline::RenderFrame()
 		UE_LOG(LogMovieRenderPipeline, VeryVerbose, TEXT("[%d] Skipping RenderFrame() call due to state being ineligable for rendering. State: %d"), GFrameCounter, CurrentCameraCut.State);
 		return;
 	}
-
-	// Allow our containers to think
-	// ToDo: Is this necessaary?
-	// for (UMoviePipelineOutputBase* Container : Config->OutputContainers)
-	// {
-	// 	Container->OnPostTick();
-	// }
 	
 	// To produce a frame from the movie pipeline we may render many frames over a period of time, additively collecting the results
 	// together before submitting it for writing on the last result - this is referred to as an "output frame". The 1 (or more) samples
@@ -113,16 +160,23 @@ void UMoviePipeline::RenderFrame()
 	// In short, for each output frame, for each accumulation frame, for each tile X/Y, for each jitter, we render a pass. This setup is
 	// designed to maximize the likely hood of deterministic rendering and that different passes line up with each other.
 	UMoviePipelineAccumulationSetting* AccumulationSettings = CurrentShot.ShotConfig->FindOrAddSetting<UMoviePipelineAccumulationSetting>();
-	UMoviePipelineOutputSetting* OutputSetting = GetPipelineMasterConfig()->FindSetting<UMoviePipelineOutputSetting>();
-	check(OutputSetting);
+	UMoviePipelineOutputSetting* OutputSettings = GetPipelineMasterConfig()->FindSetting<UMoviePipelineOutputSetting>();
+	check(OutputSettings);
+
+	FIntPoint TileCount = FIntPoint(1, 1);
+	FIntPoint OutputResolution = OutputSettings->OutputResolution;
 
+	bool bIsUsingOverlappedTiles = false;
 
-	const int32 NumTilesX = AccumulationSettings->TileCount;
-	const int32 NumTilesY = NumTilesX;
-	const bool bIsUsingOverlappedTiles = AccumulationSettings->bIsUsingOverlappedTiles;
+	UMoviePipelineHighResSetting* HighResSettings = CurrentShot.ShotConfig->FindSetting<UMoviePipelineHighResSetting>();
+	if (HighResSettings)
+	{
+		TileCount = FIntPoint(HighResSettings->TileCount, HighResSettings->TileCount);
+		bIsUsingOverlappedTiles = HighResSettings->bIsUsingOverlappedTiles;
+	}
 
 	int32 NumSpatialSamples = AccumulationSettings->SpatialSampleCount;
-	ensure(NumTilesX > 0 && NumTilesY > 0 && NumSpatialSamples > 0);
+	ensure(TileCount.X > 0 && TileCount .Y> 0 && NumSpatialSamples > 0);
 
 	FrameInfo.PrevViewLocation = FrameInfo.CurrViewLocation;
 	FrameInfo.PrevViewRotation = FrameInfo.CurrViewRotation;
@@ -152,10 +206,9 @@ void UMoviePipeline::RenderFrame()
 		}
 	}
 
-	// Each render has been initialized with ResX/NumTiles, ResY/NumTiles already to know how big each resulting render target should be.
-	for (int32 TileY = 0; TileY < NumTilesY; TileY++)
+	for (int32 TileY = 0; TileY < TileCount.Y; TileY++)
 	{
-		for (int32 TileX = 0; TileX < NumTilesX; TileX++)
+		for (int32 TileX = 0; TileX < TileCount.X; TileX++)
 		{
 			/*
 			* Two different features need to move the viewport and need to work in combination with each other.
@@ -171,7 +224,7 @@ void UMoviePipeline::RenderFrame()
 			* Single-frame anti-aliasing works in a similar way, by shifting the view by less than a pixel for each
 			* sample in the anti-aliasing. While the high-resolution rendering is a specific grid pattern, the
 			* anti-aliasing uses a more nuanced technique for picking offsets. It uses a deterministic Halton
-			* sequence to pick the offset, weighted by a guassian curve to mostly focus on samples around the
+			* sequence to pick the offset, weighted by a Gaussian curve to mostly focus on samples around the
 			* center, but occasionally picking an offset which would be a sample outside of it's current sub-pixel.
 			*
 			* In the example below we're looking at one render target pixel, split into a 2x2 tiling. To gain
@@ -205,10 +258,10 @@ void UMoviePipeline::RenderFrame()
 			// 'quadrant' we need to take our current tile index and multiply it by the size of each tile, so
 			// for the top right quadrant (in above example) size=(1/2), offset=(size*1). Finally we add half
 			// of the size to put us in the center of that pixel.
-			double TilePixelSizeFractionX = 1.0 / NumTilesX;
-			double TilePixelSizeFractionY = 1.0 / NumTilesY;
-			double TileShiftX = (TilePixelSizeFractionX * TileX) + (TilePixelSizeFractionX / 2.0);
-			double TileShiftY = (TilePixelSizeFractionY * TileY) + (TilePixelSizeFractionY / 2.0);
+			double TilePixelSizeFractionX;
+			double TilePixelSizeFractionY; 
+			double TileShiftX; 
+			double TileShiftY;
 
 			if (bIsUsingOverlappedTiles)
 			{
@@ -218,16 +271,23 @@ void UMoviePipeline::RenderFrame()
 				TilePixelSizeFractionX = 1.0f;
 				TilePixelSizeFractionY = 1.0f;
 			}
+			else
+			{
+				// Legacy interleaving
+				TilePixelSizeFractionX = 1.0 / TileCount.X;
+				TilePixelSizeFractionY = 1.0 / TileCount.Y;
+				TileShiftX = (TilePixelSizeFractionX * TileX) + (TilePixelSizeFractionX / 2.0);
+				TileShiftY = (TilePixelSizeFractionY * TileY) + (TilePixelSizeFractionY / 2.0);
+			}
 
-			// If they're trying to seed the histories, then we'll just override the NumSpatialSamples since they're contributing towards
-			// our final target buffer which needs the history.
+			// If they're trying to seed the histories, then we'll just override the NumSpatialSamples since they don't need all samples as
+			// this isn't output.
 			if (bIsHistoryOnlyFrame)
 			{
 				NumSpatialSamples = 1;
 			}
 
-			// Now we want to render a user-configured number of jitters to come up with the final output for this tile. These are considered
-			// spatial jitters which are accumulated together before contributing to the temporal accumulation.
+			// Now we want to render a user-configured number of spatial jitters to come up with the final output for this tile. 
 			for (int32 SpatialSample = 0; SpatialSample < NumSpatialSamples; SpatialSample++)
 			{
 				// Count this as a sample rendered for the current work.
@@ -236,9 +296,9 @@ void UMoviePipeline::RenderFrame()
 				float OffsetX = 0.5f;
 				float OffsetY = 0.5f; 
 
-				// CachedOutputState.TemporalSampleIndex is -1 on the first frame for motion blur, so ignore jitter on this frame.
+				// Only jitter when required, as jitter requires TAA to be off. 
 				// Also, we want jitter if either time samples or spatial samples are > 1, otherwise not.
-				if (CachedOutputState.TemporalSampleIndex >= 0 &&
+				if (!bIsHistoryOnlyFrame &&
 					(AccumulationSettings->SpatialSampleCount > 1 || AccumulationSettings->TemporalSampleCount > 1))
 				{
 					OffsetX = Halton((SpatialSample + 1) + (CachedOutputState.TemporalSampleIndex * NumSpatialSamples), 2);
@@ -252,76 +312,55 @@ void UMoviePipeline::RenderFrame()
 				double FinalSubPixelShiftX = - (TileShiftX + SpatialShiftX - 0.5);
 				double FinalSubPixelShiftY = (TileShiftY + SpatialShiftY - 0.5);
 
-				int32 TileSizeX = FMath::CeilToInt(OutputSetting->OutputResolution.X / NumTilesX);
-				int32 TileSizeY = FMath::CeilToInt(OutputSetting->OutputResolution.Y / NumTilesY);
-
-				// in the case of overlapped rendering, the render target size will be larger than tile size because of overlap
-				int32 PadX = bIsUsingOverlappedTiles ? int32(float(TileSizeX) * AccumulationSettings->PadRatioX) : 0;
-				int32 PadY = bIsUsingOverlappedTiles ? int32(float(TileSizeY) * AccumulationSettings->PadRatioY) : 0;
-
-				int TargetSizeX = TileSizeX + 2 * PadX;
-				int TargetSizeY = TileSizeY + 2 * PadY;
-
-				FMoviePipelineRenderPassMetrics FrameMetrics;
-				FrameMetrics.OutputState = CachedOutputState;
-				FrameMetrics.TileIndex = (TileY * NumTilesY) + TileX;
-				FrameMetrics.SpatialShift = FVector2D((float)( FinalSubPixelShiftX) * 2.0f / TargetSizeX, (float)FinalSubPixelShiftY * 2.0f / TargetSizeX);
-				FrameMetrics.TileIndexX = TileX;
-				FrameMetrics.TileIndexY = TileY;
-				FrameMetrics.NumTilesX = NumTilesX;
-				FrameMetrics.NumTilesY = NumTilesY;
-				FrameMetrics.bIsHistoryOnlyFrame = bIsHistoryOnlyFrame;
-				FrameMetrics.SpatialJitterIndex = SpatialSample;
-				FrameMetrics.NumSpatialJitters = NumSpatialSamples;
-				FrameMetrics.TemporalJitterIndex = CachedOutputState.TemporalSampleIndex;
-				FrameMetrics.NumTemporalJitters = AccumulationSettings->TemporalSampleCount;
-				FrameMetrics.AccumulationGamma = AccumulationSettings->AccumulationGamma;
-				FrameMetrics.JitterOffsetX = 1.0f - OffsetX; // reversing to match FinalSubPixelShiftX
-				FrameMetrics.JitterOffsetY = OffsetY;
-				FrameMetrics.bIsUsingOverlappedTiles = bIsUsingOverlappedTiles;
-
+				FIntPoint BackbufferResolution;
 				if (bIsUsingOverlappedTiles)
 				{
-					FrameMetrics.OverlappedPadX = int32(float(TileSizeX) * AccumulationSettings->PadRatioX);
-					FrameMetrics.OverlappedPadY = int32(float(TileSizeY) * AccumulationSettings->PadRatioY);
-					FrameMetrics.OverlappedOffsetX = TileX * TileSizeX - FrameMetrics.OverlappedPadX;
-					FrameMetrics.OverlappedOffsetY = TileY * TileSizeY - FrameMetrics.OverlappedPadY;
-					FrameMetrics.OverlappedSizeX = TileSizeX;
-					FrameMetrics.OverlappedSizeY = TileSizeY;
-					FrameMetrics.OverlappedSubpixelShift = FVector2D(1.0f - OffsetX,OffsetY);
+					BackbufferResolution = HighResSettings->CalculatePaddedBackbufferSize(BackbufferResolution);
 				}
 				else
 				{
-					FrameMetrics.OverlappedOffsetX = 0;
-					FrameMetrics.OverlappedOffsetY = 0;
-					FrameMetrics.OverlappedSizeX = 0;
-					FrameMetrics.OverlappedSizeY = 0;
-					FrameMetrics.OverlappedPadX = 0;
-					FrameMetrics.OverlappedPadY = 0;
-					FrameMetrics.OverlappedSubpixelShift = FVector2D(0.0f,0.0f);
+					// Legacy interleaving
+					BackbufferResolution.X = FMath::CeilToInt(OutputResolution.X / TileCount.X);
+					BackbufferResolution.Y = FMath::CeilToInt(OutputResolution.Y / TileCount.Y);
 				}
 
-				FrameMetrics.FrameInfo = FrameInfo;
+				// We take all of the information needed to render a single sample and package it into a struct.
+				FMoviePipelineRenderPassMetrics SampleState;
+				SampleState.OutputState = CachedOutputState;
+				SampleState.SpatialShift = FVector2D((float)( FinalSubPixelShiftX) * 2.0f / BackbufferResolution.X, (float)FinalSubPixelShiftY * 2.0f / BackbufferResolution.X);
+				SampleState.TileIndexes = FIntPoint(TileX, TileY);
+				SampleState.TileCounts = TileCount;
+				SampleState.bIsHistoryOnlyFrame = bIsHistoryOnlyFrame;
+				SampleState.SpatialSampleIndex = SpatialSample;
+				SampleState.SpatialSampleCount = NumSpatialSamples;
+				SampleState.TemporalSampleIndex = CachedOutputState.TemporalSampleIndex;
+				SampleState.TemporalSampleCount = AccumulationSettings->TemporalSampleCount;
+				SampleState.AccumulationGamma = AccumulationSettings->AccumulationGamma;
+				SampleState.bIsUsingOverlappedTiles = bIsUsingOverlappedTiles;
+				SampleState.BackbufferSize = BackbufferResolution;
+				SampleState.FrameInfo = FrameInfo;
 
-				// Finally, we can ask each pass to render. Passes have already been initialized knowing how many tiles they're broken into
-				for (UMoviePipelineRenderPass* Input : InputBuffers)
+				if (bIsUsingOverlappedTiles)
+				{
+					SampleState.OverlappedPad = FIntPoint(FMath::CeilToInt(BackbufferResolution.X * HighResSettings->OverlapPercentage), 
+														   FMath::CeilToInt(BackbufferResolution.Y * HighResSettings->OverlapPercentage));
+					SampleState.OverlappedOffset = FIntPoint(TileX * BackbufferResolution.X - SampleState.OverlappedPad.X, 
+															  TileY * BackbufferResolution.Y - SampleState.OverlappedPad.Y);
+					SampleState.OverlappedSubpixelShift = FVector2D(1.0f - OffsetX,OffsetY);
+				}
+				else
 				{
-					// Our Cached Output state has the expected output for this frame.
-					Input->CaptureFrame(FrameMetrics);
-
-					// If we needed to see each individual jittered image, we would need to
-					// enqueue a copy of the pass's render target here. ToDo:
-
-					// We can calculate the influence of the results of this pass before accumulating
-					// If we are accumulating to a dedicated target for sub-frame jitter results, then
-					// the influence of this pass is simply AccumTarget += ThisResult * (1/NumSpatialSamples);
-					// However, to minimize memory usage we want to use a single accumulation target for
-					// each pass on each tile. To simulate a camera shutter, different temporal sample frames
-					// may have different weights (ie: the first and last sample may be slightly darker), and
-					// in this case the influence of this pass is:
-					// AccumTarget += (ThisResult * (1/NumSpatialSamples)) * GetWeightForTemporalSample(CachedOutputState.SubFrameIndex);
+					SampleState.OverlappedOffset = FIntPoint(0,0);
+					SampleState.OverlappedPad = FIntPoint(0, 0);
+					SampleState.OverlappedSubpixelShift = FVector2D(0.0f,0.0f);
 				}
 
+				// Now we can request that all of the engine passes render. The individual render passes should have already registered delegates
+				// to receive data when the engine render pass is run, so no need to run them.
+				for (TSharedPtr<MoviePipeline::FMoviePipelineEnginePass> EnginePass : ActiveRenderPasses)
+				{
+					EnginePass->RenderSample_GameThread(SampleState);
+				}
 			}
 		}
 	}
@@ -329,39 +368,54 @@ void UMoviePipeline::RenderFrame()
 	// UE_LOG(LogMovieRenderPipeline, Warning, TEXT("[%d] Pre-FlushRenderingCommands"), GFrameCounter);
 	FlushRenderingCommands();
 	// UE_LOG(LogMovieRenderPipeline, Warning, TEXT("[%d] Post-FlushRenderingCommands"), GFrameCounter);
-
 }
 
-void UMoviePipeline::OnFrameCompletelyRendered(FMoviePipelineMergerOutputFrame&& OutputFrame)
+void UMoviePipeline::OnFrameCompletelyRendered(FMoviePipelineMergerOutputFrame&& OutputFrame, const TSharedRef<FImagePixelDataPayload, ESPMode::ThreadSafe> InFrameData)
 {
 	UE_LOG(LogMovieRenderPipeline, Warning, TEXT("[%d] Data required for output available! Frame: %d"), GFrameCounter, OutputFrame.FrameOutputState.OutputFrameNumber);
-
+	
 	for (UMoviePipelineOutputBase* OutputContainer : GetPipelineMasterConfig()->GetOutputContainers())
 	{
 		OutputContainer->OnRecieveImageDataImpl(&OutputFrame);
 	}
 }
 
-void UMoviePipeline::OnSampleRendered(TUniquePtr<FImagePixelData>&& OutputSample)
+void UMoviePipeline::OnSampleRendered(TUniquePtr<FImagePixelData>&& OutputSample, const TSharedRef<FImagePixelDataPayload, ESPMode::ThreadSafe> InFrameData)
 {
-	FImagePixelDataPayload* Payload = OutputSample->GetPayload< FImagePixelDataPayload>();
-	check(Payload);
-
 	UMoviePipelineOutputSetting* OutputSettings = GetPipelineMasterConfig()->FindSetting<UMoviePipelineOutputSetting>();
 	check(OutputSettings);
 
 	TUniquePtr<FImageWriteTask> TileImageTask = MakeUnique<FImageWriteTask>();
 
 	// Fill alpha for now 
-	TileImageTask->PixelPreProcessors.Add(TAsyncAlphaWrite<FColor>(255));
+	switch (OutputSample->GetType())
+	{
+		case EImagePixelType::Color:
+		{
+			TileImageTask->PixelPreProcessors.Add(TAsyncAlphaWrite<FColor>(255));
+			break;
+		}
+		case EImagePixelType::Float16:
+		{
+
+			break;
+		}
+		case EImagePixelType::Float32:
+		{
+			TileImageTask->PixelPreProcessors.Add(TAsyncAlphaWrite<FLinearColor>(255));
+			break;
+		}
+		default:
+			check(false);
+	}
 
 	// JPEG output
-	TileImageTask->Format = EImageFormat::JPEG;
+	TileImageTask->Format = EImageFormat::EXR;
 	TileImageTask->CompressionQuality = 100;
 
 	FString OutputName = FString::Printf(TEXT("/%s_SS_%d_TS_%d_TileX_%d_TileY_%d.%d.jpeg"),
-		*Payload->PassIdentifier.Name, Payload->SpatialJitterIndex, Payload->TemporalJitterIndex,
-		Payload->TileIndexX, Payload->TileIndexY, Payload->OutputState.OutputFrameNumber);
+		*InFrameData->PassIdentifier.Name, InFrameData->SampleState.SpatialSampleIndex, InFrameData->SampleState.TemporalSampleIndex,
+		InFrameData->SampleState.TileIndexes.X, InFrameData->SampleState.TileIndexes.Y, InFrameData->OutputState.OutputFrameNumber);
 
 	FString OutputDirectory = OutputSettings->OutputDirectory.Path;
 	FString OutputPath = OutputDirectory + OutputName;
@@ -371,3 +425,5 @@ void UMoviePipeline::OnSampleRendered(TUniquePtr<FImagePixelData>&& OutputSample
 	TileImageTask->PixelData = MoveTemp(OutputSample);
 	ImageWriteQueue->Enqueue(MoveTemp(TileImageTask));
 }
+
+#undef LOCTEXT_NAMESPACE // "MoviePipeline"
diff --git a/Engine/Plugins/MovieScene/MovieRenderPipeline/Source/MovieRenderPipelineCore/Private/MovieRenderOverlappedImage.cpp b/Engine/Plugins/MovieScene/MovieRenderPipeline/Source/MovieRenderPipelineCore/Private/MovieRenderOverlappedImage.cpp
--- a/Engine/Plugins/MovieScene/MovieRenderPipeline/Source/MovieRenderPipelineCore/Private/MovieRenderOverlappedImage.cpp
+++ b/Engine/Plugins/MovieScene/MovieRenderPipeline/Source/MovieRenderPipelineCore/Private/MovieRenderOverlappedImage.cpp
@@ -77,19 +77,19 @@ void FImageOverlappedPlane::AccumulateSinglePlane(const TArray64<float>& InRawDa
 			// Given a position on the current tile (x,y), we apply the sum of the 4 points (x+0,y+0), (x+1,y+0), (x+0,y+1), (x+1,y+1) to the destination index.
 			// So in reverse, given a destination position, our source sample positions are (x-1,y-1), (x+0,y-1), (x-1,y+0), (x+0,y+0).
 			// That's why we make sure to start at a minimum of index 1, instead of 0.
-			for (int CurrY = 1; CurrY < InSize.Y; CurrY++)
+			for (int32 CurrY = 1; CurrY < InSize.Y; CurrY++)
 			{
-				for (int CurrX = 1; CurrX < InSize.X; CurrX++)
+				for (int32 CurrX = 1; CurrX < InSize.X; CurrX++)
 				{
-					int DstY = StartY + CurrY;// +OffsetY;
-					int DstX = StartX + CurrX;// +OffsetX;
+					int32 DstY = StartY + CurrY;// +OffsetY;
+					int32 DstX = StartX + CurrX;// +OffsetX;
 
 					if (DstX >= 0 && DstY >= 0 &&
 						DstX < Size.X && DstY < Size.Y)
 					{
-						for (int OffsetY = 0; OffsetY < 2; OffsetY++)
+						for (int32 OffsetY = 0; OffsetY < 2; OffsetY++)
 						{
-							for (int OffsetX = 0; OffsetX < 2; OffsetX++)
+							for (int32 OffsetX = 0; OffsetX < 2; OffsetX++)
 							{
 								float Val = InRawData[(CurrY - 1 + OffsetY) * InSize.X + (CurrX - 1 + OffsetX)];
 								float BaseWeight = InWeightData[(CurrY - 1 + OffsetY) * InSize.X + (CurrX - 1 + OffsetX)];
@@ -106,19 +106,19 @@ void FImageOverlappedPlane::AccumulateSinglePlane(const TArray64<float>& InRawDa
 		else
 		{
 			// Slow, reference version. This is the main one.
-			for (int CurrY = 0; CurrY < InSize.Y; CurrY++)
+			for (int32 CurrY = 0; CurrY < InSize.Y; CurrY++)
 			{
-				for (int CurrX = 0; CurrX < InSize.X; CurrX++)
+				for (int32 CurrX = 0; CurrX < InSize.X; CurrX++)
 				{
 					float Val = InRawData[CurrY * InSize.X + CurrX];
 					float BaseWeight = InWeightData[CurrY * InSize.X + CurrX];
 
-					for (int OffsetY = 0; OffsetY < 2; OffsetY++)
+					for (int32 OffsetY = 0; OffsetY < 2; OffsetY++)
 					{
-						for (int OffsetX = 0; OffsetX < 2; OffsetX++)
+						for (int32 OffsetX = 0; OffsetX < 2; OffsetX++)
 						{
-							int DstY = StartY + CurrY + OffsetY;
-							int DstX = StartX + CurrX + OffsetX;
+							int32 DstY = StartY + CurrY + OffsetY;
+							int32 DstX = StartX + CurrX + OffsetX;
 
 							float Weight = BaseWeight * PixelWeight[OffsetY][OffsetX];
 
@@ -279,7 +279,7 @@ void FImageOverlappedPlane::AccumulateSinglePlane(const TArray64<float>& InRawDa
 	}
 }
 
-void FImageOverlappedAccumulator::InitMemory(FIntPoint InPlaneSize, int InNumChannels)
+void FImageOverlappedAccumulator::InitMemory(FIntPoint InPlaneSize, int32 InNumChannels)
 {
 	PlaneSize.X = InPlaneSize.X;
 	PlaneSize.Y = InPlaneSize.Y;
@@ -287,7 +287,7 @@ void FImageOverlappedAccumulator::InitMemory(FIntPoint InPlaneSize, int InNumCha
 
 	ChannelPlanes.SetNum(NumChannels);
 
-	for (int Channel = 0; Channel < NumChannels; Channel++)
+	for (int32 Channel = 0; Channel < NumChannels; Channel++)
 	{
 		ChannelPlanes[Channel].Init(PlaneSize);
 	}
@@ -303,7 +303,7 @@ void FImageOverlappedAccumulator::ZeroPlanes()
 {
 	check(ChannelPlanes.Num() == NumChannels);
 
-	for (int Channel = 0; Channel < NumChannels; Channel++)
+	for (int32 Channel = 0; Channel < NumChannels; Channel++)
 	{
 		ChannelPlanes[Channel].ZeroPlane();
 	}
@@ -336,9 +336,9 @@ void FImageOverlappedAccumulator::GenerateTileWeight(TArray64<float>& OutWeights
 	float ScaleX = 1.0f / (float(Size.X / 2) * .75f);
 	float ScaleY = 1.0f / (float(Size.Y / 2) * .75f);
 
-	for (int PixY = 0; PixY < Size.Y; PixY++)
+	for (int32 PixY = 0; PixY < Size.Y; PixY++)
 	{
-		for (int PixX = 0; PixX < Size.X; PixX++)
+		for (int32 PixX = 0; PixX < Size.X; PixX++)
 		{
 			float Y = float(PixY) + .5f;
 			float X = float(PixX) + .5f;
@@ -560,7 +560,7 @@ void FImageOverlappedAccumulator::AccumulatePixelData(const FImagePixelData& InP
 			// Unfortunately, we don't have an SSE optimized pow function. This function is quite slow (about 30-40ms).
 			float Gamma = AccumulationGamma;
 
-			for (int ChanIter = 0; ChanIter < NumChannels; ChanIter++)
+			for (int32 ChanIter = 0; ChanIter < NumChannels; ChanIter++)
 			{
 				float* DstData = RawData[ChanIter].GetData();
 
diff --git a/Engine/Plugins/MovieScene/MovieRenderPipeline/Source/MovieRenderPipelineCore/Private/MovieRenderPipelineCoreModule.cpp b/Engine/Plugins/MovieScene/MovieRenderPipeline/Source/MovieRenderPipelineCore/Private/MovieRenderPipelineCoreModule.cpp
--- a/Engine/Plugins/MovieScene/MovieRenderPipeline/Source/MovieRenderPipelineCore/Private/MovieRenderPipelineCoreModule.cpp
+++ b/Engine/Plugins/MovieScene/MovieRenderPipeline/Source/MovieRenderPipelineCore/Private/MovieRenderPipelineCoreModule.cpp
@@ -4,6 +4,7 @@
 #include "Modules/ModuleInterface.h"
 #include "Misc/CoreDelegates.h"
 #include "Modules/ModuleManager.h"
+#include "MovieRenderPipelineDataTypes.h"
 
 void FMovieRenderPipelineCoreModule::StartupModule()
 {
@@ -20,5 +21,18 @@ void FMovieRenderPipelineCoreModule::ShutdownModule()
 {
 }
 
+FDelegateHandle FMovieRenderPipelineCoreModule::RegisterEngineRenderPass(FOnCreateEngineRenderPass InOnCreateEngineRenderPass)
+{
+	EngineRenderPassDelegates.Add(InOnCreateEngineRenderPass);
+	FDelegateHandle Handle = EngineRenderPassDelegates.Last().GetHandle();
+
+	return Handle;
+}
+
+void FMovieRenderPipelineCoreModule::UnregisterEngineRenderPass(FDelegateHandle InHandle)
+{
+	EngineRenderPassDelegates.RemoveAll([=](const FOnCreateEngineRenderPass& Delegate) { return Delegate.GetHandle() == InHandle; });
+}
+
 IMPLEMENT_MODULE(FMovieRenderPipelineCoreModule, MovieRenderPipelineCore);
 DEFINE_LOG_CATEGORY(LogMovieRenderPipeline); 
diff --git a/Engine/Plugins/MovieScene/MovieRenderPipeline/Source/MovieRenderPipelineCore/Public/MoviePipeline.h b/Engine/Plugins/MovieScene/MovieRenderPipeline/Source/MovieRenderPipelineCore/Public/MoviePipeline.h
--- a/Engine/Plugins/MovieScene/MovieRenderPipeline/Source/MovieRenderPipelineCore/Public/MoviePipeline.h
+++ b/Engine/Plugins/MovieScene/MovieRenderPipeline/Source/MovieRenderPipelineCore/Public/MoviePipeline.h
@@ -74,6 +74,12 @@ public:
 	*/
 	UFUNCTION(BlueprintPure, Category = "Movie Render Pipeline")
 	UMoviePipelineMasterConfig* GetPipelineMasterConfig() const { return CurrentJob.Configuration; }
+
+	/**
+	* Get the shot config that is being used to render the current shot.
+	*/
+	UFUNCTION(BlueprintPure, Category = "Movie Render Pipeline")
+	UMoviePipelineShotConfig* GetPipelineCurrentShotConfig() const;
 public:
 
 	UFUNCTION(BlueprintPure, Category = "Movie Render Pipeline")
@@ -93,8 +99,8 @@ public:
 	FFrameRate GetEffectiveFrameRate() const;
 
 public:
-	void OnFrameCompletelyRendered(FMoviePipelineMergerOutputFrame&& OutputFrame);
-	void OnSampleRendered(TUniquePtr<FImagePixelData>&& OutputSample);
+	void OnFrameCompletelyRendered(FMoviePipelineMergerOutputFrame&& OutputFrame, const TSharedRef<FImagePixelDataPayload, ESPMode::ThreadSafe> InFrameData);
+	void OnSampleRendered(TUniquePtr<FImagePixelData>&& OutputSample, const TSharedRef<FImagePixelDataPayload, ESPMode::ThreadSafe> InFrameData);
 private:
 
 	/** Instantiate our Debug UI Widget and initialize it to ourself. */
@@ -251,16 +257,15 @@ private:
 	MoviePipeline::FMoviePipelineFrameInfo FrameInfo;
 
 public:
+	/** A list of engine passes which need to be run each frame to generate required content for all the movie render passes. */
+	TArray<TSharedPtr<MoviePipeline::FMoviePipelineEnginePass>> ActiveRenderPasses;
+
 	/** This gathers all of the produced data for an output frame (which may come in async many frames later) before passing them onto the Output Containers. */
 	TSharedPtr<FMoviePipelineOutputMerger, ESPMode::ThreadSafe> OutputBuilder;
 
 	/** A debug image sequence writer in the event they want to dump every sample generated on its own. */
 	IImageWriteQueue* ImageWriteQueue;
 
-	// debug
-	TSharedPtr<FImagePixelPipe, ESPMode::ThreadSafe> GetOutputPipe() const { return OutputPipe; }
-	TSharedPtr<FImagePixelPipe, ESPMode::ThreadSafe> OutputPipe;
-
 private:
 	/** Keep track of which job we're working on. This holds our Configuration + which shots we're supposed to render from it. */
 	UPROPERTY(Transient)
diff --git a/Engine/Plugins/MovieScene/MovieRenderPipeline/Source/MovieRenderPipelineCore/Public/MoviePipelineAccumulationSetting.h b/Engine/Plugins/MovieScene/MovieRenderPipeline/Source/MovieRenderPipelineCore/Public/MoviePipelineAccumulationSetting.h
--- a/Engine/Plugins/MovieScene/MovieRenderPipeline/Source/MovieRenderPipelineCore/Public/MoviePipelineAccumulationSetting.h
+++ b/Engine/Plugins/MovieScene/MovieRenderPipeline/Source/MovieRenderPipelineCore/Public/MoviePipelineAccumulationSetting.h
@@ -14,11 +14,7 @@ public:
 	UMoviePipelineAccumulationSetting()
 		: TemporalSampleCount(1)
 		, CameraShutterAngle(180)
-		, TileCount(1)
 		, SpatialSampleCount(1)
-		, bIsUsingOverlappedTiles(false)
-		, PadRatioX(.50f)
-		, PadRatioY(.50f)
 		, AccumulationGamma(1.f)
 	{
 	}
@@ -46,25 +42,6 @@ public:
 	UPROPERTY(EditAnywhere, BlueprintReadWrite, meta = (UIMin = 0, UIMax = 360, ClampMin = 0, ClampMax = 360), Category = "Movie Pipeline")
 	int32 CameraShutterAngle;
 
-	/**
-	* How many tiles should the resulting movie render be broken into? A tile should be no larger than
-	* the maximum texture resolution supported by your graphics card (likely 16k), so NumTiles should be
-	* ceil(Width/MaxTextureSize). More tiles mean more individual passes over the whole scene at a smaller
-	* resolution which may help with gpu timeouts. Requires at least 1 tile. Tiling is applied evenly to
-	* both X and Y (to make final interlacing more straightforward).
-	*/
-	UPROPERTY(EditAnywhere, BlueprintReadWrite, meta = (UIMin = 1, ClampMin = 1), Category = "Movie Pipeline")
-	int32 TileCount;
-
-	/**
-	* This bias encourages the engine to use a higher detail texture when it would normally use a lower detail
-	* texture due to the size of the texture on screen. A more negative number means overall sharper images
-	* (up to the detail resolution of your texture). Too much sharpness will cause visual grain/noise in the
-	* resulting image, but this can be mitigated with more spatial samples.
-	*/
-	UPROPERTY(EditAnywhere, BlueprintReadWrite, meta = (UIMin = -1, UIMax = 0), Category = "Movie Pipeline")
-	float TextureSharpnessBias;
-
 	/**
 	* How many frames should we accumulate together before contributing to one overall sample. This lets you
 	* increase the anti-aliasing quality of an sample, or have high quality anti-aliasing if you don't want
@@ -80,16 +57,11 @@ public:
 	int32 WarmUpFrameCount;
 
 	/**
-	* 
+	* For advanced users, the gamma space to apply accumulation in. During accumulation, pow(x,AccumulationGamma) 
+	* is applied and pow(x,1/AccumulationGamma) is applied after accumulation is finished
 	*/
-	UPROPERTY(EditAnywhere, BlueprintReadWrite, meta = (UIMin = 1, ClampMin = 1), Category = "Movie Pipeline")
-	bool bIsUsingOverlappedTiles;
-
-	float PadRatioX;
-	float PadRatioY;
 	float AccumulationGamma;
 
-
 	/**
 	* Shutter Timing allows you to bias the timing of your shutter angle to either be before, during, or after
 	* a frame. When set to FrameClose, it means that the motion gathered up to produce frame N comes from 
diff --git a/Engine/Plugins/MovieScene/MovieRenderPipeline/Source/MovieRenderPipelineCore/Public/MoviePipelineHighResSetting.h b/Engine/Plugins/MovieScene/MovieRenderPipeline/Source/MovieRenderPipelineCore/Public/MoviePipelineHighResSetting.h
new file mode 100644
--- /dev/null
+++ b/Engine/Plugins/MovieScene/MovieRenderPipeline/Source/MovieRenderPipelineCore/Public/MoviePipelineHighResSetting.h
@@ -0,0 +1,72 @@
+// Copyright 1998-2019 Epic Games, Inc. All Rights Reserved.
+#pragma once
+
+#include "MoviePipelineSetting.h"
+#include "MoviePipelineHighResSetting.generated.h"
+
+// Forward Declares
+
+
+UCLASS(Blueprintable)
+class MOVIERENDERPIPELINECORE_API UMoviePipelineHighResSetting : public UMoviePipelineSetting
+{
+	GENERATED_BODY()
+public:
+	UMoviePipelineHighResSetting()
+		: TileCount(1)
+		, TextureSharpnessBias(0.f)
+		, bIsUsingOverlappedTiles(true)
+		, OverlapPercentage(.5f)
+		, bWriteAllSamples(true)
+	{
+	}
+	
+	FIntPoint CalculatePaddedBackbufferSize(const FIntPoint& InTileSize)
+	{
+		int32 OverlappedPadX = FMath::CeilToInt(InTileSize.X * OverlapPercentage);
+		int32 OverlappedPadY = FMath::CeilToInt(InTileSize.Y * OverlapPercentage);
+
+		return InTileSize + FIntPoint(2 * OverlappedPadX, 2 * OverlappedPadY);
+	}
+
+public:
+	/**
+	* How many tiles should the resulting movie render be broken into? A tile should be no larger than
+	* the maximum texture resolution supported by your graphics card (likely 16k), so NumTiles should be
+	* ceil(Width/MaxTextureSize). More tiles mean more individual passes over the whole scene at a smaller
+	* resolution which may help with gpu timeouts. Requires at least 1 tile. Tiling is applied evenly to
+	* both X and Y.
+	*/
+	UPROPERTY(EditAnywhere, BlueprintReadWrite, meta = (UIMin = 1, ClampMin = 1), Category = "Movie Pipeline")
+	int32 TileCount;
+	
+	/**
+	* This bias encourages the engine to use a higher detail texture when it would normally use a lower detail
+	* texture due to the size of the texture on screen. A more negative number means overall sharper images
+	* (up to the detail resolution of your texture). Too much sharpness will cause visual grain/noise in the
+	* resulting image, but this can be mitigated with more spatial samples.
+	*/
+	UPROPERTY(EditAnywhere, BlueprintReadWrite, meta = (UIMin = -1, UIMax = 0), Category = "Movie Pipeline")
+	float TextureSharpnessBias;
+	
+	/**
+	* Do we overlap tiles or interlace? Only overlapping is supported at the moment.
+	*/
+	bool bIsUsingOverlappedTiles;
+	
+	/**
+	* How much should each tile overlap each other (0-0.5). Decreasing the overlap will result in smaller individual
+	* tiles (which means faster renders) but increases the likelyhood of edge-of-screen artifacts showing up which
+	* will become visible in the final image as a "grid" of repeated problem areas.
+	*/
+	UPROPERTY(EditAnywhere, BlueprintReadWrite, meta = (UIMin = 0, ClampMin = 0, UIMax = 0.5, ClampMax=1), Category = "Movie Pipeline")
+	float OverlapPercentage;
+	
+	/**
+	* If true, we will write all samples that get generated to disk individually. This can be useful for debugging or if you need to accumulate
+	* render passes differently than provided.
+	*/
+	UPROPERTY(EditAnywhere, BlueprintReadWrite, meta = (UIMin = 0, ClampMin = 0, UIMax = 0.5, ClampMax=1), Category = "Movie Pipeline|Debug")
+	bool bWriteAllSamples;
+
+};
\ No newline at end of file
diff --git a/Engine/Plugins/MovieScene/MovieRenderPipeline/Source/MovieRenderPipelineCore/Public/MoviePipelineOutputBuilder.h b/Engine/Plugins/MovieScene/MovieRenderPipeline/Source/MovieRenderPipelineCore/Public/MoviePipelineOutputBuilder.h
--- a/Engine/Plugins/MovieScene/MovieRenderPipeline/Source/MovieRenderPipelineCore/Public/MoviePipelineOutputBuilder.h
+++ b/Engine/Plugins/MovieScene/MovieRenderPipeline/Source/MovieRenderPipelineCore/Public/MoviePipelineOutputBuilder.h
@@ -19,8 +19,8 @@ public:
 
 public:
 	FMoviePipelineMergerOutputFrame& QueueOutputFrame_GameThread(const FMoviePipelineFrameOutputState& CachedOutputState);
-	void OnCompleteRenderPassDataAvailable_AnyThread(TUniquePtr<FImagePixelData>&& InData);
-	void OnSingleSampleDataAvailable_AnyThread(TUniquePtr<FImagePixelData>&& InData);
+	void OnCompleteRenderPassDataAvailable_AnyThread(TUniquePtr<FImagePixelData>&& InData, const TSharedRef<FImagePixelDataPayload, ESPMode::ThreadSafe> InFrameData);
+	void OnSingleSampleDataAvailable_AnyThread(TUniquePtr<FImagePixelData>&& InData, const TSharedRef<FImagePixelDataPayload, ESPMode::ThreadSafe> InFrameData);
 
 private:
 	/** The Movie Pipeline that owns us. */
diff --git a/Engine/Plugins/MovieScene/MovieRenderPipeline/Source/MovieRenderPipelineCore/Public/MoviePipelineRenderPass.h b/Engine/Plugins/MovieScene/MovieRenderPipeline/Source/MovieRenderPipelineCore/Public/MoviePipelineRenderPass.h
--- a/Engine/Plugins/MovieScene/MovieRenderPipeline/Source/MovieRenderPipelineCore/Public/MoviePipelineRenderPass.h
+++ b/Engine/Plugins/MovieScene/MovieRenderPipeline/Source/MovieRenderPipelineCore/Public/MoviePipelineRenderPass.h
@@ -12,15 +12,9 @@ class MOVIERENDERPIPELINECORE_API UMoviePipelineRenderPass : public UMoviePipeli
 {
 	GENERATED_BODY()
 public:
-	void Setup(const FMoviePipelineRenderPassInitSettings& InInitSettings)
+	void Setup(TArray<TSharedPtr<MoviePipeline::FMoviePipelineEnginePass>>& InEnginePasses)
 	{
-		InitSettings = InInitSettings;
-		SetupImpl(InInitSettings);
-	}
-
-	void CaptureFrame(const FMoviePipelineRenderPassMetrics& OutputFrameMetrics)
-	{
-		CaptureFrameImpl(OutputFrameMetrics);
+		SetupImpl(InEnginePasses);
 	}
 
 	void Teardown()
@@ -28,21 +22,25 @@ public:
 		TeardownImpl();
 	}
 
+	/** An array of identifiers for the output buffers expected as a result of this render pass. */
 	void GatherOutputPasses(TArray<FMoviePipelinePassIdentifier>& ExpectedRenderPasses)
 	{
 		GatherOutputPassesImpl(ExpectedRenderPasses);
 	}
 
+	/** The required engine passes this pass needs to operate. */
+	void GetRequiredEnginePasses(TSet<FMoviePipelinePassIdentifier>& RequiredEnginePasses)
+	{
+		GetRequiredEnginePassesImpl(RequiredEnginePasses);
+	}
+
+
 protected:
-	virtual void SetupImpl(const FMoviePipelineRenderPassInitSettings& InInitSettings) {}
+	virtual void GetRequiredEnginePassesImpl(TSet<FMoviePipelinePassIdentifier>& RequiredEnginePasses) {}
 
-	virtual void CaptureFrameImpl(const FMoviePipelineRenderPassMetrics& OutputFrameMetrics) {}
+	virtual void SetupImpl(TArray<TSharedPtr<MoviePipeline::FMoviePipelineEnginePass>>& InEnginePasses) {}
 
 	virtual void TeardownImpl() {}
 
 	virtual void GatherOutputPassesImpl(TArray<FMoviePipelinePassIdentifier>& ExpectedRenderPasses) {}
-		
-protected:
-	UPROPERTY()
-	FMoviePipelineRenderPassInitSettings InitSettings;
 };
\ No newline at end of file
diff --git a/Engine/Plugins/MovieScene/MovieRenderPipeline/Source/MovieRenderPipelineCore/Public/MovieRenderOverlappedImage.h b/Engine/Plugins/MovieScene/MovieRenderPipeline/Source/MovieRenderPipelineCore/Public/MovieRenderOverlappedImage.h
--- a/Engine/Plugins/MovieScene/MovieRenderPipeline/Source/MovieRenderPipelineCore/Public/MovieRenderOverlappedImage.h
+++ b/Engine/Plugins/MovieScene/MovieRenderPipeline/Source/MovieRenderPipelineCore/Public/MovieRenderOverlappedImage.h
@@ -92,7 +92,7 @@ public:
 	 * @param InNumTilesY - Num vertical tiles.
 	 * @param InNumChannels - Num Channels.
 	 */
-	void InitMemory(FIntPoint InPlaneSize, int InNumChannels);
+	void InitMemory(FIntPoint InPlaneSize, int32 InNumChannels);
 
 	/**
 	 * Initializes memory.
diff --git a/Engine/Plugins/MovieScene/MovieRenderPipeline/Source/MovieRenderPipelineCore/Public/MovieRenderPipelineCoreModule.h b/Engine/Plugins/MovieScene/MovieRenderPipeline/Source/MovieRenderPipelineCore/Public/MovieRenderPipelineCoreModule.h
--- a/Engine/Plugins/MovieScene/MovieRenderPipeline/Source/MovieRenderPipelineCore/Public/MovieRenderPipelineCoreModule.h
+++ b/Engine/Plugins/MovieScene/MovieRenderPipeline/Source/MovieRenderPipelineCore/Public/MovieRenderPipelineCoreModule.h
@@ -11,6 +11,8 @@ class UMoviePipelineMasterConfig;
 class ULevelSequence;
 class UMoviePipeline;
 
+namespace MoviePipeline { struct FMoviePipelineEnginePass; }
+
 namespace MoviePipelineErrorCodes
 {
 	/** Everything completed as expected or we (unfortunately) couldn't detect the error. */
@@ -23,13 +25,22 @@ namespace MoviePipelineErrorCodes
 	constexpr uint8 NoConfig = 3;
 
 }
+/** A delegate which will create an engine render pass for the Curve Editor. This declares a new engine pass which multiple Pipeline Render Passes can share to reduce re-renders. */
+DECLARE_DELEGATE_RetVal(TSharedRef<MoviePipeline::FMoviePipelineEnginePass>, FOnCreateEngineRenderPass);
 
-class FMovieRenderPipelineCoreModule : public IModuleInterface
+class MOVIERENDERPIPELINECORE_API FMovieRenderPipelineCoreModule : public IModuleInterface
 {
 public:
 	virtual void StartupModule() override;
 	virtual void ShutdownModule() override;
 
+	void UnregisterEngineRenderPass(FDelegateHandle InHandle);
+	FDelegateHandle RegisterEngineRenderPass(FOnCreateEngineRenderPass InOnCreateEngineRenderPass);
+
+	TArrayView<const FOnCreateEngineRenderPass> GetEngineRenderPasses() const
+	{
+		return EngineRenderPassDelegates;
+	}
 private:
 	bool IsTryingToRenderMovieFromCommandLine(FString& OutSequenceAssetPath, FString& OutConfigAssetPath, FString& OutExecutorType, FString& OutPipelineType) const;
 	void InitializeCommandLineMovieRender();
@@ -40,6 +51,12 @@ private:
 	uint8 GetMovieRenderClasses(const FString& InAssetPath, const FString& InConfigAssetPath, const FString& InExecutorType, const FString& InPipelineType, 
 								ULevelSequence*& OutAsset, UMoviePipelineMasterConfig*& OutConfig, UClass*& OutExecutorClass, UClass*& OutPipelineClass) const;
 
+private:
+	TSharedRef<struct MoviePipeline::FMoviePipelineEnginePass> CreateDeferredEnginePass();
+	FDelegateHandle DeferredEnginePassHandle;
+private:
+	TArray<FOnCreateEngineRenderPass> EngineRenderPassDelegates;
+
 private:
 	FString MoviePipelineLocalExecutorClassType;
 	FString MoviePipelineClassType;
diff --git a/Engine/Plugins/MovieScene/MovieRenderPipeline/Source/MovieRenderPipelineCore/Public/MovieRenderPipelineDataTypes.cpp b/Engine/Plugins/MovieScene/MovieRenderPipeline/Source/MovieRenderPipelineCore/Public/MovieRenderPipelineDataTypes.cpp
--- a/Engine/Plugins/MovieScene/MovieRenderPipeline/Source/MovieRenderPipelineCore/Public/MovieRenderPipelineDataTypes.cpp
+++ b/Engine/Plugins/MovieScene/MovieRenderPipeline/Source/MovieRenderPipelineCore/Public/MovieRenderPipelineDataTypes.cpp
@@ -4,6 +4,7 @@
 #include "MovieRenderPipelineCoreModule.h"
 #include "Sections/MovieSceneCinematicShotSection.h"
 #include "MoviePipelineAccumulationSetting.h"
+#include "MoviePipeline.h"
 
 FString FMoviePipelineShotInfo::GetDisplayName() const
 {
@@ -124,11 +125,19 @@ FMoviePipelineWorkInfo FMoviePipelineCameraCutInfo::GetTotalWorkEstimate() const
 	WorkEstimate.NumOutputFramesWithSubsampling += SubsampledOutputFrameCount.Value;
 
 	// How many different samples are we submitting to the GPU?
-	FFrameNumber TotalSampleCount = GetSampleCountEstimate(true, true) * NumTiles;
+	FFrameNumber TotalSampleCount = GetSampleCountEstimate(true, true) * NumTiles.X * NumTiles.Y;
 	WorkEstimate.NumSamples += TotalSampleCount.Value;
 
 	// How many images will we produce to create the high-res tile.
-	WorkEstimate.NumTiles = NumTiles;
+	WorkEstimate.NumTiles = NumTiles.X * NumTiles.Y;
 
 	return WorkEstimate;
+}
+
+namespace MoviePipeline
+{
+	UMoviePipeline* FMoviePipelineEnginePass::GetPipeline() const
+	{
+		return OwningPipeline.Get();
+	}
 }
\ No newline at end of file
diff --git a/Engine/Plugins/MovieScene/MovieRenderPipeline/Source/MovieRenderPipelineCore/Public/MovieRenderPipelineDataTypes.h b/Engine/Plugins/MovieScene/MovieRenderPipeline/Source/MovieRenderPipelineCore/Public/MovieRenderPipelineDataTypes.h
--- a/Engine/Plugins/MovieScene/MovieRenderPipeline/Source/MovieRenderPipelineCore/Public/MovieRenderPipelineDataTypes.h
+++ b/Engine/Plugins/MovieScene/MovieRenderPipeline/Source/MovieRenderPipelineCore/Public/MovieRenderPipelineDataTypes.h
@@ -18,6 +18,10 @@ class UMoviePipelineShotConfig;
 class UMovieSceneCameraCutSection;
 class UMoviePipelineRenderPass;
 struct FImagePixelData;
+struct FImageOverlappedAccumulator;
+class FMoviePipelineOutputMerger;
+class FRenderTarget;
+class UMoviePipeline;
 
 
 /**
@@ -293,7 +297,7 @@ public:
 		, bAccurateFirstFrameHistory(true)
 		, NumTemporalSamples(0)
 		, NumSpatialSamples(0)
-		, NumTiles(0)
+		, NumTiles(0, 0)
 		, State(EMovieRenderShotState::Uninitialized)
 		, CurrentTick(FFrameNumber(0))
 		, bHasEvaluatedMotionBlurFrame(false)
@@ -360,7 +364,7 @@ public:
 
 	/** How many image tiles are going to be rendered per temporal frame. */
 	UPROPERTY(BlueprintReadOnly, Category = "Movie Render Pipeline")
-	int32 NumTiles;
+	FIntPoint NumTiles;
 
 	/** Display name for UI Purposes. */
 	UPROPERTY(BlueprintReadOnly, Category = "Movie Render Pipeline")
@@ -618,125 +622,173 @@ public:
 	}
 };
 
-USTRUCT(BlueprintType)
-struct FMoviePipelineRenderPassInitSettings
+/**
+* These parameters define a single sample that a render pass should render with.
+*/
+struct FMoviePipelineRenderPassMetrics
 {
-	GENERATED_BODY()
 public:
-	FMoviePipelineRenderPassInitSettings()
-		: ShotConfig(nullptr)
-		, TileResolution(0, 0)
-		, TileCount(0)
+	/** How many tiles on X and Y are there total. */
+	FIntPoint TileCounts;
+
+	/** Of the TileCount, which X/Y Tile is this sample for. */
+	FIntPoint TileIndexes;
+
+	/** Get a 0-(TileCount-1) version of TileIndex. */
+	FORCEINLINE int32 GetTileIndex() const
 	{
+		return (TileIndexes.Y * TileCounts.Y) + TileIndexes.X;
 	}
 
-public:
-	/** The config associated with the shot currently being rendered. */
-	UPROPERTY()
-	UMoviePipelineShotConfig* ShotConfig;
+	FORCEINLINE int32 GetTileCount() const
+	{
+		return (TileCounts.X & TileCounts.Y);
+	}
 
-	// These are derived out of other settings and cached directly to needing to re-derive values.
-	
-	/** The back buffer should match this setting and not the final output resolution.*/
-	FIntPoint TileResolution;
+	/** How big is the back buffer for this sample? Pre-set when setup, here for convenience. */
+	FIntPoint BackbufferSize;
+
+	/** If true, we will discard this sample after rendering. Used to get history set up correctly. */
+	bool bIsHistoryOnlyFrame;
+
+	/** How many spatial jitters will there be total for this particular temporal accumulation frame? */
+	int32 SpatialSampleCount;
 	
-	/** How many tiles (in each direction) are we rendering with. */
-	int32 TileCount;
-};
+	/** Of the SpatialSampleCount, which index is this? */
+	int32 SpatialSampleIndex;
 
-USTRUCT(BlueprintType)
-struct FMoviePipelineRenderPassMetrics
-{
-	GENERATED_BODY()
+	/** How many temporal jitters will there be total for this particular output frame? */
+	int32 TemporalSampleCount;
 
-public:
-	FVector2D SpatialShift;
-	int32 TileIndex;
-	int32 TileIndexX;
-	int32 TileIndexY;
-	int32 NumTilesX;
-	int32 NumTilesY;
-	FMoviePipelineFrameOutputState OutputState;
-	bool bIsHistoryOnlyFrame;
-	int32 SpatialJitterIndex;
-	int32 NumSpatialJitters;
+	/** Of the TemporalSampleCount, which index is this? */
+	int32 TemporalSampleIndex;
+
+	/** Should we use the Overlapped Tiles method of doing high res screen shots? */
+	bool bIsUsingOverlappedTiles;
 
-	int32 TemporalJitterIndex;
-	int32 NumTemporalJitters;
+	/** How many pixels are we overlapping with adjacent tiles (on each side) */
+	FIntPoint OverlappedPad;
 
-	float JitterOffsetX;
-	float JitterOffsetY;
+	/** How much is this sample offset, taking padding into account. */
+	FIntPoint OverlappedOffset;
 
+	/** 
+	* The gamma space to apply accumulation in. During accumulation, pow(x,AccumulationGamma) is applied
+	* and pow(x,1/AccumulationGamma) is applied after accumulation is finished. 1.0 means no change."
+	*/
 	float AccumulationGamma;
 
-	bool bIsUsingOverlappedTiles;
-	int32 OverlappedOffsetX;
-	int32 OverlappedOffsetY;
-	int32 OverlappedSizeX;
-	int32 OverlappedSizeY;
-	int32 OverlappedPadX;
-	int32 OverlappedPadY;
+	/**
+	* The amount of screen-space shift used to replace TAA Projection Jitter, modified for each spatial sample
+	* of the render.
+	*/
+	FVector2D SpatialShift;
+
+	FMoviePipelineFrameOutputState OutputState;
+
 	FVector2D OverlappedSubpixelShift;
 
 
 	MoviePipeline::FMoviePipelineFrameInfo FrameInfo;
 };
 
-struct FImagePixelDataPayload : IImagePixelDataPayload
+namespace MoviePipeline
 {
-	FMoviePipelineFrameOutputState OutputState;
-	FVector2D SpatialShift;
-	int32 TileIndexX;
-	int32 TileIndexY;
-	int32 TileSizeX;
-	int32 TileSizeY;
-	int32 NumTilesX;
-	int32 NumTilesY;
-	int32 SpatialJitterIndex;
-	int32 NumSpatialJitters;
+	struct FMoviePipelineRenderPassInitSettings
+	{
+	public:
+		FMoviePipelineRenderPassInitSettings()
+			: BackbufferResolution(0, 0)
+			, TileCount(0, 0)
+		{
+		}
 
-	int32 TemporalJitterIndex;
-	int32 NumTemporalJitters;
+	public:
+		/** This takes into account any padding needed for tiled rendering overlap. Different than the output resolution of the final image. */
+		FIntPoint BackbufferResolution;
 
-	float JitterOffsetX;
-	float JitterOffsetY;
+		/** How many tiles (in each direction) are we rendering with. */
+		FIntPoint TileCount;
+	};
 
-	float AccumulationGamma;
+	struct MOVIERENDERPIPELINECORE_API FMoviePipelineEnginePass
+	{
+		FMoviePipelineEnginePass(const FMoviePipelinePassIdentifier& InPassIdentifier)
+			: PassIdentifier(InPassIdentifier)
+		{
+		}
 
-	bool bIsUsingOverlappedTiles;
-	int32 OverlappedOffsetX;
-	int32 OverlappedOffsetY;
-	int32 OverlappedSizeX;
-	int32 OverlappedSizeY;
-	int32 OverlappedPadX;
-	int32 OverlappedPadY;
-	FVector2D OverlappedSubpixelShift;
+		virtual ~FMoviePipelineEnginePass()
+		{}
+
+
+		virtual void Setup(TWeakObjectPtr<UMoviePipeline> InOwningPipeline, const FMoviePipelineRenderPassInitSettings& InInitSettings)
+		{
+			OwningPipeline = InOwningPipeline;
+			InitSettings = InInitSettings;
+		}
+
+		virtual void RenderSample_GameThread(const FMoviePipelineRenderPassMetrics& InSampleState)
+		{
+		}
+
+		virtual void Teardown()
+		{
+		}
+
+	protected:
+		UMoviePipeline* GetPipeline() const;
+
+	public:
+		/** A unique name for this engine pass. This is how an individual output pass specifies what data source it wants. */
+		FMoviePipelinePassIdentifier PassIdentifier;
 
+	protected:
+		FMoviePipelineRenderPassInitSettings InitSettings;
+
+	private:
+		TWeakObjectPtr<UMoviePipeline> OwningPipeline;
+	};
+
+	struct FSampleRenderThreadParams
+	{
+		FMoviePipelineRenderPassMetrics SampleState;
+		TSharedPtr<FImageOverlappedAccumulator, ESPMode::ThreadSafe> ImageAccumulator;
+		TSharedPtr<FMoviePipelineOutputMerger, ESPMode::ThreadSafe> OutputMerger;
+		bool bWriteTiles;
+		FMoviePipelinePassIdentifier PassIdentifier;
+	};
+}
+
+struct FImagePixelDataPayload : IImagePixelDataPayload, public TSharedFromThis<FImagePixelDataPayload>
+{
+	FMoviePipelineFrameOutputState OutputState;
+	FMoviePipelineRenderPassMetrics SampleState;
 
 	FMoviePipelinePassIdentifier PassIdentifier;
 
 	/** Is this the first tile of an image and we should start accumulating? */
 	FORCEINLINE bool IsFirstTile() const
 	{
-		return TileIndexX == 0 && TileIndexY == 0 && SpatialJitterIndex == 0;
+		return SampleState.TileIndexes.X == 0 && SampleState.TileIndexes.Y == 0 && SampleState.SpatialSampleIndex == 0;
 	}
 
 	/** Is this the last tile of an image and we should finish accumulating? */
 	FORCEINLINE bool IsLastTile() const
 	{
-		return TileIndexX == NumTilesX - 1 &&
-			   TileIndexY == NumTilesY - 1 &&
-			   SpatialJitterIndex == NumSpatialJitters - 1;
+		return SampleState.TileIndexes.X == SampleState.TileCounts.X - 1 &&
+			   SampleState.TileIndexes.Y == SampleState.TileCounts.Y - 1 &&
+			   SampleState.SpatialSampleIndex == SampleState.SpatialSampleCount - 1;
 	}
 
 	FORCEINLINE bool IsFirstTemporalSample() const
 	{
-		return TemporalJitterIndex == 0;
+		return SampleState.TemporalSampleIndex == 0;
 	}
 
 	FORCEINLINE bool IsLastTemporalSample() const
 	{
-		return TemporalJitterIndex == NumTemporalJitters - 1;
+		return SampleState.TemporalSampleIndex == SampleState.TemporalSampleCount - 1;
 	}
 };
 
diff --git a/Engine/Plugins/MovieScene/MovieRenderPipeline/Source/MovieRenderPipelineEditor/Private/MovieRenderPipelineEditorTest.cpp b/Engine/Plugins/MovieScene/MovieRenderPipeline/Source/MovieRenderPipelineEditor/Private/MovieRenderPipelineEditorTest.cpp
--- a/Engine/Plugins/MovieScene/MovieRenderPipeline/Source/MovieRenderPipelineEditor/Private/MovieRenderPipelineEditorTest.cpp
+++ b/Engine/Plugins/MovieScene/MovieRenderPipeline/Source/MovieRenderPipelineEditor/Private/MovieRenderPipelineEditorTest.cpp
@@ -10,13 +10,15 @@
 #include "Misc/Paths.h"
 #include "MoviePipelineMasterConfig.h"
 #include "MoviePipelineOutputSetting.h"
+#include "MoviePipelineHighResSetting.h"
+#include "MoviePipelineDeferredPasses.h"
 
 UMoviePipelineShotConfig* GenerateTestShotConfig(UObject* InOwner, int32 InSampleCount,
 	int32 InShutterAngle, EMoviePipelineShutterTiming InFrameTiming, int32 InTileCount,
 	int32 InSpatialSampleCount, bool bIsUsingOverlappedTiles, float PadRatioX, float PadRatioY, float AccumulationGamma)
 {
 	UMoviePipelineShotConfig* OutConfig = NewObject<UMoviePipelineShotConfig>(InOwner);
-	OutConfig->FindOrAddSetting<UMoviePipelineBackbufferPass>();
+	OutConfig->FindOrAddSetting<UMoviePipelineDeferredPassBase>();
 
 	UMoviePipelineGameOverrideSetting* GamemodeOverride = OutConfig->FindOrAddSetting<UMoviePipelineGameOverrideSetting>();
 	GamemodeOverride->GameModeOverride = AGameMode::StaticClass();
@@ -24,14 +26,18 @@ UMoviePipelineShotConfig* GenerateTestShotConfig(UObject* InOwner, int32 InSampl
 	UMoviePipelineAccumulationSetting* Accumulation = OutConfig->FindOrAddSetting< UMoviePipelineAccumulationSetting>();
 	Accumulation->TemporalSampleCount = InSampleCount;
 	Accumulation->CameraShutterAngle = InShutterAngle;
-	Accumulation->TileCount = InTileCount;
 	Accumulation->ShutterTiming = InFrameTiming;
 	Accumulation->SpatialSampleCount = InSpatialSampleCount;
-	Accumulation->bIsUsingOverlappedTiles = bIsUsingOverlappedTiles;
-	Accumulation->PadRatioX = PadRatioX;
-	Accumulation->PadRatioY = PadRatioY;
 	Accumulation->AccumulationGamma = AccumulationGamma;
 
+	if (InTileCount > 1)
+	{
+		UMoviePipelineHighResSetting* HighRes = OutConfig->FindOrAddSetting<UMoviePipelineHighResSetting>();
+		HighRes->TileCount = InTileCount;
+		HighRes->bIsUsingOverlappedTiles = bIsUsingOverlappedTiles;
+		HighRes->OverlapPercentage = PadRatioX;
+	}
+
 	return OutConfig;
 }
 
@@ -63,8 +69,8 @@ TArray<UMoviePipelineMasterConfig*> FMovieRenderPipelineEditorModule::GenerateTe
 
 				int32 SizeX = 1920;
 				int32 SizeY = 1080;
-				int32 TileX = 2;
-				int32 TileY = 2;
+				int32 TileX = 1;
+				int32 TileY = 1;
 				int32 TestNumSamples = 2;
 				float PadRatioX = 0.5f;
 				float PadRatioY = 0.5f;
diff --git a/Engine/Plugins/MovieScene/MovieRenderPipeline/Source/MovieRenderPipelineRenderPasses/Private/MoviePipelineBackbufferPass.cpp b/Engine/Plugins/MovieScene/MovieRenderPipeline/Source/MovieRenderPipelineRenderPasses/Private/MoviePipelineBackbufferPass.cpp
--- a/Engine/Plugins/MovieScene/MovieRenderPipeline/Source/MovieRenderPipelineRenderPasses/Private/MoviePipelineBackbufferPass.cpp
+++ b/Engine/Plugins/MovieScene/MovieRenderPipeline/Source/MovieRenderPipelineRenderPasses/Private/MoviePipelineBackbufferPass.cpp
@@ -27,410 +27,7 @@
 #include "Async/Async.h"
 #include "MovieRenderPipelineCoreModule.h"
 #include "MovieRenderOverlappedImage.h"
-
-namespace MoviePipeline
-{
-	struct FSampleRenderThreadParams
-	{
-		FMoviePipelineRenderPassMetrics PassMetrics;
-		FRenderTarget* CanvasRenderTarget;
-		TSharedPtr<FImageOverlappedAccumulator, ESPMode::ThreadSafe> ImageAccumulator;
-		TSharedPtr<FMoviePipelineOutputMerger, ESPMode::ThreadSafe> OutputMerger;
-		bool bWriteTiles;
-	};
-}
-
-// Forward Declares
-static void ReadAndAccumulateSample_RenderThread(FRHICommandListImmediate &RHICmdList, const MoviePipeline::FSampleRenderThreadParams& InParams);
-
-void UMoviePipelineBackbufferPass::SetupImpl(const FMoviePipelineRenderPassInitSettings& InInitSettings)
-{
-	TileRenderTarget = NewObject<UTextureRenderTarget2D>(this);
-	TileRenderTarget->ClearColor = FLinearColor(0.0f, 0.0f, 0.0f, 0.0f);
-
-	// Initialize to the tile size (not final size) and use a 16 bit back buffer to avoid precision issues when accumulating later
-
-	int32 TargetSizeX = InInitSettings.TileResolution.X;
-	int32 TargetSizeY = InInitSettings.TileResolution.Y;
-/*
-	*/
-	TileRenderTarget->InitCustomFormat(TargetSizeX, TargetSizeY, EPixelFormat::PF_A16B16G16R16, false);
-	// Create a Render Target to hold our back buffer.
-	ViewState.Allocate();
-
-	ImageTileAccumulator = MakeShared<FImageOverlappedAccumulator, ESPMode::ThreadSafe>();
-}
-
-void UMoviePipelineBackbufferPass::TeardownImpl()
-{
-	ViewState.Destroy();
-}
-
-void UMoviePipelineBackbufferPass::GatherOutputPassesImpl(TArray<FMoviePipelinePassIdentifier>& ExpectedRenderPasses)
-{
-	ExpectedRenderPasses.Add(FMoviePipelinePassIdentifier(TEXT("Backbuffer")));
-}
-
-void UMoviePipelineBackbufferPass::CaptureFrameImpl(const FMoviePipelineRenderPassMetrics& InPassMetrics)
-{
-	float TimeSeconds = InPassMetrics.OutputState.WorldSeconds;
-	float RealTimeSeconds = InPassMetrics.OutputState.WorldSeconds;
-	float DeltaTimeSeconds = InPassMetrics.OutputState.FrameDeltaTime;
-
-	FEngineShowFlags ShowFlags = FEngineShowFlags(EShowFlagInitMode::ESFIM_Game);
-
-	if (InPassMetrics.bIsUsingOverlappedTiles)
-	{
-		// We can actually keep these enabled for now
-		//ShowFlags.SetDepthOfField(false);
-		//ShowFlags.SetMotionBlur(false);
-	}
-
-	FSceneViewFamilyContext ViewFamily(FSceneViewFamily::ConstructionValues(
-		TileRenderTarget->GameThread_GetRenderTargetResource(),
-		GetWorld()->Scene,
-		ShowFlags)
-		.SetWorldTimes(TimeSeconds, DeltaTimeSeconds, RealTimeSeconds)
-		.SetRealtimeUpdate(true));
-
-	ViewFamily.SetScreenPercentageInterface(new FLegacyScreenPercentageDriver(ViewFamily, 1.0f, false));
-
-
-	FVector ViewLocation;
-	FRotator ViewRotation;
-
-	//APlayerController* LocalPlayerController = GetWorld()->GetFirstPlayerController();
-	//LocalPlayerController->GetPlayerViewPoint(ViewLocation, ViewRotation);
-
-	// View is added as a child of the ViewFamily
-	FSceneView* View = CalcSceneView(&ViewFamily, InPassMetrics);
-	ViewFamily.FrameNumber += InPassMetrics.SpatialJitterIndex;
-	
-	// Override the Motion Blur settings since these are controlled by the movie pipeline.
-	{
-		View->FinalPostProcessSettings.MotionBlurTargetFPS = FMath::RoundToInt(GetPipeline()->GetEffectiveFrameRate().AsDecimal());
-		View->FinalPostProcessSettings.MotionBlurAmount = InPassMetrics.OutputState.MotionBlurFraction;
-		View->FinalPostProcessSettings.MotionBlurMax = 100.f;
-		View->FinalPostProcessSettings.bOverride_MotionBlurAmount = true;
-		View->FinalPostProcessSettings.bOverride_MotionBlurTargetFPS = true;
-		View->FinalPostProcessSettings.bOverride_MotionBlurMax = true;
-	}
-
-	// Autoexposure is not allowed
-	if (View->FinalPostProcessSettings.AutoExposureMethod != EAutoExposureMethod::AEM_Manual)
-	{
-		UE_LOG(LogMovieRenderPipeline, Log, TEXT("AutoExposure Method should always be Manual."));
-
-		View->FinalPostProcessSettings.AutoExposureMethod = EAutoExposureMethod::AEM_Manual;
-	}
-
-	// Disable TAA while we're working on pixel accuracy testing
-	{
-		View->AntiAliasingMethod = EAntiAliasingMethod::AAM_None;
-	}
-
-	// Override the view matrix with our offsets provided by tiling/jittering. This requires
-	// TAA to be disabled.
-	{
-		View->ViewMatrices.HackAddTemporalAAProjectionJitter(InPassMetrics.SpatialShift);
-	}
-
-	// Bias all mip-mapping to pretend to be working at our target resolution and not our tile resolution
-	// so that the images don't end up soft.
-	{
-		float EffectivePrimaryResolutionFraction = 1.f / InitSettings.TileCount;
-		View->MaterialTextureMipBias = FMath::Log2(EffectivePrimaryResolutionFraction);
-
-		// Add an additional bias per user settings. This allows them to choose to make the textures sharper if it
-		// looks better with their particular settings.
-		UMoviePipelineAccumulationSetting* AccumulationSettings = InitSettings.ShotConfig->FindOrAddSetting<UMoviePipelineAccumulationSetting>();
-		View->MaterialTextureMipBias += AccumulationSettings->TextureSharpnessBias;
-	}
-
-	// If this is a history only frame it means we just came from a camera cut, flag it as such
-	// which will have knock-on effects in histories (such as clearing old TAA buffers).
-	if (InPassMetrics.bIsHistoryOnlyFrame)
-	{
-		View->bCameraCut = true;
-	}
-
-	// The world is always paused, but we are setting the motion blur and view matrices manually
-	ViewFamily.bWorldIsPaused = true;
-
-	FCanvas Canvas = FCanvas(TileRenderTarget->GameThread_GetRenderTargetResource(), nullptr, GetWorld(), ERHIFeatureLevel::SM5, FCanvas::CDM_DeferDrawing, 1.0f);
-
-
-	// Draw the world into this View Family
-	GetRendererModule().BeginRenderingViewFamily(&Canvas, &ViewFamily);
-
-	FRenderTarget* CanvasRenderTarget = Canvas.GetRenderTarget();
-	TSharedPtr<FMoviePipelineOutputMerger, ESPMode::ThreadSafe> OutputMerger = GetPipeline()->OutputBuilder;
-	TSharedPtr<FImageOverlappedAccumulator, ESPMode::ThreadSafe> LocalImageTileAccumulator = ImageTileAccumulator;
-
-	MoviePipeline::FSampleRenderThreadParams AccumulationParams;
-	{
-		AccumulationParams.bWriteTiles = true;
-		AccumulationParams.CanvasRenderTarget = CanvasRenderTarget;
-		AccumulationParams.OutputMerger = GetPipeline()->OutputBuilder;
-		AccumulationParams.ImageAccumulator = ImageTileAccumulator;
-		AccumulationParams.PassMetrics = InPassMetrics;
-	}
-
-	ENQUEUE_RENDER_COMMAND(CanvasRenderTargetResolveCommand)(
-		[AccumulationParams](FRHICommandListImmediate& RHICmdList)
-	{
-		UE_LOG(LogMovieRenderPipeline, Log, TEXT("Rendered Tile: %d Sample: %d"), AccumulationParams.PassMetrics.TileIndex, AccumulationParams.PassMetrics.SpatialJitterIndex);
-		ReadAndAccumulateSample_RenderThread(RHICmdList, AccumulationParams);
-	});
-}
-
-FSceneView* UMoviePipelineBackbufferPass::CalcSceneView(FSceneViewFamily* ViewFamily, const FMoviePipelineRenderPassMetrics& InPassMetrics)
-{
-	APlayerController* LocalPlayerController = GetWorld()->GetFirstPlayerController();
-
-	int32 TileSizeX = InitSettings.TileResolution.X;
-	int32 TileSizeY = InitSettings.TileResolution.Y;
-
-	FSceneViewInitOptions ViewInitOptions;
-	ViewInitOptions.ViewFamily = ViewFamily;
-	ViewInitOptions.ViewOrigin = InPassMetrics.FrameInfo.CurrViewLocation;
-	ViewInitOptions.SetViewRectangle(FIntRect(FIntPoint(0, 0), FIntPoint(TileSizeX, TileSizeY)));
-	ViewInitOptions.ViewRotationMatrix = FInverseRotationMatrix(InPassMetrics.FrameInfo.CurrViewRotation);
-
-	// Rotate the view 90 degrees (reason: unknown)
-	ViewInitOptions.ViewRotationMatrix = ViewInitOptions.ViewRotationMatrix * FMatrix(
-		FPlane(0, 0, 1, 0),
-		FPlane(1, 0, 0, 0),
-		FPlane(0, 1, 0, 0),
-		FPlane(0, 0, 0, 1));
-	float ViewFOV = 90.f;
-	if (GetWorld()->GetFirstPlayerController()->PlayerCameraManager)
-	{
-		ViewFOV = GetWorld()->GetFirstPlayerController()->PlayerCameraManager->GetFOVAngle();
-	}
-
-	// Calculate a Projection Matrix
-	{
-		float XAxisMultiplier;
-		float YAxisMultiplier;
-
-		if (ViewInitOptions.GetViewRect().Width() > ViewInitOptions.GetViewRect().Height())
-		{
-			//if the viewport is wider than it is tall
-			XAxisMultiplier = 1.0f;
-			YAxisMultiplier = ViewInitOptions.GetViewRect().Width() / (float)ViewInitOptions.GetViewRect().Height();
-		}
-		else
-		{
-			//if the viewport is taller than it is wide
-			XAxisMultiplier = ViewInitOptions.GetViewRect().Height() / (float)ViewInitOptions.GetViewRect().Width();
-			YAxisMultiplier = 1.0f;
-		}
-
-		const float MinZ = GNearClippingPlane;
-		const float MaxZ = MinZ;
-		// Avoid zero ViewFOV's which cause divide by zero's in projection matrix
-		const float MatrixFOV = FMath::Max(0.001f, ViewFOV) * (float)PI / 360.0f;
-
-		FMatrix BaseProjMatrix;
-
-		if ((bool)ERHIZBuffer::IsInverted)
-		{
-			BaseProjMatrix = FReversedZPerspectiveMatrix(
-				MatrixFOV,
-				MatrixFOV,
-				XAxisMultiplier,
-				YAxisMultiplier,
-				MinZ,
-				MaxZ
-			);
-		}
-		else
-		{
-			BaseProjMatrix = FPerspectiveMatrix(
-				MatrixFOV,
-				MatrixFOV,
-				XAxisMultiplier,
-				YAxisMultiplier,
-				MinZ,
-				MaxZ
-			);
-		}
-
-		if (InPassMetrics.bIsUsingOverlappedTiles)
-		{
-			float PadRatioX = float(InPassMetrics.OverlappedPadX * 2 + InPassMetrics.OverlappedSizeX) / float (InPassMetrics.OverlappedSizeX);
-			float PadRatioY = float(InPassMetrics.OverlappedPadY * 2 + InPassMetrics.OverlappedSizeY) / float (InPassMetrics.OverlappedSizeY);
-
-			float ScaleX = PadRatioX / float(InitSettings.TileCount);
-			float ScaleY = PadRatioY / float(InitSettings.TileCount);
-
-			BaseProjMatrix.M[0][0] /= ScaleX;
-			BaseProjMatrix.M[1][1] /= ScaleY;
-
-			// this offset would be correct with no pad
-			float OffsetX = -((float(InPassMetrics.TileIndexX) + 0.5f - float(InitSettings.TileCount) / 2.0f)* 2.0f);
-			float OffsetY =  ((float(InPassMetrics.TileIndexY) + 0.5f - float(InitSettings.TileCount) / 2.0f)* 2.0f);
-
-			BaseProjMatrix.M[2][0] += OffsetX / PadRatioX;
-			BaseProjMatrix.M[2][1] += OffsetY / PadRatioX;
-		}
-
-		ViewInitOptions.ProjectionMatrix = BaseProjMatrix;
-	}
-
-	ViewInitOptions.SceneViewStateInterface = ViewState.GetReference();
-	ViewInitOptions.FOV = ViewFOV;
-
-	FSceneView* View = new FSceneView(ViewInitOptions);
-	ViewFamily->Views.Add(View);
-	View->ViewLocation = InPassMetrics.FrameInfo.CurrViewLocation;
-	View->ViewRotation = InPassMetrics.FrameInfo.CurrViewRotation;
-	View->PreviousViewTransform = FTransform(InPassMetrics.FrameInfo.PrevViewRotation,InPassMetrics.FrameInfo.PrevViewLocation);
-	View->StartFinalPostprocessSettings(View->ViewLocation);
-	View->OverrideFrameIndexValue = InPassMetrics.SpatialJitterIndex;
-
-	// CameraAnim override
-	if (LocalPlayerController->PlayerCameraManager)
-	{
-		TArray<FPostProcessSettings> const* CameraAnimPPSettings;
-		TArray<float> const* CameraAnimPPBlendWeights;
-		LocalPlayerController->PlayerCameraManager->GetCachedPostProcessBlends(CameraAnimPPSettings, CameraAnimPPBlendWeights);
-
-		FMinimalViewInfo ViewInfo = LocalPlayerController->PlayerCameraManager->GetCameraCachePOV();
-		for (int32 PPIdx = 0; PPIdx < CameraAnimPPBlendWeights->Num(); ++PPIdx)
-		{
-			View->OverridePostProcessSettings((*CameraAnimPPSettings)[PPIdx], (*CameraAnimPPBlendWeights)[PPIdx]);
-		}
-
-		View->OverridePostProcessSettings(ViewInfo.PostProcessSettings, ViewInfo.PostProcessBlendWeight);
-	}
-
-	View->EndFinalPostprocessSettings(ViewInitOptions);
-
-	return View;
-}
-
-void ReadAndAccumulateSample_RenderThread(FRHICommandListImmediate &RHICmdList, const MoviePipeline::FSampleRenderThreadParams& InParams)
-{
-	check(IsInRenderingThread());
-
-	// If they only needed this frame for the history, we don't actually care about trying to output it,
-	// so we can skip the copy and pushing to output.
-	if (InParams.PassMetrics.bIsHistoryOnlyFrame)
-	{
-		return;
-	}
-
-	FReadSurfaceDataFlags ReadDataFlags;
-	ReadDataFlags.SetLinearToGamma(false);
-
-	FIntRect SourceRect = FIntRect(0, 0, InParams.CanvasRenderTarget->GetSizeXY().X, InParams.CanvasRenderTarget->GetSizeXY().Y);
-
-	// Build our metadata to associate with this frame.
-	TSharedRef<FImagePixelDataPayload, ESPMode::ThreadSafe> FrameData = MakeShared<FImagePixelDataPayload, ESPMode::ThreadSafe>();
-	FrameData->OutputState = InParams.PassMetrics.OutputState;
-	FrameData->PassIdentifier = FMoviePipelinePassIdentifier(TEXT("Backbuffer"));
-	FrameData->TileIndexX = InParams.PassMetrics.TileIndexX;
-	FrameData->TileIndexY = InParams.PassMetrics.TileIndexY;
-	FrameData->TileSizeX = SourceRect.Width();
-	FrameData->TileSizeY = SourceRect.Height();
-	FrameData->NumTilesX = InParams.PassMetrics.NumTilesX;
-	FrameData->NumTilesY = InParams.PassMetrics.NumTilesY;
-	FrameData->SpatialShift = InParams.PassMetrics.SpatialShift;
-	FrameData->SpatialJitterIndex = InParams.PassMetrics.SpatialJitterIndex;
-	FrameData->NumSpatialJitters = InParams.PassMetrics.NumSpatialJitters;
-	FrameData->TemporalJitterIndex = InParams.PassMetrics.TemporalJitterIndex;
-	FrameData->NumTemporalJitters = InParams.PassMetrics.NumTemporalJitters;
-	FrameData->JitterOffsetX = InParams.PassMetrics.JitterOffsetX;
-	FrameData->JitterOffsetY = InParams.PassMetrics.JitterOffsetY;
-	FrameData->bIsUsingOverlappedTiles = InParams.PassMetrics.bIsUsingOverlappedTiles;
-
-	FrameData->AccumulationGamma = InParams.PassMetrics.AccumulationGamma;
-
-	FrameData->OverlappedOffsetX = InParams.PassMetrics.OverlappedOffsetX;
-	FrameData->OverlappedOffsetY = InParams.PassMetrics.OverlappedOffsetY;
-	FrameData->OverlappedSizeX = InParams.PassMetrics.OverlappedSizeX;
-	FrameData->OverlappedSizeY = InParams.PassMetrics.OverlappedSizeY;
-	FrameData->OverlappedPadX = InParams.PassMetrics.OverlappedPadX;
-	FrameData->OverlappedPadY = InParams.PassMetrics.OverlappedPadY;
-	FrameData->OverlappedSubpixelShift = InParams.PassMetrics.OverlappedSubpixelShift;
-
-	// Read the data back to the CPU
-	TArray<FColor> RawPixels;
-	RawPixels.SetNum(SourceRect.Width() * SourceRect.Height());
-
-	RHICmdList.ReadSurfaceData(InParams.CanvasRenderTarget->GetRenderTargetTexture(), SourceRect, RawPixels, ReadDataFlags);
-	TUniquePtr<TImagePixelData<FColor>> PixelData = MakeUnique<TImagePixelData<FColor>>(SourceRect.Size(), TArray64<FColor>(MoveTemp(RawPixels)), FrameData);
-
-	check(PixelData->IsDataWellFormed());
-
-	const bool bWriteTiles = true;
-	
-	// Writing tiles can be useful for debug reasons. These get passed onto the output every frame.
-	if (bWriteTiles)
-	{
-		// Send the data to the Output Builder. This has to be a copy of the pixel data from the GPU, since
-		// it enqueues it onto the game thread and won't be read/sent to write to disk for another frame. 
-		// The extra copy is unfortunate, but is only the size of a single sample (ie: 1920x1080 -> 17mb)
-		TUniquePtr<FImagePixelData> SampleData = PixelData->Copy();
-		InParams.OutputMerger->OnSingleSampleDataAvailable_AnyThread(MoveTemp(SampleData));
-	}
-
-	// For the first sample in a new output, we allocate memory
-	if (FrameData->IsFirstTile() && FrameData->IsFirstTemporalSample())
-	{
-		InParams.ImageAccumulator->InitMemory(FIntPoint(FrameData->OverlappedSizeX * FrameData->NumTilesX, FrameData->OverlappedSizeY * FrameData->NumTilesY), 3);
-		InParams.ImageAccumulator->ZeroPlanes();
-		InParams.ImageAccumulator->AccumulationGamma = FrameData->AccumulationGamma;
-	}
-
-	// Accumulate the new sample to our target
-	{
-		const double AccumulateBeginTime = FPlatformTime::Seconds();
-
-		FIntPoint RawSize = PixelData->GetSize();
-
-		check(FrameData->OverlappedSizeX + 2 * FrameData->OverlappedPadX == RawSize.X);
-		check(FrameData->OverlappedSizeY + 2 * FrameData->OverlappedPadY == RawSize.Y);
-
-		InParams.ImageAccumulator->AccumulatePixelData(*PixelData.Get(), FIntPoint(FrameData->OverlappedOffsetX, FrameData->OverlappedOffsetY), FrameData->OverlappedSubpixelShift);
-
-		const double AccumulateEndTime = FPlatformTime::Seconds();
-		const float ElapsedMs = float((AccumulateEndTime - AccumulateBeginTime)*1000.0f);
-	
-		UE_LOG(LogMovieRenderPipeline, Log, TEXT("Accumulation time: %8.2fms"), ElapsedMs);
-
-	}
-
-	if (FrameData->IsLastTile() && FrameData->IsLastTemporalSample())
-	{
-		int32 FullSizeX = InParams.ImageAccumulator->PlaneSize.X;
-		int32 FullSizeY = InParams.ImageAccumulator->PlaneSize.Y;
-
-		// Now that a tile is fully built and accumulated we can notify the output builder that the
-		// data is ready so it can pass that onto the output containers (if needed).
-		static bool bFullLinearColor = false; // Make this an option. For now, if it's linear, write EXR. Otherwise PNGs for 8bit.
-		if (bFullLinearColor)
-		{
-			// 32 bit FLinearColor
-			TUniquePtr<TImagePixelData<FLinearColor> > FinalPixelData = MakeUnique<TImagePixelData<FLinearColor>>(FIntPoint(FullSizeX, FullSizeY), FrameData);
-			InParams.ImageAccumulator->FetchFinalPixelDataLinearColor(FinalPixelData->Pixels);
-
-			// Send the data to the Output Builder
-			InParams.OutputMerger->OnCompleteRenderPassDataAvailable_AnyThread(MoveTemp(FinalPixelData));
-		}
-		else
-		{
-			// 8bit FColors
-			TUniquePtr<TImagePixelData<FColor>> FinalPixelData = MakeUnique<TImagePixelData<FColor>>(FIntPoint(FullSizeX, FullSizeY), FrameData);
-			InParams.ImageAccumulator->FetchFinalPixelDataByte(FinalPixelData->Pixels);
-
-			// Send the data to the Output Builder
-			InParams.OutputMerger->OnCompleteRenderPassDataAvailable_AnyThread(MoveTemp(FinalPixelData));
-		}
-
-		// Free the memory in the accumulator.
-		InParams.ImageAccumulator->Reset();
-	}
-}
\ No newline at end of file
+#include "BufferVisualizationData.h"
+#include "Materials/Material.h"
+#include "FinalPostProcessSettings.h"
+#include "ImageWriteStream.h"
diff --git a/Engine/Plugins/MovieScene/MovieRenderPipeline/Source/MovieRenderPipelineRenderPasses/Private/MoviePipelineDeferredPasses.cpp b/Engine/Plugins/MovieScene/MovieRenderPipeline/Source/MovieRenderPipelineRenderPasses/Private/MoviePipelineDeferredPasses.cpp
new file mode 100644
--- /dev/null
+++ b/Engine/Plugins/MovieScene/MovieRenderPipeline/Source/MovieRenderPipelineRenderPasses/Private/MoviePipelineDeferredPasses.cpp
@@ -0,0 +1,563 @@
+// Copyright 1998-2019 Epic Games, Inc. All Rights Reserved.
+
+#include "MoviePipelineDeferredPasses.h"
+#include "MoviePipelineHighResSetting.h"
+#include "Engine/TextureRenderTarget2D.h"
+#include "ShowFlags.h"
+#include "SceneView.h"
+#include "LegacyScreenPercentageDriver.h"
+#include "MovieRenderPipelineDataTypes.h"
+#include "GameFramework/PlayerController.h"
+#include "MoviePipelineRenderPass.h"
+#include "EngineModule.h"
+#include "Engine/World.h"
+#include "Engine/TextureRenderTarget.h"
+#include "MoviePipeline.h"
+#include "Misc/FrameRate.h"
+#include "CanvasTypes.h"
+#include "MoviePipelineShotConfig.h"
+#include "MovieRenderOverlappedImage.h"
+#include "MovieRenderPipelineCoreModule.h"
+#include "ImagePixelData.h"
+#include "MoviePipelineOutputBuilder.h"
+#include "BufferVisualizationData.h"
+#include "Containers/Array.h"
+#include "FinalPostProcessSettings.h"
+
+namespace MoviePipeline
+{
+	// Forward Declare
+	void ReadBackbufferAndBroadcast_RenderThread(FRHICommandListImmediate &RHICmdList, FRenderTarget* InFromRenderTarget, const FMoviePipelineRenderPassMetrics& InSampleState, const FMoviePipelineSampleReady& OnReadDelegate);
+	void AccumulateSample_RenderThread(TUniquePtr<FImagePixelData>&& InPixelData, TSharedRef<FImagePixelDataPayload, ESPMode::ThreadSafe> InFrameData, MoviePipeline::FSampleRenderThreadParams& InParams);
+
+	void FDeferredRenderEnginePass::Setup(TWeakObjectPtr<UMoviePipeline> InOwningPipeline, const FMoviePipelineRenderPassInitSettings& InInitSettings)
+	{
+		FMoviePipelineEnginePass::Setup(InOwningPipeline, InInitSettings);
+
+		TileRenderTarget = NewObject<UTextureRenderTarget2D>(GetTransientPackage());
+		TileRenderTarget->ClearColor = FLinearColor(0.0f, 0.0f, 0.0f, 0.0f);
+		TileRenderTarget->AddToRoot();
+
+		// Initialize to the tile size (not final size) and use a 16 bit back buffer to avoid precision issues when accumulating later
+		TileRenderTarget->InitCustomFormat(InInitSettings.BackbufferResolution.X, InInitSettings.BackbufferResolution.Y, EPixelFormat::PF_A16B16G16R16, false);
+
+		// Allocate 
+		ViewState.Allocate();
+	}
+
+	void FDeferredRenderEnginePass::Teardown()
+	{
+		FMoviePipelineEnginePass::Teardown();
+
+		if (TileRenderTarget.IsValid())
+		{
+			TileRenderTarget->RemoveFromRoot();
+		}
+
+		ViewState.Destroy();
+	}
+
+	void FDeferredRenderEnginePass::RenderSample_GameThread(const FMoviePipelineRenderPassMetrics& InSampleState)
+	{
+		FMoviePipelineEnginePass::RenderSample_GameThread(InSampleState);
+
+		float TimeSeconds = InSampleState.OutputState.WorldSeconds;
+		float RealTimeSeconds = InSampleState.OutputState.WorldSeconds;
+		float DeltaTimeSeconds = InSampleState.OutputState.FrameDeltaTime;
+
+		FEngineShowFlags ShowFlags = FEngineShowFlags(EShowFlagInitMode::ESFIM_Game);
+
+		if (InSampleState.bIsUsingOverlappedTiles)
+		{
+			// Disable these for now
+			// ShowFlags.SetDepthOfField(false);
+			// ShowFlags.SetMotionBlur(false);
+		}
+
+		FSceneViewFamilyContext ViewFamily(FSceneViewFamily::ConstructionValues(
+			TileRenderTarget->GameThread_GetRenderTargetResource(),
+			GetPipeline()->GetWorld()->Scene,
+			ShowFlags)
+			.SetWorldTimes(TimeSeconds, DeltaTimeSeconds, RealTimeSeconds)
+			.SetRealtimeUpdate(true));
+
+		ViewFamily.SceneCaptureSource = ESceneCaptureSource::SCS_FinalColorHDR;
+		ViewFamily.SetScreenPercentageInterface(new FLegacyScreenPercentageDriver(ViewFamily, 1.0f, false));
+
+		// View is added as a child of the ViewFamily
+		FSceneView* View = CalcSceneView(&ViewFamily, InSampleState);
+
+		// ToDo: This doesn't work.
+		ViewFamily.FrameNumber += InSampleState.SpatialSampleIndex;
+
+		// Override the Motion Blur settings since these are controlled by the movie pipeline.
+		{
+			View->FinalPostProcessSettings.MotionBlurTargetFPS = FMath::RoundToInt(GetPipeline()->GetEffectiveFrameRate().AsDecimal());
+			View->FinalPostProcessSettings.MotionBlurAmount = InSampleState.OutputState.MotionBlurFraction;
+			View->FinalPostProcessSettings.MotionBlurMax = 100.f;
+			View->FinalPostProcessSettings.bOverride_MotionBlurAmount = true;
+			View->FinalPostProcessSettings.bOverride_MotionBlurTargetFPS = true;
+			View->FinalPostProcessSettings.bOverride_MotionBlurMax = true;
+		}
+
+		// Autoexposure is not allowed
+		if (View->FinalPostProcessSettings.AutoExposureMethod != EAutoExposureMethod::AEM_Manual)
+		{
+			UE_LOG(LogMovieRenderPipeline, Log, TEXT("AutoExposure Method should always be Manual."));
+
+			View->FinalPostProcessSettings.AutoExposureMethod = EAutoExposureMethod::AEM_Manual;
+		}
+
+		const bool bCanUseAA = InSampleState.SpatialSampleCount == 1 && 
+							   (InSampleState.TileCounts.X == 1 && InSampleState.TileCounts.Y == 1) && 
+							   InSampleState.TemporalSampleCount == 1;
+		if (!bCanUseAA)
+		{
+			View->AntiAliasingMethod = EAntiAliasingMethod::AAM_None;
+
+			// Override the view matrix with our offsets provided by tiling/jittering. This requires
+			// TAA to be disabled.
+			{
+				View->ViewMatrices.HackAddTemporalAAProjectionJitter(InSampleState.SpatialShift);
+			}
+		}
+
+		// Bias all mip-mapping to pretend to be working at our target resolution and not our tile resolution
+		// so that the images don't end up soft.
+		{
+			float EffectivePrimaryResolutionFraction = 1.f / InSampleState.TileCounts.X;
+			View->MaterialTextureMipBias = FMath::Log2(EffectivePrimaryResolutionFraction);
+
+			// Add an additional bias per user settings. This allows them to choose to make the textures sharper if it
+			// looks better with their particular settings.
+			UMoviePipelineHighResSetting* HighResSettings = GetPipeline()->GetPipelineCurrentShotConfig()->FindSetting<UMoviePipelineHighResSetting>();
+			if (HighResSettings)
+			{
+				View->MaterialTextureMipBias += HighResSettings->TextureSharpnessBias;
+			}
+		}
+
+		// If this is a history only frame it means we just came from a camera cut, flag it as such
+		// which will have knock-on effects in histories (such as clearing old TAA buffers).
+		if (InSampleState.bIsHistoryOnlyFrame)
+		{
+			View->bCameraCut = true;
+		}
+
+		// The world is always paused, but we are setting the motion blur and view matrices manually
+		ViewFamily.bWorldIsPaused = true;
+
+		FCanvas Canvas = FCanvas(TileRenderTarget->GameThread_GetRenderTargetResource(), nullptr, GetPipeline()->GetWorld(), ERHIFeatureLevel::SM5, FCanvas::CDM_DeferDrawing, 1.0f);
+
+
+		// SetupViewDelegate.Broadcast(ViewFamily, *View, InSampleState);
+
+		// Draw the world into this View Family
+		GetRendererModule().BeginRenderingViewFamily(&Canvas, &ViewFamily);
+
+		FRenderTarget* BackbufferRenderTarget = TileRenderTarget->GameThread_GetRenderTargetResource();
+		FMoviePipelineSampleReady& OnBackbufferReadDelegate = BackbufferReadyDelegate;
+
+		ENQUEUE_RENDER_COMMAND(CanvasRenderTargetResolveCommand)(
+			[InSampleState, BackbufferRenderTarget, OnBackbufferReadDelegate](FRHICommandListImmediate& RHICmdList)
+		{
+			// If they only needed this frame for the history, we don't actually care about trying to output it,
+			// so we can skip the copy and pushing to output.
+			if (InSampleState.bIsHistoryOnlyFrame)
+			{
+				return;
+			}
+
+			// If they are using tiles, we'll output some more fine-grained progress via logging since it freezes the UI for such a long time.
+			if (InSampleState.GetTileCount() > 1)
+			{
+				UE_LOG(LogMovieRenderPipeline, Log, TEXT("Frame: %d TemporalSample: %d/%d Rendered Tile: %d/%d Sample: %d/%d"),
+					InSampleState.OutputState.OutputFrameNumber, InSampleState.TemporalSampleIndex + 1, InSampleState.TemporalSampleCount,
+					InSampleState.GetTileIndex() + 1 / InSampleState.GetTileCount(), InSampleState.SpatialSampleIndex + 1, InSampleState.SpatialSampleCount);
+			}
+			
+			ReadBackbufferAndBroadcast_RenderThread(RHICmdList, BackbufferRenderTarget, InSampleState, OnBackbufferReadDelegate);
+		});
+	}
+
+	void ReadBackbufferAndBroadcast_RenderThread(FRHICommandListImmediate &RHICmdList, FRenderTarget* InFromRenderTarget, const FMoviePipelineRenderPassMetrics& InSampleState, const FMoviePipelineSampleReady& OnReadDelegate)
+	{
+		check(IsInRenderingThread());
+
+		FIntRect SourceRect = FIntRect(0, 0, InFromRenderTarget->GetSizeXY().X, InFromRenderTarget->GetSizeXY().Y);
+
+		// Read the data back to the CPU
+		TArray<FLinearColor> RawPixels;
+		RawPixels.SetNum(SourceRect.Width() * SourceRect.Height());
+
+		FReadSurfaceDataFlags ReadDataFlags;
+		ReadDataFlags.SetLinearToGamma(false);
+
+		RHICmdList.ReadSurfaceData(InFromRenderTarget->GetRenderTargetTexture(), SourceRect, RawPixels, ReadDataFlags);
+
+		OnReadDelegate.Broadcast(RawPixels, InSampleState);
+	}
+
+	FSceneView* FDeferredRenderEnginePass::CalcSceneView(FSceneViewFamily* ViewFamily, const FMoviePipelineRenderPassMetrics& InSampleState)
+	{
+		APlayerController* LocalPlayerController = GetPipeline()->GetWorld()->GetFirstPlayerController();
+
+		int32 TileSizeX = InitSettings.BackbufferResolution.X;
+		int32 TileSizeY = InitSettings.BackbufferResolution.Y;
+
+		FSceneViewInitOptions ViewInitOptions;
+		ViewInitOptions.ViewFamily = ViewFamily;
+		ViewInitOptions.ViewOrigin = InSampleState.FrameInfo.CurrViewLocation;
+		ViewInitOptions.SetViewRectangle(FIntRect(FIntPoint(0, 0), FIntPoint(TileSizeX, TileSizeY)));
+		ViewInitOptions.ViewRotationMatrix = FInverseRotationMatrix(InSampleState.FrameInfo.CurrViewRotation);
+
+		// Rotate the view 90 degrees (reason: unknown)
+		ViewInitOptions.ViewRotationMatrix = ViewInitOptions.ViewRotationMatrix * FMatrix(
+			FPlane(0, 0, 1, 0),
+			FPlane(1, 0, 0, 0),
+			FPlane(0, 1, 0, 0),
+			FPlane(0, 0, 0, 1));
+		float ViewFOV = 90.f;
+		if (GetPipeline()->GetWorld()->GetFirstPlayerController()->PlayerCameraManager)
+		{
+			ViewFOV = GetPipeline()->GetWorld()->GetFirstPlayerController()->PlayerCameraManager->GetFOVAngle();
+		}
+
+		// Calculate a Projection Matrix
+		{
+			float XAxisMultiplier;
+			float YAxisMultiplier;
+
+			if (ViewInitOptions.GetViewRect().Width() > ViewInitOptions.GetViewRect().Height())
+			{
+				//if the viewport is wider than it is tall
+				XAxisMultiplier = 1.0f;
+				YAxisMultiplier = ViewInitOptions.GetViewRect().Width() / (float)ViewInitOptions.GetViewRect().Height();
+			}
+			else
+			{
+				//if the viewport is taller than it is wide
+				XAxisMultiplier = ViewInitOptions.GetViewRect().Height() / (float)ViewInitOptions.GetViewRect().Width();
+				YAxisMultiplier = 1.0f;
+			}
+
+			const float MinZ = GNearClippingPlane;
+			const float MaxZ = MinZ;
+			// Avoid zero ViewFOV's which cause divide by zero's in projection matrix
+			const float MatrixFOV = FMath::Max(0.001f, ViewFOV) * (float)PI / 360.0f;
+
+			FMatrix BaseProjMatrix;
+
+			if ((bool)ERHIZBuffer::IsInverted)
+			{
+				BaseProjMatrix = FReversedZPerspectiveMatrix(
+					MatrixFOV,
+					MatrixFOV,
+					XAxisMultiplier,
+					YAxisMultiplier,
+					MinZ,
+					MaxZ
+				);
+			}
+			else
+			{
+				BaseProjMatrix = FPerspectiveMatrix(
+					MatrixFOV,
+					MatrixFOV,
+					XAxisMultiplier,
+					YAxisMultiplier,
+					MinZ,
+					MaxZ
+				);
+			}
+
+			if (InSampleState.bIsUsingOverlappedTiles)
+			{
+				float PadRatioX = 1.0f;
+				float PadRatioY = 1.0f;
+				
+				if (InSampleState.OverlappedPad.X > 0 && InSampleState.OverlappedPad.Y > 0)
+				{
+					PadRatioX = float(InSampleState.OverlappedPad.X * 2 + InSampleState.OverlappedPad.X) / float(InSampleState.OverlappedPad.X);
+					PadRatioY = float(InSampleState.OverlappedPad.Y * 2 + InSampleState.OverlappedPad.Y) / float(InSampleState.OverlappedPad.Y);
+				}
+				
+				float ScaleX = PadRatioX / float(InitSettings.TileCount.X);
+				float ScaleY = PadRatioY / float(InitSettings.TileCount.Y);
+				
+				BaseProjMatrix.M[0][0] /= ScaleX;
+				BaseProjMatrix.M[1][1] /= ScaleY;
+				
+				// this offset would be correct with no pad
+				float OffsetX = -((float(InSampleState.TileIndexes.X) + 0.5f - float(InitSettings.TileCount.X) / 2.0f)* 2.0f);
+				float OffsetY = ((float(InSampleState.TileIndexes.Y) + 0.5f - float(InitSettings.TileCount.Y) / 2.0f)* 2.0f);
+				
+				BaseProjMatrix.M[2][0] += OffsetX / PadRatioX;
+				BaseProjMatrix.M[2][1] += OffsetY / PadRatioX;
+			}
+
+			ViewInitOptions.ProjectionMatrix = BaseProjMatrix;
+		}
+
+		ViewInitOptions.SceneViewStateInterface = ViewState.GetReference();
+		ViewInitOptions.FOV = ViewFOV;
+
+		FSceneView* View = new FSceneView(ViewInitOptions);
+		ViewFamily->Views.Add(View);
+		View->ViewLocation = InSampleState.FrameInfo.CurrViewLocation;
+		View->ViewRotation = InSampleState.FrameInfo.CurrViewRotation;
+		View->PreviousViewTransform = FTransform(InSampleState.FrameInfo.PrevViewRotation, InSampleState.FrameInfo.PrevViewLocation);
+		
+		View->StartFinalPostprocessSettings(View->ViewLocation);
+		View->OverrideFrameIndexValue = InSampleState.SpatialSampleIndex;
+
+		// CameraAnim override
+		if (LocalPlayerController->PlayerCameraManager)
+		{
+			TArray<FPostProcessSettings> const* CameraAnimPPSettings;
+			TArray<float> const* CameraAnimPPBlendWeights;
+			LocalPlayerController->PlayerCameraManager->GetCachedPostProcessBlends(CameraAnimPPSettings, CameraAnimPPBlendWeights);
+
+			FMinimalViewInfo ViewInfo = LocalPlayerController->PlayerCameraManager->GetCameraCachePOV();
+			for (int32 PPIdx = 0; PPIdx < CameraAnimPPBlendWeights->Num(); ++PPIdx)
+			{
+				View->OverridePostProcessSettings((*CameraAnimPPSettings)[PPIdx], (*CameraAnimPPBlendWeights)[PPIdx]);
+			}
+
+			View->OverridePostProcessSettings(ViewInfo.PostProcessSettings, ViewInfo.PostProcessBlendWeight);
+
+		}
+
+		View->EndFinalPostprocessSettings(ViewInitOptions);
+
+		return View;
+	}
+}
+
+
+void UMoviePipelineDeferredPassBase::GetRequiredEnginePassesImpl(TSet<FMoviePipelinePassIdentifier>& RequiredEnginePasses)
+{
+	Super::GetRequiredEnginePassesImpl(RequiredEnginePasses);
+
+	RequiredEnginePasses.Add(FMoviePipelinePassIdentifier(TEXT("MainDeferredPass")));
+}
+
+void UMoviePipelineDeferredPassBase::SetupImpl(TArray<TSharedPtr<MoviePipeline::FMoviePipelineEnginePass>>& InEnginePasses)
+{
+	Super::SetupImpl(InEnginePasses);
+
+	// Look for the Deferred Engine Pass
+	const FMoviePipelinePassIdentifier DesiredPass = FMoviePipelinePassIdentifier(TEXT("MainDeferredPass"));
+	for (const TSharedPtr<MoviePipeline::FMoviePipelineEnginePass>& Pass : InEnginePasses)
+	{
+		if (Pass->PassIdentifier == DesiredPass)
+		{
+			// Register our callbacks here. We need access to the Backbuffer and the Composition Graph Pipes.
+			TSharedPtr<MoviePipeline::FDeferredRenderEnginePass> DeferredPass = StaticCastSharedPtr<MoviePipeline::FDeferredRenderEnginePass>(Pass);
+			DeferredPass->BackbufferReadyDelegate.AddUObject(this, &UMoviePipelineDeferredPassBase::OnBackbufferSampleReady);
+			DeferredPass->SetupViewDelegate.AddUObject(this, &UMoviePipelineDeferredPassBase::OnSetupView);
+			break;
+		}
+	}
+
+	DesiredOutputPasses.Add(TEXT("Backbuffer"));
+	// DesiredOutputPasses.Add(TEXT("SceneDepth"));
+	// DesiredOutputPasses.Add(TEXT("AmbientOcclusion"));
+
+	for (const FString& Pass : DesiredOutputPasses)
+	{
+		// Create an allocator for each expected output pass.
+		ImageTileAccumulator.Add(MakeShared<FImageOverlappedAccumulator, ESPMode::ThreadSafe>());
+	}
+}
+
+void UMoviePipelineDeferredPassBase::GatherOutputPassesImpl(TArray<FMoviePipelinePassIdentifier>& ExpectedRenderPasses)
+{
+	Super::GatherOutputPassesImpl(ExpectedRenderPasses);
+
+	// We will produce the following output files.
+	for (const FString& Pass : DesiredOutputPasses)
+	{
+		ExpectedRenderPasses.Add(FMoviePipelinePassIdentifier(Pass));
+	}
+}
+
+void UMoviePipelineDeferredPassBase::OnBackbufferSampleReady(TArray<FLinearColor> InPixelData, FMoviePipelineRenderPassMetrics InSampleState)
+{
+	MoviePipeline::FSampleRenderThreadParams AccumulationParams;
+	{
+		AccumulationParams.bWriteTiles = true;
+		AccumulationParams.OutputMerger = GetPipeline()->OutputBuilder;
+		AccumulationParams.ImageAccumulator = ImageTileAccumulator[0];
+		AccumulationParams.SampleState = InSampleState;
+	}
+
+	TSharedRef<FImagePixelDataPayload, ESPMode::ThreadSafe> FrameData = MakeShared<FImagePixelDataPayload, ESPMode::ThreadSafe>();
+	FrameData->OutputState = InSampleState.OutputState;
+	FrameData->PassIdentifier = FMoviePipelinePassIdentifier(TEXT("Backbuffer"));
+	FrameData->SampleState = InSampleState;
+
+	TUniquePtr<TImagePixelData<FLinearColor>> PixelData = MakeUnique<TImagePixelData<FLinearColor>>(InSampleState.BackbufferSize, TArray64<FLinearColor>(MoveTemp(InPixelData)), FrameData);
+
+	AccumulateSample_RenderThread(MoveTemp(PixelData), FrameData, AccumulationParams);
+}
+
+void UMoviePipelineDeferredPassBase::OnSetupView(FSceneViewFamily& InViewFamily, FSceneView& InView, const FMoviePipelineRenderPassMetrics& InSampleState)
+{
+	struct FIterator
+	{
+		FFinalPostProcessSettings& FinalPostProcessSettings;
+		const TArray<FString>& RenderPasses;
+
+		FIterator(FFinalPostProcessSettings& InFinalPostProcessSettings, const TArray<FString>& InRenderPasses)
+			: FinalPostProcessSettings(InFinalPostProcessSettings), RenderPasses(InRenderPasses)
+		{}
+
+		void ProcessValue(const FString& InName, UMaterial* Material, const FText& InText)
+		{
+			if (!RenderPasses.Num() || RenderPasses.Contains(InName) || RenderPasses.Contains(InText.ToString()))
+			{
+				FinalPostProcessSettings.BufferVisualizationOverviewMaterials.Add(Material);
+			}
+		}
+	} Iterator(InView.FinalPostProcessSettings, DesiredOutputPasses);
+	GetBufferVisualizationData().IterateOverAvailableMaterials(Iterator);
+
+	for(int32 PassIndex = 0; PassIndex < InView.FinalPostProcessSettings.BufferVisualizationOverviewMaterials.Num(); PassIndex++)
+	{
+		const UMaterialInterface* Pass = InView.FinalPostProcessSettings.BufferVisualizationOverviewMaterials[PassIndex];
+		UE_LOG(LogTemp, Log, TEXT("Pass: %s"), *Pass->GetName());
+
+		FName PassName = FName(*Pass->GetName());
+		TSharedPtr<FImagePixelPipe, ESPMode::ThreadSafe>* Pipe = InView.FinalPostProcessSettings.BufferVisualizationPipes.Find(PassName);
+		if (!Pipe)
+		{
+			Pipe = &InView.FinalPostProcessSettings.BufferVisualizationPipes.Add(PassName, MakeShared<FImagePixelPipe, ESPMode::ThreadSafe>());
+		}
+
+		MoviePipeline::FSampleRenderThreadParams AccumulationParams;
+		{
+			AccumulationParams.bWriteTiles = true;
+			AccumulationParams.OutputMerger = GetPipeline()->OutputBuilder;
+			AccumulationParams.ImageAccumulator = ImageTileAccumulator[PassIndex + 1]; // 0 is reserved for backbuffer.
+			AccumulationParams.SampleState = InSampleState;
+		}
+
+		auto OnImageReceived = [PassName, AccumulationParams](TUniquePtr<FImagePixelData>&& InPixelData) mutable
+		{
+			UE_LOG(LogTemp, Log, TEXT("Recieved Data for Pass: %d Size: %dx%d BitDepth: %d"), *PassName.ToString(),
+				InPixelData->GetSize().X, InPixelData->GetSize().Y, InPixelData->GetBitDepth());
+
+			// Don't try to accumulate samples rendered for the sake of setting up history.
+			if (AccumulationParams.SampleState.bIsHistoryOnlyFrame)
+			{
+				return;
+			}
+			
+			TSharedRef<FImagePixelDataPayload, ESPMode::ThreadSafe> FrameData = MakeShared<FImagePixelDataPayload, ESPMode::ThreadSafe>();
+			FrameData->OutputState = AccumulationParams.SampleState.OutputState;
+			FrameData->PassIdentifier = FMoviePipelinePassIdentifier(PassName.ToString());
+			FrameData->SampleState = AccumulationParams.SampleState;
+
+			MoviePipeline::AccumulateSample_RenderThread(MoveTemp(InPixelData), FrameData, AccumulationParams);
+		};
+
+
+		(*Pipe)->AddEndpoint(OnImageReceived);
+	}
+
+
+	// if (PostProcessingMaterial)
+	// {
+	// 	FWeightedBlendable Blendable(1.f, PostProcessingMaterial);
+	// 	PostProcessingMaterial->OverrideBlendableSettings(InView, 1.f);
+	// }
+}
+
+
+namespace MoviePipeline
+{
+	void AccumulateSample_RenderThread(TUniquePtr<FImagePixelData>&& InPixelData, TSharedRef<FImagePixelDataPayload, ESPMode::ThreadSafe> InFrameData, MoviePipeline::FSampleRenderThreadParams& InParams)
+	{
+		check(InPixelData->IsDataWellFormed());
+
+		// Writing tiles can be useful for debug reasons. These get passed onto the output every frame.
+		if (InParams.bWriteTiles)
+		{
+			// Send the data to the Output Builder. This has to be a copy of the pixel data from the GPU, since
+			// it enqueues it onto the game thread and won't be read/sent to write to disk for another frame. 
+			// The extra copy is unfortunate, but is only the size of a single sample (ie: 1920x1080 -> 17mb)
+			TUniquePtr<FImagePixelData> SampleData = InPixelData->CopyImageData();
+			InParams.OutputMerger->OnSingleSampleDataAvailable_AnyThread(MoveTemp(SampleData), InFrameData);
+		}
+
+		// For the first sample in a new output, we allocate memory
+		if (InFrameData->IsFirstTile() && InFrameData->IsFirstTemporalSample())
+		{
+			InParams.ImageAccumulator->InitMemory(FIntPoint(InFrameData->SampleState.BackbufferSize.X * InFrameData->SampleState.TileCounts.X, InFrameData->SampleState.BackbufferSize.Y * InFrameData->SampleState.TileCounts.Y), 3);
+			InParams.ImageAccumulator->ZeroPlanes();
+			InParams.ImageAccumulator->AccumulationGamma = InFrameData->SampleState.AccumulationGamma;
+		}
+
+		// Accumulate the new sample to our target
+		{
+			const double AccumulateBeginTime = FPlatformTime::Seconds();
+
+			FIntPoint RawSize = InPixelData->GetSize();
+
+			check(InFrameData->SampleState.BackbufferSize.X + 2 * InFrameData->SampleState.OverlappedPad.X == RawSize.X);
+			check(InFrameData->SampleState.BackbufferSize.Y + 2 * InFrameData->SampleState.OverlappedPad.Y == RawSize.Y);
+
+			InParams.ImageAccumulator->AccumulatePixelData(*InPixelData.Get(), InFrameData->SampleState.OverlappedOffset, InFrameData->SampleState.OverlappedSubpixelShift);
+
+			const double AccumulateEndTime = FPlatformTime::Seconds();
+			const float ElapsedMs = float((AccumulateEndTime - AccumulateBeginTime)*1000.0f);
+
+			UE_LOG(LogMovieRenderPipeline, Log, TEXT("Accumulation time: %8.2fms"), ElapsedMs);
+
+		}
+
+		if (InFrameData->IsLastTile() && InFrameData->IsLastTemporalSample())
+		{
+			int32 FullSizeX = InParams.ImageAccumulator->PlaneSize.X;
+			int32 FullSizeY = InParams.ImageAccumulator->PlaneSize.Y;
+
+			// Now that a tile is fully built and accumulated we can notify the output builder that the
+			// data is ready so it can pass that onto the output containers (if needed).
+			if (InPixelData->GetType() == EImagePixelType::Float32)
+			{
+				// 32 bit FLinearColor
+				TUniquePtr<TImagePixelData<FLinearColor> > FinalPixelData = MakeUnique<TImagePixelData<FLinearColor>>(FIntPoint(FullSizeX, FullSizeY));
+				InParams.ImageAccumulator->FetchFinalPixelDataLinearColor(FinalPixelData->Pixels);
+
+				// Send the data to the Output Builder
+				InParams.OutputMerger->OnCompleteRenderPassDataAvailable_AnyThread(MoveTemp(FinalPixelData), InFrameData);
+			}
+			else if (InPixelData->GetType() == EImagePixelType::Float16)
+			{
+				// 32 bit FLinearColor
+				TUniquePtr<TImagePixelData<FFloat16Color> > FinalPixelData = MakeUnique<TImagePixelData<FFloat16Color>>(FIntPoint(FullSizeX, FullSizeY));
+				InParams.ImageAccumulator->FetchFinalPixelDataHalfFloat(FinalPixelData->Pixels);
+			
+				// Send the data to the Output Builder
+				InParams.OutputMerger->OnCompleteRenderPassDataAvailable_AnyThread(MoveTemp(FinalPixelData), InFrameData);
+			}
+			else if(InPixelData->GetType() == EImagePixelType::Color)
+			{
+				// 8bit FColors
+				TUniquePtr<TImagePixelData<FColor>> FinalPixelData = MakeUnique<TImagePixelData<FColor>>(FIntPoint(FullSizeX, FullSizeY));
+				InParams.ImageAccumulator->FetchFinalPixelDataByte(FinalPixelData->Pixels);
+
+				// Send the data to the Output Builder
+				InParams.OutputMerger->OnCompleteRenderPassDataAvailable_AnyThread(MoveTemp(FinalPixelData), InFrameData);
+			}
+			else
+			{
+				check(0);
+			}
+
+			// Free the memory in the accumulator.
+			InParams.ImageAccumulator->Reset();
+		}
+	}
+}
\ No newline at end of file
diff --git a/Engine/Plugins/MovieScene/MovieRenderPipeline/Source/MovieRenderPipelineRenderPasses/Private/MoviePipelineImageSequenceContainer.cpp b/Engine/Plugins/MovieScene/MovieRenderPipeline/Source/MovieRenderPipelineRenderPasses/Private/MoviePipelineImageSequenceContainer.cpp
--- a/Engine/Plugins/MovieScene/MovieRenderPipeline/Source/MovieRenderPipelineRenderPasses/Private/MoviePipelineImageSequenceContainer.cpp
+++ b/Engine/Plugins/MovieScene/MovieRenderPipeline/Source/MovieRenderPipelineRenderPasses/Private/MoviePipelineImageSequenceContainer.cpp
@@ -49,13 +49,32 @@ void UMoviePipelineImageSequenceContainerBase::OnRecieveImageDataImpl(FMoviePipe
 		TUniquePtr<FImageWriteTask> TileImageTask = MakeUnique<FImageWriteTask>();
 
 		// Fill alpha for now 
-		TileImageTask->PixelPreProcessors.Add(TAsyncAlphaWrite<FColor>(255));
+		switch (RenderPassData.Value->GetType())
+		{
+		case EImagePixelType::Color:
+		{
+			TileImageTask->PixelPreProcessors.Add(TAsyncAlphaWrite<FColor>(255));
+			break;
+		}
+		case EImagePixelType::Float16:
+		{
+
+			break;
+		}
+		case EImagePixelType::Float32:
+		{
+			TileImageTask->PixelPreProcessors.Add(TAsyncAlphaWrite<FLinearColor>(255));
+			break;
+		}
+		default:
+			check(false);
+		}
 
 		// JPEG output
-		TileImageTask->Format = EImageFormat::JPEG;
+		TileImageTask->Format = EImageFormat::EXR;
 		TileImageTask->CompressionQuality = 100;
 
-		FString OutputName = FString::Printf(TEXT("/%s.%d.jpeg"),
+		FString OutputName = FString::Printf(TEXT("/%s.%d.exr"),
 			*RenderPassData.Key.Name, InMergedOutputFrame->FrameOutputState.OutputFrameNumber);
 
 		FString OutputPath = OutputDirectory + OutputName;
diff --git a/Engine/Plugins/MovieScene/MovieRenderPipeline/Source/MovieRenderPipelineRenderPasses/Private/MovieRenderPipelineRenderPassesModule.cpp b/Engine/Plugins/MovieScene/MovieRenderPipeline/Source/MovieRenderPipelineRenderPasses/Private/MovieRenderPipelineRenderPassesModule.cpp
--- a/Engine/Plugins/MovieScene/MovieRenderPipeline/Source/MovieRenderPipelineRenderPasses/Private/MovieRenderPipelineRenderPassesModule.cpp
+++ b/Engine/Plugins/MovieScene/MovieRenderPipeline/Source/MovieRenderPipelineRenderPasses/Private/MovieRenderPipelineRenderPassesModule.cpp
@@ -1,18 +1,36 @@
 // Copyright 1998-2019 Epic Games, Inc. All Rights Reserved.
 
 #include "Modules/ModuleManager.h"
+#include "MovieRenderPipelineDataTypes.h"
+#include "MovieRenderPipelineCoreModule.h"
+#include "MoviePipelineDeferredPasses.h"
 
 class FMovieRenderPipelineRenderPassesModule : public IModuleInterface
 {
-	virtual void StartupModule() override
+
+	static TSharedRef<MoviePipeline::FMoviePipelineEnginePass> CreateDeferredEnginePass()
 	{
+		return MakeShared<MoviePipeline::FDeferredRenderEnginePass>();
+	}
 
+	virtual void StartupModule() override
+	{
+		FMovieRenderPipelineCoreModule& CoreModule = FModuleManager::Get().LoadModuleChecked<FMovieRenderPipelineCoreModule>("MovieRenderPipelineCore");
+		DeferredEnginePassHandle = CoreModule.RegisterEngineRenderPass(FOnCreateEngineRenderPass::CreateStatic(&FMovieRenderPipelineRenderPassesModule::CreateDeferredEnginePass));
 	}
 
 	virtual void ShutdownModule() override
 	{
+		FMovieRenderPipelineCoreModule* CoreModule = (FMovieRenderPipelineCoreModule*)FModuleManager::Get().LoadModule("MovieRenderPipelineCore");
+		if (CoreModule)
+		{
+			CoreModule->UnregisterEngineRenderPass(DeferredEnginePassHandle);
+		}
 
 	}
+private:
+
+	FDelegateHandle DeferredEnginePassHandle;
 };
 
 IMPLEMENT_MODULE(FMovieRenderPipelineRenderPassesModule, MovieRenderPipelineRenderPasses);
diff --git a/Engine/Plugins/MovieScene/MovieRenderPipeline/Source/MovieRenderPipelineRenderPasses/Public/MoviePipelineBackbufferPass.h b/Engine/Plugins/MovieScene/MovieRenderPipeline/Source/MovieRenderPipelineRenderPasses/Public/MoviePipelineBackbufferPass.h
--- a/Engine/Plugins/MovieScene/MovieRenderPipeline/Source/MovieRenderPipelineRenderPasses/Public/MoviePipelineBackbufferPass.h
+++ b/Engine/Plugins/MovieScene/MovieRenderPipeline/Source/MovieRenderPipelineRenderPasses/Public/MoviePipelineBackbufferPass.h
@@ -22,19 +22,6 @@ class MOVIERENDERPIPELINERENDERPASSES_API UMoviePipelineBackbufferPass : public
 
 protected:
 
-	virtual void SetupImpl(const FMoviePipelineRenderPassInitSettings& InInitSettings) override;
-	virtual void TeardownImpl() override;
-	virtual void CaptureFrameImpl(const FMoviePipelineRenderPassMetrics& OutputFrameMetrics) override;
-	virtual void GatherOutputPassesImpl(TArray<FMoviePipelinePassIdentifier>& ExpectedRenderPasses) override;
-
-	FSceneView* CalcSceneView(FSceneViewFamily* ViewFamily, const FMoviePipelineRenderPassMetrics& InPassMetrics);
-
 private:
-	FSceneViewStateReference ViewState;
-
-
-	UPROPERTY(Transient)
-	UTextureRenderTarget2D* TileRenderTarget;
-
-	TSharedPtr<FImageOverlappedAccumulator, ESPMode::ThreadSafe> ImageTileAccumulator;
+	
 };
\ No newline at end of file
diff --git a/Engine/Plugins/MovieScene/MovieRenderPipeline/Source/MovieRenderPipelineRenderPasses/Public/MoviePipelineDeferredPasses.h b/Engine/Plugins/MovieScene/MovieRenderPipeline/Source/MovieRenderPipelineRenderPasses/Public/MoviePipelineDeferredPasses.h
new file mode 100644
--- /dev/null
+++ b/Engine/Plugins/MovieScene/MovieRenderPipeline/Source/MovieRenderPipelineRenderPasses/Public/MoviePipelineDeferredPasses.h
@@ -0,0 +1,67 @@
+// Copyright 1998-2019 Epic Games, Inc. All Rights Reserved.
+#pragma once
+
+#include "MoviePipelineRenderPass.h"
+#include "MovieRenderPipelineDataTypes.h"
+#include "SceneTypes.h"
+#include "SceneView.h"
+#include "MoviePipelineDeferredPasses.generated.h"
+
+class UTextureRenderTarget2D;
+struct FImageOverlappedAccumulator;
+class FSceneViewFamily;
+class FSceneView;
+
+UCLASS(BlueprintType)
+class MOVIERENDERPIPELINERENDERPASSES_API UMoviePipelineDeferredPassBase : public UMoviePipelineRenderPass
+{
+	GENERATED_BODY()
+
+protected:
+	
+	// UMoviePipelineRenderPass API
+	virtual void GetRequiredEnginePassesImpl(TSet<FMoviePipelinePassIdentifier>& RequiredEnginePasses) override;
+	virtual void GatherOutputPassesImpl(TArray<FMoviePipelinePassIdentifier>& ExpectedRenderPasses) override;
+	virtual void SetupImpl(TArray<TSharedPtr<MoviePipeline::FMoviePipelineEnginePass>>& InEnginePasses) override;
+	// ~UMovieRenderPassAPI
+
+	void OnBackbufferSampleReady(TArray<FLinearColor> InPixelData, FMoviePipelineRenderPassMetrics InSampleState);
+	void OnSetupView(FSceneViewFamily& InViewFamily, FSceneView& InView, const FMoviePipelineRenderPassMetrics& InSampleState);
+
+private:
+	/** List of passes by name that we should output. */
+	TArray<FString> DesiredOutputPasses;
+
+	/** One accumulator per sample source being saved. */
+	TArray<TSharedPtr<FImageOverlappedAccumulator, ESPMode::ThreadSafe>> ImageTileAccumulator;
+};
+
+DECLARE_MULTICAST_DELEGATE_TwoParams(FMoviePipelineSampleReady, TArray<FLinearColor>, FMoviePipelineRenderPassMetrics);
+DECLARE_MULTICAST_DELEGATE_ThreeParams(FMoviePipelineSetupView, FSceneViewFamily&, FSceneView&, const FMoviePipelineRenderPassMetrics&);
+
+namespace MoviePipeline
+{
+	struct FDeferredRenderEnginePass : public FMoviePipelineEnginePass
+	{
+		FDeferredRenderEnginePass()
+			: FMoviePipelineEnginePass(FMoviePipelinePassIdentifier(TEXT("MainDeferredPass")))
+		{
+		}
+
+		virtual void Setup(TWeakObjectPtr<UMoviePipeline> InOwningPipeline, const FMoviePipelineRenderPassInitSettings& InInitSettings) override;
+		virtual void RenderSample_GameThread(const FMoviePipelineRenderPassMetrics& InSampleState) override;
+		virtual void Teardown() override;
+
+		FMoviePipelineSampleReady BackbufferReadyDelegate;
+		FMoviePipelineSetupView SetupViewDelegate;
+
+	protected:
+		FSceneView* CalcSceneView(FSceneViewFamily* ViewFamily, const FMoviePipelineRenderPassMetrics& InSampleState);
+
+	protected:
+
+		FSceneViewStateReference ViewState;
+
+		TWeakObjectPtr<UTextureRenderTarget2D> TileRenderTarget;
+	};
+}
\ No newline at end of file
