MoviePipeline: Created a Camera settings object to better gather camera-specific settings. Allowed user to control need for auto-exposure and allowed user to override exposure compensation via the settings.

diff --git a/Engine/Plugins/MovieScene/MovieRenderPipeline/Source/MovieRenderPipelineCore/Private/MoviePipeline.cpp b/Engine/Plugins/MovieScene/MovieRenderPipeline/Source/MovieRenderPipelineCore/Private/MoviePipeline.cpp
--- a/Engine/Plugins/MovieScene/MovieRenderPipeline/Source/MovieRenderPipelineCore/Private/MoviePipeline.cpp
+++ b/Engine/Plugins/MovieScene/MovieRenderPipeline/Source/MovieRenderPipelineCore/Private/MoviePipeline.cpp
@@ -32,6 +32,7 @@
 #include "MoviePipelineBlueprintLibrary.h"
 #include "ImageWriteQueue.h"
 #include "MoviePipelineHighResSetting.h"
+#include "MoviePipelineCameraSetting.h"
 
 #define LOCTEXT_NAMESPACE "MoviePipeline"
 
@@ -622,6 +623,9 @@ TArray<FMoviePipelineShotInfo> UMoviePipeline::BuildShotListFromSequence(const U
 		UMoviePipelineAccumulationSetting* AccumulationSettings = Shot.ShotConfig->FindOrAddSetting<UMoviePipelineAccumulationSetting>();
 		Shot.NumHandleFrames = AccumulationSettings->HandleFrameCount;
 
+		UMoviePipelineCameraSetting* CameraSettings = Shot.ShotConfig->FindOrAddSetting<UMoviePipelineCameraSetting>();
+
+
 		// Expand the shot to encompass handle frames. This will modify our Camera Cuts bounds.
 		ExpandShot(Shot);
 
@@ -635,7 +639,7 @@ TArray<FMoviePipelineShotInfo> UMoviePipeline::BuildShotListFromSequence(const U
 		{
 			CameraCut.NumWarmUpFramesRemaining = CameraCut.NumWarmUpFrames = AccumulationSettings->WarmUpFrameCount;
 			CameraCut.bAccurateFirstFrameHistory = true;
-			CameraCut.NumTemporalSamples = AccumulationSettings->TemporalSampleCount;
+			CameraCut.NumTemporalSamples = CameraSettings->TemporalSampleCount;
 			CameraCut.NumSpatialSamples = AccumulationSettings->SpatialSampleCount;
 			CameraCut.CachedFrameRate = GetEffectiveFrameRate();
 			CameraCut.CachedTickResolution = TargetSequence->GetMovieScene()->GetTickResolution();
@@ -945,10 +949,10 @@ MoviePipeline::FFrameConstantMetrics UMoviePipeline::CalculateShotFrameMetrics(c
 	Output.FrameRate = GetEffectiveFrameRate();
 	Output.TicksPerOutputFrame = FFrameRate::TransformTime(FFrameTime(FFrameNumber(1)), Output.FrameRate, Output.TickResolution);
 
-	UMoviePipelineAccumulationSetting* AccumulationSettings = InShot.ShotConfig->FindOrAddSetting<UMoviePipelineAccumulationSetting>();
+	UMoviePipelineCameraSetting* CameraSettings = InShot.ShotConfig->FindOrAddSetting<UMoviePipelineCameraSetting>();
 
 	// (CameraShutterAngle/360) gives us the fraction-of-the-output-frame the accumulation frames should cover.
-	Output.ShutterAnglePercentage = AccumulationSettings->CameraShutterAngle / 360.0;
+	Output.ShutterAnglePercentage = CameraSettings->CameraShutterAngle / 360.0;
 
 	{
 		/*
@@ -967,14 +971,14 @@ MoviePipeline::FFrameConstantMetrics UMoviePipeline::CalculateShotFrameMetrics(c
 		Output.TicksWhileShutterOpen = Output.TicksPerOutputFrame * Output.ShutterAnglePercentage;
 
 		// Divide that amongst all of our accumulation sample frames.
-		Output.TicksPerSample = Output.TicksWhileShutterOpen / AccumulationSettings->TemporalSampleCount;
+		Output.TicksPerSample = Output.TicksWhileShutterOpen / CameraSettings->TemporalSampleCount;
 	}
 
-	Output.ShutterClosedFraction = (360 - AccumulationSettings->CameraShutterAngle) / 360.0;
+	Output.ShutterClosedFraction = (360 - CameraSettings->CameraShutterAngle) / 360.0;
 	Output.TicksWhileShutterClosed = Output.TicksPerOutputFrame * Output.ShutterClosedFraction;
 
 	// Shutter Offset
-	switch (AccumulationSettings->ShutterTiming)
+	switch (CameraSettings->ShutterTiming)
 	{
 		// Subtract the entire time the shutter is open.
 	case EMoviePipelineShutterTiming::FrameClose:
diff --git a/Engine/Plugins/MovieScene/MovieRenderPipeline/Source/MovieRenderPipelineCore/Private/MoviePipelineRendering.cpp b/Engine/Plugins/MovieScene/MovieRenderPipeline/Source/MovieRenderPipelineCore/Private/MoviePipelineRendering.cpp
--- a/Engine/Plugins/MovieScene/MovieRenderPipeline/Source/MovieRenderPipelineCore/Private/MoviePipelineRendering.cpp
+++ b/Engine/Plugins/MovieScene/MovieRenderPipeline/Source/MovieRenderPipelineCore/Private/MoviePipelineRendering.cpp
@@ -19,6 +19,7 @@
 #include "GameFramework/PlayerController.h"
 #include "MoviePipelineHighResSetting.h"
 #include "Modules/ModuleManager.h"
+#include "MoviePipelineCameraSetting.h"
 
 #define LOCTEXT_NAMESPACE "MoviePipeline"
 
@@ -160,6 +161,7 @@ void UMoviePipeline::RenderFrame()
 	// In short, for each output frame, for each accumulation frame, for each tile X/Y, for each jitter, we render a pass. This setup is
 	// designed to maximize the likely hood of deterministic rendering and that different passes line up with each other.
 	UMoviePipelineAccumulationSetting* AccumulationSettings = CurrentShot.ShotConfig->FindOrAddSetting<UMoviePipelineAccumulationSetting>();
+	UMoviePipelineCameraSetting* CameraSettings = GetPipelineMasterConfig()->FindOrAddSetting<UMoviePipelineCameraSetting>();
 	UMoviePipelineOutputSetting* OutputSettings = GetPipelineMasterConfig()->FindSetting<UMoviePipelineOutputSetting>();
 	check(OutputSettings);
 
@@ -299,7 +301,7 @@ void UMoviePipeline::RenderFrame()
 				// Only jitter when required, as jitter requires TAA to be off. 
 				// Also, we want jitter if either time samples or spatial samples are > 1, otherwise not.
 				if (!bIsHistoryOnlyFrame &&
-					(AccumulationSettings->SpatialSampleCount > 1 || AccumulationSettings->TemporalSampleCount > 1))
+					(AccumulationSettings->SpatialSampleCount > 1 || CameraSettings->TemporalSampleCount > 1))
 				{
 					OffsetX = Halton((SpatialSample + 1) + (CachedOutputState.TemporalSampleIndex * NumSpatialSamples), 2);
 					OffsetY = Halton((SpatialSample + 1) + (CachedOutputState.TemporalSampleIndex * NumSpatialSamples), 3);
@@ -334,7 +336,7 @@ void UMoviePipeline::RenderFrame()
 				SampleState.SpatialSampleIndex = SpatialSample;
 				SampleState.SpatialSampleCount = NumSpatialSamples;
 				SampleState.TemporalSampleIndex = CachedOutputState.TemporalSampleIndex;
-				SampleState.TemporalSampleCount = AccumulationSettings->TemporalSampleCount;
+				SampleState.TemporalSampleCount = CameraSettings->TemporalSampleCount;
 				SampleState.AccumulationGamma = AccumulationSettings->AccumulationGamma;
 				SampleState.bIsUsingOverlappedTiles = bIsUsingOverlappedTiles;
 				SampleState.BackbufferSize = BackbufferResolution;
diff --git a/Engine/Plugins/MovieScene/MovieRenderPipeline/Source/MovieRenderPipelineCore/Private/MoviePipelineTiming.cpp b/Engine/Plugins/MovieScene/MovieRenderPipeline/Source/MovieRenderPipelineCore/Private/MoviePipelineTiming.cpp
--- a/Engine/Plugins/MovieScene/MovieRenderPipeline/Source/MovieRenderPipelineCore/Private/MoviePipelineTiming.cpp
+++ b/Engine/Plugins/MovieScene/MovieRenderPipeline/Source/MovieRenderPipelineCore/Private/MoviePipelineTiming.cpp
@@ -15,6 +15,7 @@
 #include "GameFramework/PlayerController.h"
 #include "GameFramework/WorldSettings.h"
 #include "MoviePipelineMasterConfig.h"
+#include "MoviePipelineCameraSetting.h"
 
 void UMoviePipeline::TickProducingFrames()
 {
@@ -126,8 +127,8 @@ void UMoviePipeline::TickProducingFrames()
 		FFrameTime TicksToEndOfPreviousFrame;
 		float WorldTimeDilation = GetWorld()->GetWorldSettings()->GetEffectiveTimeDilation();
 
-		UMoviePipelineAccumulationSetting* AccumulationSettings = CurrentShot.ShotConfig->FindOrAddSetting<UMoviePipelineAccumulationSetting>();
-		if (AccumulationSettings->TemporalSampleCount == 1)
+		UMoviePipelineCameraSetting* CameraSettings = CurrentShot.ShotConfig->FindOrAddSetting<UMoviePipelineCameraSetting>();
+		if (CameraSettings->TemporalSampleCount == 1)
 		{
 			TicksToEndOfPreviousFrame = FrameMetrics.TicksPerOutputFrame;
 		}
@@ -208,18 +209,20 @@ void UMoviePipeline::TickProducingFrames()
 		UMoviePipelineAccumulationSetting* AccumulationSettings = CurrentShot.ShotConfig->FindOrAddSetting<UMoviePipelineAccumulationSetting>();
 		check(AccumulationSettings);
 
+		UMoviePipelineCameraSetting* CameraSettings = CurrentShot.ShotConfig->FindOrAddSetting<UMoviePipelineCameraSetting>();
+
 		// If we've rendered the last sample and wrapped around, then we're going to be
 		// working on the next output frame for this frame. Since OutputFrameNumber gets
 		// initialized to -1, this is ok on the very first frame of rendering as it pushes
 		// us to be working on outputting frame zero.
-		if (CachedOutputState.TemporalSampleIndex >= AccumulationSettings->TemporalSampleCount - 1)
+		if (CachedOutputState.TemporalSampleIndex >= CameraSettings->TemporalSampleCount - 1)
 		{
 			CachedOutputState.TemporalSampleIndex = -1;
 		}
 
 		CachedOutputState.TemporalSampleIndex++;
 
-		if (AccumulationSettings->TemporalSampleCount == 1)
+		if (CameraSettings->TemporalSampleCount == 1)
 		{
 			// It should always be zero when using no sub-sampling.
 			check(CachedOutputState.TemporalSampleIndex == 0);
@@ -336,8 +339,8 @@ void UMoviePipeline::TickProducingFrames()
 
 			
 			// Shouldn't be in this else branch if they don't have an Accumulation Setting for this shot.
-			check(AccumulationSettings);
-			check(AccumulationSettings->TemporalSampleCount > 1);
+			check(CameraSettings);
+			check(CameraSettings->TemporalSampleCount > 1);
 				
 			const bool bFirstFrame = CachedOutputState.TemporalSampleIndex == 0;
 			if (bFirstFrame)
@@ -457,7 +460,7 @@ void UMoviePipeline::TickProducingFrames()
 
 		// The combination of shutter angle percentage, non-uniform render frame delta times and dividing by sample
 		// count produce the correct length for motion blur in all cases.
-		CachedOutputState.MotionBlurFraction = FrameMetrics.ShutterAnglePercentage / AccumulationSettings->TemporalSampleCount;
+		CachedOutputState.MotionBlurFraction = FrameMetrics.ShutterAnglePercentage / CameraSettings->TemporalSampleCount;
 
 		// Update our Output State with any data we need to derive based on our current time.
 		CalculateFrameNumbersForOutputState(FrameMetrics, CurrentCameraCut, CachedOutputState);
diff --git a/Engine/Plugins/MovieScene/MovieRenderPipeline/Source/MovieRenderPipelineCore/Public/MoviePipelineAccumulationSetting.h b/Engine/Plugins/MovieScene/MovieRenderPipeline/Source/MovieRenderPipelineCore/Public/MoviePipelineAccumulationSetting.h
--- a/Engine/Plugins/MovieScene/MovieRenderPipeline/Source/MovieRenderPipelineCore/Public/MoviePipelineAccumulationSetting.h
+++ b/Engine/Plugins/MovieScene/MovieRenderPipeline/Source/MovieRenderPipelineCore/Public/MoviePipelineAccumulationSetting.h
@@ -12,35 +12,13 @@ class MOVIERENDERPIPELINECORE_API UMoviePipelineAccumulationSetting : public UMo
 	GENERATED_BODY()
 public:
 	UMoviePipelineAccumulationSetting()
-		: TemporalSampleCount(1)
-		, CameraShutterAngle(180)
-		, SpatialSampleCount(1)
+		: SpatialSampleCount(1)
+		, HandleFrameCount(0)
+		, WarmUpFrameCount(0)
 		, AccumulationGamma(1.f)
 	{
 	}
-	
-	/** 
-	* The number of frames we should combine together to produce each output frame. This blends the
-	* results of this many sub-steps together to produce one output frame. See CameraShutterAngle to
-	* control how much time passes between each sub-frame. See SpatialSampleCount to see how many
-	* samples we average together to produce a sub-step. (This means rendering complexity is
-	* SampleCount * TileCount^2 * SpatialSampleCount * NumPasses).
-	*/
-	UPROPERTY(EditAnywhere, BlueprintReadWrite, meta=(UIMin = 1, ClampMin = 1), Category = "Movie Pipeline")
-	int32 TemporalSampleCount;
-	
-	/** 
-	* The camera shutter angle determines how much of a given frame the accumulation frames span.
-	* For example, a 24fps film with a shutter angle of 180 will spread the accumulation frames
-	* out over (180/360) percent of the frame (1/24th of a second). With 2 accumulation frames,
-	* the first one will be at the very start of the frame and the second accumulation frame will
-	* be land on the half way point between the start of this frame and the start of the next.
-	*
-	* A shutter angle of 360 means continuous movement, while a shutter angle of zero means no
-	* motion blur.
-	*/
-	UPROPERTY(EditAnywhere, BlueprintReadWrite, meta = (UIMin = 0, UIMax = 360, ClampMin = 0, ClampMax = 360), Category = "Movie Pipeline")
-	int32 CameraShutterAngle;
+
 
 	/**
 	* How many frames should we accumulate together before contributing to one overall sample. This lets you
@@ -61,13 +39,4 @@ public:
 	* is applied and pow(x,1/AccumulationGamma) is applied after accumulation is finished
 	*/
 	float AccumulationGamma;
-
-	/**
-	* Shutter Timing allows you to bias the timing of your shutter angle to either be before, during, or after
-	* a frame. When set to FrameClose, it means that the motion gathered up to produce frame N comes from 
-	* before and right up to frame N. When set to FrameCenter, the motion represents half the time before the
-	* frame and half the time after the frame. When set to FrameOpen, the motion represents the time from 
-	* Frame N onwards.
-	*/
-	EMoviePipelineShutterTiming ShutterTiming;
 };
\ No newline at end of file
diff --git a/Engine/Plugins/MovieScene/MovieRenderPipeline/Source/MovieRenderPipelineCore/Public/MoviePipelineCameraSetting.h b/Engine/Plugins/MovieScene/MovieRenderPipeline/Source/MovieRenderPipelineCore/Public/MoviePipelineCameraSetting.h
new file mode 100644
--- /dev/null
+++ b/Engine/Plugins/MovieScene/MovieRenderPipeline/Source/MovieRenderPipelineCore/Public/MoviePipelineCameraSetting.h
@@ -0,0 +1,70 @@
+// Copyright 1998-2019 Epic Games, Inc. All Rights Reserved.
+#pragma once
+
+#include "MoviePipelineSetting.h"
+#include "MovieRenderPipelineDataTypes.h"
+#include "MoviePipelineCameraSetting.generated.h"
+
+
+UCLASS(Blueprintable)
+class MOVIERENDERPIPELINECORE_API UMoviePipelineCameraSetting : public UMoviePipelineSetting
+{
+	GENERATED_BODY()
+public:
+	UMoviePipelineCameraSetting()
+		: TemporalSampleCount(1)
+		, CameraShutterAngle(180)
+		, ShutterTiming(EMoviePipelineShutterTiming::FrameCenter)
+		, bManualExposure(true)
+		, ExposureCompensation(8.0)
+	{}
+	
+public:
+
+	/** 
+	* The number of frames we should combine together to produce each output frame. This blends the
+	* results of this many sub-steps together to produce one output frame. See CameraShutterAngle to
+	* control how much time passes between each sub-frame. See SpatialSampleCount to see how many
+	* samples we average together to produce a sub-step. (This means rendering complexity is
+	* SampleCount * TileCount^2 * SpatialSampleCount * NumPasses).
+	*/
+	UPROPERTY(EditAnywhere, BlueprintReadWrite, meta=(UIMin = 1, ClampMin = 1), Category = "Movie Pipeline")
+	int32 TemporalSampleCount;
+	
+	/** 
+	* The camera shutter angle determines how much of a given frame the accumulation frames span.
+	* For example, a 24fps film with a shutter angle of 180 will spread the accumulation frames
+	* out over (180/360) percent of the frame (1/24th of a second). With 2 accumulation frames,
+	* the first one will be at the very start of the frame and the second accumulation frame will
+	* be land on the half way point between the start of this frame and the start of the next.
+	*
+	* A shutter angle of 360 means continuous movement, while a shutter angle of zero means no
+	* motion blur.
+	*/
+	UPROPERTY(EditAnywhere, BlueprintReadWrite, meta = (UIMin = 0, UIMax = 360, ClampMin = 0, ClampMax = 360), Category = "Movie Pipeline")
+	int32 CameraShutterAngle;
+	
+	/**
+	* Shutter Timing allows you to bias the timing of your shutter angle to either be before, during, or after
+	* a frame. When set to FrameClose, it means that the motion gathered up to produce frame N comes from 
+	* before and right up to frame N. When set to FrameCenter, the motion represents half the time before the
+	* frame and half the time after the frame. When set to FrameOpen, the motion represents the time from 
+	* Frame N onwards.
+	*/
+	UPROPERTY(EditAnywhere, BlueprintReadWrite, meta = (UIMin = 0, UIMax = 360, ClampMin = 0, ClampMax = 360), Category = "Movie Pipeline")
+	EMoviePipelineShutterTiming ShutterTiming;
+	
+	/**
+	* Should we override the exposure on the camera to use a fixed value? This is required for high res screenshots to use consistent
+	* exposure between the different tiles of the image. Leaving this off lets the camera determine based on the previous frame rendered,
+	* which can require long warm up times between shots to allow the exposure to settle.
+	*/
+	UPROPERTY(EditAnywhere, BlueprintReadWrite, meta = (UIMin = 0, UIMax = 360, ClampMin = 0, ClampMax = 360), Category = "Movie Pipeline")
+	bool bManualExposure;
+	
+	/**
+	* What exposure should we use when using Manual Exposure? Same behavior as the Post Processing volume.
+	*/
+	UPROPERTY(EditAnywhere, BlueprintReadWrite, meta = (UIMin = -10, UIMax = 10, EditCondition="bManualExposure"), Category = "Movie Pipeline")
+	float ExposureCompensation;
+};
\ No newline at end of file
diff --git a/Engine/Plugins/MovieScene/MovieRenderPipeline/Source/MovieRenderPipelineEditor/Private/MovieRenderPipelineEditorTest.cpp b/Engine/Plugins/MovieScene/MovieRenderPipeline/Source/MovieRenderPipelineEditor/Private/MovieRenderPipelineEditorTest.cpp
--- a/Engine/Plugins/MovieScene/MovieRenderPipeline/Source/MovieRenderPipelineEditor/Private/MovieRenderPipelineEditorTest.cpp
+++ b/Engine/Plugins/MovieScene/MovieRenderPipeline/Source/MovieRenderPipelineEditor/Private/MovieRenderPipelineEditorTest.cpp
@@ -12,6 +12,7 @@
 #include "MoviePipelineOutputSetting.h"
 #include "MoviePipelineHighResSetting.h"
 #include "MoviePipelineDeferredPasses.h"
+#include "MoviePipelineCameraSetting.h"
 
 UMoviePipelineShotConfig* GenerateTestShotConfig(UObject* InOwner, int32 InSampleCount,
 	int32 InShutterAngle, EMoviePipelineShutterTiming InFrameTiming, int32 InTileCount,
@@ -24,12 +25,18 @@ UMoviePipelineShotConfig* GenerateTestShotConfig(UObject* InOwner, int32 InSampl
 	GamemodeOverride->GameModeOverride = AGameMode::StaticClass();
 
 	UMoviePipelineAccumulationSetting* Accumulation = OutConfig->FindOrAddSetting< UMoviePipelineAccumulationSetting>();
-	Accumulation->TemporalSampleCount = InSampleCount;
-	Accumulation->CameraShutterAngle = InShutterAngle;
-	Accumulation->ShutterTiming = InFrameTiming;
 	Accumulation->SpatialSampleCount = InSpatialSampleCount;
 	Accumulation->AccumulationGamma = AccumulationGamma;
 
+	{
+		UMoviePipelineCameraSetting* CameraSetting = OutConfig->FindOrAddSetting<UMoviePipelineCameraSetting>();
+		CameraSetting->TemporalSampleCount = InSampleCount;
+		CameraSetting->CameraShutterAngle = InShutterAngle;
+		CameraSetting->ShutterTiming = InFrameTiming;
+		CameraSetting->bManualExposure = true;
+		CameraSetting->ExposureCompensation = 8;
+	}
+
 	if (InTileCount > 1)
 	{
 		UMoviePipelineHighResSetting* HighRes = OutConfig->FindOrAddSetting<UMoviePipelineHighResSetting>();
diff --git a/Engine/Plugins/MovieScene/MovieRenderPipeline/Source/MovieRenderPipelineRenderPasses/Private/MoviePipelineDeferredPasses.cpp b/Engine/Plugins/MovieScene/MovieRenderPipeline/Source/MovieRenderPipelineRenderPasses/Private/MoviePipelineDeferredPasses.cpp
--- a/Engine/Plugins/MovieScene/MovieRenderPipeline/Source/MovieRenderPipelineRenderPasses/Private/MoviePipelineDeferredPasses.cpp
+++ b/Engine/Plugins/MovieScene/MovieRenderPipeline/Source/MovieRenderPipelineRenderPasses/Private/MoviePipelineDeferredPasses.cpp
@@ -23,6 +23,8 @@
 #include "BufferVisualizationData.h"
 #include "Containers/Array.h"
 #include "FinalPostProcessSettings.h"
+#include "Materials/Material.h"
+#include "MoviePipelineCameraSetting.h"
 
 namespace MoviePipeline
 {
@@ -100,12 +102,22 @@ namespace MoviePipeline
 			View->FinalPostProcessSettings.bOverride_MotionBlurMax = true;
 		}
 
-		// Autoexposure is not allowed
-		if (View->FinalPostProcessSettings.AutoExposureMethod != EAutoExposureMethod::AEM_Manual)
-		{
-			UE_LOG(LogMovieRenderPipeline, Log, TEXT("AutoExposure Method should always be Manual."));
+		UMoviePipelineCameraSetting* CameraSettings = GetPipeline()->GetPipelineCurrentShotConfig()->FindSetting<UMoviePipelineCameraSetting>();
+		UMoviePipelineHighResSetting* HighResSettings = GetPipeline()->GetPipelineCurrentShotConfig()->FindSetting<UMoviePipelineHighResSetting>();
 
-			View->FinalPostProcessSettings.AutoExposureMethod = EAutoExposureMethod::AEM_Manual;
+		if(CameraSettings)
+		{
+			if (CameraSettings->bManualExposure)
+			{
+				View->FinalPostProcessSettings.AutoExposureMethod = EAutoExposureMethod::AEM_Manual;
+				View->FinalPostProcessSettings.AutoExposureBias = CameraSettings->ExposureCompensation;
+			}
+			else if (HighResSettings && (HighResSettings->TileCount > 1) && (View->FinalPostProcessSettings.AutoExposureMethod != EAutoExposureMethod::AEM_Manual))
+			{
+				// Autoexposure is not allowed
+				UE_LOG(LogMovieRenderPipeline, Warning, TEXT("AutoExposure Method should always be Manual."));
+				View->FinalPostProcessSettings.AutoExposureMethod = EAutoExposureMethod::AEM_Manual;
+			}
 		}
 
 		const bool bCanUseAA = InSampleState.SpatialSampleCount == 1 && 
@@ -130,7 +142,6 @@ namespace MoviePipeline
 
 			// Add an additional bias per user settings. This allows them to choose to make the textures sharper if it
 			// looks better with their particular settings.
-			UMoviePipelineHighResSetting* HighResSettings = GetPipeline()->GetPipelineCurrentShotConfig()->FindSetting<UMoviePipelineHighResSetting>();
 			if (HighResSettings)
 			{
 				View->MaterialTextureMipBias += HighResSettings->TextureSharpnessBias;
