MoviePipeline: Fixed using TileCounts greater than one causing a crash due to uninitialzied variables.

#rnx

diff --git a/Engine/Plugins/MovieScene/MovieRenderPipeline/Source/MovieRenderPipelineCore/Private/MoviePipelineRendering.cpp b/Engine/Plugins/MovieScene/MovieRenderPipeline/Source/MovieRenderPipelineCore/Private/MoviePipelineRendering.cpp
--- a/Engine/Plugins/MovieScene/MovieRenderPipeline/Source/MovieRenderPipelineCore/Private/MoviePipelineRendering.cpp
+++ b/Engine/Plugins/MovieScene/MovieRenderPipeline/Source/MovieRenderPipelineCore/Private/MoviePipelineRendering.cpp
@@ -314,9 +314,19 @@ void UMoviePipeline::RenderFrame()
 				double FinalSubPixelShiftX = - (TileShiftX + SpatialShiftX - 0.5);
 				double FinalSubPixelShiftY = (TileShiftY + SpatialShiftY - 0.5);
 
-				FIntPoint BackbufferResolution;
+				FIntPoint BackbufferResolution = OutputSettings->OutputResolution;
+				FIntPoint TileResolution = BackbufferResolution;
+
 				if (bIsUsingOverlappedTiles)
 				{
+					BackbufferResolution = FIntPoint(
+						FMath::CeilToInt(BackbufferResolution.X / HighResSettings->TileCount),
+						FMath::CeilToInt(BackbufferResolution.Y / HighResSettings->TileCount));
+
+					// Cache tile resolution before we add padding.
+					TileResolution = BackbufferResolution;
+
+					// Apply size padding.
 					BackbufferResolution = HighResSettings->CalculatePaddedBackbufferSize(BackbufferResolution);
 				}
 				else
@@ -324,6 +334,9 @@ void UMoviePipeline::RenderFrame()
 					// Legacy interleaving
 					BackbufferResolution.X = FMath::CeilToInt(OutputResolution.X / TileCount.X);
 					BackbufferResolution.Y = FMath::CeilToInt(OutputResolution.Y / TileCount.Y);
+
+					// No padding is used when interleaving
+					TileResolution = BackbufferResolution;
 				}
 
 				// We take all of the information needed to render a single sample and package it into a struct.
@@ -340,14 +353,15 @@ void UMoviePipeline::RenderFrame()
 				SampleState.AccumulationGamma = AccumulationSettings->AccumulationGamma;
 				SampleState.bIsUsingOverlappedTiles = bIsUsingOverlappedTiles;
 				SampleState.BackbufferSize = BackbufferResolution;
+				SampleState.TileSize = TileResolution;
 				SampleState.FrameInfo = FrameInfo;
 
 				if (bIsUsingOverlappedTiles)
 				{
-					SampleState.OverlappedPad = FIntPoint(FMath::CeilToInt(BackbufferResolution.X * HighResSettings->OverlapPercentage), 
-														   FMath::CeilToInt(BackbufferResolution.Y * HighResSettings->OverlapPercentage));
-					SampleState.OverlappedOffset = FIntPoint(TileX * BackbufferResolution.X - SampleState.OverlappedPad.X, 
-															  TileY * BackbufferResolution.Y - SampleState.OverlappedPad.Y);
+					SampleState.OverlappedPad = FIntPoint(FMath::CeilToInt(TileResolution.X * HighResSettings->OverlapPercentage), 
+														   FMath::CeilToInt(TileResolution.Y * HighResSettings->OverlapPercentage));
+					SampleState.OverlappedOffset = FIntPoint(TileX * TileResolution.X - SampleState.OverlappedPad.X,
+															  TileY * TileResolution.Y - SampleState.OverlappedPad.Y);
 					SampleState.OverlappedSubpixelShift = FVector2D(1.0f - OffsetX,OffsetY);
 				}
 				else
diff --git a/Engine/Plugins/MovieScene/MovieRenderPipeline/Source/MovieRenderPipelineCore/Public/MovieRenderPipelineDataTypes.h b/Engine/Plugins/MovieScene/MovieRenderPipeline/Source/MovieRenderPipelineCore/Public/MovieRenderPipelineDataTypes.h
--- a/Engine/Plugins/MovieScene/MovieRenderPipeline/Source/MovieRenderPipelineCore/Public/MovieRenderPipelineDataTypes.h
+++ b/Engine/Plugins/MovieScene/MovieRenderPipeline/Source/MovieRenderPipelineCore/Public/MovieRenderPipelineDataTypes.h
@@ -645,9 +645,12 @@ public:
 		return (TileCounts.X & TileCounts.Y);
 	}
 
-	/** How big is the back buffer for this sample? Pre-set when setup, here for convenience. */
+	/** How big is the back buffer for this sample? This is the final target size divided by the number of tiles with padding added. */
 	FIntPoint BackbufferSize;
 
+	/** How big is an individual tile for this sample? This is the final target size divided by the number of tiles (without padding added) */
+	FIntPoint TileSize;
+
 	/** If true, we will discard this sample after rendering. Used to get history set up correctly. */
 	bool bIsHistoryOnlyFrame;
 
diff --git a/Engine/Plugins/MovieScene/MovieRenderPipeline/Source/MovieRenderPipelineEditor/Private/MovieRenderPipelineEditorTest.cpp b/Engine/Plugins/MovieScene/MovieRenderPipeline/Source/MovieRenderPipelineEditor/Private/MovieRenderPipelineEditorTest.cpp
--- a/Engine/Plugins/MovieScene/MovieRenderPipeline/Source/MovieRenderPipelineEditor/Private/MovieRenderPipelineEditorTest.cpp
+++ b/Engine/Plugins/MovieScene/MovieRenderPipeline/Source/MovieRenderPipelineEditor/Private/MovieRenderPipelineEditorTest.cpp
@@ -74,10 +74,10 @@ TArray<UMoviePipelineMasterConfig*> FMovieRenderPipelineEditorModule::GenerateTe
 				UMoviePipelineMasterConfig* OutPipeline = NewObject<UMoviePipelineMasterConfig>(GetTransientPackage());
 				UMoviePipelineOutputSetting* OutputSetting = OutPipeline->FindOrAddSetting<UMoviePipelineOutputSetting>();
 
-				int32 SizeX = 1920;
-				int32 SizeY = 1080;
-				int32 TileX = 1;
-				int32 TileY = 1;
+				int32 SizeX = 1920/4;
+				int32 SizeY = 1080/4;
+				int32 TileX = 4;
+				int32 TileY = 4;
 				int32 TestNumSamples = 2;
 				float PadRatioX = 0.5f;
 				float PadRatioY = 0.5f;
diff --git a/Engine/Plugins/MovieScene/MovieRenderPipeline/Source/MovieRenderPipelineRenderPasses/Private/MoviePipelineDeferredPasses.cpp b/Engine/Plugins/MovieScene/MovieRenderPipeline/Source/MovieRenderPipelineRenderPasses/Private/MoviePipelineDeferredPasses.cpp
--- a/Engine/Plugins/MovieScene/MovieRenderPipeline/Source/MovieRenderPipelineRenderPasses/Private/MoviePipelineDeferredPasses.cpp
+++ b/Engine/Plugins/MovieScene/MovieRenderPipeline/Source/MovieRenderPipelineRenderPasses/Private/MoviePipelineDeferredPasses.cpp
@@ -505,7 +505,7 @@ namespace MoviePipeline
 		// For the first sample in a new output, we allocate memory
 		if (InFrameData->IsFirstTile() && InFrameData->IsFirstTemporalSample())
 		{
-			InParams.ImageAccumulator->InitMemory(FIntPoint(InFrameData->SampleState.BackbufferSize.X * InFrameData->SampleState.TileCounts.X, InFrameData->SampleState.BackbufferSize.Y * InFrameData->SampleState.TileCounts.Y), 3);
+			InParams.ImageAccumulator->InitMemory(FIntPoint(InFrameData->SampleState.TileSize.X * InFrameData->SampleState.TileCounts.X, InFrameData->SampleState.TileSize.Y * InFrameData->SampleState.TileCounts.Y), 3);
 			InParams.ImageAccumulator->ZeroPlanes();
 			InParams.ImageAccumulator->AccumulationGamma = InFrameData->SampleState.AccumulationGamma;
 		}
@@ -516,8 +516,8 @@ namespace MoviePipeline
 
 			FIntPoint RawSize = InPixelData->GetSize();
 
-			check(InFrameData->SampleState.BackbufferSize.X + 2 * InFrameData->SampleState.OverlappedPad.X == RawSize.X);
-			check(InFrameData->SampleState.BackbufferSize.Y + 2 * InFrameData->SampleState.OverlappedPad.Y == RawSize.Y);
+			check(InFrameData->SampleState.TileSize.X + 2 * InFrameData->SampleState.OverlappedPad.X == RawSize.X);
+			check(InFrameData->SampleState.TileSize.Y + 2 * InFrameData->SampleState.OverlappedPad.Y == RawSize.Y);
 
 			InParams.ImageAccumulator->AccumulatePixelData(*InPixelData.Get(), InFrameData->SampleState.OverlappedOffset, InFrameData->SampleState.OverlappedSubpixelShift);
 
