MoviePipeline: Tiled Accumulation no longer overlaps samples incorrectly.

#rnx

diff --git a/Engine/Plugins/MovieScene/MovieRenderPipeline/Source/MovieRenderPipelineRenderPasses/Private/MoviePipelineDeferredPasses.cpp b/Engine/Plugins/MovieScene/MovieRenderPipeline/Source/MovieRenderPipelineRenderPasses/Private/MoviePipelineDeferredPasses.cpp
--- a/Engine/Plugins/MovieScene/MovieRenderPipeline/Source/MovieRenderPipelineRenderPasses/Private/MoviePipelineDeferredPasses.cpp
+++ b/Engine/Plugins/MovieScene/MovieRenderPipeline/Source/MovieRenderPipelineRenderPasses/Private/MoviePipelineDeferredPasses.cpp
@@ -289,8 +289,8 @@ namespace MoviePipeline
 				
 				if (InSampleState.OverlappedPad.X > 0 && InSampleState.OverlappedPad.Y > 0)
 				{
-					PadRatioX = float(InSampleState.OverlappedPad.X * 2 + InSampleState.OverlappedPad.X) / float(InSampleState.OverlappedPad.X);
-					PadRatioY = float(InSampleState.OverlappedPad.Y * 2 + InSampleState.OverlappedPad.Y) / float(InSampleState.OverlappedPad.Y);
+					PadRatioX = float(InSampleState.OverlappedPad.X * 2 + InSampleState.TileSize.X) / float(InSampleState.TileSize.X);
+					PadRatioY = float(InSampleState.OverlappedPad.Y * 2 + InSampleState.TileSize.Y) / float(InSampleState.TileSize.Y);
 				}
 				
 				float ScaleX = PadRatioX / float(InitSettings.TileCount.X);
@@ -519,7 +519,11 @@ namespace MoviePipeline
 			check(InFrameData->SampleState.TileSize.X + 2 * InFrameData->SampleState.OverlappedPad.X == RawSize.X);
 			check(InFrameData->SampleState.TileSize.Y + 2 * InFrameData->SampleState.OverlappedPad.Y == RawSize.Y);
 
-			InParams.ImageAccumulator->AccumulatePixelData(*InPixelData.Get(), InFrameData->SampleState.OverlappedOffset, InFrameData->SampleState.OverlappedSubpixelShift);
+			// bool bSkip = InFrameData->SampleState.TileIndexes.X != 0 || InFrameData->SampleState.TileIndexes.Y != 1;
+			// if (!bSkip)
+			{
+				InParams.ImageAccumulator->AccumulatePixelData(*InPixelData.Get(), InFrameData->SampleState.OverlappedOffset, InFrameData->SampleState.OverlappedSubpixelShift);
+			}
 
 			const double AccumulateEndTime = FPlatformTime::Seconds();
 			const float ElapsedMs = float((AccumulateEndTime - AccumulateBeginTime)*1000.0f);
