Changing weighting function for overlapped tile rendering.

diff --git a/Engine/Plugins/MovieScene/MovieRenderPipeline/Source/MovieRenderPipelineCore/Private/MoviePipelineRendering.cpp b/Engine/Plugins/MovieScene/MovieRenderPipeline/Source/MovieRenderPipelineCore/Private/MoviePipelineRendering.cpp
--- a/Engine/Plugins/MovieScene/MovieRenderPipeline/Source/MovieRenderPipelineCore/Private/MoviePipelineRendering.cpp
+++ b/Engine/Plugins/MovieScene/MovieRenderPipeline/Source/MovieRenderPipelineCore/Private/MoviePipelineRendering.cpp
@@ -371,6 +371,9 @@ void UMoviePipeline::RenderFrame()
 					SampleState.OverlappedSubpixelShift = FVector2D(0.0f,0.0f);
 				}
 
+				SampleState.WeightFunctionX.InitHelper(SampleState.OverlappedPad.X, SampleState.TileSize.X, SampleState.OverlappedPad.X);
+				SampleState.WeightFunctionY.InitHelper(SampleState.OverlappedPad.Y, SampleState.TileSize.Y, SampleState.OverlappedPad.Y);
+
 				// Now we can request that all of the engine passes render. The individual render passes should have already registered delegates
 				// to receive data when the engine render pass is run, so no need to run them.
 				for (TSharedPtr<MoviePipeline::FMoviePipelineEnginePass> EnginePass : ActiveRenderPasses)
diff --git a/Engine/Plugins/MovieScene/MovieRenderPipeline/Source/MovieRenderPipelineCore/Private/MovieRenderOverlappedImage.cpp b/Engine/Plugins/MovieScene/MovieRenderPipeline/Source/MovieRenderPipelineCore/Private/MovieRenderOverlappedImage.cpp
--- a/Engine/Plugins/MovieScene/MovieRenderPipeline/Source/MovieRenderPipelineCore/Private/MovieRenderOverlappedImage.cpp
+++ b/Engine/Plugins/MovieScene/MovieRenderPipeline/Source/MovieRenderPipelineCore/Private/MovieRenderOverlappedImage.cpp
@@ -34,13 +34,16 @@ void FImageOverlappedPlane::Reset()
 }
 
 // a subpixel offset of (0.5,0.5) means that the raw data is exactly centered on destination pixels.
-void FImageOverlappedPlane::AccumulateSinglePlane(const TArray64<float>& InRawData, const TArray64<float>& InWeightData, FIntPoint InSize, FIntPoint InOffset,
+void FImageOverlappedPlane::AccumulateSinglePlane(const TArray64<float>& InRawData, FIntPoint InSize, FIntPoint InOffset,
 								float SubpixelOffsetX, float SubpixelOffsetY,
 								FIntPoint SubRectOffset,
-								FIntPoint SubRectSize)
+								FIntPoint SubRectSize,
+								const TArray<float> & WeightDataX,
+								const TArray<float> & WeightDataY)
 {
 	check(InRawData.Num() == InSize.X * InSize.Y);
-	check(InWeightData.Num() == InSize.X * InSize.Y);
+	check(WeightDataX.Num() == InSize.X);
+	check(WeightDataY.Num() == InSize.Y);
 
 	check(SubpixelOffsetX >= 0.0f);
 	check(SubpixelOffsetX <= 1.0f);
@@ -92,7 +95,10 @@ void FImageOverlappedPlane::AccumulateSinglePlane(const TArray64<float>& InRawDa
 							for (int32 OffsetX = 0; OffsetX < 2; OffsetX++)
 							{
 								float Val = InRawData[(CurrY - 1 + OffsetY) * InSize.X + (CurrX - 1 + OffsetX)];
-								float BaseWeight = InWeightData[(CurrY - 1 + OffsetY) * InSize.X + (CurrX - 1 + OffsetX)];
+								
+								float WX = WeightDataX[(CurrX - 1 + OffsetX)];
+								float WY = WeightDataY[(CurrY - 1 + OffsetY)];
+								float BaseWeight = WX * WY;
 
 								float Weight = BaseWeight * PixelWeight[1-OffsetY][1-OffsetX];
 
@@ -111,7 +117,10 @@ void FImageOverlappedPlane::AccumulateSinglePlane(const TArray64<float>& InRawDa
 				for (int32 CurrX = 0; CurrX < InSize.X; CurrX++)
 				{
 					float Val = InRawData[CurrY * InSize.X + CurrX];
-					float BaseWeight = InWeightData[CurrY * InSize.X + CurrX];
+					float WX = WeightDataX[CurrX];
+					float WY = WeightDataX[CurrY];
+
+					float BaseWeight = WX * WY;
 
 					for (int32 OffsetY = 0; OffsetY < 2; OffsetY++)
 					{
@@ -164,8 +173,11 @@ void FImageOverlappedPlane::AccumulateSinglePlane(const TArray64<float>& InRawDa
 				const float * SrcLineRaw0 = &InRawData[SrcY0 * InSize.X];
 				const float * SrcLineRaw1 = &InRawData[SrcY1 * InSize.X];
 
-				const float * SrcLineWeight0 = &InWeightData[SrcY0 * InSize.X];
-				const float * SrcLineWeight1 = &InWeightData[SrcY1 * InSize.X];
+
+				const float * SrcLineWeight = WeightDataX.GetData();
+
+				const float RowWeight0 = WeightDataY[SrcY0];
+				const float RowWeight1 = WeightDataY[SrcY1];
 
 				float * DstLine = &ChannelData[DstY * Size.X];
 				for (int32 DstX = ActualDstX0; DstX < ActualDstX1; DstX++)
@@ -176,10 +188,10 @@ void FImageOverlappedPlane::AccumulateSinglePlane(const TArray64<float>& InRawDa
 
 					float Sum = DstLine[DstX];
 
-					Sum += SrcLineWeight0[X0] * PixelWeight[1][1] * SrcLineRaw0[X0];
-					Sum += SrcLineWeight0[X1] * PixelWeight[1][0] * SrcLineRaw0[X1];
-					Sum += SrcLineWeight1[X0] * PixelWeight[0][1] * SrcLineRaw1[X0];
-					Sum += SrcLineWeight1[X1] * PixelWeight[0][0] * SrcLineRaw1[X1];
+					Sum += RowWeight0 * SrcLineWeight[X0] * PixelWeight[1][1] * SrcLineRaw0[X0];
+					Sum += RowWeight0 * SrcLineWeight[X1] * PixelWeight[1][0] * SrcLineRaw0[X1];
+					Sum += RowWeight1 * SrcLineWeight[X0] * PixelWeight[0][1] * SrcLineRaw1[X0];
+					Sum += RowWeight1 * SrcLineWeight[X1] * PixelWeight[0][0] * SrcLineRaw1[X1];
 
 					DstLine[DstX] = Sum;
 				}
@@ -209,8 +221,17 @@ void FImageOverlappedPlane::AccumulateSinglePlane(const TArray64<float>& InRawDa
 				const float * SrcLineRaw0 = &InRawData[SrcY0 * InSize.X];
 				const float * SrcLineRaw1 = &InRawData[SrcY1 * InSize.X];
 
-				const float * SrcLineWeight0 = &InWeightData[SrcY0 * InSize.X];
-				const float * SrcLineWeight1 = &InWeightData[SrcY1 * InSize.X];
+				//const float * SrcLineWeight0 = &InWeightData[SrcY0 * InSize.X];
+				//const float * SrcLineWeight1 = &InWeightData[SrcY1 * InSize.X];
+
+				const float * SrcLineWeight = WeightDataX.GetData();
+
+				const float RowWeight0 = WeightDataY[SrcY0];
+				const float RowWeight1 = WeightDataY[SrcY1];
+
+
+				VectorRegister VecRowWeight0 = VectorSetFloat1(RowWeight0);
+				VectorRegister VecRowWeight1 = VectorSetFloat1(RowWeight1);
 
 				float * DstLine = &ChannelData[DstY * Size.X];
 
@@ -226,10 +247,10 @@ void FImageOverlappedPlane::AccumulateSinglePlane(const TArray64<float>& InRawDa
 
 						VectorRegister Sum = VectorLoad(&DstLine[BaseDstX]);
 						
-						Sum = VectorMultiplyAdd(VecPixelWeight11, VectorMultiply(VectorLoad(&SrcLineWeight0[BaseX0]), VectorLoad(&SrcLineRaw0[BaseX0])), Sum);
-						Sum = VectorMultiplyAdd(VecPixelWeight10, VectorMultiply(VectorLoad(&SrcLineWeight0[BaseX1]), VectorLoad(&SrcLineRaw0[BaseX1])), Sum);
-						Sum = VectorMultiplyAdd(VecPixelWeight01, VectorMultiply(VectorLoad(&SrcLineWeight1[BaseX0]), VectorLoad(&SrcLineRaw1[BaseX0])), Sum);
-						Sum = VectorMultiplyAdd(VecPixelWeight00, VectorMultiply(VectorLoad(&SrcLineWeight1[BaseX1]), VectorLoad(&SrcLineRaw1[BaseX1])), Sum);
+						Sum = VectorMultiplyAdd(VecPixelWeight11, VectorMultiply(VecRowWeight0,VectorMultiply(VectorLoad(&SrcLineWeight[BaseX0]), VectorLoad(&SrcLineRaw0[BaseX0]))), Sum);
+						Sum = VectorMultiplyAdd(VecPixelWeight10, VectorMultiply(VecRowWeight0,VectorMultiply(VectorLoad(&SrcLineWeight[BaseX1]), VectorLoad(&SrcLineRaw0[BaseX1]))), Sum);
+						Sum = VectorMultiplyAdd(VecPixelWeight01, VectorMultiply(VecRowWeight1,VectorMultiply(VectorLoad(&SrcLineWeight[BaseX0]), VectorLoad(&SrcLineRaw1[BaseX0]))), Sum);
+						Sum = VectorMultiplyAdd(VecPixelWeight00, VectorMultiply(VecRowWeight1,VectorMultiply(VectorLoad(&SrcLineWeight[BaseX1]), VectorLoad(&SrcLineRaw1[BaseX1]))), Sum);
 
 						VectorStore(Sum, &DstLine[BaseDstX]);
 					}
@@ -266,10 +287,10 @@ void FImageOverlappedPlane::AccumulateSinglePlane(const TArray64<float>& InRawDa
 
 					float Sum = DstLine[DstX];
 
-					Sum += PixelWeight11 * SrcLineWeight0[X0] * SrcLineRaw0[X0];
-					Sum += PixelWeight10 * SrcLineWeight0[X1] * SrcLineRaw0[X1];
-					Sum += PixelWeight01 * SrcLineWeight1[X0] * SrcLineRaw1[X0];
-					Sum += PixelWeight00 * SrcLineWeight1[X1] * SrcLineRaw1[X1];
+					Sum += PixelWeight11 * RowWeight0 * SrcLineWeight[X0] * SrcLineRaw0[X0];
+					Sum += PixelWeight10 * RowWeight0 * SrcLineWeight[X1] * SrcLineRaw0[X1];
+					Sum += PixelWeight01 * RowWeight1 * SrcLineWeight[X0] * SrcLineRaw1[X0];
+					Sum += PixelWeight00 * RowWeight1 * SrcLineWeight[X1] * SrcLineRaw1[X1];
 
 					DstLine[DstX] = Sum;
 				}
@@ -421,7 +442,7 @@ void FImageOverlappedAccumulator::CheckTileSubRect(const TArray64<float>& OutWei
 	}
 }
 
-void FImageOverlappedAccumulator::AccumulatePixelData(const FImagePixelData& InPixelData, FIntPoint InTileOffset, FVector2D InSubpixelOffset)
+void FImageOverlappedAccumulator::AccumulatePixelData(const FImagePixelData& InPixelData, FIntPoint InTileOffset, FVector2D InSubpixelOffset, const MoviePipeline::FTileWeight1D & WeightX, const MoviePipeline::FTileWeight1D & WeightY)
 {
 	EImagePixelType Fmt = InPixelData.GetType();
 
@@ -615,7 +636,12 @@ void FImageOverlappedAccumulator::AccumulatePixelData(const FImagePixelData& InP
 
 			GenerateTileWeight(TileMaskData, RawSize);
 		}
-		
+
+		TArray<float> WeightDataX;
+		TArray<float> WeightDataY;
+		WeightX.CalculateArrayWeight(WeightDataX, RawSize.X);
+		WeightY.CalculateArrayWeight(WeightDataY, RawSize.Y);
+
 		FIntPoint SubRectOffset(0, 0);
 		FIntPoint SubRectSize(0, 0);
 		GetTileSubRect(SubRectOffset, SubRectSize, TileMaskData, RawSize);
@@ -636,8 +662,8 @@ void FImageOverlappedAccumulator::AccumulatePixelData(const FImagePixelData& InP
 
 		for (int32 ChanIter = 0; ChanIter < NumChannels; ChanIter++)
 		{
-			ChannelPlanes[ChanIter].AccumulateSinglePlane(RawData[ChanIter], TileMaskData, RawSize, InTileOffset, InSubpixelOffset.X, InSubpixelOffset.Y,
-				SubRectOffset, SubRectSize);
+			ChannelPlanes[ChanIter].AccumulateSinglePlane(RawData[ChanIter], RawSize, InTileOffset, InSubpixelOffset.X, InSubpixelOffset.Y,
+				SubRectOffset, SubRectSize, WeightDataX, WeightDataY);
 		}
 
 		// Accumulate weights as well.
@@ -649,8 +675,8 @@ void FImageOverlappedAccumulator::AccumulatePixelData(const FImagePixelData& InP
 				VecOne[Index] = 1.0f;
 			}
 
-			WeightPlane.AccumulateSinglePlane(VecOne, TileMaskData, RawSize, InTileOffset, InSubpixelOffset.X, InSubpixelOffset.Y,
-				SubRectOffset, SubRectSize);
+			WeightPlane.AccumulateSinglePlane(VecOne, RawSize, InTileOffset, InSubpixelOffset.X, InSubpixelOffset.Y,
+				SubRectOffset, SubRectSize, WeightDataX, WeightDataY);
 		}
 
 		const double PlaneEndTime = FPlatformTime::Seconds();
diff --git a/Engine/Plugins/MovieScene/MovieRenderPipeline/Source/MovieRenderPipelineCore/Public/MovieRenderOverlappedImage.h b/Engine/Plugins/MovieScene/MovieRenderPipeline/Source/MovieRenderPipelineCore/Public/MovieRenderOverlappedImage.h
--- a/Engine/Plugins/MovieScene/MovieRenderPipeline/Source/MovieRenderPipelineCore/Public/MovieRenderOverlappedImage.h
+++ b/Engine/Plugins/MovieScene/MovieRenderPipeline/Source/MovieRenderPipelineCore/Public/MovieRenderOverlappedImage.h
@@ -6,6 +6,8 @@
 
 #include "ImagePixelData.h"
 
+#include "MovieRenderPipelineDataTypes.h"
+
 /* Types
  *****************************************************************************/
 
@@ -54,10 +56,12 @@ public:
 	 * @param SubRectOffset - The offset of the SubRect inside the raw data (InRawData) that actually has a weight > 0.0.
 	 * @param SubRectSize - The size of the SubRect inside the raw data (InRawData) that actually has a weight > 0.0.
 	 */
-	void AccumulateSinglePlane(const TArray64<float>& InRawData, const TArray64<float>& InWeightData, FIntPoint InSize, FIntPoint InOffset,
+	void AccumulateSinglePlane(const TArray64<float>& InRawData, FIntPoint InSize, FIntPoint InOffset,
 								float SubpixelOffsetX, float SubpixelOffsetY,
 								FIntPoint SubRectOffset,
-								FIntPoint SubRectSize);
+								FIntPoint SubRectSize,
+								const TArray<float> & WeightDataX,
+								const TArray<float> & WeightDataY);
 
 	/** Actual channel data*/
 	TArray64<float> ChannelData;
@@ -140,8 +144,10 @@ public:
 	 * @param InPixelData - Raw pixel data.
 	 * @param InTileOffset - Tile offset in pixels.
 	 * @param InSubPixelOffset - SubPixel offset, should be in the range [0,1)
+	 * @param WeightX - 1D Weights in X
+	 * @param WeightY - 1D Weights in Y
 	 */
-	void AccumulatePixelData(const FImagePixelData& InPixelData, FIntPoint InTileOffset, FVector2D InSubpixelOffset);
+	void AccumulatePixelData(const FImagePixelData& InPixelData, FIntPoint InTileOffset, FVector2D InSubpixelOffset, const MoviePipeline::FTileWeight1D & WeightX, const MoviePipeline::FTileWeight1D & WeightY);
 
 	/**
 	 * After accumulation is finished, fetch the final image as bytes. In theory we don't need this, because we could
diff --git a/Engine/Plugins/MovieScene/MovieRenderPipeline/Source/MovieRenderPipelineCore/Public/MovieRenderPipelineDataTypes.cpp b/Engine/Plugins/MovieScene/MovieRenderPipeline/Source/MovieRenderPipelineCore/Public/MovieRenderPipelineDataTypes.cpp
--- a/Engine/Plugins/MovieScene/MovieRenderPipeline/Source/MovieRenderPipelineCore/Public/MovieRenderPipelineDataTypes.cpp
+++ b/Engine/Plugins/MovieScene/MovieRenderPipeline/Source/MovieRenderPipelineCore/Public/MovieRenderPipelineDataTypes.cpp
@@ -134,10 +134,99 @@ FMoviePipelineWorkInfo FMoviePipelineCameraCutInfo::GetTotalWorkEstimate() const
 	return WorkEstimate;
 }
 
-namespace MoviePipeline
+
+bool MoviePipeline::FTileWeight1D::operator==(const FTileWeight1D& Rhs) const
+{
+	const bool bIsSame =
+		(X0 == Rhs.X0) &&
+		(X1 == Rhs.X1) &&
+		(X2 == Rhs.X2) &&
+		(X3 == Rhs.X3);
+
+	return bIsSame;
+}
+
+void MoviePipeline::FTileWeight1D::InitHelper(int32 PadLeft, int32 SizeCenter, int32 PadRight)
+{
+	check(PadLeft >= 0);
+	check(SizeCenter > 0);
+	check(PadRight >= 0);
+
+	int32 Size = PadLeft + SizeCenter + PadRight;
+	int32 Midpoint = PadLeft + SizeCenter / 2;
+
+	X0 = PadLeft / 2;
+	X1 = (3 * PadLeft) / 2;
+	X2 = Size - ((3 * PadLeft) / 2);
+	X3 = Size - (PadRight / 2);
+
+	X1 = FMath::Min<int32>(X1, Midpoint);
+	X2 = FMath::Max<int32>(X2, Midpoint);
+}
+
+float MoviePipeline::FTileWeight1D::CalculateWeight(int32 Pixel) const
 {
-	UMoviePipeline* FMoviePipelineEnginePass::GetPipeline() const
+	// the order of the if is important, because if X0=X1, then
+	// we take the Pixel < X0 path, not the Pixel < X1 path which would cause a
+	// divide by zero. Same for X2 and X3.
+	if (Pixel < X0)
 	{
-		return OwningPipeline.Get();
+		return 0.0f;
 	}
-}
\ No newline at end of file
+	else if (Pixel < X1)
+	{
+		// note that it's impossible for X0==X1, because the X0 if is earlier
+		return float(Pixel - X0) / float(X1 - X0);
+	}
+	else if (Pixel < X2)
+	{
+		return 1.0f;
+	}
+	else if (Pixel < X3)
+	{
+		// note that it's impossible for X2==X3
+		return float(X3 - Pixel) / float(X3 - X2);
+	}
+
+	return 0.0f;
+}
+
+void MoviePipeline::FTileWeight1D::CalculateArrayWeight(TArray<float>& WeightData, int Size) const
+{
+	WeightData.SetNum(Size);
+	for (int32 Index = 0; Index < X0; Index++)
+	{
+		WeightData[Index] = 0.0f;
+	}
+
+	float ScaleLhs = (X0 != X1) ? 1.0f / float(X1 - X0) : 0.0f;
+	float ScaleRhs = (X2 != X3) ? 1.0f / float(X3 - X2) : 0.0f;
+
+	for (int32 Index = X0; Index < X1; Index++)
+	{
+		float W = FMath::Clamp<float>(((float(Index) + .5f) - float(X0)) * ScaleLhs, 0.0f, 1.0f);
+		WeightData[Index] = W;
+	}
+
+	for (int32 Index = X1; Index < X2; Index++)
+	{
+		WeightData[Index] = 1.0f;
+	}
+
+	for (int32 Index = X2; Index < X3; Index++)
+	{
+		float W = FMath::Clamp<float>(1.0f - ((float(Index) + .5f) - float(X2)) * ScaleRhs, 0.0f, 1.0f);
+		WeightData[Index] = W;
+	}
+
+	for (int32 Index = X3; Index < Size; Index++)
+	{
+		WeightData[Index] = 0.0f;
+	}
+}
+
+
+UMoviePipeline* MoviePipeline::FMoviePipelineEnginePass::GetPipeline() const
+{
+	return OwningPipeline.Get();
+}
diff --git a/Engine/Plugins/MovieScene/MovieRenderPipeline/Source/MovieRenderPipelineCore/Public/MovieRenderPipelineDataTypes.h b/Engine/Plugins/MovieScene/MovieRenderPipeline/Source/MovieRenderPipelineCore/Public/MovieRenderPipelineDataTypes.h
--- a/Engine/Plugins/MovieScene/MovieRenderPipeline/Source/MovieRenderPipelineCore/Public/MovieRenderPipelineDataTypes.h
+++ b/Engine/Plugins/MovieScene/MovieRenderPipeline/Source/MovieRenderPipelineCore/Public/MovieRenderPipelineDataTypes.h
@@ -159,6 +159,73 @@ namespace MoviePipeline
 		FString PassName;
 	};
 
+	struct MOVIERENDERPIPELINECORE_API FTileWeight1D
+	{
+	public:
+		/**
+		 * For tiles, this struct stores the weight function in one dimension.
+		 *
+		 *                X1           X2
+		 * 1.0              |------------|
+		 *                 /              \
+		 *                /                \
+		 *               /                  \
+		 * 0.0 |--------|                    |----------------|
+		 *     0       X0                    X3             Size
+		 *
+		 *  For a tile, the weight goes from:
+		 *  [0 ,X0]    : 0.0
+		 *  [X0,X1]    : 0.0 to 1.0
+		 *  [X1,X2]    : 1.0
+		 *  [X2,X3]    : 1.0 to 0.0
+		 *  [X3,SizeX] : 0.0
+		 **/
+
+		/** Constructor. Sets to invalid values and must be initialized to reasonable values before using. **/
+		FTileWeight1D()
+		{
+			X0 = 0;
+			X1 = 0;
+			X2 = 0;
+			X3 = 0;
+		}
+
+		/** Make sure this map has valid values. **/
+		void CheckValid() const
+		{
+			check(0 <= X0);
+			check(X0 <= X1);
+			check(X1 <= X2);
+			check(X2 <= X3);
+		}
+
+		/** Is equal operator.**/
+		bool operator==(const FTileWeight1D& Rhs) const;
+
+		/**
+		 * The full tile is of size PadLeft + SizeCenter + PadRight
+		 *
+		 *  |------PadLeft-------|------------SizeCenter-----------|------PadRight-----|
+		 *  
+		 * This function puts X0 in the middle of PadLeft, X3 in the middle of PadRight.
+		 * And X1 is "reflected" around PadLeft, and X2 is "reflected" around PadRight.
+		 *
+		 *  |------PadLeft-------|------------SizeCenter-----------|------PadRight-----|
+		 *           X0                  X1                  X2              X3
+		 * Also note that PadLeft and PadRight are allowed to be zero, but SizeCenter must be > 0;
+		 **/
+		void InitHelper(int32 PadLeft, int32 SizeCenter, int32 PadRight);
+
+		float CalculateWeight(int32 Pixel) const;
+
+		void CalculateArrayWeight(TArray<float>& WeightData, int Size) const;
+
+		int32 X0;
+		int32 X1;
+		int32 X2;
+		int32 X3;
+	};
+
 	struct FFrameConstantMetrics
 	{
 		/** What is the tick resolution fo the sequence */
@@ -691,6 +758,8 @@ public:
 
 	FVector2D OverlappedSubpixelShift;
 
+	MoviePipeline::FTileWeight1D WeightFunctionX;
+	MoviePipeline::FTileWeight1D WeightFunctionY;
 
 	MoviePipeline::FMoviePipelineFrameInfo FrameInfo;
 };
@@ -723,8 +792,6 @@ namespace MoviePipeline
 
 		virtual ~FMoviePipelineEnginePass()
 		{}
-
-
 		virtual void Setup(TWeakObjectPtr<UMoviePipeline> InOwningPipeline, const FMoviePipelineRenderPassInitSettings& InInitSettings)
 		{
 			OwningPipeline = InOwningPipeline;
