Fixed compile issue.

diff --git a/Engine/Plugins/MovieScene/MovieRenderPipeline/Source/MovieRenderPipelineRenderPasses/Private/MoviePipelineDeferredPasses.cpp b/Engine/Plugins/MovieScene/MovieRenderPipeline/Source/MovieRenderPipelineRenderPasses/Private/MoviePipelineDeferredPasses.cpp
--- a/Engine/Plugins/MovieScene/MovieRenderPipeline/Source/MovieRenderPipelineRenderPasses/Private/MoviePipelineDeferredPasses.cpp
+++ b/Engine/Plugins/MovieScene/MovieRenderPipeline/Source/MovieRenderPipelineRenderPasses/Private/MoviePipelineDeferredPasses.cpp
@@ -72,8 +72,8 @@ namespace MoviePipeline
 		if (InSampleState.bIsUsingOverlappedTiles)
 		{
 			// Disable these for now
-			// ShowFlags.SetDepthOfField(false);
-			// ShowFlags.SetMotionBlur(false);
+			//ShowFlags.SetDepthOfField(false);
+			//ShowFlags.SetMotionBlur(false);
 		}
 
 		FSceneViewFamilyContext ViewFamily(FSceneViewFamily::ConstructionValues(
@@ -485,7 +485,6 @@ void UMoviePipelineDeferredPassBase::OnSetupView(FSceneViewFamily& InViewFamily,
 	// }
 }
 
-
 namespace MoviePipeline
 {
 	void AccumulateSample_RenderThread(TUniquePtr<FImagePixelData>&& InPixelData, TSharedRef<FImagePixelDataPayload, ESPMode::ThreadSafe> InFrameData, MoviePipeline::FSampleRenderThreadParams& InParams)
@@ -522,7 +521,8 @@ namespace MoviePipeline
 			// bool bSkip = InFrameData->SampleState.TileIndexes.X != 0 || InFrameData->SampleState.TileIndexes.Y != 1;
 			// if (!bSkip)
 			{
-				InParams.ImageAccumulator->AccumulatePixelData(*InPixelData.Get(), InFrameData->SampleState.OverlappedOffset, InFrameData->SampleState.OverlappedSubpixelShift);
+				InParams.ImageAccumulator->AccumulatePixelData(*InPixelData.Get(), InFrameData->SampleState.OverlappedOffset, InFrameData->SampleState.OverlappedSubpixelShift,
+					InFrameData->SampleState.WeightFunctionX, InFrameData->SampleState.WeightFunctionY);
 			}
 
 			const double AccumulateEndTime = FPlatformTime::Seconds();
