Optimizing tiled accumulator.

diff --git a/Engine/Plugins/MovieScene/MovieRenderPipeline/Source/MovieRenderPipelineCore/Private/MovieRenderOverlappedImage.cpp b/Engine/Plugins/MovieScene/MovieRenderPipeline/Source/MovieRenderPipelineCore/Private/MovieRenderOverlappedImage.cpp
--- a/Engine/Plugins/MovieScene/MovieRenderPipeline/Source/MovieRenderPipelineCore/Private/MovieRenderOverlappedImage.cpp
+++ b/Engine/Plugins/MovieScene/MovieRenderPipeline/Source/MovieRenderPipelineCore/Private/MovieRenderOverlappedImage.cpp
@@ -161,7 +161,7 @@ void FImageOverlappedPlane::AccumulateSinglePlane(const TArray64<float>& InRawDa
 		const float PixelWeight10 = PixelWeight[1][0];
 		const float PixelWeight11 = PixelWeight[1][1];
 
-		static bool bRunVectorized = false;
+		static bool bRunVectorized = true;
 		if (!bRunVectorized)
 		{
 			for (int32 DstY = ActualDstY0; DstY < ActualDstY1; DstY++)
@@ -314,10 +314,6 @@ void FImageOverlappedAccumulator::InitMemory(FIntPoint InPlaneSize, int32 InNumC
 	}
 
 	WeightPlane.Init(PlaneSize);
-
-	TileMaskData.Empty();
-	TileMaskSize.X = 0;
-	TileMaskSize.Y = 0;
 }
 
 void FImageOverlappedAccumulator::ZeroPlanes()
@@ -341,12 +337,9 @@ void FImageOverlappedAccumulator::Reset()
 	// Let the desctructor clean up
 	ChannelPlanes.Empty();
 	WeightPlane.Reset();
-
-	TileMaskData.Empty();
-	TileMaskSize.X = 0;
-	TileMaskSize.Y = 0;
 }
 
+#if 0
 
 void FImageOverlappedAccumulator::GenerateTileWeight(TArray64<float>& OutWeights, FIntPoint Size)
 {
@@ -441,6 +434,7 @@ void FImageOverlappedAccumulator::CheckTileSubRect(const TArray64<float>& OutWei
 		}
 	}
 }
+#endif
 
 void FImageOverlappedAccumulator::AccumulatePixelData(const FImagePixelData& InPixelData, FIntPoint InTileOffset, FVector2D InSubpixelOffset, const MoviePipeline::FTileWeight1D & WeightX, const MoviePipeline::FTileWeight1D & WeightY)
 {
@@ -458,6 +452,15 @@ void FImageOverlappedAccumulator::AccumulatePixelData(const FImagePixelData& InP
 
 	if (IsFetchOk)
 	{
+
+		TArray<float> WeightDataX;
+		TArray<float> WeightDataY;
+		WeightX.CalculateArrayWeight(WeightDataX, RawSize.X);
+		WeightY.CalculateArrayWeight(WeightDataY, RawSize.Y);
+
+		FIntPoint SubRectOffset(WeightX.X0, WeightY.X0);
+		FIntPoint SubRectSize(WeightX.X3 - WeightX.X0, WeightY.X3 - WeightX.X0);
+
 		// hardcode to 4 channels (RGBA), even if we are only saving fewer channels
 		TArray64<float> RawData[4];
 
@@ -547,17 +550,46 @@ void FImageOverlappedAccumulator::AccumulatePixelData(const FImagePixelData& InP
 		}
 		else if (Fmt == EImagePixelType::Float32 && RawNumChan == 4 && RawBitDepth == 32)
 		{
-			const float* RawDataPtr = static_cast<const float*>(SrcRawDataPtr);
 
-			// reference version for float
-			for (int32 Y = 0; Y < RawSize.Y; Y++)
+			static bool IsReferenceUnpack = false;
+			if (IsReferenceUnpack)
 			{
-				for (int32 X = 0; X < RawSize.X; X++)
+				const float* RawDataPtr = static_cast<const float*>(SrcRawDataPtr);
+
+				// reference version for float
+				for (int32 Y = 0; Y < RawSize.Y; Y++)
 				{
-					for (int32 ChanIter = 0; ChanIter < 4; ChanIter++)
+					for (int32 X = 0; X < RawSize.X; X++)
 					{
-						float Value = RawDataPtr[(Y*RawSize.X + X)*RawNumChan + ChanIter];
-						RawData[ChanIter][Y*RawSize.X + X] = Value;
+						for (int32 ChanIter = 0; ChanIter < 4; ChanIter++)
+						{
+							float Value = RawDataPtr[(Y*RawSize.X + X)*RawNumChan + ChanIter];
+							RawData[ChanIter][Y*RawSize.X + X] = Value;
+						}
+					}
+				}
+			}
+			else
+			{
+				const float* RawDataPtr = static_cast<const float*>(SrcRawDataPtr);
+
+				// slightly optimized, takes about 3-7ms on a 1080p image
+				for (int32 Y = 0; Y < RawSize.Y; Y++)
+				{
+					const float* SrcRowDataPtr = &RawDataPtr[Y*RawSize.X*RawNumChan];
+
+					float* DstRowDataR = &RawData[0][Y*RawSize.X];
+					float* DstRowDataG = &RawData[1][Y*RawSize.X];
+					float* DstRowDataB = &RawData[2][Y*RawSize.X];
+					float* DstRowDataA = &RawData[3][Y*RawSize.X];
+
+					// simple, one pixel at a time vectorized version, we could do better
+					for (int32 X = 0; X < RawSize.X; X++)
+					{
+						DstRowDataR[X] = SrcRowDataPtr[X*4+0];
+						DstRowDataG[X] = SrcRowDataPtr[X*4+1];
+						DstRowDataB[X] = SrcRowDataPtr[X*4+2];
+						DstRowDataA[X] = SrcRowDataPtr[X*4+3];
 					}
 				}
 			}
@@ -567,7 +599,7 @@ void FImageOverlappedAccumulator::AccumulatePixelData(const FImagePixelData& InP
 			check(0);
 		}
 
-		static bool bLogTiming = false;
+		static bool bLogTiming = true;
 
 		const double AccumulateEndTime = FPlatformTime::Seconds();
 		if(bLogTiming)
@@ -625,33 +657,6 @@ void FImageOverlappedAccumulator::AccumulatePixelData(const FImagePixelData& InP
 			UE_LOG(LogTemp, Log, TEXT("    [%8.2f] Gamma time."), ElapsedMs);
 		}
 
-		// Calculate weight data for this tile. If it is the same, use the cached size. In theory it is allowed
-		// to have a bunch of tiles of different sizes, but so far we are only using the same size for all tiles.
-		if (RawSize.X != TileMaskSize.X ||
-			RawSize.Y != TileMaskSize.Y ||
-			TileMaskSize.X * TileMaskSize.Y != TileMaskData.Num())
-		{
-			TileMaskSize.X = RawSize.X;
-			TileMaskSize.Y = RawSize.Y;
-
-			GenerateTileWeight(TileMaskData, RawSize);
-		}
-
-		TArray<float> WeightDataX;
-		TArray<float> WeightDataY;
-		WeightX.CalculateArrayWeight(WeightDataX, RawSize.X);
-		WeightY.CalculateArrayWeight(WeightDataY, RawSize.Y);
-
-		FIntPoint SubRectOffset(0, 0);
-		FIntPoint SubRectSize(0, 0);
-		GetTileSubRect(SubRectOffset, SubRectSize, TileMaskData, RawSize);
-		
-		const bool bVerifySubRect = false;
-		if (bVerifySubRect)
-		{
-			CheckTileSubRect(TileMaskData, RawSize, SubRectOffset, SubRectSize);
-		}
-
 		const double TileEndTime = FPlatformTime::Seconds();
 		if (bLogTiming)
 		{
diff --git a/Engine/Plugins/MovieScene/MovieRenderPipeline/Source/MovieRenderPipelineCore/Public/MovieRenderOverlappedImage.h b/Engine/Plugins/MovieScene/MovieRenderPipeline/Source/MovieRenderPipelineCore/Public/MovieRenderOverlappedImage.h
--- a/Engine/Plugins/MovieScene/MovieRenderPipeline/Source/MovieRenderPipelineCore/Public/MovieRenderOverlappedImage.h
+++ b/Engine/Plugins/MovieScene/MovieRenderPipeline/Source/MovieRenderPipelineCore/Public/MovieRenderOverlappedImage.h
@@ -83,7 +83,6 @@ public:
 		: PlaneSize(0,0)
 		, NumChannels(0)
 		, AccumulationGamma(1.0f)
-		, TileMaskSize(0,0)
 	{
 	}
 
@@ -110,6 +109,7 @@ public:
 	 */
 	void Reset();
 
+#if 0
 	/**
 	 * Given the size of a tile, calculate its weight mask.
 	 * 
@@ -137,6 +137,7 @@ public:
 	 * @param SubRectSize - Size of the found subrect.
 	 */
 	static void CheckTileSubRect(const TArray64<float>& Weights, const FIntPoint Size, FIntPoint SubRectOffset, FIntPoint SubRectSize);
+#endif
 
 	/**
 	 * Given a rendered tile, accumulate the data to the full size image.
@@ -196,12 +197,6 @@ public:
 	TArray64<FImageOverlappedPlane> ChannelPlanes;
 	
 	FImageOverlappedPlane WeightPlane;
-
-	/** Weights image. Recalculate if it changes. */
-	TArray64<float> TileMaskData;
-
-	/** Width and height of weight mask image. */
-	FIntPoint TileMaskSize;
 };
 
 
