Fixing off by one issues with overlapped tile accumulator.

diff --git a/Engine/Plugins/MovieScene/MovieRenderPipeline/Source/MovieRenderPipelineCore/Private/MoviePipelineRendering.cpp b/Engine/Plugins/MovieScene/MovieRenderPipeline/Source/MovieRenderPipelineCore/Private/MoviePipelineRendering.cpp
--- a/Engine/Plugins/MovieScene/MovieRenderPipeline/Source/MovieRenderPipelineCore/Private/MoviePipelineRendering.cpp
+++ b/Engine/Plugins/MovieScene/MovieRenderPipeline/Source/MovieRenderPipelineCore/Private/MoviePipelineRendering.cpp
@@ -368,7 +368,7 @@ void UMoviePipeline::RenderFrame()
 				{
 					SampleState.OverlappedOffset = FIntPoint(0,0);
 					SampleState.OverlappedPad = FIntPoint(0, 0);
-					SampleState.OverlappedSubpixelShift = FVector2D(0.0f,0.0f);
+					SampleState.OverlappedSubpixelShift = FVector2D(0.5f,0.5f);
 				}
 
 				SampleState.WeightFunctionX.InitHelper(SampleState.OverlappedPad.X, SampleState.TileSize.X, SampleState.OverlappedPad.X);
diff --git a/Engine/Plugins/MovieScene/MovieRenderPipeline/Source/MovieRenderPipelineCore/Private/MovieRenderOverlappedImage.cpp b/Engine/Plugins/MovieScene/MovieRenderPipeline/Source/MovieRenderPipelineCore/Private/MovieRenderOverlappedImage.cpp
--- a/Engine/Plugins/MovieScene/MovieRenderPipeline/Source/MovieRenderPipelineCore/Private/MovieRenderOverlappedImage.cpp
+++ b/Engine/Plugins/MovieScene/MovieRenderPipeline/Source/MovieRenderPipelineCore/Private/MovieRenderOverlappedImage.cpp
@@ -50,6 +50,10 @@ void FImageOverlappedPlane::AccumulateSinglePlane(const TArray64<float>& InRawDa
 	check(SubpixelOffsetY >= 0.0f);
 	check(SubpixelOffsetY <= 1.0f);
 
+	check(SubRectSize.X >= 1);
+	check(SubRectSize.Y >= 1);
+
+
 	// if the subpixel offset is less than 0.5, go back one pixel
 	int32 StartX = (SubpixelOffsetX >= 0.5 ? InOffset.X : InOffset.X - 1);
 	int32 StartY = (SubpixelOffsetY >= 0.5 ? InOffset.Y : InOffset.Y - 1);
@@ -70,7 +74,7 @@ void FImageOverlappedPlane::AccumulateSinglePlane(const TArray64<float>& InRawDa
 	static bool bSlowReference = false;
 	if (bSlowReference)
 	{
-		static bool bReversedReference = false;
+		static bool bReversedReference = true;
 		if (bReversedReference)
 		{
 			// This is the reversed reference. Instead of taking a tile and applying it to the accumulation, we start from the
@@ -80,9 +84,9 @@ void FImageOverlappedPlane::AccumulateSinglePlane(const TArray64<float>& InRawDa
 			// Given a position on the current tile (x,y), we apply the sum of the 4 points (x+0,y+0), (x+1,y+0), (x+0,y+1), (x+1,y+1) to the destination index.
 			// So in reverse, given a destination position, our source sample positions are (x-1,y-1), (x+0,y-1), (x-1,y+0), (x+0,y+0).
 			// That's why we make sure to start at a minimum of index 1, instead of 0.
-			for (int32 CurrY = 1; CurrY < InSize.Y; CurrY++)
+			for (int32 CurrY = 0; CurrY < InSize.Y; CurrY++)
 			{
-				for (int32 CurrX = 1; CurrX < InSize.X; CurrX++)
+				for (int32 CurrX = 0; CurrX < InSize.X; CurrX++)
 				{
 					int32 DstY = StartY + CurrY;// +OffsetY;
 					int32 DstX = StartX + CurrX;// +OffsetX;
@@ -94,10 +98,14 @@ void FImageOverlappedPlane::AccumulateSinglePlane(const TArray64<float>& InRawDa
 						{
 							for (int32 OffsetX = 0; OffsetX < 2; OffsetX++)
 							{
-								float Val = InRawData[(CurrY - 1 + OffsetY) * InSize.X + (CurrX - 1 + OffsetX)];
+								int32 SrcX = FMath::Max<int32>(CurrX - 1 + OffsetX,0);
+								int32 SrcY = FMath::Max<int32>(CurrY - 1 + OffsetY,0);
+
+
+								float Val = InRawData[SrcY * InSize.X + SrcX];
 								
-								float WX = WeightDataX[(CurrX - 1 + OffsetX)];
-								float WY = WeightDataY[(CurrY - 1 + OffsetY)];
+								float WX = WeightDataX[SrcX];
+								float WY = WeightDataY[SrcY];
 								float BaseWeight = WX * WY;
 
 								float Weight = BaseWeight * PixelWeight[1-OffsetY][1-OffsetX];
@@ -167,8 +175,8 @@ void FImageOverlappedPlane::AccumulateSinglePlane(const TArray64<float>& InRawDa
 			for (int32 DstY = ActualDstY0; DstY < ActualDstY1; DstY++)
 			{
 				int32 CurrY = DstY - StartY;
-				int32 SrcY0 = CurrY + (0 - 1);
-				int32 SrcY1 = CurrY + (1 - 1);
+				int32 SrcY0 = FMath::Max<int32>(0,CurrY + (0 - 1));
+				int32 SrcY1 = FMath::Max<int32>(0,CurrY + (1 - 1));
 
 				const float * SrcLineRaw0 = &InRawData[SrcY0 * InSize.X];
 				const float * SrcLineRaw1 = &InRawData[SrcY1 * InSize.X];
@@ -183,8 +191,8 @@ void FImageOverlappedPlane::AccumulateSinglePlane(const TArray64<float>& InRawDa
 				for (int32 DstX = ActualDstX0; DstX < ActualDstX1; DstX++)
 				{
 					int32 CurrX = DstX - StartX;
-					int32 X0 = CurrX - 1 + 0;
-					int32 X1 = CurrX - 1 + 1;
+					int32 X0 = FMath::Max<int32>(0,CurrX - 1 + 0);
+					int32 X1 = FMath::Max<int32>(0,CurrX - 1 + 1);
 
 					float Sum = DstLine[DstX];
 
@@ -199,14 +207,6 @@ void FImageOverlappedPlane::AccumulateSinglePlane(const TArray64<float>& InRawDa
 		}
 		else
 		{
-			int32 ActualWidth = ActualDstX1 - ActualDstX0;
-
-			// how many groups of 4
-			int32 ActualWidthGroup4 = ActualWidth / 4;
-
-			// extra 0-3
-			int32 ActualWidthExtra4 = ActualWidth - 4 * ActualWidthGroup4;
-
 			VectorRegister VecPixelWeight00 = VectorSetFloat1(PixelWeight[0][0]);
 			VectorRegister VecPixelWeight01 = VectorSetFloat1(PixelWeight[0][1]);
 			VectorRegister VecPixelWeight10 = VectorSetFloat1(PixelWeight[1][0]);
@@ -215,8 +215,8 @@ void FImageOverlappedPlane::AccumulateSinglePlane(const TArray64<float>& InRawDa
 			for (int32 DstY = ActualDstY0; DstY < ActualDstY1; DstY++)
 			{
 				int32 CurrY = DstY - StartY;
-				int32 SrcY0 = CurrY + (0 - 1);
-				int32 SrcY1 = CurrY + (1 - 1);
+				int32 SrcY0 = FMath::Max<int32>(CurrY + (0 - 1), 0); // we need the max for the first pixel which could go off the edge
+				int32 SrcY1 = FMath::Max<int32>(CurrY + (1 - 1), 0); 
 
 				const float * SrcLineRaw0 = &InRawData[SrcY0 * InSize.X];
 				const float * SrcLineRaw1 = &InRawData[SrcY1 * InSize.X];
@@ -235,10 +235,43 @@ void FImageOverlappedPlane::AccumulateSinglePlane(const TArray64<float>& InRawDa
 
 				float * DstLine = &ChannelData[DstY * Size.X];
 
+				check(ActualDstX0 - StartX >= 0);
+
+				// skip a single pixel to avoid maxes inside the loop
+				int32 NumSkipX = (ActualDstX0 - StartX == 0) ? 1 : 0;
+
+				int32 ActualWidth = FMath::Max<int32>(0,(ActualDstX1 - ActualDstX0) - NumSkipX);
+
+				// how many groups of 4
+				int32 ActualWidthGroup4 = ActualWidth / 4;
+
+				// extra 0-3
+				int32 ActualWidthExtra4 = ActualWidth - 4 * ActualWidthGroup4;
+
+
+				for (int32 IterStart = 0; IterStart < NumSkipX; IterStart++)
+				{
+					int32 DstX = ActualDstX0 + IterStart;
+
+					int32 CurrX = DstX - StartX;
+
+					int32 X0 = FMath::Max<int32>(0,CurrX - 1 + 0);
+					int32 X1 = FMath::Max<int32>(0,CurrX - 1 + 1);
+
+					float Sum = DstLine[DstX];
+
+					Sum += PixelWeight11 * RowWeight0 * SrcLineWeight[X0] * SrcLineRaw0[X0];
+					Sum += PixelWeight10 * RowWeight0 * SrcLineWeight[X1] * SrcLineRaw0[X1];
+					Sum += PixelWeight01 * RowWeight1 * SrcLineWeight[X0] * SrcLineRaw1[X0];
+					Sum += PixelWeight00 * RowWeight1 * SrcLineWeight[X1] * SrcLineRaw1[X1];
+
+					DstLine[DstX] = Sum;
+				}
+
 				for (int32 GroupIter = 0; GroupIter < ActualWidthGroup4; GroupIter++)
 				{
 					{
-						int32 BaseDstX = ActualDstX0 + GroupIter * 4;
+						int32 BaseDstX = ActualDstX0 + NumSkipX + GroupIter * 4;
 
 						int32 BaseCurrX = BaseDstX - StartX;
 
@@ -274,16 +307,16 @@ void FImageOverlappedPlane::AccumulateSinglePlane(const TArray64<float>& InRawDa
 						DstLine[DstX] = Sum;
 					}
 #endif
-
 				}
 
 				for (int32 IterExtra = 0; IterExtra < ActualWidthExtra4; IterExtra++)
 				{
-					int32 DstX = ActualDstX0 + ActualWidthExtra4 * 4 + IterExtra;
+					int32 DstX = ActualDstX0 + NumSkipX + ActualWidthGroup4 * 4 + IterExtra;
 
 					int32 CurrX = DstX - StartX;
-					int32 X0 = CurrX - 1 + 0;
-					int32 X1 = CurrX - 1 + 1;
+
+					int32 X0 = FMath::Max<int32>(0,CurrX - 1 + 0);
+					int32 X1 = FMath::Max<int32>(0,CurrX - 1 + 1);
 
 					float Sum = DstLine[DstX];
 
diff --git a/Engine/Plugins/MovieScene/MovieRenderPipeline/Source/MovieRenderPipelineRenderPasses/Private/MoviePipelineDeferredPasses.cpp b/Engine/Plugins/MovieScene/MovieRenderPipeline/Source/MovieRenderPipelineRenderPasses/Private/MoviePipelineDeferredPasses.cpp
--- a/Engine/Plugins/MovieScene/MovieRenderPipeline/Source/MovieRenderPipelineRenderPasses/Private/MoviePipelineDeferredPasses.cpp
+++ b/Engine/Plugins/MovieScene/MovieRenderPipeline/Source/MovieRenderPipelineRenderPasses/Private/MoviePipelineDeferredPasses.cpp
@@ -69,11 +69,12 @@ namespace MoviePipeline
 
 		FEngineShowFlags ShowFlags = FEngineShowFlags(EShowFlagInitMode::ESFIM_Game);
 
-		if (InSampleState.bIsUsingOverlappedTiles)
+		static bool TestDisableDofAndMotionBlur = false;
+		if (TestDisableDofAndMotionBlur)
 		{
 			// Disable these for now
-			//ShowFlags.SetDepthOfField(false);
-			//ShowFlags.SetMotionBlur(false);
+			ShowFlags.SetDepthOfField(false);
+			ShowFlags.SetMotionBlur(false);
 		}
 
 		FSceneViewFamilyContext ViewFamily(FSceneViewFamily::ConstructionValues(
