Convert the Movie Rendering back to sRGB FColor for the moment. Improved Movie Pipeline Blueprint API.

#rnx

diff --git a/Engine/Plugins/MovieScene/MovieRenderPipeline/Source/MovieRenderPipelineCore/Private/MoviePipelineMasterConfig.cpp b/Engine/Plugins/MovieScene/MovieRenderPipeline/Source/MovieRenderPipelineCore/Private/MoviePipelineMasterConfig.cpp
--- a/Engine/Plugins/MovieScene/MovieRenderPipeline/Source/MovieRenderPipelineCore/Private/MoviePipelineMasterConfig.cpp
+++ b/Engine/Plugins/MovieScene/MovieRenderPipeline/Source/MovieRenderPipelineCore/Private/MoviePipelineMasterConfig.cpp
@@ -119,11 +119,11 @@ UMoviePipelineShotConfig* UMoviePipelineMasterConfig::GetConfigForShot(const FSt
 
 UMoviePipelineMasterConfig::UMoviePipelineMasterConfig()
 {
-	DefaultShotConfig = NewObject<UMoviePipelineShotConfig>(this, "DefaultShotConfig");
+	DefaultShotConfig = CreateDefaultSubobject<UMoviePipelineShotConfig>("DefaultShotConfig");
 	
-	// Always add at least the output settings block since having a framerate/directory/etc. is crucial.
-	Settings.Add(NewObject<UMoviePipelineOutputSetting>(this, "OutputSetting"));
-
+	// Always add at least the output settings block since having a framerate/directory/etc. is critical.
+	UMoviePipelineOutputSetting* DefaultOutputBlock = CreateDefaultSubobject<UMoviePipelineOutputSetting>("DefaultOutputSetting");
+	Settings.Add(DefaultOutputBlock);
 }
 
 TArray<UMoviePipelineOutputBase*> UMoviePipelineMasterConfig::GetOutputContainers() const
diff --git a/Engine/Plugins/MovieScene/MovieRenderPipeline/Source/MovieRenderPipelineCore/Private/MoviePipelineRendering.cpp b/Engine/Plugins/MovieScene/MovieRenderPipeline/Source/MovieRenderPipelineCore/Private/MoviePipelineRendering.cpp
--- a/Engine/Plugins/MovieScene/MovieRenderPipeline/Source/MovieRenderPipelineCore/Private/MoviePipelineRendering.cpp
+++ b/Engine/Plugins/MovieScene/MovieRenderPipeline/Source/MovieRenderPipelineCore/Private/MoviePipelineRendering.cpp
@@ -161,7 +161,7 @@ void UMoviePipeline::RenderFrame()
 	// In short, for each output frame, for each accumulation frame, for each tile X/Y, for each jitter, we render a pass. This setup is
 	// designed to maximize the likely hood of deterministic rendering and that different passes line up with each other.
 	UMoviePipelineAccumulationSetting* AccumulationSettings = CurrentShot.ShotConfig->FindOrAddSetting<UMoviePipelineAccumulationSetting>();
-	UMoviePipelineCameraSetting* CameraSettings = GetPipelineMasterConfig()->FindOrAddSetting<UMoviePipelineCameraSetting>();
+	UMoviePipelineCameraSetting* CameraSettings = CurrentShot.ShotConfig->FindOrAddSetting<UMoviePipelineCameraSetting>();
 	UMoviePipelineOutputSetting* OutputSettings = GetPipelineMasterConfig()->FindSetting<UMoviePipelineOutputSetting>();
 	check(OutputSettings);
 
@@ -355,6 +355,14 @@ void UMoviePipeline::RenderFrame()
 				SampleState.BackbufferSize = BackbufferResolution;
 				SampleState.TileSize = TileResolution;
 				SampleState.FrameInfo = FrameInfo;
+				if (HighResSettings)
+				{
+					SampleState.bWriteSampleToDisk = HighResSettings->bWriteAllSamples;
+				}
+				else
+				{
+					SampleState.bWriteSampleToDisk = false;
+				}
 
 				if (bIsUsingOverlappedTiles)
 				{
@@ -429,7 +437,7 @@ void UMoviePipeline::OnSampleRendered(TUniquePtr<FImagePixelData>&& OutputSample
 	}
 
 	// JPEG output
-	TileImageTask->Format = EImageFormat::EXR;
+	TileImageTask->Format = EImageFormat::JPEG;
 	TileImageTask->CompressionQuality = 100;
 
 	FString OutputName = FString::Printf(TEXT("/%s_SS_%d_TS_%d_TileX_%d_TileY_%d.%d.jpeg"),
diff --git a/Engine/Plugins/MovieScene/MovieRenderPipeline/Source/MovieRenderPipelineCore/Public/MoviePipelineAccumulationSetting.h b/Engine/Plugins/MovieScene/MovieRenderPipeline/Source/MovieRenderPipelineCore/Public/MoviePipelineAccumulationSetting.h
--- a/Engine/Plugins/MovieScene/MovieRenderPipeline/Source/MovieRenderPipelineCore/Public/MoviePipelineAccumulationSetting.h
+++ b/Engine/Plugins/MovieScene/MovieRenderPipeline/Source/MovieRenderPipelineCore/Public/MoviePipelineAccumulationSetting.h
@@ -19,6 +19,11 @@ public:
 	{
 	}
 
+protected:
+	virtual bool IsValidOnShots() const override { return true; }
+	virtual bool IsValidOnMaster() const override { return false; }
+	
+public:
 
 	/**
 	* How many frames should we accumulate together before contributing to one overall sample. This lets you
diff --git a/Engine/Plugins/MovieScene/MovieRenderPipeline/Source/MovieRenderPipelineCore/Public/MoviePipelineBlueprintLibrary.h b/Engine/Plugins/MovieScene/MovieRenderPipeline/Source/MovieRenderPipelineCore/Public/MoviePipelineBlueprintLibrary.h
--- a/Engine/Plugins/MovieScene/MovieRenderPipeline/Source/MovieRenderPipelineCore/Public/MoviePipelineBlueprintLibrary.h
+++ b/Engine/Plugins/MovieScene/MovieRenderPipeline/Source/MovieRenderPipelineCore/Public/MoviePipelineBlueprintLibrary.h
@@ -17,6 +17,13 @@ class MOVIERENDERPIPELINECORE_API UMoviePipelineBlueprintLibrary : public UBluep
 	GENERATED_BODY()
 
 public:
+	/**
+	* Takes a given Output Directory and FileName which contain {formatStrings} and convert those format strings
+	* to their actual values as best as possible. Merges the result into one final filepath string.
+	*/
+	// UFUNCTION(BlueprintPure, Category = "Movie Render Pipeline")
+	// static FString GetCameraCutCounts(const UMoviePipeline* InPipeline, const FMoviePipelineFrameOutputState& OutputState, const FDirectoryPath& InDirectoryFormatString, const FString& InFileNameFormatString);
+
 	/** 
 	* Get the total number of Camera Cuts in a Movie Pipeline, and an index of which one is being processed.
 	* Note: Index doesn't index into any arrays (since internally we store CameraCuts inside of Shots)
diff --git a/Engine/Plugins/MovieScene/MovieRenderPipeline/Source/MovieRenderPipelineCore/Public/MoviePipelineCameraSetting.h b/Engine/Plugins/MovieScene/MovieRenderPipeline/Source/MovieRenderPipelineCore/Public/MoviePipelineCameraSetting.h
--- a/Engine/Plugins/MovieScene/MovieRenderPipeline/Source/MovieRenderPipelineCore/Public/MoviePipelineCameraSetting.h
+++ b/Engine/Plugins/MovieScene/MovieRenderPipeline/Source/MovieRenderPipelineCore/Public/MoviePipelineCameraSetting.h
@@ -18,7 +18,10 @@ public:
 		, bManualExposure(true)
 		, ExposureCompensation(8.0)
 	{}
-	
+protected:
+	virtual bool IsValidOnShots() const override { return true; }
+	virtual bool IsValidOnMaster() const override { return false; }
+
 public:
 
 	/** 
diff --git a/Engine/Plugins/MovieScene/MovieRenderPipeline/Source/MovieRenderPipelineCore/Public/MoviePipelineConfigBase.h b/Engine/Plugins/MovieScene/MovieRenderPipeline/Source/MovieRenderPipelineCore/Public/MoviePipelineConfigBase.h
--- a/Engine/Plugins/MovieScene/MovieRenderPipeline/Source/MovieRenderPipelineCore/Public/MoviePipelineConfigBase.h
+++ b/Engine/Plugins/MovieScene/MovieRenderPipeline/Source/MovieRenderPipelineCore/Public/MoviePipelineConfigBase.h
@@ -7,11 +7,11 @@
 #include "MoviePipelineConfigBase.generated.h"
 
 
-UCLASS(Blueprintable)
+UCLASS(BlueprintType, Abstract)
 class MOVIERENDERPIPELINECORE_API UMoviePipelineConfigBase : public UObject
 {
 	GENERATED_BODY()
-	
+
 public:
 	UMoviePipelineConfigBase()
 	{
@@ -20,13 +20,13 @@ public:
 
 public:
 	/**
-	* Removes the specific instance from our Setting list. 
+	* Removes the specific instance from our Setting list.
 	*/
 	UFUNCTION(BlueprintCallable, Category = "Movie Render Pipeline")
-	void RemoveSetting(UMoviePipelineSetting* InSetting);
+		void RemoveSetting(UMoviePipelineSetting* InSetting);
 
 	int32 GetSettingsSerialNumber() const { return SettingsSerialNumber; }
-	
+
 	/**
 	* Returns an array of all settings in this config.
 	*/
@@ -38,29 +38,38 @@ public:
 	* @param InClass - Class that you wish to find the setting object for.
 	* @return An instance of this class if it already exists as a setting on this config, otherwise null.
 	*/
-	UFUNCTION(BlueprintPure, Category = "Movie Render Pipeline")
-	UMoviePipelineSetting* FindSettingByClass(TSubclassOf<UMoviePipelineSetting> InClass) const
+	UFUNCTION(BlueprintPure, meta = (DeterminesOutputType = "InClass"), Category = "Movie Render Pipeline")
+		UMoviePipelineSetting* FindSettingByClass(TSubclassOf<UMoviePipelineSetting> InClass) const
 	{
 		UMoviePipelineSetting* const* Found = Settings.FindByPredicate([InClass](UMoviePipelineSetting* In) { return In && In->GetClass() == InClass; });
 		return Found ? CastChecked<UMoviePipelineSetting>(*Found) : nullptr;
 	}
 
-	/** 
+	/**
 	* Finds a setting of a particular type for this pipeline config, adding it if it doesn't already exist.
 	* @param InClass - Class you wish to find or create the setting object for.
 	* @return An instance of this class as a setting on this config.
 	*/
-	UFUNCTION(BlueprintCallable, Category = "Movie Render Pipeline")
-	UMoviePipelineSetting* FindOrAddSettingByClass(TSubclassOf<UMoviePipelineSetting> InClass)
+	UFUNCTION(BlueprintCallable, meta = (DeterminesOutputType = "InClass"), Category = "Movie Render Pipeline")
+		UMoviePipelineSetting* FindOrAddSettingByClass(TSubclassOf<UMoviePipelineSetting> InClass)
 	{
 		UMoviePipelineSetting* Found = FindSettingByClass(InClass);
 		if (!Found)
 		{
 			Modify();
-
+			
 			Found = NewObject<UMoviePipelineSetting>(this, InClass);
-			Settings.Add(Found);
-			++SettingsSerialNumber;
+
+			if (CanSettingBeAdded(Found))
+			{
+				Settings.Add(Found);
+				++SettingsSerialNumber;
+			}
+			else
+			{
+				FFrame::KismetExecutionMessage(*FString::Printf(TEXT("Setting %d is not compatible with this Config Type and was not added."), *InClass->GetName()), ELogVerbosity::Error);
+				return nullptr;
+			}
 		}
 
 		return Found;
@@ -89,15 +98,26 @@ public:
 			Modify();
 
 			Found = NewObject<SettingType>(this);
-			Settings.Add(Found);
-			++SettingsSerialNumber;
+			if (CanSettingBeAdded(Found))
+			{
+				Settings.Add(Found);
+				++SettingsSerialNumber;
+			}
+			else
+			{
+				FFrame::KismetExecutionMessage(*FString::Printf(TEXT("Setting %d is not compatible with this Config Type and was not added."), *Found->GetName()), ELogVerbosity::Error);
+				return nullptr;
+			}
 		}
 		return Found;
 	}
 
+protected:
+	virtual bool CanSettingBeAdded(UMoviePipelineSetting* InSetting) PURE_VIRTUAL( UMoviePipelineConfigBase::CanSettingBeAdded, return false; );
+
 protected:
 	/** Array of settings classes that affect various parts of the output pipeline. */
-	UPROPERTY(VisibleAnywhere, Category = "Movie Pipeline")
+	UPROPERTY(VisibleAnywhere, Instanced, Category = "Movie Pipeline")
 	TArray<UMoviePipelineSetting*> Settings;
 
 private:
diff --git a/Engine/Plugins/MovieScene/MovieRenderPipeline/Source/MovieRenderPipelineCore/Public/MoviePipelineHighResSetting.h b/Engine/Plugins/MovieScene/MovieRenderPipeline/Source/MovieRenderPipelineCore/Public/MoviePipelineHighResSetting.h
--- a/Engine/Plugins/MovieScene/MovieRenderPipeline/Source/MovieRenderPipelineCore/Public/MoviePipelineHighResSetting.h
+++ b/Engine/Plugins/MovieScene/MovieRenderPipeline/Source/MovieRenderPipelineCore/Public/MoviePipelineHighResSetting.h
@@ -28,6 +28,9 @@ public:
 
 		return InTileSize + FIntPoint(2 * OverlappedPadX, 2 * OverlappedPadY);
 	}
+protected:
+	virtual bool IsValidOnShots() const override { return true; }
+	virtual bool IsValidOnMaster() const override { return false; }
 
 public:
 	/**
diff --git a/Engine/Plugins/MovieScene/MovieRenderPipeline/Source/MovieRenderPipelineCore/Public/MoviePipelineMasterConfig.h b/Engine/Plugins/MovieScene/MovieRenderPipeline/Source/MovieRenderPipelineCore/Public/MoviePipelineMasterConfig.h
--- a/Engine/Plugins/MovieScene/MovieRenderPipeline/Source/MovieRenderPipelineCore/Public/MoviePipelineMasterConfig.h
+++ b/Engine/Plugins/MovieScene/MovieRenderPipeline/Source/MovieRenderPipelineCore/Public/MoviePipelineMasterConfig.h
@@ -38,6 +38,12 @@ public:
 	/** Returns a pointer to the config specified for the shot, otherwise the default for this pipeline. */
 	UMoviePipelineShotConfig* GetConfigForShot(const FString& ShotName) const;
 
+protected:
+	virtual bool CanSettingBeAdded(UMoviePipelineSetting* InSetting) override
+	{
+		check(InSetting);
+		return InSetting->IsValidOnMaster();
+	}
 public:
 	
 	/** The default shot-setup to use for any shot that doesn't a specific implementation. This is required! */
diff --git a/Engine/Plugins/MovieScene/MovieRenderPipeline/Source/MovieRenderPipelineCore/Public/MoviePipelineOutputBase.h b/Engine/Plugins/MovieScene/MovieRenderPipeline/Source/MovieRenderPipelineCore/Public/MoviePipelineOutputBase.h
--- a/Engine/Plugins/MovieScene/MovieRenderPipeline/Source/MovieRenderPipelineCore/Public/MoviePipelineOutputBase.h
+++ b/Engine/Plugins/MovieScene/MovieRenderPipeline/Source/MovieRenderPipelineCore/Public/MoviePipelineOutputBase.h
@@ -58,6 +58,10 @@ public:
 	void OnPostTick() { OnPostTickImpl(); }
 
 	virtual void OnRecieveImageDataImpl(FMoviePipelineMergerOutputFrame* InMergedOutputFrame) {}
+
+protected:
+	virtual bool IsValidOnShots() const override { return false; }
+	virtual bool IsValidOnMaster() const override { return true; }
 protected:
 	virtual void OnShotInitializedImpl(const TOptional<FMoviePipelineShotInfo> PrevShot, const FMoviePipelineShotInfo& NewShot) {}
 	virtual void OnShotFinishedImpl(const FMoviePipelineShotInfo& Shot) {}
diff --git a/Engine/Plugins/MovieScene/MovieRenderPipeline/Source/MovieRenderPipelineCore/Public/MoviePipelineRenderPass.h b/Engine/Plugins/MovieScene/MovieRenderPipeline/Source/MovieRenderPipelineCore/Public/MoviePipelineRenderPass.h
--- a/Engine/Plugins/MovieScene/MovieRenderPipeline/Source/MovieRenderPipelineCore/Public/MoviePipelineRenderPass.h
+++ b/Engine/Plugins/MovieScene/MovieRenderPipeline/Source/MovieRenderPipelineCore/Public/MoviePipelineRenderPass.h
@@ -34,7 +34,9 @@ public:
 		GetRequiredEnginePassesImpl(RequiredEnginePasses);
 	}
 
-
+protected:
+	virtual bool IsValidOnShots() const override { return true; }
+	virtual bool IsValidOnMaster() const override { return false; }
 protected:
 	virtual void GetRequiredEnginePassesImpl(TSet<FMoviePipelinePassIdentifier>& RequiredEnginePasses) {}
 
diff --git a/Engine/Plugins/MovieScene/MovieRenderPipeline/Source/MovieRenderPipelineCore/Public/MoviePipelineSetting.h b/Engine/Plugins/MovieScene/MovieRenderPipeline/Source/MovieRenderPipelineCore/Public/MoviePipelineSetting.h
--- a/Engine/Plugins/MovieScene/MovieRenderPipeline/Source/MovieRenderPipelineCore/Public/MoviePipelineSetting.h
+++ b/Engine/Plugins/MovieScene/MovieRenderPipeline/Source/MovieRenderPipelineCore/Public/MoviePipelineSetting.h
@@ -37,11 +37,16 @@ protected:
 
 	virtual void SetupForPipelineImpl(UMoviePipeline* InPipeline) {}
 	virtual void TeardownForPipelineImpl(UMoviePipeline* InPipeline) {}
-
+	
 public:
 #if WITH_EDITOR
 	virtual FText GetDisplayText() const { return this->GetClass()->GetDisplayNameText(); }
 #endif
+	/** Can this configuration setting be added to shots? If not, it will throw an error when trying to add it to a shot config. */
+	virtual bool IsValidOnShots() const PURE_VIRTUAL(UMoviePipelineSetting::IsValidOnShots, return false; );
+	/** Can this configuration setting be added to the master configuration? If not, it will throw an error when trying to add it to the master configuration. */
+	virtual bool IsValidOnMaster() const PURE_VIRTUAL(UMoviePipelineSetting::IsValidOnMaster, return false; );
+
 	/** Should the Pipeline automatically create an instance of this under the hood so calling code can rely on it existing? */
 	virtual bool IsRequired() const { return false; }
 	
diff --git a/Engine/Plugins/MovieScene/MovieRenderPipeline/Source/MovieRenderPipelineCore/Public/MoviePipelineShotConfig.h b/Engine/Plugins/MovieScene/MovieRenderPipeline/Source/MovieRenderPipelineCore/Public/MoviePipelineShotConfig.h
--- a/Engine/Plugins/MovieScene/MovieRenderPipeline/Source/MovieRenderPipelineCore/Public/MoviePipelineShotConfig.h
+++ b/Engine/Plugins/MovieScene/MovieRenderPipeline/Source/MovieRenderPipelineCore/Public/MoviePipelineShotConfig.h
@@ -24,4 +24,12 @@ public:
 	}
 	
 	TArray<UMoviePipelineRenderPass*> GetRenderPasses() const;
+
+protected:
+	virtual bool CanSettingBeAdded(UMoviePipelineSetting* InSetting) override
+	{
+		check(InSetting);
+		return InSetting->IsValidOnShots();
+	}
+
 };
\ No newline at end of file
diff --git a/Engine/Plugins/MovieScene/MovieRenderPipeline/Source/MovieRenderPipelineCore/Public/MovieRenderPipelineDataTypes.h b/Engine/Plugins/MovieScene/MovieRenderPipeline/Source/MovieRenderPipelineCore/Public/MovieRenderPipelineDataTypes.h
--- a/Engine/Plugins/MovieScene/MovieRenderPipeline/Source/MovieRenderPipelineCore/Public/MovieRenderPipelineDataTypes.h
+++ b/Engine/Plugins/MovieScene/MovieRenderPipeline/Source/MovieRenderPipelineCore/Public/MovieRenderPipelineDataTypes.h
@@ -709,7 +709,7 @@ public:
 
 	FORCEINLINE int32 GetTileCount() const
 	{
-		return (TileCounts.X & TileCounts.Y);
+		return (TileCounts.X * TileCounts.Y);
 	}
 
 	/** How big is the back buffer for this sample? This is the final target size divided by the number of tiles with padding added. */
@@ -721,6 +721,9 @@ public:
 	/** If true, we will discard this sample after rendering. Used to get history set up correctly. */
 	bool bIsHistoryOnlyFrame;
 
+	/** If true, we will write this sample to disk (and also send it to the accumulator.) */
+	bool bWriteSampleToDisk;
+
 	/** How many spatial jitters will there be total for this particular temporal accumulation frame? */
 	int32 SpatialSampleCount;
 	
@@ -825,7 +828,6 @@ namespace MoviePipeline
 		FMoviePipelineRenderPassMetrics SampleState;
 		TSharedPtr<FImageOverlappedAccumulator, ESPMode::ThreadSafe> ImageAccumulator;
 		TSharedPtr<FMoviePipelineOutputMerger, ESPMode::ThreadSafe> OutputMerger;
-		bool bWriteTiles;
 		FMoviePipelinePassIdentifier PassIdentifier;
 	};
 }
diff --git a/Engine/Plugins/MovieScene/MovieRenderPipeline/Source/MovieRenderPipelineEditor/Private/MovieRenderPipelineEditorTest.cpp b/Engine/Plugins/MovieScene/MovieRenderPipeline/Source/MovieRenderPipelineEditor/Private/MovieRenderPipelineEditorTest.cpp
--- a/Engine/Plugins/MovieScene/MovieRenderPipeline/Source/MovieRenderPipelineEditor/Private/MovieRenderPipelineEditorTest.cpp
+++ b/Engine/Plugins/MovieScene/MovieRenderPipeline/Source/MovieRenderPipelineEditor/Private/MovieRenderPipelineEditorTest.cpp
@@ -34,7 +34,7 @@ UMoviePipelineShotConfig* GenerateTestShotConfig(UObject* InOwner, int32 InSampl
 		CameraSetting->CameraShutterAngle = InShutterAngle;
 		CameraSetting->ShutterTiming = InFrameTiming;
 		CameraSetting->bManualExposure = true;
-		CameraSetting->ExposureCompensation = 8;
+		CameraSetting->ExposureCompensation = 11;
 	}
 
 	if (InTileCount > 1)
@@ -43,6 +43,7 @@ UMoviePipelineShotConfig* GenerateTestShotConfig(UObject* InOwner, int32 InSampl
 		HighRes->TileCount = InTileCount;
 		HighRes->bIsUsingOverlappedTiles = bIsUsingOverlappedTiles;
 		HighRes->OverlapPercentage = PadRatioX;
+		HighRes->bWriteAllSamples = false;
 	}
 
 	return OutConfig;
@@ -74,11 +75,11 @@ TArray<UMoviePipelineMasterConfig*> FMovieRenderPipelineEditorModule::GenerateTe
 				UMoviePipelineMasterConfig* OutPipeline = NewObject<UMoviePipelineMasterConfig>(GetTransientPackage());
 				UMoviePipelineOutputSetting* OutputSetting = OutPipeline->FindOrAddSetting<UMoviePipelineOutputSetting>();
 
-				int32 SizeX = 1920/4;
-				int32 SizeY = 1080/4;
+				int32 SizeX = 1920;
+				int32 SizeY = 1080;
 				int32 TileX = 4;
 				int32 TileY = 4;
-				int32 TestNumSamples = 2;
+				int32 TestNumSamples = 8;
 				float PadRatioX = 0.5f;
 				float PadRatioY = 0.5f;
 				float AccumulationGamma = 1.0f;
diff --git a/Engine/Plugins/MovieScene/MovieRenderPipeline/Source/MovieRenderPipelineRenderPasses/Private/MoviePipelineDeferredPasses.cpp b/Engine/Plugins/MovieScene/MovieRenderPipeline/Source/MovieRenderPipelineRenderPasses/Private/MoviePipelineDeferredPasses.cpp
--- a/Engine/Plugins/MovieScene/MovieRenderPipeline/Source/MovieRenderPipelineRenderPasses/Private/MoviePipelineDeferredPasses.cpp
+++ b/Engine/Plugins/MovieScene/MovieRenderPipeline/Source/MovieRenderPipelineRenderPasses/Private/MoviePipelineDeferredPasses.cpp
@@ -84,7 +84,7 @@ namespace MoviePipeline
 			.SetWorldTimes(TimeSeconds, DeltaTimeSeconds, RealTimeSeconds)
 			.SetRealtimeUpdate(true));
 
-		ViewFamily.SceneCaptureSource = ESceneCaptureSource::SCS_FinalColorHDR;
+		ViewFamily.SceneCaptureSource = ESceneCaptureSource::SCS_FinalColorLDR;
 		ViewFamily.SetScreenPercentageInterface(new FLegacyScreenPercentageDriver(ViewFamily, 1.0f, false));
 
 		// View is added as a child of the ViewFamily
@@ -185,7 +185,7 @@ namespace MoviePipeline
 			{
 				UE_LOG(LogMovieRenderPipeline, Log, TEXT("Frame: %d TemporalSample: %d/%d Rendered Tile: %d/%d Sample: %d/%d"),
 					InSampleState.OutputState.OutputFrameNumber, InSampleState.TemporalSampleIndex + 1, InSampleState.TemporalSampleCount,
-					InSampleState.GetTileIndex() + 1 / InSampleState.GetTileCount(), InSampleState.SpatialSampleIndex + 1, InSampleState.SpatialSampleCount);
+					InSampleState.GetTileIndex() + 1 , InSampleState.GetTileCount(), InSampleState.SpatialSampleIndex + 1, InSampleState.SpatialSampleCount);
 			}
 			
 			ReadBackbufferAndBroadcast_RenderThread(RHICmdList, BackbufferRenderTarget, InSampleState, OnBackbufferReadDelegate);
@@ -199,11 +199,11 @@ namespace MoviePipeline
 		FIntRect SourceRect = FIntRect(0, 0, InFromRenderTarget->GetSizeXY().X, InFromRenderTarget->GetSizeXY().Y);
 
 		// Read the data back to the CPU
-		TArray<FLinearColor> RawPixels;
+		TArray<FColor> RawPixels;
 		RawPixels.SetNum(SourceRect.Width() * SourceRect.Height());
 
 		FReadSurfaceDataFlags ReadDataFlags;
-		ReadDataFlags.SetLinearToGamma(false);
+		ReadDataFlags.SetLinearToGamma(true);
 
 		RHICmdList.ReadSurfaceData(InFromRenderTarget->GetRenderTargetTexture(), SourceRect, RawPixels, ReadDataFlags);
 
@@ -394,11 +394,10 @@ void UMoviePipelineDeferredPassBase::GatherOutputPassesImpl(TArray<FMoviePipelin
 	}
 }
 
-void UMoviePipelineDeferredPassBase::OnBackbufferSampleReady(TArray<FLinearColor> InPixelData, FMoviePipelineRenderPassMetrics InSampleState)
+void UMoviePipelineDeferredPassBase::OnBackbufferSampleReady(TArray<FColor> InPixelData, FMoviePipelineRenderPassMetrics InSampleState)
 {
 	MoviePipeline::FSampleRenderThreadParams AccumulationParams;
 	{
-		AccumulationParams.bWriteTiles = true;
 		AccumulationParams.OutputMerger = GetPipeline()->OutputBuilder;
 		AccumulationParams.ImageAccumulator = ImageTileAccumulator[0];
 		AccumulationParams.SampleState = InSampleState;
@@ -409,7 +408,7 @@ void UMoviePipelineDeferredPassBase::OnBackbufferSampleReady(TArray<FLinearColor
 	FrameData->PassIdentifier = FMoviePipelinePassIdentifier(TEXT("Backbuffer"));
 	FrameData->SampleState = InSampleState;
 
-	TUniquePtr<TImagePixelData<FLinearColor>> PixelData = MakeUnique<TImagePixelData<FLinearColor>>(InSampleState.BackbufferSize, TArray64<FLinearColor>(MoveTemp(InPixelData)), FrameData);
+	TUniquePtr<TImagePixelData<FColor>> PixelData = MakeUnique<TImagePixelData<FColor>>(InSampleState.BackbufferSize, TArray64<FColor>(MoveTemp(InPixelData)), FrameData);
 
 	AccumulateSample_RenderThread(MoveTemp(PixelData), FrameData, AccumulationParams);
 }
@@ -449,7 +448,6 @@ void UMoviePipelineDeferredPassBase::OnSetupView(FSceneViewFamily& InViewFamily,
 
 		MoviePipeline::FSampleRenderThreadParams AccumulationParams;
 		{
-			AccumulationParams.bWriteTiles = true;
 			AccumulationParams.OutputMerger = GetPipeline()->OutputBuilder;
 			AccumulationParams.ImageAccumulator = ImageTileAccumulator[PassIndex + 1]; // 0 is reserved for backbuffer.
 			AccumulationParams.SampleState = InSampleState;
@@ -493,7 +491,7 @@ namespace MoviePipeline
 		check(InPixelData->IsDataWellFormed());
 
 		// Writing tiles can be useful for debug reasons. These get passed onto the output every frame.
-		if (InParams.bWriteTiles)
+		if (InParams.SampleState.bWriteSampleToDisk)
 		{
 			// Send the data to the Output Builder. This has to be a copy of the pixel data from the GPU, since
 			// it enqueues it onto the game thread and won't be read/sent to write to disk for another frame. 
@@ -540,7 +538,7 @@ namespace MoviePipeline
 
 			// Now that a tile is fully built and accumulated we can notify the output builder that the
 			// data is ready so it can pass that onto the output containers (if needed).
-			if (InPixelData->GetType() == EImagePixelType::Float32)
+			/*if (InPixelData->GetType() == EImagePixelType::Float32)
 			{
 				// 32 bit FLinearColor
 				TUniquePtr<TImagePixelData<FLinearColor> > FinalPixelData = MakeUnique<TImagePixelData<FLinearColor>>(FIntPoint(FullSizeX, FullSizeY));
@@ -558,7 +556,7 @@ namespace MoviePipeline
 				// Send the data to the Output Builder
 				InParams.OutputMerger->OnCompleteRenderPassDataAvailable_AnyThread(MoveTemp(FinalPixelData), InFrameData);
 			}
-			else if(InPixelData->GetType() == EImagePixelType::Color)
+			else*/ if(true)
 			{
 				// 8bit FColors
 				TUniquePtr<TImagePixelData<FColor>> FinalPixelData = MakeUnique<TImagePixelData<FColor>>(FIntPoint(FullSizeX, FullSizeY));
diff --git a/Engine/Plugins/MovieScene/MovieRenderPipeline/Source/MovieRenderPipelineRenderPasses/Public/MoviePipelineDeferredPasses.h b/Engine/Plugins/MovieScene/MovieRenderPipeline/Source/MovieRenderPipelineRenderPasses/Public/MoviePipelineDeferredPasses.h
--- a/Engine/Plugins/MovieScene/MovieRenderPipeline/Source/MovieRenderPipelineRenderPasses/Public/MoviePipelineDeferredPasses.h
+++ b/Engine/Plugins/MovieScene/MovieRenderPipeline/Source/MovieRenderPipelineRenderPasses/Public/MoviePipelineDeferredPasses.h
@@ -25,7 +25,7 @@ protected:
 	virtual void SetupImpl(TArray<TSharedPtr<MoviePipeline::FMoviePipelineEnginePass>>& InEnginePasses) override;
 	// ~UMovieRenderPassAPI
 
-	void OnBackbufferSampleReady(TArray<FLinearColor> InPixelData, FMoviePipelineRenderPassMetrics InSampleState);
+	void OnBackbufferSampleReady(TArray<FColor> InPixelData, FMoviePipelineRenderPassMetrics InSampleState);
 	void OnSetupView(FSceneViewFamily& InViewFamily, FSceneView& InView, const FMoviePipelineRenderPassMetrics& InSampleState);
 
 private:
@@ -36,7 +36,7 @@ private:
 	TArray<TSharedPtr<FImageOverlappedAccumulator, ESPMode::ThreadSafe>> ImageTileAccumulator;
 };
 
-DECLARE_MULTICAST_DELEGATE_TwoParams(FMoviePipelineSampleReady, TArray<FLinearColor>, FMoviePipelineRenderPassMetrics);
+DECLARE_MULTICAST_DELEGATE_TwoParams(FMoviePipelineSampleReady, TArray<FColor>, FMoviePipelineRenderPassMetrics);
 DECLARE_MULTICAST_DELEGATE_ThreeParams(FMoviePipelineSetupView, FSceneViewFamily&, FSceneView&, const FMoviePipelineRenderPassMetrics&);
 
 namespace MoviePipeline
diff --git a/Engine/Plugins/MovieScene/MovieRenderPipeline/Source/MovieRenderPipelineSettings/Public/MoviePipelineGameOverrideSetting.h b/Engine/Plugins/MovieScene/MovieRenderPipeline/Source/MovieRenderPipelineSettings/Public/MoviePipelineGameOverrideSetting.h
--- a/Engine/Plugins/MovieScene/MovieRenderPipeline/Source/MovieRenderPipelineSettings/Public/MoviePipelineGameOverrideSetting.h
+++ b/Engine/Plugins/MovieScene/MovieRenderPipeline/Source/MovieRenderPipelineSettings/Public/MoviePipelineGameOverrideSetting.h
@@ -19,7 +19,12 @@ public:
 		, bShowPlayer(false)
 	{
 	}
-	
+
+protected:
+	virtual bool IsValidOnShots() const override { return false; }
+	virtual bool IsValidOnMaster() const override { return true; }
+
+public:
 	/** Optional Game Mode to override the map's default game mode with. This can be useful if the game's normal mode displays UI elements or loading screens that you don't want captured. */
 	UPROPERTY(EditAnywhere, BlueprintReadWrite, Category = "Movie Pipeline")
 	TSubclassOf<AGameModeBase> GameModeOverride;
diff --git a/Engine/Plugins/MovieScene/MovieRenderPipeline/Source/MovieRenderPipelineSettings/Public/MoviePipelineRayTracingSetting.h b/Engine/Plugins/MovieScene/MovieRenderPipeline/Source/MovieRenderPipelineSettings/Public/MoviePipelineRayTracingSetting.h
new file mode 100644
--- /dev/null
+++ b/Engine/Plugins/MovieScene/MovieRenderPipeline/Source/MovieRenderPipelineSettings/Public/MoviePipelineRayTracingSetting.h
@@ -0,0 +1,35 @@
+// Copyright 1998-2019 Epic Games, Inc. All Rights Reserved.
+#pragma once
+
+#include "MoviePipelineSetting.h"
+#include "LevelSequenceActor.h"
+#include "MoviePipelineRayTracingSetting.generated.h"
+
+UCLASS(BlueprintType)
+class UMoviePipelineRayTracingSetting : public UMoviePipelineSetting
+{
+	GENERATED_BODY()
+
+	UMoviePipelineRayTracingSetting()
+		: bEnableDenoisers(true)
+	{}
+public:
+	/**
+	* Enabling the ray tracing denoisers will reduce the number of samples required per frame,
+	* but can introduce artifacts (ghosting). If you disable this expect to need 64-128 samples
+	* per frame to overcome the lack of denoiser.
+	*/
+	UPROPERTY(EditAnywhere, BlueprintReadWrite, meta = (EditCondition="bUsePathTracer"), Category = "Movie Pipeline")
+	bool bEnableDenoisers;
+	
+	int32 AmbientOcclusionSPP;
+	int32 GlobalIlluminationSPP;
+	int32 ReflectionsSPP;
+	int32 ShadowSPP;
+		
+	bool bUsePathTracer;
+	/**
+	
+	UPROPERTY(Instanced, EditAnywhere, BlueprintReadWrite, Category="Movie Pipeline")
+	ULevelSequenceBurnInInitSettings* Settings;*/
+};
\ No newline at end of file
