Copy up from Dev-Editor [REDACTED_HANDLE]

diff --git a/Engine/Plugins/MovieScene/MatineeToLevelSequence/Source/MatineeToLevelSequence/Private/MatineeToLevelSequenceModule.cpp b/Engine/Plugins/MovieScene/MatineeToLevelSequence/Source/MatineeToLevelSequence/Private/MatineeToLevelSequenceModule.cpp
--- a/Engine/Plugins/MovieScene/MatineeToLevelSequence/Source/MatineeToLevelSequence/Private/MatineeToLevelSequenceModule.cpp
+++ b/Engine/Plugins/MovieScene/MatineeToLevelSequence/Source/MatineeToLevelSequence/Private/MatineeToLevelSequenceModule.cpp
@@ -12,7 +12,13 @@
 #include "AssetData.h"
 #include "IAssetTools.h"
 #include "AssetToolsModule.h"
+#include "AssetTypeActions_Base.h"
+#include "ContentBrowserModule.h"
+#include "ContentBrowserMenuContexts.h"
+#include "Misc/Paths.h"
 #include "Textures/SlateIcon.h"
+#include "ToolMenus.h"
+#include "ToolMenuSection.h"
 #include "Framework/Commands/UICommandList.h"
 #include "Framework/MultiBox/MultiBoxExtender.h"
 #include "Framework/MultiBox/MultiBoxBuilder.h"
@@ -20,6 +26,7 @@
 #include "UObject/UnrealType.h"
 #include "GameFramework/Character.h"
 #include "Camera/CameraActor.h"
+#include "Camera/CameraAnim.h"
 #include "Particles/ParticleSystemComponent.h"
 #include "Components/LightComponentBase.h"
 #include "Engine/Light.h"
@@ -49,7 +56,6 @@
 #include "Matinee/InterpTrackMoveAxis.h"
 #include "MatineeUtils.h"
 #include "Editor.h"
-
 #include "LevelEditor.h"
 #include "MovieScene.h"
 #include "Evaluation/MovieSceneEvaluationTemplateInstance.h"
@@ -65,12 +71,14 @@
 #include "Tracks/MovieScene3DTransformTrack.h"
 #include "Tracks/MovieSceneParticleTrack.h"
 #include "Tracks/MovieSceneSkeletalAnimationTrack.h"
+#include "Tracks/MovieSceneSpawnTrack.h"
 #include "Tracks/MovieSceneFadeTrack.h"
 #include "Tracks/MovieSceneSlomoTrack.h"
 #include "Tracks/MovieSceneCameraCutTrack.h"
 #include "Tracks/MovieSceneEventTrack.h"
 #include "Tracks/MovieSceneVectorTrack.h"
 #include "Tracks/MovieSceneVisibilityTrack.h"
+#include "Sections/MovieSceneSpawnSection.h"
 #include "Animation/SkeletalMeshActor.h"
 #include "Dialogs/Dialogs.h"
 #include "Framework/Notifications/NotificationManager.h"
@@ -81,6 +89,7 @@
 #include "Subsystems/AssetEditorSubsystem.h"
 #include "EngineAnalytics.h"
 #include "Interfaces/IAnalyticsProvider.h"
+#include "CameraAnimationSequence.h"
 
 #define LOCTEXT_NAMESPACE "MatineeToLevelSequence"
 
@@ -103,14 +112,16 @@ public:
 		}
 		
 		RegisterMenuExtensions();
+		RegisterAssetTools();
 	}
 	
 	virtual void ShutdownModule() override
 	{
+		UnregisterAssetTools();
 		UnregisterMenuExtensions();
 	}
 
- 	FDelegateHandle RegisterTrackConverterForMatineeClass(TSubclassOf<UInterpTrack> InterpTrackClass, FOnConvertMatineeTrack OnConvertMatineeTrack)
+ 	FDelegateHandle RegisterTrackConverterForMatineeClass(TSubclassOf<UInterpTrack> InterpTrackClass, FOnConvertMatineeTrack OnConvertMatineeTrack) override
 	{
 		if (ExtendedInterpConverters.Contains(InterpTrackClass))
 		{
@@ -121,7 +132,7 @@ public:
 		return ExtendedInterpConverters.Add(InterpTrackClass, OnConvertMatineeTrack).GetHandle();
 	}
  	
-	void UnregisterTrackConverterForMatineeClass(FDelegateHandle RemoveDelegate)
+	void UnregisterTrackConverterForMatineeClass(FDelegateHandle RemoveDelegate) override
 	{
 		for (auto InterpConverter : ExtendedInterpConverters)
 		{
@@ -148,7 +159,21 @@ protected:
 		LevelEditorExtenderDelegateHandle = MenuExtenders.Last().GetHandle();
 	}
 
-protected:
+	/** Register new asset tools actions on the content browser. */
+	void RegisterAssetTools()
+	{
+		FToolMenuOwnerScoped OnwerScoped(this);
+
+		UToolMenu* Menu = UToolMenus::Get()->ExtendMenu("ContentBrowser.AssetContextMenu.CameraAnim");
+		FToolMenuSection& MenuSection = Menu->AddSection("ConversionActions", LOCTEXT("CameraAnim_ConversionSection", "Conversion"), FToolMenuInsert(FName("CommonAssetActions"), EToolMenuInsertType::Before));
+		MenuSection.AddMenuEntry(
+				"CameraAnim_ConvertToSequence",
+				LOCTEXT("CameraAnim_ConvertToSequence", "Convert to Sequence"),
+				LOCTEXT("CameraAnim_ConvertToSequence_Tooltip", "Converts the CameraAnim asset to a CameraAnimationSequence, to be used with the Sequencer."),
+				FSlateIcon(),
+				FToolMenuExecuteAction(FToolMenuExecuteAction::CreateRaw(this, &FMatineeToLevelSequenceModule::OnExecuteConvertCameraAnim))
+				);
+	}
 
 	/** Unregisters menu extensions for the level editor toolbar. */
 	void UnregisterMenuExtensions()
@@ -163,6 +188,12 @@ protected:
 		}
 	}
 
+	/** Unregisters content browser extensions. */
+	void UnregisterAssetTools()
+	{
+		UToolMenus::UnregisterOwner(this);
+	}
+
 	TSharedRef<FExtender> ExtendLevelViewportContextMenu(const TSharedRef<FUICommandList> CommandList, const TArray<AActor*> SelectedActors)
 	{
 		TSharedRef<FExtender> Extender(new FExtender());
@@ -237,6 +268,102 @@ protected:
 		return true;
 	}
 
+	void OnExecuteConvertCameraAnim(const FToolMenuContext& MenuContext)
+	{
+		UContentBrowserAssetContextMenuContext* Context = MenuContext.FindContext<UContentBrowserAssetContextMenuContext>();
+		if (Context == nullptr)
+		{
+			return;
+		}
+
+		// Get the assets to convert.
+		TArray<UCameraAnim*> CameraAnimsToConvert;
+		for (TWeakObjectPtr<UObject> SelectedObject : Context->SelectedObjects)
+		{
+			if (UCameraAnim* CameraAnimToConvert = CastChecked<UCameraAnim>(SelectedObject.Get(), ECastCheckedType::NullAllowed))
+			{
+				CameraAnimsToConvert.Add(CameraAnimToConvert);
+			}
+		}
+		if (CameraAnimsToConvert.Num() == 0)
+		{
+			return;
+		}
+
+		// Find the factory class.
+		UFactory* CameraAnimationSequenceFactoryNew = FindFactoryForClass(UCameraAnimationSequence::StaticClass());
+		ensure(CameraAnimationSequenceFactoryNew != nullptr);
+
+		// Convert all selected camera anims.
+		int32 NumWarnings = 0;
+		bool bConvertSuccess = false;
+		IAssetTools& AssetTools = FModuleManager::GetModuleChecked<FAssetToolsModule>("AssetTools").Get();
+		for (UCameraAnim* CameraAnimToConvert : CameraAnimsToConvert)
+		{
+			bConvertSuccess = ConvertSingleCameraAnimToTemplateSequence(
+					CameraAnimToConvert, AssetTools, CameraAnimationSequenceFactoryNew, NumWarnings) 
+				|| bConvertSuccess;
+		}
+
+		if (bConvertSuccess)
+		{
+			FText NotificationText = FText::Format(
+					LOCTEXT("CameraAnim_ConvertToSequence_Notification", "Converted {0} assets with {1} warnings"),
+					FText::AsNumber(CameraAnimsToConvert.Num()), FText::AsNumber(NumWarnings));
+			FNotificationInfo NotificationInfo(NotificationText);
+			NotificationInfo.ExpireDuration = 5.f;
+			NotificationInfo.Hyperlink = FSimpleDelegate::CreateStatic([](){ FGlobalTabmanager::Get()->InvokeTab(FName("OutputLog")); });
+			NotificationInfo.HyperlinkText = LOCTEXT("ShowMessageLogHyperlink", "Show Output Log");
+			FSlateNotificationManager::Get().AddNotification(NotificationInfo);
+		}
+	}
+
+	bool ConvertSingleCameraAnimToTemplateSequence(UCameraAnim* CameraAnimToConvert, IAssetTools& AssetTools, UFactory* CameraAnimationSequenceFactoryNew, int32& NumWarnings)
+	{
+		// Ask user for the new asset's name and folder.
+		UPackage* AssetPackage = CameraAnimToConvert->GetOutermost();
+		FString NewCameraAnimSequenceName = CameraAnimToConvert->GetName() + FString("Sequence");
+		FString NewCameraAnimSequencePath = FPaths::GetPath(AssetPackage->GetName());
+
+		UObject* NewAsset = AssetTools.CreateAssetWithDialog(NewCameraAnimSequenceName, NewCameraAnimSequencePath, UCameraAnimationSequence::StaticClass(), CameraAnimationSequenceFactoryNew);
+		if (NewAsset == nullptr)
+		{
+			return false;
+		}
+
+		// Create the new sequence.
+		UCameraAnimationSequence* NewSequence = Cast<UCameraAnimationSequence>(NewAsset);
+		NewSequence->BoundActorClass = ACameraActor::StaticClass();
+
+		UMovieScene* NewMovieScene = NewSequence->GetMovieScene();
+
+		// Add the spawnable for the camera.
+		ACameraActor* CameraTemplate = NewObject<ACameraActor>(NewMovieScene, ACameraActor::StaticClass());
+		FGuid SpawnableGuid = NewMovieScene->AddSpawnable("CameraActor", *CameraTemplate);
+		
+		// Set sequence length.
+		const int32 LengthInFrames = (CameraAnimToConvert->AnimLength * NewMovieScene->GetTickResolution()).FrameNumber.Value;
+		NewMovieScene->SetPlaybackRange(FFrameNumber(0), LengthInFrames + 1);
+
+		// Add spawning track for the camera.
+		UMovieSceneSpawnTrack* NewSpawnTrack = NewMovieScene->AddTrack<UMovieSceneSpawnTrack>(SpawnableGuid);
+		UMovieSceneSpawnSection* NewSpawnSection = CastChecked<UMovieSceneSpawnSection>(NewSpawnTrack->CreateNewSection());
+		NewSpawnSection->GetChannel().SetDefault(true);
+		NewSpawnSection->SetStartFrame(TRangeBound<FFrameNumber>());
+		NewSpawnSection->SetEndFrame(TRangeBound<FFrameNumber>());
+		NewSpawnTrack->AddSection(*NewSpawnSection);
+
+		// Add camera animation data.
+		if (CameraAnimToConvert->CameraInterpGroup != nullptr)
+		{
+			ConvertInterpGroup(
+					CameraAnimToConvert->CameraInterpGroup, SpawnableGuid, nullptr, 
+					NewSequence, NewMovieScene, NumWarnings);
+		}
+
+		return true;
+	}
+
 	/** Callback for converting a matinee to a level sequence asset. */
 	void OnConvertMatineeToLevelSequence(TArray<TWeakObjectPtr<AActor> > ActorsToConvert)
 	{
@@ -304,6 +431,24 @@ protected:
 		}
 	}
 
+	/** Find the factory that can create instances of the given class */
+	static UFactory* FindFactoryForClass(UClass* InClass)
+	{
+		for (TObjectIterator<UClass> It ; It ; ++It)
+		{
+			UClass* CurrentClass = *It;
+			if (CurrentClass->IsChildOf(UFactory::StaticClass()) && !(CurrentClass->HasAnyClassFlags(CLASS_Abstract)))
+			{
+				UFactory* Factory = Cast<UFactory>(CurrentClass->GetDefaultObject());
+				if (Factory->CanCreateNew() && Factory->ImportPriority >= 0 && Factory->SupportedClass == InClass)
+				{
+					return Factory;
+				}
+			}
+		}
+		return nullptr;
+	}
+
 	/** Find or add a folder for the given actor **/
 	static UMovieSceneFolder* FindOrAddFolder(UMovieScene* MovieScene, FName FolderName)
 	{
@@ -366,7 +511,7 @@ protected:
 
 	/** Add property to possessable node **/
 	template <typename T>
-	static T* AddPropertyTrack(FName InPropertyName, AActor* InActor, const FGuid& PossessableGuid, IMovieScenePlayer& Player, UMovieSceneSequence* NewSequence, UMovieScene* NewMovieScene, int32& NumWarnings, TMap<UObject*, FGuid>& BoundObjectsToGuids)
+	static T* AddPropertyTrack(FName InPropertyName, AActor* InActor, const FGuid& PossessableGuid, UMovieSceneSequence* NewSequence, UMovieScene* NewMovieScene, int32& NumWarnings, TMap<UObject*, FGuid>& BoundObjectsToGuids)
 	{
 		T* PropertyTrack = nullptr;
 
@@ -549,10 +694,8 @@ protected:
 	}
 
 	/** Convert an interp group */
-	void ConvertInterpGroup(UInterpGroup* Group, AActor* GroupActor, IMovieScenePlayer& Player, UMovieSceneSequence* NewSequence, UMovieScene* NewMovieScene, int32& NumWarnings)
+	void ConvertInterpGroup(UInterpGroup* Group, AActor* GroupActor, UMovieSceneSequence* NewSequence, UMovieScene* NewMovieScene, int32& NumWarnings)
 	{
-		TMap<UObject*, FGuid> BoundObjectsToGuids;
-
 		FGuid PossessableGuid;
 
 		// Bind the group actor as a possessable						
@@ -565,6 +708,13 @@ protected:
 			FindOrAddFolder(NewMovieScene, GroupActor, PossessableGuid);
 		}
 
+		ConvertInterpGroup(Group, PossessableGuid, GroupActor, NewSequence, NewMovieScene, NumWarnings);
+	}
+
+	void ConvertInterpGroup(UInterpGroup* Group, FGuid ObjectBindingGuid, AActor* GroupActor, UMovieSceneSequence* NewSequence, UMovieScene* NewMovieScene, int32& NumWarnings)
+	{
+		TMap<UObject*, FGuid> BoundObjectsToGuids;
+
 		for (int32 j=0; j<Group->InterpTracks.Num(); ++j)
 		{
 			UInterpTrack* Track = Group->InterpTracks[j];
@@ -576,7 +726,7 @@ protected:
 			// Handle each track class
 			if (ExtendedInterpConverters.Find(Track->GetClass()))
 			{
-				ExtendedInterpConverters.Find(Track->GetClass())->Execute(Track, PossessableGuid, NewMovieScene);
+				ExtendedInterpConverters.Find(Track->GetClass())->Execute(Track, ObjectBindingGuid, NewMovieScene);
 			}
 			else if (Track->IsA(UInterpTrackMove::StaticClass()))
 			{
@@ -600,19 +750,19 @@ protected:
 					}
 				}
 
-				if ( bHasKeyframes && PossessableGuid.IsValid())
+				if ( bHasKeyframes && ObjectBindingGuid.IsValid())
 				{
 					FVector DefaultScale = GroupActor != nullptr ? GroupActor->GetActorScale() : FVector(1.f);
-					UMovieScene3DTransformTrack* TransformTrack = NewMovieScene->AddTrack<UMovieScene3DTransformTrack>(PossessableGuid);								
+					UMovieScene3DTransformTrack* TransformTrack = NewMovieScene->AddTrack<UMovieScene3DTransformTrack>(ObjectBindingGuid);								
 					FMatineeImportTools::CopyInterpMoveTrack(MatineeMoveTrack, TransformTrack, DefaultScale);
 				}
 			}
 			else if (Track->IsA(UInterpTrackAnimControl::StaticClass()))
 			{
 				UInterpTrackAnimControl* MatineeAnimControlTrack = StaticCast<UInterpTrackAnimControl*>(Track);
-				if (MatineeAnimControlTrack->GetNumKeyframes() != 0 && PossessableGuid.IsValid())
+				if (MatineeAnimControlTrack->GetNumKeyframes() != 0 && ObjectBindingGuid.IsValid())
 				{
-					UMovieSceneSkeletalAnimationTrack* SkeletalAnimationTrack = NewMovieScene->AddTrack<UMovieSceneSkeletalAnimationTrack>(PossessableGuid);	
+					UMovieSceneSkeletalAnimationTrack* SkeletalAnimationTrack = NewMovieScene->AddTrack<UMovieSceneSkeletalAnimationTrack>(ObjectBindingGuid);	
 					FFrameNumber EndPlaybackRange = MovieScene::DiscreteExclusiveUpper(NewMovieScene->GetPlaybackRange());
 					FMatineeImportTools::CopyInterpAnimControlTrack(MatineeAnimControlTrack, SkeletalAnimationTrack, EndPlaybackRange);
 				}
@@ -620,9 +770,9 @@ protected:
 			else if (Track->IsA(UInterpTrackToggle::StaticClass()))
 			{
 				UInterpTrackToggle* MatineeParticleTrack = StaticCast<UInterpTrackToggle*>(Track);
-				if (MatineeParticleTrack->GetNumKeyframes() != 0 && PossessableGuid.IsValid())
+				if (MatineeParticleTrack->GetNumKeyframes() != 0 && ObjectBindingGuid.IsValid())
 				{
-					UMovieSceneParticleTrack* ParticleTrack = NewMovieScene->AddTrack<UMovieSceneParticleTrack>(PossessableGuid);	
+					UMovieSceneParticleTrack* ParticleTrack = NewMovieScene->AddTrack<UMovieSceneParticleTrack>(ObjectBindingGuid);	
 					FMatineeImportTools::CopyInterpParticleTrack(MatineeParticleTrack, ParticleTrack);
 				}
 			}
@@ -657,9 +807,9 @@ protected:
 			else if (Track->IsA(UInterpTrackBoolProp::StaticClass()))
 			{
 				UInterpTrackBoolProp* MatineeBoolTrack = StaticCast<UInterpTrackBoolProp*>(Track);
-				if (MatineeBoolTrack->GetNumKeyframes() != 0 && GroupActor && PossessableGuid.IsValid())
+				if (MatineeBoolTrack->GetNumKeyframes() != 0 && ObjectBindingGuid.IsValid())
 				{
-					UMovieSceneBoolTrack* BoolTrack = AddPropertyTrack<UMovieSceneBoolTrack>(MatineeBoolTrack->PropertyName, GroupActor, PossessableGuid, Player, NewSequence, NewMovieScene, NumWarnings, BoundObjectsToGuids);
+					UMovieSceneBoolTrack* BoolTrack = AddPropertyTrack<UMovieSceneBoolTrack>(MatineeBoolTrack->PropertyName, GroupActor, ObjectBindingGuid, NewSequence, NewMovieScene, NumWarnings, BoundObjectsToGuids);
 					if (BoolTrack)
 					{
 						FMatineeImportTools::CopyInterpBoolTrack(MatineeBoolTrack, BoolTrack);
@@ -669,9 +819,9 @@ protected:
 			else if (Track->IsA(UInterpTrackFloatProp::StaticClass()))
 			{
 				UInterpTrackFloatProp* MatineeFloatTrack = StaticCast<UInterpTrackFloatProp*>(Track);
-				if (MatineeFloatTrack->GetNumKeyframes() != 0 && GroupActor && PossessableGuid.IsValid())
+				if (MatineeFloatTrack->GetNumKeyframes() != 0 && ObjectBindingGuid.IsValid())
 				{
-					UMovieSceneFloatTrack* FloatTrack = AddPropertyTrack<UMovieSceneFloatTrack>(MatineeFloatTrack->PropertyName, GroupActor, PossessableGuid, Player, NewSequence, NewMovieScene, NumWarnings, BoundObjectsToGuids);
+					UMovieSceneFloatTrack* FloatTrack = AddPropertyTrack<UMovieSceneFloatTrack>(MatineeFloatTrack->PropertyName, GroupActor, ObjectBindingGuid, NewSequence, NewMovieScene, NumWarnings, BoundObjectsToGuids);
 					if (FloatTrack)
 					{
 						FMatineeImportTools::CopyInterpFloatTrack(MatineeFloatTrack, FloatTrack);
@@ -681,9 +831,9 @@ protected:
 			else if (Track->IsA(UInterpTrackFloatMaterialParam::StaticClass()))
 			{
 				UInterpTrackFloatMaterialParam* MatineeMaterialParamTrack = StaticCast<UInterpTrackFloatMaterialParam*>(Track);
-				if (MatineeMaterialParamTrack->GetNumKeyframes() != 0 && GroupActor && PossessableGuid.IsValid())
+				if (MatineeMaterialParamTrack->GetNumKeyframes() != 0 && ObjectBindingGuid.IsValid())
 				{
-					FGuid ComponentGuid = FindComponentGuid(GroupActor, NewSequence, NewMovieScene, PossessableGuid);
+					FGuid ComponentGuid = FindComponentGuid(GroupActor, NewSequence, NewMovieScene, ObjectBindingGuid);
 
 					if (ComponentGuid == FGuid())
 					{
@@ -697,9 +847,9 @@ protected:
 			else if (Track->IsA(UInterpTrackVectorMaterialParam::StaticClass()))
 			{
 				UInterpTrackVectorMaterialParam* MatineeMaterialParamTrack = StaticCast<UInterpTrackVectorMaterialParam*>(Track);
-				if (MatineeMaterialParamTrack->GetNumKeyframes() != 0 && GroupActor && PossessableGuid.IsValid())
+				if (MatineeMaterialParamTrack->GetNumKeyframes() != 0 && ObjectBindingGuid.IsValid())
 				{
-					FGuid ComponentGuid = FindComponentGuid(GroupActor, NewSequence, NewMovieScene, PossessableGuid);
+					FGuid ComponentGuid = FindComponentGuid(GroupActor, NewSequence, NewMovieScene, ObjectBindingGuid);
 
 					if (ComponentGuid == FGuid())
 					{
@@ -713,9 +863,9 @@ protected:
 			else if (Track->IsA(UInterpTrackVectorProp::StaticClass()))
 			{
 				UInterpTrackVectorProp* MatineeVectorTrack = StaticCast<UInterpTrackVectorProp*>(Track);
-				if (MatineeVectorTrack->GetNumKeyframes() != 0 && GroupActor && PossessableGuid.IsValid())
+				if (MatineeVectorTrack->GetNumKeyframes() != 0 && ObjectBindingGuid.IsValid())
 				{
-					UMovieSceneVectorTrack* VectorTrack = AddPropertyTrack<UMovieSceneVectorTrack>(MatineeVectorTrack->PropertyName, GroupActor, PossessableGuid, Player, NewSequence, NewMovieScene, NumWarnings, BoundObjectsToGuids);
+					UMovieSceneVectorTrack* VectorTrack = AddPropertyTrack<UMovieSceneVectorTrack>(MatineeVectorTrack->PropertyName, GroupActor, ObjectBindingGuid, NewSequence, NewMovieScene, NumWarnings, BoundObjectsToGuids);
 					if (VectorTrack)
 					{
 						VectorTrack->SetNumChannelsUsed(3);
@@ -726,9 +876,9 @@ protected:
 			else if (Track->IsA(UInterpTrackColorProp::StaticClass()))
 			{
 				UInterpTrackColorProp* MatineeColorTrack = StaticCast<UInterpTrackColorProp*>(Track);
-				if (MatineeColorTrack->GetNumKeyframes() != 0 && GroupActor && PossessableGuid.IsValid())
+				if (MatineeColorTrack->GetNumKeyframes() != 0 && ObjectBindingGuid.IsValid())
 				{
-					UMovieSceneColorTrack* ColorTrack = AddPropertyTrack<UMovieSceneColorTrack>(MatineeColorTrack->PropertyName, GroupActor, PossessableGuid, Player, NewSequence, NewMovieScene, NumWarnings, BoundObjectsToGuids);
+					UMovieSceneColorTrack* ColorTrack = AddPropertyTrack<UMovieSceneColorTrack>(MatineeColorTrack->PropertyName, GroupActor, ObjectBindingGuid, NewSequence, NewMovieScene, NumWarnings, BoundObjectsToGuids);
 					if (ColorTrack)
 					{
 						FMatineeImportTools::CopyInterpColorTrack(MatineeColorTrack, ColorTrack);
@@ -738,9 +888,9 @@ protected:
 			else if (Track->IsA(UInterpTrackLinearColorProp::StaticClass()))
 			{
 				UInterpTrackLinearColorProp* MatineeLinearColorTrack = StaticCast<UInterpTrackLinearColorProp*>(Track);
-				if (MatineeLinearColorTrack->GetNumKeyframes() != 0 && GroupActor && PossessableGuid.IsValid())
+				if (MatineeLinearColorTrack->GetNumKeyframes() != 0 && ObjectBindingGuid.IsValid())
 				{
-					UMovieSceneColorTrack* ColorTrack = AddPropertyTrack<UMovieSceneColorTrack>(MatineeLinearColorTrack->PropertyName, GroupActor, PossessableGuid, Player, NewSequence, NewMovieScene, NumWarnings, BoundObjectsToGuids);
+					UMovieSceneColorTrack* ColorTrack = AddPropertyTrack<UMovieSceneColorTrack>(MatineeLinearColorTrack->PropertyName, GroupActor, ObjectBindingGuid, NewSequence, NewMovieScene, NumWarnings, BoundObjectsToGuids);
 					if (ColorTrack)
 					{
 						FMatineeImportTools::CopyInterpLinearColorTrack(MatineeLinearColorTrack, ColorTrack);
@@ -750,9 +900,9 @@ protected:
 			else if (Track->IsA(UInterpTrackVisibility::StaticClass()))
 			{
 				UInterpTrackVisibility* MatineeVisibilityTrack = StaticCast<UInterpTrackVisibility*>(Track);
-				if (MatineeVisibilityTrack->GetNumKeyframes() != 0 && GroupActor && PossessableGuid.IsValid())
+				if (MatineeVisibilityTrack->GetNumKeyframes() != 0 && ObjectBindingGuid.IsValid())
 				{
-					UMovieSceneVisibilityTrack* VisibilityTrack = NewMovieScene->AddTrack<UMovieSceneVisibilityTrack>(PossessableGuid);	
+					UMovieSceneVisibilityTrack* VisibilityTrack = NewMovieScene->AddTrack<UMovieSceneVisibilityTrack>(ObjectBindingGuid);	
 					if (VisibilityTrack)
 					{
 						VisibilityTrack->SetPropertyNameAndPath(TEXT("bHidden"), GroupActor->GetPathName() + TEXT(".bHidden"));
@@ -796,20 +946,11 @@ protected:
 		IAssetTools& AssetTools = FModuleManager::GetModuleChecked<FAssetToolsModule>("AssetTools").Get();
 
 		UObject* NewAsset = nullptr;
-		for (TObjectIterator<UClass> It ; It ; ++It)
+		UFactory* Factory = FindFactoryForClass(ULevelSequence::StaticClass());
+		if (Factory != nullptr)
 		{
-			UClass* CurrentClass = *It;
-			if (CurrentClass->IsChildOf(UFactory::StaticClass()) && !(CurrentClass->HasAnyClassFlags(CLASS_Abstract)))
-			{
-				UFactory* Factory = Cast<UFactory>(CurrentClass->GetDefaultObject());
-				if (Factory->CanCreateNew() && Factory->ImportPriority >= 0 && Factory->SupportedClass == ULevelSequence::StaticClass())
-				{
-					NewAsset = AssetTools.CreateAssetWithDialog(NewLevelSequenceAssetName, NewLevelSequenceAssetPath, ULevelSequence::StaticClass(), Factory);
-					break;
-				}
-			}
+			NewAsset = AssetTools.CreateAssetWithDialog(NewLevelSequenceAssetName, NewLevelSequenceAssetPath, ULevelSequence::StaticClass(), Factory);
 		}
-
 		if (!NewAsset)
 		{
 			return nullptr;
@@ -864,7 +1005,7 @@ protected:
 				UInterpGroupInst* GrInst = MatineeActor->GroupInst[i];
 				UInterpGroup* Group = GrInst->Group;
 				AActor* GroupActor = GrInst->GetGroupActor();
-				ConvertInterpGroup(Group, GroupActor, TemporaryPlayer, NewSequence, NewMovieScene, NumWarnings);
+				ConvertInterpGroup(Group, GroupActor, NewSequence, NewMovieScene, NumWarnings);
 			}
 
 			// Force an evaluation so that bound objects will 
diff --git a/Engine/Plugins/MovieScene/MovieRenderPipeline/MovieRenderPipeline.uplugin b/Engine/Plugins/MovieScene/MovieRenderPipeline/MovieRenderPipeline.uplugin
new file mode 100644
--- /dev/null
+++ b/Engine/Plugins/MovieScene/MovieRenderPipeline/MovieRenderPipeline.uplugin
@@ -0,0 +1,45 @@
+{
+	"FileVersion": 3,
+	"Version": 1,
+	"VersionName": "1.0",
+	"FriendlyName": "Movie Render Pipeline",
+	"Description": "Advanced movie rendering pipeline for use in creating rendered cinematics or other multi-media creation.",
+	"Category": "Scripting",
+	"CreatedBy": "Epic Games Inc.",
+	"CreatedByURL": "",
+	"DocsURL": "",
+	"MarketplaceURL": "",
+	"SupportURL": "",
+	"CanContainContent": true,
+	"IsBetaVersion": true,
+	"Installed": false,
+	"EnabledByDefault": true,
+	"Plugins": [
+		{
+			"Name": "EditorScriptingUtilities",
+			"Enabled": true
+		}
+	],
+	"Modules": [
+		{
+			"Name": "MovieRenderPipelineCore",
+			"Type": "UncookedOnly",
+			"LoadingPhase": "Default"
+		},
+		{
+			"Name": "MovieRenderPipelineSettings",
+			"Type": "UncookedOnly",
+			"LoadingPhase": "Default"
+		}, 
+		{
+			"Name" : "MovieRenderPipelineRenderPasses",
+			"Type": "UncookedOnly",
+			"LoadingPhase": "Default"
+		},
+		{
+			"Name": "MovieRenderPipelineEditor",
+			"Type": "Editor",
+			"LoadingPhase": "Default"
+		}
+	]
+}
\ No newline at end of file
diff --git a/Engine/Plugins/MovieScene/MovieRenderPipeline/Source/MovieRenderPipelineCore/Private/MoviePipeline.cpp b/Engine/Plugins/MovieScene/MovieRenderPipeline/Source/MovieRenderPipelineCore/Private/MoviePipeline.cpp
new file mode 100644
--- /dev/null
+++ b/Engine/Plugins/MovieScene/MovieRenderPipeline/Source/MovieRenderPipelineCore/Private/MoviePipeline.cpp
@@ -0,0 +1,1043 @@
+// Copyright 1998-2019 Epic Games, Inc. All Rights Reserved.
+#include "MoviePipeline.h"
+#include "MovieRenderPipelineCoreModule.h"
+#include "LevelSequence.h"
+#include "Tracks/MovieSceneSubTrack.h"
+#include "Tracks/MovieSceneCameraCutTrack.h"
+#include "LevelSequenceActor.h"
+#include "EngineUtils.h"
+#include "MovieSceneSection.h"
+#include "Sections/MovieSceneCameraCutSection.h"
+#include "Sections/MovieSceneCinematicShotSection.h"
+#include "Tracks/MovieSceneCinematicShotTrack.h"
+#include "CanvasTypes.h"
+#include "AudioDeviceManager.h"
+#include "MovieSceneTimeHelpers.h"
+#include "Engine/World.h"
+#include "GameFramework/PlayerController.h"
+#include "MovieRenderDebugWidget.h"
+#include "MoviePipelineShotConfig.h"
+#include "MovieRenderPipelineDataTypes.h"
+#include "MoviePipelineRenderPass.h"
+#include "MoviePipelineOutputBase.h"
+#include "ShaderCompiler.h"
+#include "ImageWriteStream.h"
+#include "MoviePipelineAccumulationSetting.h"
+#include "MoviePipelineOutputBuilder.h"
+#include "DistanceFieldAtlas.h"
+#include "UObject/SoftObjectPath.h"
+#include "MoviePipelineOutputSetting.h"
+#include "MoviePipelineMasterConfig.h"
+#include "MoviePipelineOutputSetting.h"
+#include "MoviePipelineBlueprintLibrary.h"
+#include "ImageWriteQueue.h"
+#include "MoviePipelineHighResSetting.h"
+#include "MoviePipelineCameraSetting.h"
+#include "MoviePipelineQueue.h"
+
+#define LOCTEXT_NAMESPACE "MoviePipeline"
+
+static TAutoConsoleVariable<int32> CVarMovieRenderPipelineFrameStepper(
+	TEXT("MovieRenderPipeline.FrameStepDebug"),
+	-1,
+	TEXT("How many frames should the Movie Render Pipeline produce before pausing. Set to zero on launch to stall at the first frame. Debug tool.\n")
+	TEXT("-1: Don't pause after each frame (default)\n")
+	TEXT("0: Process engine ticks but don't progress in the movie rendering pipeline.\n")
+	TEXT("1+: Run this many loops of the movie rendering pipeline before pausing again.\n"),
+	ECVF_Default);
+
+
+UMoviePipeline::UMoviePipeline()
+	: CustomTimeStep(nullptr)
+	, CachedPrevCustomTimeStep(nullptr)
+	, TargetSequence(nullptr)
+	, LevelSequenceActor(nullptr)
+	, PipelineState(EMovieRenderPipelineState::Uninitialized)
+	, CurrentShotIndex(-1)
+	, bHasRunBeginFrameOnce(false)
+	, bPauseAtEndOfFrame(false)
+	, AccumulatedTickSubFrameDeltas(0.f)
+{
+	CustomTimeStep = CreateDefaultSubobject<UMoviePipelineCustomTimeStep>("MoviePipelineCustomTimeStep");
+	CustomSequenceTimeController = MakeShared<FMoviePipelineTimeController>();
+	OutputBuilder = MakeShared<FMoviePipelineOutputMerger, ESPMode::ThreadSafe>(this);
+
+	ImageWriteQueue = &FModuleManager::Get().LoadModuleChecked<IImageWriteQueueModule>("ImageWriteQueue").GetWriteQueue();
+}
+
+
+void UMoviePipeline::ValidateSequenceAndSettings()
+{
+	// ToDo: 
+	// Warn for Blueprint Streaming Levels
+	// Warn for sections that aren't extended far enough (once handle frames are taken into account)
+	// Warn for not whole frame aligned sections
+
+	// Check to see if they're trying to output alpha and don't have the required project setting set.
+	{
+		IConsoleVariable* TonemapAlphaCVar = IConsoleManager::Get().FindConsoleVariable(TEXT("r.PostProcessing.PropagateAlpha"));
+		check(TonemapAlphaCVar);
+
+		TArray<UMoviePipelineOutputBase*> OutputSettings = GetPipelineMasterConfig()->FindSettings<UMoviePipelineOutputBase>();
+		bool bAnyOutputWantsAlpha = false;
+
+		for (const UMoviePipelineOutputBase* Output : OutputSettings)
+		{
+			bAnyOutputWantsAlpha |= Output->IsAlphaSupported();
+		}
+
+		if (bAnyOutputWantsAlpha && TonemapAlphaCVar->GetInt() == 0)
+		{
+			UE_LOG(LogMovieRenderPipeline, Warning, TEXT("An output requested Alpha Support but the required project setting is not enabled! Go to Project Settings > Rendering > Tonemapping > 'Enable Alpha Channel Support in Post Processing' and set it to 'Linear Color Space Only'."));
+		}
+	}
+}
+
+void UMoviePipeline::Initialize(UMoviePipelineExecutorJob* InJob)
+{
+	// This function is called after the PIE world has finished initializing, but before
+	// the PIE world is ticked for the first time. We'll end up waiting for the next tick
+	// for FCoreDelegateS::OnBeginFrame to get called to actually start processing.
+	UE_LOG(LogMovieRenderPipeline, Log, TEXT("[%d] Initializing overall Movie Pipeline"), GFrameCounter);
+
+	if (!ensureAlwaysMsgf(InJob, TEXT("MoviePipeline cannot be initialized with null job. Aborting.")))
+	{
+		OnMoviePipelineErroredDelegate.Broadcast(this, true, LOCTEXT("MissingJob", "Job was not specified, movie render is aborting."));
+
+		Shutdown();
+		return;
+	}
+	
+	if (!ensureAlwaysMsgf(InJob->GetConfiguration(), TEXT("MoviePipeline cannot be initialized with null configuration. Aborting.")))
+	{
+		OnMoviePipelineErroredDelegate.Broadcast(this, true, LOCTEXT("MissingConfiguration", "Job did not specify a configuration, movie render is aborting."));
+
+		Shutdown();
+		return;
+	}
+
+	if (!ensureAlwaysMsgf(PipelineState == EMovieRenderPipelineState::Uninitialized, TEXT("Pipeline cannot be reused. Create a new pipeline to execute a job.")))
+	{
+		OnMoviePipelineErroredDelegate.Broadcast(this, true, LOCTEXT("DontReusePipeline", "Attempted to reuse an existing Movie Pipeline. Initialize a new pipeline instead of reusing an existing one."));
+
+		Shutdown();
+		return;
+	}
+
+	// Ensure this object has the World as part of its Outer (so that it has context to spawn things)
+	if (!ensureAlwaysMsgf(GetWorld(), TEXT("Pipeline does not contain the world as an outer.")))
+	{
+		OnMoviePipelineErroredDelegate.Broadcast(this, true, LOCTEXT("MissingWorld", "Could not find World in the Outer Path for Pipeline. The world must be an outer to give the Pipeline enough context to spawn things."));
+
+		Shutdown();
+		return;
+	}
+
+	CurrentJob = InJob;
+
+	// When start up we want to override the engine's Custom Timestep with our own.
+	// This gives us the ability to completely control the engine tick/delta time before the frame
+	// is started so that we don't have to always be thinking of delta times one frame ahead.
+	CachedPrevCustomTimeStep = GEngine->GetCustomTimeStep();
+	GEngine->SetCustomTimeStep(CustomTimeStep);
+	
+	ULevelSequence* OriginalSequence = Cast<ULevelSequence>(InJob->Sequence.TryLoad());
+	if (!ensureAlwaysMsgf(OriginalSequence, TEXT("Failed to load Sequence Asset from specified path, aborting movie render! Attempted to load Path: %s"), *InJob->Sequence.ToString()))
+	{
+		OnMoviePipelineErroredDelegate.Broadcast(this, true, LOCTEXT("MissingSequence", "Could not load sequence asset, movie render is aborting. Check logs for additional details."));
+
+		Shutdown();
+		return;
+	}
+
+	// Duplicate the target level sequence. This way modifications don't need to be undone.
+	TargetSequence = Cast<ULevelSequence>(UMoviePipelineBlueprintLibrary::DuplicateSequence(this, OriginalSequence));
+
+	// Initialize all of our master config settings. Shot specific ones will be called for their appropriate shot.
+	for (UMoviePipelineSetting* Setting : GetPipelineMasterConfig()->GetSettings())
+	{
+		Setting->OnMoviePipelineInitialized(this);
+	}
+
+	// Allow master settings to modify the sequence. This can be useful when working with dynamic content, you might
+	// want to modify things in the sequence, or modify things in the world before rendering.
+	// @ToDo: ModifySequenceViaExtensions(TargetSequence);
+
+	// Now that we've post-processed it, we're going to run a validation pass on it. This will produce warnings
+	// for anything we can't fix that might be an issue - extending sections, etc. This should be const as this
+	// validation should re-use what was used in the UI.
+	// @ToDo: ValidateSequence(TargetSequence);
+
+	// Now that we've fixed up the sequence and validated it, we're going to build a list of shots that we need
+	// to produce in a simplified data structure. The simplified structure makes the flow/debugging easier.
+	ShotList = BuildShotListFromSequence(TargetSequence);
+
+	// Finally, we're going to create a Level Sequence Actor in the world that has its settings configured by us.
+	// Because this callback is at the end of startup (and before tick) we should be able to spawn the actor
+	// and give it a chance to tick once (where it should do nothing) before we start manually manipulating it.
+	InitializeLevelSequenceActor(OriginalSequence, TargetSequence);
+
+	// Register any additional engine callbacks needed.
+	{
+		// Called before the Custom Timestep is updated. This gives us a chance to calculate
+		// what we want the frame to look like and then cache that information so that the
+		// Custom Timestep doesn't have to perform its own logic.
+		FCoreDelegates::OnBeginFrame.AddUObject(this, &UMoviePipeline::OnEngineTickBeginFrame);
+		// Called at the end of the frame after everything has been ticked and rendered for the frame.
+		FCoreDelegates::OnEndFrame.AddUObject(this, &UMoviePipeline::OnEngineTickEndFrame);
+	}
+
+	// Construct a debug UI and bind it to this instance.
+	LoadDebugWidget();
+
+	TotalExpectedWork = CalculateExpectedOutputMetrics();
+
+	PipelineState = EMovieRenderPipelineState::ProducingFrames;
+	CurrentShotIndex = 0;
+	InitializationTime = FDateTime::UtcNow();
+
+	// Initialization is complete. This engine frame is a wash (because the tick started with a 
+	// delta time not generated by us) so we'll wait until the next engine frame to start rendering.
+	UE_LOG(LogMovieRenderPipeline, Log, TEXT("[%d] Initialized Movie Pipeline with [%d Camera Cuts]. Expected Total Work: %s"),
+		GFrameCounter, TotalExpectedWork.NumCameraCuts, *TotalExpectedWork.ToDisplayString());
+}
+
+void UMoviePipeline::Shutdown()
+{
+	// Pipeline should not be shut down if already shut down. You probably have a delegate
+	// loop where we shutdown because finish and you responded by requesting we shutdown.
+	ensure(PipelineState != EMovieRenderPipelineState::Shutdown);
+
+	// Uninitialize our master config settings.
+	for (UMoviePipelineSetting* Setting : GetPipelineMasterConfig()->GetSettings())
+	{
+		Setting->OnMoviePipelineShutdown(this);
+	}
+
+	// Restore any custom Time Step that may have been set before.
+	GEngine->SetCustomTimeStep(CachedPrevCustomTimeStep);
+
+	// Ensure our delegates don't get called anymore as we're going to become null soon.
+	FCoreDelegates::OnBeginFrame.RemoveAll(this);
+	FCoreDelegates::OnEndFrame.RemoveAll(this);
+
+	// If we're not shutting down naturally (ie: we got canceled) then force the output
+	// containers to flush so that all in-flight write processes get finished. If we
+	// did shut down naturally, we should be in the export state.
+	if (PipelineState != EMovieRenderPipelineState::Export)
+	{
+		const bool bForceFinish = true;
+		TickFinalizeOutputContainers(bForceFinish);
+	}
+
+	if (DebugWidget)
+	{
+		DebugWidget->RemoveFromParent();
+		DebugWidget = nullptr;
+	}
+
+	for (UMoviePipelineOutputBase* Setting : GetPipelineMasterConfig()->GetOutputContainers())
+	{
+		Setting->OnPipelineFinished();
+	}
+
+	UE_LOG(LogMovieRenderPipeline, Log, TEXT("Movie Pipeline completed. Duration: %s"), *(FDateTime::UtcNow() - InitializationTime).ToString());
+	PipelineState = EMovieRenderPipelineState::Shutdown;
+}
+
+void UMoviePipeline::OnEngineTickBeginFrame()
+{
+	UE_LOG(LogMovieRenderPipeline, VeryVerbose, TEXT("OnEngineTickBeginFrame (Start) Engine Frame: [%d]"), GFrameCounter);
+	
+	// We should have a custom timestep set up by now.
+	check(CustomTimeStep);
+
+	switch (PipelineState)
+	{
+	case EMovieRenderPipelineState::Uninitialized:
+		// We shouldn't register this delegate until we're initialized.
+		check(false);
+		break;
+	case EMovieRenderPipelineState::ProducingFrames:
+		TickProducingFrames();
+		break;
+	case EMovieRenderPipelineState::Finalize:
+		TickFinalizeOutputContainers(false);
+		break;
+	case EMovieRenderPipelineState::Export:
+		// We support ticking export state as well to allow async long export steps.
+		TickPostFinalizeExport();
+		break;
+	}
+
+	bHasRunBeginFrameOnce = true;
+	UE_LOG(LogMovieRenderPipeline, VeryVerbose, TEXT("OnEngineTickBeginFrame (End) Engine Frame: [%d]"), GFrameCounter);
+}
+
+void UMoviePipeline::OnSequenceEvaluated(const UMovieSceneSequencePlayer& Player, FFrameTime CurrentTime, FFrameTime PreviousTime)
+{
+	// This callback exists for logging purposes. DO NOT HINGE LOGIC ON THIS CALLBACK
+	// because this may get called multiple times per frame and may be the result of
+	// a seek operation which is reverted before a frame is even rendered.
+	UE_LOG(LogMovieRenderPipeline, VeryVerbose, TEXT("[GFrameCounter: %d] Sequence Evaluated. CurrentTime: %s PreviousTime: %s"), GFrameCounter, *LexToString(CurrentTime), *LexToString(PreviousTime));
+}
+
+void UMoviePipeline::OnEngineTickEndFrame()
+{
+	// Unfortunately, since we can't control when our Initialization function is called
+	// we can end up in a situation where this callback is registered but the matching
+	// OnEngineTickBeginFrame() hasn't been called for that given engine tick. Instead of
+	// changing this registration to hang off of the end of the first OnEngineTickBeginFrame()
+	// we instead just early out here if that hasn't actually been called once. This decision
+	// is designed to minimize places where callbacks are registered and where flow changes.
+	if (!bHasRunBeginFrameOnce)
+	{
+		return;
+	}
+
+	// Early out if we're idling as we don't want to process a frame. This prevents us from
+	// overwriting render state when the engine is processing ticks but we're not trying to
+	// change the evaluation. 
+	if (IsDebugFrameStepIdling())
+	{
+		return;
+	}
+
+	// It is important that there is no early out that skips hitting this
+	// (Otherwise we don't pause on the frame we transition from step -> idle
+	// and the world plays even though the state is frozen).
+	if (bPauseAtEndOfFrame)
+	{
+		GetWorld()->GetFirstPlayerController()->SetPause(true);
+		bPauseAtEndOfFrame = false;
+	}
+
+	UE_LOG(LogMovieRenderPipeline, VeryVerbose, TEXT("OnEngineTickEndFrame (Start) Engine Frame: [%d]"), GFrameCounter);
+
+	RenderFrame();
+
+	UE_LOG(LogMovieRenderPipeline, VeryVerbose, TEXT("OnEngineTickEndFrame (End) Engine Frame: [%d]"), GFrameCounter);
+}
+
+bool UMoviePipeline::ProcessEndOfCameraCut(FMoviePipelineShotInfo &CurrentShot, FMoviePipelineCameraCutInfo &CurrentCameraCut)
+{
+	UE_LOG(LogMovieRenderPipeline, Log, TEXT("[%d] Finished processing Camera Cut [%d/%d] on Shot [%d] (%s) with stats [...]"), GFrameCounter,
+		CurrentShot.CurrentCameraCutIndex + 1, CurrentShot.CameraCuts.Num(),
+		CurrentShotIndex + 1, *CurrentShot.GetDisplayName());
+
+	CurrentCameraCut.State = EMovieRenderShotState::Finished;
+
+	// Compare our expected vs. actual results for logging.
+	if (CurrentCameraCut.CurrentWorkInfo != CurrentCameraCut.TotalWorkInfo)
+	{
+		UE_LOG(LogMovieRenderPipeline, Warning, TEXT("Mismatch in work done vs. expected work done.\nExpected: %s\nTotal: %s"),
+			*CurrentCameraCut.TotalWorkInfo.ToDisplayString(), *CurrentCameraCut.CurrentWorkInfo.ToDisplayString());
+	}
+
+	// We pause at the end too, just so that frames during finalize don't continue to trigger Sequence Eval messages.
+	LevelSequenceActor->GetSequencePlayer()->Pause();
+
+	bool bWasLastCameraCut = CurrentShot.SetNextShotActive();
+	if (bWasLastCameraCut)
+	{
+		TeardownShot(CurrentShot);
+
+		CurrentShotIndex++;
+
+		// Notify our containers that the current shot has ended.
+		for (UMoviePipelineOutputBase* Container : GetPipelineMasterConfig()->GetOutputContainers())
+		{
+			Container->OnShotFinished(CurrentShot);
+		}
+	}
+
+	// Check to see if this was the last shot in the Pipeline, otherwise on the next
+	// tick the new shot will be initialized and processed.
+	if (CurrentShotIndex >= ShotList.Num())
+	{
+		UE_LOG(LogMovieRenderPipeline, Log, TEXT("[%d] Finished rendering last shot. Moving to Finalize to finish writing items to disk."), GFrameCounter);
+		PipelineState = EMovieRenderPipelineState::Finalize;
+
+		// Unregister our OnEngineTickEndFrame delegate. We don't unregister BeginFrame as we need
+		// to continue to call it to allow ticking the Finalization stage.
+		FCoreDelegates::OnEndFrame.RemoveAll(this);
+
+		// Reset the Custom Timestep because we don't care how long the engine takes now
+		GEngine->SetCustomTimeStep(CachedPrevCustomTimeStep);
+
+		BeginFinalize();
+		return true;
+	}
+
+	return false;
+}
+
+void UMoviePipeline::BeginFinalize()
+{
+	// Notify all of our output containers that we have finished producing and
+	// submitting all frames to them and that they should start any async flushes.
+	for (UMoviePipelineOutputBase* Container : GetPipelineMasterConfig()->GetOutputContainers())
+	{
+		Container->BeginFinalize();
+	}
+}
+
+void UMoviePipeline::TickFinalizeOutputContainers(const bool bInForceFinish)
+{
+	// Tick all containers until they all report that they have finalized.
+	bool bAllContainsFinishedProcessing;
+
+	do
+	{
+		bAllContainsFinishedProcessing = true;
+
+		// Ask the containers if they're all done processing.
+		for (UMoviePipelineOutputBase* Container : GetPipelineMasterConfig()->GetOutputContainers())
+		{
+			bAllContainsFinishedProcessing &= Container->HasFinishedProcessing();
+		}
+
+		// If we aren't forcing a finish, early out after one loop to keep
+		// the editor/ui responsive.
+		if (!bInForceFinish || bAllContainsFinishedProcessing)
+		{
+			break;
+		}
+
+		// If they've reached here, they're forcing them to finish so we'll sleep for a touch to give
+		// everyone a chance to actually do work before asking them if they're done.
+		FPlatformProcess::Sleep(1.f);
+
+	} while (true);
+
+	// If an output container is still working, we'll early out to keep the UI responsive.
+	// If they've forced a finish this will have to be true before we can reach this block.
+	if (!bAllContainsFinishedProcessing)
+	{
+		return;
+	}
+
+	for (UMoviePipelineOutputBase* Container : GetPipelineMasterConfig()->GetOutputContainers())
+	{
+		// All containers have finished processing, final shutdown.
+		Container->Finalize();
+	}
+
+	// Move onto the Export stage and call the Export function directly. Most export steps are
+	// instantaneous and will complete here, but on the next tick we will check them all anyways
+	// to allow running long-running export processes (such as encoding).
+	PipelineState = EMovieRenderPipelineState::Export;
+}
+
+void UMoviePipeline::TickPostFinalizeExport()
+{
+	// This step assumes you have produced data and filled the data structures.
+	check(PipelineState == EMovieRenderPipelineState::Export);
+	UE_LOG(LogMovieRenderPipeline, Log, TEXT("[%d] PostFinalize Export (Start)."), GFrameCounter);
+
+	// ToDo: Loop through any extensions (such as XML export) and let them export using all of the
+	// data that was generated during this run such as containers, output names and lengths.
+
+
+
+	bool bAllExportsFinishedProcessing = true;
+
+	// Ask the containers if they're all done processing.
+	//for (UMoviePipelineOutput* Container : Config->OutputContainers)
+	//{
+	//	bAllContainsFinishedProcessing &= Container->HasFinishedProcessing();
+	//}
+
+	if (!bAllExportsFinishedProcessing)
+	{
+		return;
+	}
+
+
+	UE_LOG(LogMovieRenderPipeline, Log, TEXT("[%d] PostFinalize Export (End)."), GFrameCounter);
+
+	// Now that we have done the last step of the process we're going to call Shutdown on ourself
+	// to unregister all delegates, and then we'll broadcast to anyone listening that we're done.
+	// This allows an external controller to decide what to do, ie: start another one, shut down
+	// PIE or exit the game.
+	Shutdown();
+
+	UE_LOG(LogMovieRenderPipeline, Verbose, TEXT("[%d] Broadcasting Post-Finalize Finished event."), GFrameCounter);
+	OnMoviePipelineFinished().Broadcast(this);
+}
+
+bool UMoviePipelineCustomTimeStep::UpdateTimeStep(UEngine* /*InEngine*/)
+{
+	if (ensureMsgf(!FMath::IsNearlyZero(TimeCache.DeltaTime), TEXT("An incorrect or uninitialized time step was used! Delta Time of 0 isn't allowed.")))
+	{
+		FApp::UpdateLastTime();
+		FApp::SetDeltaTime(TimeCache.DeltaTime);
+		FApp::SetCurrentTime(FApp::GetCurrentTime() + FApp::GetDeltaTime());
+	}
+
+	// Clear our cached time to ensure we're always explicitly telling this what to do and never relying on the last set value.
+	// (This will cause the ensure above to check on the next tick if someone didn't reset our value.)
+	TimeCache = MoviePipeline::FFrameTimeStepCache();
+
+	// Return false so the engine doesn't run its own logic to overwrite FApp timings.
+	return false;
+}
+
+void UMoviePipeline::ModifySequenceViaExtensions(ULevelSequence* InSequence)
+{
+}
+
+FMoviePipelineWorkInfo UMoviePipeline::CalculateExpectedOutputMetrics()
+{
+	UE_LOG(LogMovieRenderPipeline, Log, TEXT("[%d] Calculating expected output metrics..."), GFrameCounter);
+
+	// We can calculate the expected amount of work that we will be doing, unless slow-mo/timescale is used. 
+	FMoviePipelineWorkInfo ExpectedWork;
+
+	for (const FMoviePipelineShotInfo& Shot : GetShotList())
+	{
+		ExpectedWork.NumCameraCuts += Shot.CameraCuts.Num();
+
+		int32 NumTiles = 0;
+		UMoviePipelineHighResSetting* HighResSettings = FindOrAddSetting<UMoviePipelineHighResSetting>(Shot);
+		if (HighResSettings)
+		{
+			NumTiles = HighResSettings->TileCount * HighResSettings->TileCount;
+		}
+
+		for (const FMoviePipelineCameraCutInfo& CameraCut : Shot.CameraCuts)
+		{
+			FMoviePipelineWorkInfo CameraCutWorkInfo = CameraCut.TotalWorkInfo;
+
+			ExpectedWork += CameraCut.TotalWorkInfo;
+
+			UE_LOG(LogMovieRenderPipeline, Log, TEXT("[%d] CameraCut: %d Expected Work: %s"), 
+				GFrameCounter, ExpectedWork.NumCameraCuts, *CameraCutWorkInfo.ToDisplayString());
+		}
+	}
+
+	return ExpectedWork;
+}
+
+void UMoviePipeline::InitializeLevelSequenceActor(ULevelSequence* OriginalLevelSequence, ULevelSequence* InSequenceToApply)
+{
+	// There is a reasonable chance that there exists a Level Sequence Actor in the world
+	// already set up to play this sequence. We don't want that however as it may have unknown
+	// settings on it.
+	for (auto It = TActorIterator<ALevelSequenceActor>(GetWorld()); It; ++It)
+	{
+		// Iterate through all of them in the event someone has multiple copies in the world on accident.
+		if (It->LevelSequence == OriginalLevelSequence)
+		{
+			// Found it!
+			ALevelSequenceActor* ExistingActor = *It;
+
+			//  Clear the sequence ref instead of killing actor to reduce likelyhood of game code throwing a NRE
+			// that may have expected this actor to exist.
+			if (ExistingActor->GetSequencePlayer())
+			{
+				ExistingActor->GetSequencePlayer()->Stop();
+			}
+
+			// Sequence must have been stopped before we can switch sequences.
+			ExistingActor->SetSequence(nullptr);
+		}
+	}
+
+	// Now we can construct a nice new actor which has the settings we need.
+	LevelSequenceActor = GetWorld()->SpawnActor<ALevelSequenceActor>();
+	check(LevelSequenceActor);
+
+	// Use our duplicated sequence
+	LevelSequenceActor->SetSequence(InSequenceToApply);
+
+	// Enforce settings.
+	LevelSequenceActor->PlaybackSettings.LoopCount.Value = 0;
+	LevelSequenceActor->PlaybackSettings.bAutoPlay = false;
+	LevelSequenceActor->PlaybackSettings.bPauseAtEnd = true;
+	LevelSequenceActor->GetSequencePlayer()->SetTimeController(CustomSequenceTimeController);
+
+	LevelSequenceActor->GetSequencePlayer()->OnSequenceUpdated().AddUObject(this, &UMoviePipeline::OnSequenceEvaluated);
+}
+
+
+FMoviePipelineShotInfo CreateShotFromMovieScene(const UMovieScene* InMovieScene, const TRange<FFrameNumber>& InIntersectionRange)
+{
+	check(InMovieScene);
+
+	// We will generate a range for each sub-section within the shot we wish to render. These will be clipped by
+	// the overall range of the shot to represent the final desired render duration.
+	struct FCameraCutRange
+	{
+		TRange<FFrameNumber> Range;
+		UMovieSceneCameraCutSection* Section;
+	};
+
+	TArray<FCameraCutRange> IntersectedRanges;
+
+	// We're going to search for Camera Cut tracks within this shot. If none are found, we'll use the whole range of the shot.
+	UMovieSceneCameraCutTrack* CameraCutTrack = InMovieScene->FindMasterTrack<UMovieSceneCameraCutTrack>();
+	if (CameraCutTrack)
+	{
+		for (UMovieSceneSection* Section : CameraCutTrack->GetAllSections())
+		{
+			UMovieSceneCameraCutSection* CameraCutSection = CastChecked<UMovieSceneCameraCutSection>(Section);
+
+			// ToDo: Inner vs. Outer resolution differences.
+			// Intersect this cut with the outer range in the likely event that the section goes past the bounds.
+			TRange<FFrameNumber> IntersectingRange = TRange<FFrameNumber>::Intersection(Section->GetRange(), InIntersectionRange);
+
+			FCameraCutRange& NewRange = IntersectedRanges.AddDefaulted_GetRef();
+			NewRange.Range = IntersectingRange;
+			NewRange.Section = CameraCutSection;
+		}
+	}
+	else
+	{
+		// No camera cut track was found inside. We'll treat the whole shot as the desired range.
+		FCameraCutRange& NewRange = IntersectedRanges.AddDefaulted_GetRef();
+		NewRange.Range = InIntersectionRange;
+		NewRange.Section = nullptr;
+	}
+
+	FMoviePipelineShotInfo NewShot;
+	NewShot.OriginalRange = InIntersectionRange;
+	NewShot.TotalOutputRange = NewShot.OriginalRange;
+
+	for (const FCameraCutRange& Range : IntersectedRanges)
+	{
+		// Generate a CameraCut for each range.
+		FMoviePipelineCameraCutInfo& CameraCut = NewShot.CameraCuts.AddDefaulted_GetRef();
+
+		CameraCut.CameraCutSection = Range.Section; // May be nullptr.
+		CameraCut.OriginalRange = Range.Range;
+		CameraCut.TotalOutputRange = CameraCut.OriginalRange;
+	}
+
+	return NewShot;
+}
+
+TArray<FMoviePipelineShotInfo> UMoviePipeline::BuildShotListFromSequence(const ULevelSequence* InSequence)
+{
+	TArray<FMoviePipelineShotInfo> NewShotList;
+	
+	// Shot Tracks take precedent over camera cuts, as settings can only be applied as granular as a shot.
+	UMovieSceneCinematicShotTrack* CinematicShotTrack = InSequence->GetMovieScene()->FindMasterTrack<UMovieSceneCinematicShotTrack>();
+	
+	if (CinematicShotTrack)
+	{
+		ensureMsgf(CinematicShotTrack->GetAllSections().Num() > 0, TEXT("Cinematic Shot track must have at least one section to be rendered."));
+
+		for (UMovieSceneSection* Section : CinematicShotTrack->GetAllSections())
+		{
+			UMovieSceneCinematicShotSection* ShotSection = CastChecked<UMovieSceneCinematicShotSection>(Section);
+
+			// If the user has manually marked a section as inactive we don't produce a shot for it.
+			if (!Section->IsActive())
+			{
+				UE_LOG(LogMovieRenderPipeline, Verbose, TEXT("Skipped adding Shot %s to Shot List due to being inactive."), *ShotSection->GetShotDisplayName());
+				continue;
+			}
+
+			if (!ShotSection->GetSequence())
+			{
+				UE_LOG(LogMovieRenderPipeline, Warning, TEXT("Skipped adding Shot %s to Shot List due to no inner sequence."), *ShotSection->GetShotDisplayName());
+				continue;
+			}
+
+			//if (CurrentJob->ShotRenderMask.Num() > 0)
+			//{
+			//	if (!CurrentJob->ShotRenderMask.Contains(ShotSection->GetShotDisplayName()))
+			//	{
+			//		UE_LOG(LogMovieRenderPipeline, Verbose, TEXT("Skipped adding Shot %s to Shot List due to a shot render mask being active, and this shot not being on the list."), *ShotSection->GetShotDisplayName());
+			//		continue;
+			//	}
+			//}
+
+			// The Shot Section may extend past our Sequence's Playback Bounds. We intersect the two bounds to ensure that
+			// the Playback Start/Playback End of the overall sequence is respected.
+			TRange<FFrameNumber> CinematicShotSectionRange = TRange<FFrameNumber>::Intersection(ShotSection->GetRange(), InSequence->GetMovieScene()->GetPlaybackRange());
+
+			FMoviePipelineShotInfo NewShot = CreateShotFromMovieScene(ShotSection->GetSequence()->GetMovieScene(), CinematicShotSectionRange);
+
+			// The first thing we do is find the appropriate configuration from the settings. Each shot can have its own config
+			// or they fall back to a default one specified for the whole pipeline.
+			// NewShot.ShotConfig = GetPipelineMasterConfig()->GetConfigForShot(ShotSection->GetShotDisplayName());
+			NewShot.CinematicShotSection = ShotSection;
+
+			// There should always be a shot config as the Pipeline default is returned in the event they didn't customize.
+			// check(NewShot.ShotConfig);
+
+			NewShotList.Add(MoveTemp(NewShot));
+		}
+	}
+	else
+	{
+		// They don't have a cinematic shot track. We'll slice them up by camera cuts instead.
+		FMoviePipelineShotInfo NewShot = CreateShotFromMovieScene(TargetSequence->GetMovieScene(), TargetSequence->GetMovieScene()->GetPlaybackRange());
+		// NewShot.ShotConfig = GetPipelineMasterConfig()->DefaultShotConfig;
+		
+		NewShotList.Add(MoveTemp(NewShot));
+	}
+
+	// If they don't have a cinematic shot track, or a camera cut track then they want to control the camera
+	// through their own logic. We'll just use the duration of the Sequence as the render, plus warn them.
+	if (NewShotList.Num() == 0)
+	{
+		UE_LOG(LogMovieRenderPipeline, Warning, TEXT("No Cinematic Shot Tracks found, and no Camera Cut Tracks found. Playback Range will be used but camera will render from Pawns perspective."));
+		FMoviePipelineShotInfo NewShot;
+		// NewShot.ShotConfig = GetPipelineMasterConfig()->DefaultShotConfig;
+
+		FMoviePipelineCameraCutInfo& CameraCut = NewShot.CameraCuts.AddDefaulted_GetRef();
+		CameraCut.OriginalRange = TargetSequence->GetMovieScene()->GetPlaybackRange();
+	}
+
+	// Now that we've gathered at least one or more shots with one or more cuts, we can apply settings. It's easier to
+	// debug when all of the shots are calculated up front and debug info is printed in a block instead of as it is reached.
+	int32 ShotIndex = 1;
+	for (FMoviePipelineShotInfo& Shot : NewShotList)
+	{
+		UE_LOG(LogMovieRenderPipeline, Log, TEXT("Shot %d/%d has %d Camera Cuts."), ShotIndex, NewShotList.Num(), Shot.CameraCuts.Num());
+		ShotIndex++;
+
+		UMoviePipelineAccumulationSetting* AccumulationSettings = FindOrAddSetting<UMoviePipelineAccumulationSetting>(Shot);
+		Shot.NumHandleFrames = AccumulationSettings->HandleFrameCount;
+
+		UMoviePipelineCameraSetting* CameraSettings = FindOrAddSetting<UMoviePipelineCameraSetting>(Shot);
+
+
+		// Expand the shot to encompass handle frames. This will modify our Camera Cuts bounds.
+		ExpandShot(Shot);
+
+		FString ShotName;
+		if (Shot.CinematicShotSection.IsValid())
+		{
+			ShotName = Shot.CinematicShotSection->GetShotDisplayName();
+		}
+
+		for (FMoviePipelineCameraCutInfo& CameraCut : Shot.CameraCuts)
+		{
+			CameraCut.NumWarmUpFramesRemaining = CameraCut.NumWarmUpFrames = AccumulationSettings->WarmUpFrameCount;
+			CameraCut.bAccurateFirstFrameHistory = true;
+			CameraCut.NumTemporalSamples = CameraSettings->TemporalSampleCount;
+			CameraCut.NumSpatialSamples = AccumulationSettings->SpatialSampleCount;
+			CameraCut.CachedFrameRate = GetPipelineMasterConfig()->GetEffectiveFrameRate(TargetSequence);
+			CameraCut.CachedTickResolution = TargetSequence->GetMovieScene()->GetTickResolution();
+			CameraCut.TotalWorkInfo = CameraCut.GetTotalWorkEstimate();
+
+			FString CameraName;
+			if (CameraCut.CameraCutSection.IsValid())
+			{
+				CameraName = TargetSequence->GetMovieScene()->FindBinding(CameraCut.CameraCutSection->GetCameraBindingID().GetGuid())->GetName();
+			}
+
+			CameraCut.DisplayName = FString::Printf(TEXT("Shot: %s Cam: %s"), *ShotName, *CameraName);
+			
+			// When we expanded the shot above, it pushed the first/last camera cuts ranges to account for Handle Frames.
+			// We want to start rendering from the first handle frame. Shutter Timing is a fixed offset from this number.
+			CameraCut.CurrentTick = CameraCut.TotalOutputRange.GetLowerBoundValue();
+		}
+	}
+
+	return NewShotList;
+}
+void UMoviePipeline::InitializeShot(FMoviePipelineShotInfo& InShot)
+{
+	// Set the new shot as the active shot. This enables the specified shot section and disables all other shot sections.
+	SetSoloShot(InShot);
+
+	if (InShot.ShotOverrideConfig != nullptr)
+	{
+		// Any shot-specific overrides haven't had first time initialization. So we'll do that now.
+		for (UMoviePipelineSetting* Setting : InShot.ShotOverrideConfig->GetSettings())
+		{
+			Setting->OnMoviePipelineInitialized(this);
+		}
+	}
+
+	// Setup required rendering architecture for all passes in this shot.
+	SetupRenderingPipelineForShot(InShot);
+}
+
+void UMoviePipeline::TeardownShot(FMoviePipelineShotInfo& InShot)
+{
+	if (InShot.ShotOverrideConfig != nullptr)
+	{
+		// Any shot-specific overrides haven't had first time initialization. So we'll do that now.
+		for (UMoviePipelineSetting* Setting : InShot.ShotOverrideConfig->GetSettings())
+		{
+			Setting->OnMoviePipelineShutdown(this);
+		}
+	}
+
+	// Setup required rendering architecture for all passes in this shot.
+	TeardownRenderingPipelineForShot(InShot);
+}
+
+void UMoviePipeline::SetSoloShot(const FMoviePipelineShotInfo& InShot)
+{
+	// Iterate through the shot list and ensure all shots have been set to inactive.
+	for (FMoviePipelineShotInfo& Shot : ShotList)
+	{
+		UMovieSceneCinematicShotSection* Section = Shot.CinematicShotSection.Get();
+
+		// Skip the one we will be modifying below to avoid extra MarkAsChanged calls.
+		if (Section && Section != InShot.CinematicShotSection.Get())
+		{
+			Section->SetIsActive(false);
+			Section->MarkAsChanged();
+		}
+	}
+
+	// Now that we've set them all to inactive we'll ensure that our passed in shot is active.
+	if (UMovieSceneCinematicShotSection* CurrentShotSection = InShot.CinematicShotSection.Get())
+	{
+		UE_LOG(LogMovieRenderPipeline, Verbose, TEXT("Disabled all shot tracks and re-enabling %s for solo."), *CurrentShotSection->GetShotDisplayName());
+		CurrentShotSection->SetIsActive(true);
+		CurrentShotSection->MarkAsChanged();
+	}
+	else
+	{
+		UE_LOG(LogMovieRenderPipeline, Verbose, TEXT("Disabled all shot tracks and skipped enabling a shot track due to no shot section associated with the provided shot."));
+	}
+}
+
+void UMoviePipeline::ExpandShot(FMoviePipelineShotInfo& InShot)
+{
+	const MoviePipeline::FFrameConstantMetrics FrameMetrics = CalculateShotFrameMetrics(InShot);
+
+	// Handle Frames will be added onto our original shot size. We track the ranges separately for counting purposes
+	// later - the actual rendering code is unaware of the handle frames. Handle Frames only apply to the shot
+	// and expand the first/last inner cut to cover this area.
+	FFrameNumber HandleFrameTicks = FrameMetrics.TicksPerOutputFrame.FloorToFrame().Value * InShot.NumHandleFrames;
+	InShot.TotalOutputRange = MovieScene::DilateRange(InShot.OriginalRange, -HandleFrameTicks, HandleFrameTicks);
+
+	// We'll expand the overall sequence by at least this much to ensure we don't get clamped. This is +1 frame for shutter timing overlap.
+	TRange<FFrameNumber> TotalPlaybackRange = MovieScene::DilateRange(InShot.TotalOutputRange, -FrameMetrics.TicksPerOutputFrame.CeilToFrame(), FrameMetrics.TicksPerOutputFrame.CeilToFrame());
+
+	// We can now take the difference between the two ranges to isolate just the opening/closing range (so we can use the range for counting later)
+	TArray<TRange<FFrameNumber>> HandleFrameRanges = TRange<FFrameNumber>::Difference(InShot.TotalOutputRange, InShot.OriginalRange);
+	if (HandleFrameRanges.Num() == 2)
+	{
+		InShot.HandleFrameRangeStart = HandleFrameRanges[0];
+		InShot.HandleFrameRangeEnd = HandleFrameRanges[1];
+	}
+
+	// Expand the first and last cut inside this shot to cover the handle frame distance.
+	if (InShot.CameraCuts.Num() > 0)
+	{
+		// These may be the same or they may be different.
+		FMoviePipelineCameraCutInfo& FirstCut = InShot.CameraCuts[0];
+		FMoviePipelineCameraCutInfo& LastCut = InShot.CameraCuts[InShot.CameraCuts.Num() - 1];
+
+		FirstCut.TotalOutputRange = MovieScene::DilateRange(FirstCut.TotalOutputRange, -HandleFrameTicks, FFrameNumber(0));
+		LastCut.TotalOutputRange = MovieScene::DilateRange(LastCut.TotalOutputRange, FFrameNumber(0), HandleFrameTicks);
+	}
+
+	// When we expand a shot section, we also need to expand the inner sequence's range by the same amount.
+	UMovieSceneCinematicShotSection* ShotSection = InShot.CinematicShotSection.Get();
+	if(ShotSection)
+	{
+		UMovieScene* ShotMovieScene = ShotSection->GetSequence() ? ShotSection->GetSequence()->GetMovieScene() : nullptr;
+		if (!ShotMovieScene)
+		{
+			UE_LOG(LogMovieRenderPipeline, Warning, TEXT("ShotSection (%s) with no inner sequence, skipping expansion!"), *ShotSection->GetShotDisplayName());
+			return;
+		}
+
+		// Expand the inner shot section range by the handle size, multiplied by the difference between the outer and inner tick resolutions (and factoring in the time scale)
+		const bool bRateMatches = TargetSequence->GetMovieScene()->GetTickResolution() == ShotMovieScene->GetTickResolution();
+		const double OuterToInnerRateDilation = bRateMatches ? 1.0 : (ShotMovieScene->GetTickResolution() / TargetSequence->GetMovieScene()->GetTickResolution()).AsDecimal();
+		const double OuterToInnerScale = OuterToInnerRateDilation * (double)ShotSection->Parameters.TimeScale;
+
+		// We push the playback bounds one frame further (in both directions) to handle the case where shutter timings offset them +/- the actual edge.
+		// These bounds don't affect the period of time which is rendered, so it's okay that this pushes beyond the actual evaluated range.
+
+		const FFrameNumber StartOffset = HandleFrameTicks + FrameMetrics.TicksPerOutputFrame.FloorToFrame();
+		const FFrameNumber EndOffset = HandleFrameTicks + FrameMetrics.TicksPerOutputFrame.FloorToFrame();
+
+		const FFrameNumber DilatedStartOffset = FFrameNumber(FMath::CeilToInt(StartOffset.Value * OuterToInnerScale));
+		const FFrameNumber DilatedEndOffset = FFrameNumber(FMath::CeilToInt(EndOffset.Value * OuterToInnerScale));
+
+		TRange<FFrameNumber> OldPlaybackRange = ShotMovieScene->GetPlaybackRange();
+		TRange<FFrameNumber> NewPlaybackRange = MovieScene::DilateRange(ShotMovieScene->GetPlaybackRange(), -DilatedStartOffset, DilatedEndOffset);
+
+		FFrameNumber TotalExpansionSizeInFrames = FFrameRate::TransformTime(FFrameTime(HandleFrameTicks) + FrameMetrics.TicksPerOutputFrame, FrameMetrics.TickResolution, FrameMetrics.FrameRate).CeilToFrame();
+		// ToDo: ShotList isn't assigned yet so IndexOfByKey fails.
+		UE_LOG(LogMovieRenderPipeline, Log, TEXT("ShotSection ([%d] %s) Expanding Shot by [-%d, %d] Ticks (+/-%d Frames). Old Inner Playback Range: %s New Inner Playback Range: %s"),
+			ShotList.IndexOfByKey(InShot), *ShotSection->GetShotDisplayName(), 
+			StartOffset.Value, EndOffset.Value, TotalExpansionSizeInFrames.Value,
+			*LexToString(OldPlaybackRange), *LexToString(NewPlaybackRange));
+
+		// Expand the inner scene
+		ShotMovieScene->SetPlaybackRange(NewPlaybackRange, false);
+
+		// Expand the outer owning section
+		ShotSection->SetRange(MovieScene::DilateRange(ShotSection->GetRange(), -StartOffset, EndOffset));
+
+	}
+
+	// Ensure the overall Movie Scene Playback Range is large enough. This will clamp evaluation if we don't expand it. We hull the existing range
+	// with the new range.
+	TRange<FFrameNumber> EncompassingPlaybackRange = TRange<FFrameNumber>::Hull(TotalPlaybackRange, TargetSequence->MovieScene->GetPlaybackRange());
+	TargetSequence->GetMovieScene()->SetPlaybackRange(EncompassingPlaybackRange);
+
+}
+
+bool UMoviePipeline::IsDebugFrameStepIdling() const
+{
+	// We're only idling when we're at zero, otherwise there's more frames to process.
+	// Caveat is that this will be zero on the last frame we want to render, so we
+	// take into account whether or not we've queued up a pause at the end of the frame
+	// which is indicator that we want to process the current frame.
+	int32 DebugFrameStepValue = CVarMovieRenderPipelineFrameStepper.GetValueOnGameThread();
+	return DebugFrameStepValue == 0 && !bPauseAtEndOfFrame;
+}
+
+bool UMoviePipeline::DebugFrameStepPreTick()
+{
+	int32 DebugFrameStepValue = CVarMovieRenderPipelineFrameStepper.GetValueOnGameThread();
+	if (DebugFrameStepValue == 0)
+	{
+		// A value of 0 means that they are using the frame stepper and that we have stepped
+		// the specified number of frames. We will create a DeltaTime for the engine
+		// and not process anything below, which prevents us from trying to produce an
+		// output frame later.
+		CustomTimeStep->SetCachedFrameTiming(MoviePipeline::FFrameTimeStepCache(1 / 24.0));
+		return true;
+	}
+	else if (DebugFrameStepValue > 0)
+	{
+		// They want to process at least one frame, deincrement, then we
+		// process the frame. We pause the game here to preserve render state.
+		CVarMovieRenderPipelineFrameStepper->Set(DebugFrameStepValue - 1, ECVF_SetByConsole);
+
+		// We want to run this one frame and then pause again at the end.
+		bPauseAtEndOfFrame = true;
+	}
+
+	return false;
+}
+
+void UMoviePipeline::LoadDebugWidget()
+{
+	FSoftClassPath DebugWidgetClassRef(TEXT("/MovieRenderPipeline/Blueprints/UI_MovieRenderPipelineScreenOverlay.UI_MovieRenderPipelineScreenOverlay_C"));
+	if (UClass* DebugWidgetClass = DebugWidgetClassRef.TryLoadClass<UMovieRenderDebugWidget>())
+	{
+		DebugWidget = CreateWidget<UMovieRenderDebugWidget>(GetWorld(), DebugWidgetClass);
+		if (DebugWidget)
+		{
+			DebugWidget->OnInitializedForPipeline(this);
+			DebugWidget->AddToViewport();
+		}
+		else
+		{
+			UE_LOG(LogMovieRenderPipeline, Warning, TEXT("Failed to create Debug Screen UMG Widget. No debug overlay available."));
+		}
+	}
+	else
+	{
+		UE_LOG(LogMovieRenderPipeline, Warning, TEXT("Failed to find Debug Screen UMG Widget class. No debug overlay available."));
+	}
+}
+
+FFrameTime FMoviePipelineTimeController::OnRequestCurrentTime(const FQualifiedFrameTime& InCurrentTime, float InPlayRate)
+{
+	FFrameTime RequestTime = FFrameRate::TransformTime(TimeCache.Time, TimeCache.Rate, InCurrentTime.Rate);
+	UE_LOG(LogMovieRenderPipeline, VeryVerbose, TEXT("[%d] OnRequestCurrentTime: %d %f"), GFrameCounter, RequestTime.FloorToFrame().Value, RequestTime.GetSubFrame());
+
+	return RequestTime;
+}
+
+MoviePipeline::FFrameConstantMetrics UMoviePipeline::CalculateShotFrameMetrics(const FMoviePipelineShotInfo& InShot) const
+{
+	MoviePipeline::FFrameConstantMetrics Output;
+	Output.TickResolution = TargetSequence->GetMovieScene()->GetTickResolution();
+	Output.FrameRate = GetPipelineMasterConfig()->GetEffectiveFrameRate(TargetSequence);
+	Output.TicksPerOutputFrame = FFrameRate::TransformTime(FFrameTime(FFrameNumber(1)), Output.FrameRate, Output.TickResolution);
+
+	UMoviePipelineCameraSetting* CameraSettings = FindOrAddSetting<UMoviePipelineCameraSetting>(InShot);
+
+	// (CameraShutterAngle/360) gives us the fraction-of-the-output-frame the accumulation frames should cover.
+	Output.ShutterAnglePercentage = CameraSettings->CameraShutterAngle / 360.0;
+
+	{
+		/*
+		* Calculate out how many ticks a normal sub-frame occupies.
+		* (TickRes/FrameRate) gives you ticks-per-second, and then divide that by the percentage of time the
+		* shutter is open. Finally, divide the percentage of time the shutter is open by the number of frames
+		* we're accumulating.
+		*
+		* It is common that there is potential to have some leftover here. ie:
+		* 24000 Ticks / 24fps = 1000 ticks per second. At 180 degree shutter angle that gives you 500 ticks
+		* spread out amongst n=3 sub-frames. 500/3 = 166.66 ticks. We'll floor that when we use it, and ensure
+		* we accumulate the sub-tick and choose when to apply it.
+		*/
+
+		// Now we take the amount of time the shutter is open.
+		Output.TicksWhileShutterOpen = Output.TicksPerOutputFrame * Output.ShutterAnglePercentage;
+
+		// Divide that amongst all of our accumulation sample frames.
+		Output.TicksPerSample = Output.TicksWhileShutterOpen / CameraSettings->TemporalSampleCount;
+	}
+
+	Output.ShutterClosedFraction = (360 - CameraSettings->CameraShutterAngle) / 360.0;
+	Output.TicksWhileShutterClosed = Output.TicksPerOutputFrame * Output.ShutterClosedFraction;
+
+	// Shutter Offset
+	switch (CameraSettings->ShutterTiming)
+	{
+		// Subtract the entire time the shutter is open.
+	case EMoviePipelineShutterTiming::FrameClose:
+		Output.ShutterOffsetTicks = -Output.TicksWhileShutterOpen;
+		break;
+		// Only subtract half the time the shutter is open.
+	case EMoviePipelineShutterTiming::FrameCenter:
+		Output.ShutterOffsetTicks = -Output.TicksWhileShutterOpen / 2.0;
+		break;
+		// No offset needed
+	case EMoviePipelineShutterTiming::FrameOpen:
+		break;
+	}
+
+	// Then, calculate our motion blur offset. Motion Blur in the engine is always
+	// centered around the object so we offset our time sampling by half of the
+	// motion blur distance so that the distance blurred represents that time.
+	Output.MotionBlurCenteringOffsetTicks = Output.TicksPerSample / 2.0;
+
+	return Output;
+}
+
+
+UMoviePipelineMasterConfig* UMoviePipeline::GetPipelineMasterConfig() const
+{ 
+	return CurrentJob->GetConfiguration(); 
+}
+
+UMoviePipelineSetting* UMoviePipeline::FindOrAddSetting(TSubclassOf<UMoviePipelineSetting> InSetting, const FMoviePipelineShotInfo& InShot) const
+{
+	// Check to see if this setting is in the shot override, if it is we'll use the shot version of that.
+	if (InShot.ShotOverrideConfig)
+	{
+		UMoviePipelineSetting* Setting = InShot.ShotOverrideConfig->FindSettingByClass(InSetting);
+		if (Setting)
+		{
+			// If they specified the setting at all, respect the enabled setting. If it's not enabled, we return the
+			// default instead which is the same as if they hadn't added the setting at all.
+			return Setting->bEnabled ? Setting : InSetting->GetDefaultObject<UMoviePipelineSetting>();
+		}
+	}
+
+	// If they didn't have a shot override, or the setting wasn't enabled, we'll check the master config.
+	UMoviePipelineSetting* Setting = GetPipelineMasterConfig()->FindSettingByClass(InSetting);
+	if (Setting)
+	{
+		return Setting->bEnabled ? Setting : InSetting->GetDefaultObject<UMoviePipelineSetting>();
+	}
+
+	// If no one overrode it, then we return the default.
+	return InSetting->GetDefaultObject<UMoviePipelineSetting>();
+}
+
+
+#undef LOCTEXT_NAMESPACE // "MoviePipeline"
diff --git a/Engine/Plugins/MovieScene/MovieRenderPipeline/Source/MovieRenderPipelineCore/Private/MoviePipelineInProcessExecutor.cpp b/Engine/Plugins/MovieScene/MovieRenderPipeline/Source/MovieRenderPipelineCore/Private/MoviePipelineInProcessExecutor.cpp
new file mode 100644
--- /dev/null
+++ b/Engine/Plugins/MovieScene/MovieRenderPipeline/Source/MovieRenderPipelineCore/Private/MoviePipelineInProcessExecutor.cpp
@@ -0,0 +1,50 @@
+// Copyright 1998-2019 Epic Games, Inc. All Rights Reserved.
+#include "MoviePipelineInProcessExecutor.h"
+#include "MoviePipeline.h"
+#include "Engine/World.h"
+#include "MoviePipelineQueue.h"
+
+#define LOCTEXT_NAMESPACE "MoviePipelineInProcessExecutor"
+void UMoviePipelineInProcessExecutor::Start(const UMoviePipelineExecutorJob* InJob)
+{
+	FCoreUObjectDelegates::PostLoadMapWithWorld.AddUObject(this, &UMoviePipelineInProcessExecutor::OnMapLoadFinished);
+}
+
+void UMoviePipelineInProcessExecutor::OnMapLoadFinished(UWorld* NewWorld)
+{
+	// NewWorld can be null if a world is being destroyed.
+	if (!NewWorld)
+	{
+		return;
+	}
+
+	// Stop listening for map load until we're done and know we want to start the next config.
+	FCoreUObjectDelegates::PostLoadMapWithWorld.RemoveAll(this);
+	
+	UMoviePipelineExecutorJob* CurrentJob = Queue->GetJobs()[CurrentPipelineIndex];
+
+	ActiveMoviePipeline = NewObject<UMoviePipeline>(NewWorld, TargetPipelineClass);
+	ActiveMoviePipeline->Initialize(CurrentJob);
+
+	// Listen for when the pipeline thinks it has finished.
+	ActiveMoviePipeline->OnMoviePipelineFinished().AddUObject(this, &UMoviePipelineInProcessExecutor::OnMoviePipelineFinished);
+}
+
+void UMoviePipelineInProcessExecutor::OnMoviePipelineFinished(UMoviePipeline* InMoviePipeline)
+{
+	UMoviePipeline* MoviePipeline = ActiveMoviePipeline;
+
+	if (ActiveMoviePipeline)
+	{
+		// Unsubscribe in the event that it gets called twice we don't have issues.
+		ActiveMoviePipeline->OnMoviePipelineFinished().RemoveAll(this);
+	}
+
+	// Null these out now since OnIndividualPipelineFinished might invoke something that causes a GC
+	// and we want them to go away with the GC.
+	ActiveMoviePipeline = nullptr;
+
+	// Now that another frame has passed and we should be OK to start another PIE session, notify our owner.
+	OnIndividualPipelineFinished(MoviePipeline);
+}
+#undef LOCTEXT_NAMESPACE // "MoviePipelineInProcessExecutor"
diff --git a/Engine/Plugins/MovieScene/MovieRenderPipeline/Source/MovieRenderPipelineCore/Private/MoviePipelineLinearExecutor.cpp b/Engine/Plugins/MovieScene/MovieRenderPipeline/Source/MovieRenderPipelineCore/Private/MoviePipelineLinearExecutor.cpp
new file mode 100644
--- /dev/null
+++ b/Engine/Plugins/MovieScene/MovieRenderPipeline/Source/MovieRenderPipelineCore/Private/MoviePipelineLinearExecutor.cpp
@@ -0,0 +1,74 @@
+// Copyright 1998-2019 Epic Games, Inc. All Rights Reserved.
+#include "MoviePipelineLinearExecutor.h"
+#include "MovieRenderPipelineCoreModule.h"
+#include "MoviePipelineQueue.h"
+
+#define LOCTEXT_NAMESPACE "MoviePipelineLinearExecutorBase"
+
+void UMoviePipelineLinearExecutorBase::ExecuteImpl(UMoviePipelineQueue* InPipelineQueue)
+{
+	check(InPipelineQueue);
+
+	if (InPipelineQueue->GetJobs().Num() == 0)
+	{
+		UE_LOG(LogMovieRenderPipeline, Warning, TEXT("Executor asked to execute on empty list of pipelines."));
+		OnExecutorErroredImpl(nullptr, true, LOCTEXT("EmptyPipelineError", "Executor asked to execute empty list of jobs. This was probably not intended!"));
+		OnExecutorFinishedImpl();
+		return;
+	}
+
+	// We'll process them in linear fashion and wait until each one is canceled or finishes on its own
+	// before moving onto the next one. This may be parallelizable in the future (either multiple PIE
+	// sessions, or multiple external processes) but ideally one render would maximize resource usage anyways...
+	Queue = InPipelineQueue;
+	InitializationTime = FDateTime::UtcNow();
+
+	UE_LOG(LogMovieRenderPipeline, Log, TEXT("MoviePipelineLinearExecutorBase starting %d jobs."), InPipelineQueue->GetJobs().Num());
+
+	StartPipelineByIndex(0);
+}
+
+void UMoviePipelineLinearExecutorBase::StartPipelineByIndex(int32 InPipelineIndex)
+{
+	check(InPipelineIndex >= 0 && InPipelineIndex < Queue->GetJobs().Num());
+	
+	CurrentPipelineIndex = InPipelineIndex;
+
+	if (!Queue->GetJobs()[CurrentPipelineIndex]->GetConfiguration())
+	{
+		UE_LOG(LogMovieRenderPipeline, Warning, TEXT("Found null config in list of configs to render. Aborting the pipeline processing!"));
+		OnExecutorErroredImpl(nullptr, true, LOCTEXT("NullPipelineError", "Found null config in list of configs to render with. Does your config have the wrong outer?"));
+		OnExecutorFinishedImpl();
+		return;
+	}
+	
+	UE_LOG(LogMovieRenderPipeline, Log, TEXT("MoviePipelineLinearExecutorBase starting jobs [%d/%d]"), CurrentPipelineIndex, Queue->GetJobs().Num());
+	Start(Queue->GetJobs()[CurrentPipelineIndex]);
+}
+
+void UMoviePipelineLinearExecutorBase::OnIndividualPipelineFinished(UMoviePipeline* /* FinishedPipeline */)
+{
+	if (CurrentPipelineIndex == Queue->GetJobs().Num() - 1)
+	{
+		OnExecutorFinishedImpl();
+	}
+	else
+	{
+		// Onto the next one!
+		StartPipelineByIndex(CurrentPipelineIndex + 1);
+	}
+}
+
+void UMoviePipelineLinearExecutorBase::OnPipelineErrored(UMoviePipeline* InPipeline, bool bIsFatal, FText ErrorText)
+{
+	OnExecutorErroredImpl(InPipeline, bIsFatal, ErrorText);
+}
+
+void UMoviePipelineLinearExecutorBase::OnExecutorFinishedImpl()
+{
+	UE_LOG(LogMovieRenderPipeline, Log, TEXT("MoviePipelineLinearExecutorBase finished %d jobs in %s."), Queue->GetJobs().Num(), *(FDateTime::UtcNow() - InitializationTime).ToString());
+
+	Super::OnExecutorFinishedImpl();
+}
+
+#undef LOCTEXT_NAMESPACE // "MoviePipelineLinearExecutorBase"
diff --git a/Engine/Plugins/MovieScene/MovieRenderPipeline/Source/MovieRenderPipelineCore/Private/MoviePipelineMasterConfig.cpp b/Engine/Plugins/MovieScene/MovieRenderPipeline/Source/MovieRenderPipelineCore/Private/MoviePipelineMasterConfig.cpp
new file mode 100644
--- /dev/null
+++ b/Engine/Plugins/MovieScene/MovieRenderPipeline/Source/MovieRenderPipelineCore/Private/MoviePipelineMasterConfig.cpp
@@ -0,0 +1,188 @@
+// Copyright 1998-2019 Epic Games, Inc. All Rights Reserved.
+
+#include "MoviePipelineMasterConfig.h"
+#include "MovieScene.h"
+#include "LevelSequence.h"
+#include "MoviePipelineShotConfig.h"
+#include "MoviePipelineOutputBase.h"
+#include "MoviePipelineOutputSetting.h"
+#include "MoviePipelineQueue.h"
+
+#define LOCTEXT_NAMESPACE "MoviePipelineMasterConfig"
+
+
+
+UMoviePipelineShotConfig* UMoviePipelineMasterConfig::GetConfigForShot(const FString& ShotName) const
+{
+	UMoviePipelineShotConfig* OutConfig = PerShotConfigMapping.FindRef(ShotName);
+
+	// They didn't customize this shot, return the global pipeline default
+	if (!OutConfig)
+	{
+		OutConfig = DefaultShotConfig;
+	}
+
+	return OutConfig;
+}
+
+void UMoviePipelineMasterConfig::GetFilenameFormatArguments(FFormatNamedArguments& OutArguments, const UMoviePipelineExecutorJob* InJob) const
+{
+	if (!InJob)
+	{
+		return;
+	}
+
+	// Add "global" ones not specific to a setting.
+	{
+		OutArguments.Add(TEXT("level_name"), FText::FromString(InJob->Map.GetAssetName()));
+		OutArguments.Add(TEXT("sequence_name"), FText::FromString(InJob->Sequence.GetAssetName()));
+		OutArguments.Add(TEXT("date"), FText::AsDate(FDateTime::UtcNow()));
+		OutArguments.Add(TEXT("time"), FText::AsTime(FDateTime::UtcNow()));
+		
+		// FrameRate is a combination of Output Settings and Sequence so we do it here instead of in OutputSetting
+		OutArguments.Add(TEXT("frame_rate"), FText::AsNumber(GetEffectiveFrameRate(Cast<ULevelSequence>(InJob->Sequence.TryLoad())).AsDecimal()));
+
+	}
+
+	// Let each setting provide its own set of key/value pairs.
+	{
+		for (UMoviePipelineSetting* Setting : GetSettings())
+		{
+			Setting->GetFilenameFormatArguments(OutArguments, InJob);
+		}
+
+		// ToDo: Should shots be able to provide arguments too? They're only overrides, and
+		// outside of that shot the format would fail and end up leaving the {text} anyways.
+	}
+}
+
+FFrameRate UMoviePipelineMasterConfig::GetEffectiveFrameRate(const ULevelSequence* InSequence) const
+{
+	UMoviePipelineOutputSetting* OutputSettings = FindSetting<UMoviePipelineOutputSetting>();
+	check(OutputSettings);
+
+	// Check to see if they overrode the frame rate.
+	if (OutputSettings->bUseCustomFrameRate)
+	{
+		return OutputSettings->OutputFrameRate;
+	}
+
+	// Pull it from the sequence if they didn't.
+	if (InSequence)
+	{
+		return InSequence->GetMovieScene()->GetDisplayRate();
+	}
+
+	return FFrameRate();
+}
+
+
+/*bool UMoviePipelineShotConfig::ValidateConfig(TArray<FText>& OutValidationErrors) const{
+	bool bValidPipeline = true;
+	OutValidationErrors.Reset();
+	
+	// Check that we have a sequence to render
+	if(!Sequence.IsValid())
+	{
+		OutValidationErrors.Add(LOCTEXT("InvalidInputSequence", "No sequence to render has been specified."));
+		bValidPipeline = false;
+	}
+	
+	// Check that we have a valid output resolution
+	if(OutputResolution.X == 0 || OutputResolution.Y == 0)
+	{
+		OutValidationErrors.Add(LOCTEXT("InvalidOutputResolution", "Output Resolution must be greater than zero."));
+		bValidPipeline = false;
+	}
+	
+	// Check for a non-zero framerate
+	if(OutputFrameRate.AsDecimal() == 0.0)
+	{
+		OutValidationErrors.Add(LOCTEXT("InvalidOutputFrameRate", "Output FrameRate must be greater than zero."));
+		bValidPipeline = false;
+	}
+	
+	// Output Directory can probably be ignored. The folder path you output to may not exist yet anyways.
+	
+	// Give each setting a chance to validate the pipeline
+	for(const UMoviePipelineSetting* Setting : Settings)
+	{
+		bValidPipeline &= Setting->ValidatePipeline(OutValidationErrors);
+	}
+	
+	// ToDo-MoviePipeline: Reflect over all settings classes and make sure all required classes have been added.
+	// ToDo-MoviePipeline: Scan for any shots/sections that don't start on whole frames issue warning.
+	
+	// Give each Input Buffer a chance to validate the pipeline
+	for(const UMoviePipelineRenderPass* InputBuffer : InputBuffers)
+	{
+		bValidPipeline &= InputBuffer->ValidatePipeline(OutValidationErrors);
+	}
+	
+	// Give each output container a chance to validate the pipeline
+	for(const UMoviePipelineOutput* OutputContainer : OutputContainers)
+	{
+		bValidPipeline &= OutputContainer->ValidatePipeline(OutValidationErrors);
+	}
+
+	// Ensure we have at least one input buffer and one output buffer.
+	if (InputBuffers.Num() == 0)
+	{
+		OutValidationErrors.Add(LOCTEXT("NoInputBuffer", "Must specify at least one Input Buffer to capture."));
+		bValidPipeline = false;
+	}
+
+	if (OutputContainers.Num() == 0)
+	{
+		OutValidationErrors.Add(LOCTEXT("NoOutputContainer", "Must specify at least one Output Container to write to."));
+		bValidPipeline = false;
+	}
+	
+	return bValidPipeline;
+}*/
+
+UMoviePipelineMasterConfig::UMoviePipelineMasterConfig()
+{
+	DefaultShotConfig = CreateDefaultSubobject<UMoviePipelineShotConfig>("DefaultShotConfig");
+	
+	// Always add at least the output settings block since having a framerate/directory/etc. is critical.
+	OutputSetting = CreateDefaultSubobject<UMoviePipelineOutputSetting>("DefaultOutputSetting");
+}
+
+TArray<UMoviePipelineSetting*> UMoviePipelineMasterConfig::GetSettings() const
+{
+	TArray<UMoviePipelineSetting*> BaseSettings = Super::GetSettings();
+	BaseSettings.Add(OutputSetting);
+
+	return BaseSettings;
+}
+
+void UMoviePipelineMasterConfig::CopyFrom(UMoviePipelineConfigBase* InConfig)
+{
+	Super::CopyFrom(InConfig);
+
+	if (InConfig->IsA<UMoviePipelineMasterConfig>())
+	{
+		UMoviePipelineMasterConfig* MasterConfig = CastChecked<UMoviePipelineMasterConfig>(InConfig);
+		OutputSetting = Cast<UMoviePipelineOutputSetting>(StaticDuplicateObject(MasterConfig->OutputSetting, this, MasterConfig->OutputSetting->GetFName()));
+	}
+}
+
+
+TArray<UMoviePipelineOutputBase*> UMoviePipelineMasterConfig::GetOutputContainers() const
+{
+	TArray<UMoviePipelineOutputBase*> OutputContainers;
+
+	for (UMoviePipelineSetting* Setting : GetSettings())
+	{
+		UMoviePipelineOutputBase* Output = Cast<UMoviePipelineOutputBase>(Setting);
+		if (Output)
+		{
+			OutputContainers.Add(Output);
+		}
+	}
+
+	return OutputContainers;
+}
+
+#undef LOCTEXT_NAMESPACE // "MovieRenderPipelineConfig"
\ No newline at end of file
diff --git a/Engine/Plugins/MovieScene/MovieRenderPipeline/Source/MovieRenderPipelineCore/Private/MoviePipelineOutputBuilder.cpp b/Engine/Plugins/MovieScene/MovieRenderPipeline/Source/MovieRenderPipelineCore/Private/MoviePipelineOutputBuilder.cpp
new file mode 100644
--- /dev/null
+++ b/Engine/Plugins/MovieScene/MovieRenderPipeline/Source/MovieRenderPipelineCore/Private/MoviePipelineOutputBuilder.cpp
@@ -0,0 +1,104 @@
+// Copyright 1998-2019 Epic Games, Inc. All Rights Reserved.
+
+#include "MoviePipelineOutputBuilder.h"
+#include "MoviePipeline.h"
+#include "Async/Async.h"
+
+
+// When the Movie Pipeline Render goes to produce a frame, it will push an expected output frame into
+// this list. This gets handed all of the component parts of that frame in an arbitrary order/threads
+// so we just wait and gather everything that comes in until a frame is complete and then we can kick
+// it off to the output step. Because there's potentially a long delay between asking a frame to be
+// produced, and it actually being produced, this is responsible for notifying the output containers
+// about shot changes, etc. - This has to be pushed back to the Game Thread though because those want to
+// live in UObject land for settings.
+
+FMoviePipelineMergerOutputFrame& FMoviePipelineOutputMerger::QueueOutputFrame_GameThread(const FMoviePipelineFrameOutputState& CachedOutputState)
+{
+	// Lock the ActiveData while we get a new output frame.
+	FScopeLock ScopeLock(&ActiveDataMutex);
+
+	// Ensure this frame hasn't already been entered somehow.
+	check(!ActiveData.Find(CachedOutputState));
+
+	FMoviePipelineMergerOutputFrame& NewFrame = ActiveData.Add(CachedOutputState);
+	NewFrame.FrameOutputState = CachedOutputState;
+
+	return NewFrame;
+}
+
+void FMoviePipelineOutputMerger::OnSingleSampleDataAvailable_AnyThread(TUniquePtr<FImagePixelData>&& InData, const TSharedRef<FImagePixelDataPayload, ESPMode::ThreadSafe> InFrameData)
+{
+	// This is to support outputting individual samples (skipping accumulation) for debug reasons,
+	// or because you want to post-process them yourself. We just forward this directly on for output to disk.
+
+	TWeakObjectPtr<UMoviePipeline> LocalWeakPipeline = WeakMoviePipeline;
+
+	AsyncTask(ENamedThreads::GameThread, [LocalData = MoveTemp(InData), InFrameData, LocalWeakPipeline]() mutable
+	{
+		if (ensureAlwaysMsgf(LocalWeakPipeline.IsValid(), TEXT("A memory lifespan issue has left an output builder alive without an owning Movie Pipeline.")))
+		{
+			LocalWeakPipeline->OnSampleRendered(MoveTemp(LocalData), InFrameData);
+		}
+	}
+	);
+}
+
+void FMoviePipelineOutputMerger::OnCompleteRenderPassDataAvailable_AnyThread(TUniquePtr<FImagePixelData>&& InData, const TSharedRef<FImagePixelDataPayload, ESPMode::ThreadSafe> InFrameData)
+{
+	// Lock the ActiveData when we're updating what data has been gathered.
+	FScopeLock ScopeLock(&ActiveDataMutex);
+
+	// See if we can find the frame this data is for. This should always be valid, if it's not
+	// valid it means they either forgot to declare they were going to produce it, or this is
+	// coming in after the system already thinks it's finished that frame.
+	FMoviePipelineMergerOutputFrame* OutputFrame = nullptr;
+	
+	// Instead of just finding the result in the TMap with the equality operator, we find it by hand so that we can
+	// ignore certain parts of equality (such as Temporal Sample, as the last sample has a temporal index different
+	// than the first sample!)
+	for (TPair<FMoviePipelineFrameOutputState, FMoviePipelineMergerOutputFrame>& KVP : ActiveData)
+	{
+		if (KVP.Key.OutputFrameNumber == InFrameData->OutputState.OutputFrameNumber)
+	 	{
+	 		OutputFrame = &KVP.Value;
+	 		break;
+	 	}
+	}
+	if (!ensureAlwaysMsgf(OutputFrame, TEXT("Recieved data for unknown frame. Frame was either already processed or not queued yet!")))
+	{
+		return;
+	}
+
+	// Ensure this pass is expected as well...
+	if (!ensureAlwaysMsgf(OutputFrame->ExpectedRenderPasses.Contains(InFrameData->PassIdentifier), TEXT("Recieved data for unexpected render pass: %s"), *InFrameData->PassIdentifier.Name))
+	{
+		return;
+	}
+
+	// If this data was expected and this frame is still in progress, pass the data to the frame.
+	OutputFrame->ImageOutputData.FindOrAdd(InFrameData->PassIdentifier) = MoveTemp(InData);
+
+	// Check to see if this was the last piece of data needed for this frame.
+	int32 TotalPasses = OutputFrame->ExpectedRenderPasses.Num();
+	int32 SucceededPasses = OutputFrame->ImageOutputData.Num();
+
+	if (SucceededPasses == TotalPasses)
+	{
+		// Transfer ownership from the map to here;
+		FMoviePipelineMergerOutputFrame FinalFrame;
+		ActiveData.RemoveAndCopyValue(InFrameData->OutputState, FinalFrame);
+
+		// Notify the Movie Pipeline that this frame has been completed and can be forwarded to the output containers.
+		TWeakObjectPtr<UMoviePipeline> LocalWeakPipeline = WeakMoviePipeline;
+
+		AsyncTask(ENamedThreads::GameThread, [LocalFinalFrame = MoveTemp(FinalFrame), InFrameData, LocalWeakPipeline]() mutable
+		{
+			if (ensureAlwaysMsgf(LocalWeakPipeline.IsValid(), TEXT("A memory lifespan issue has left an output builder alive without an owning Movie Pipeline.")))
+			{
+				LocalWeakPipeline->OnFrameCompletelyRendered(MoveTemp(LocalFinalFrame), InFrameData);
+			}
+		}
+		);
+	}
+}
\ No newline at end of file
diff --git a/Engine/Plugins/MovieScene/MovieRenderPipeline/Source/MovieRenderPipelineCore/Private/MoviePipelineOutputSetting.cpp b/Engine/Plugins/MovieScene/MovieRenderPipeline/Source/MovieRenderPipelineCore/Private/MoviePipelineOutputSetting.cpp
new file mode 100644
--- /dev/null
+++ b/Engine/Plugins/MovieScene/MovieRenderPipeline/Source/MovieRenderPipelineCore/Private/MoviePipelineOutputSetting.cpp
@@ -0,0 +1,54 @@
+// Copyright 1998-2019 Epic Games, Inc. All Rights Reserved.
+
+#include "MoviePipelineOutputSetting.h"
+#include "Misc/Paths.h"
+#include "MoviePipelineMasterConfig.h"
+#include "MoviePipelineQueue.h"
+
+UMoviePipelineOutputSetting::UMoviePipelineOutputSetting()
+	: OutputResolution(FIntPoint(1920, 1080))
+	, bUseCustomFrameRate(false)
+	, OutputFrameRate(FFrameRate(24, 1))
+	, OutputFrameStep(FFrameNumber(1))
+	, bOverrideExistingOutput(true)
+	, bUseCustomPlaybackRange(false)
+	, CustomStartFrame(FFrameNumber(0))
+	, CustomEndFrame(FFrameNumber(0))
+{
+	FileNameFormat = TEXT("{sequence_name}.{frame_number}");
+	OutputDirectory.Path = FPaths::ProjectSavedDir() / TEXT("MovieRenders/");
+}
+
+FText UMoviePipelineOutputSetting::GetFooterText(UMoviePipelineExecutorJob* InJob) const 
+{
+	FTextBuilder TextBuilder;
+	TextBuilder.AppendLine(NSLOCTEXT("MovieRenderPipeline", "OutputSettingFooterText_Fmt",
+		"A list of {format_strings} and example values that are valid to use in the File Name Format:\n"));
+
+	FFormatNamedArguments Arguments;
+	
+	// Find the master configuration that owns us
+	UMoviePipelineMasterConfig* MasterConfig = GetTypedOuter<UMoviePipelineMasterConfig>();
+	if (MasterConfig)
+	{
+		MasterConfig->GetFilenameFormatArguments(Arguments, InJob);
+	}
+
+	for (const TPair<FString, FFormatArgumentValue>& KVP : Arguments)
+	{
+		TextBuilder.AppendLineFormat(NSLOCTEXT("MovieRenderPipeline", "FormatKeyValuePair_Fmt", "`{{0}`} => {1}"), FText::FromString(KVP.Key), KVP.Value);
+	}
+	
+	return TextBuilder.ToText();
+}
+
+void UMoviePipelineOutputSetting::GetFilenameFormatArguments(FFormatNamedArguments& OutArguments, const UMoviePipelineExecutorJob* InJob) const
+{
+	// Resolution Arguments
+	{
+		FString Resolution = FString::Printf(TEXT("%d_%d"), OutputResolution.X, OutputResolution.Y);
+		OutArguments.Add(TEXT("output_resolution"), FText::FromString(Resolution));
+		OutArguments.Add(TEXT("output_width"), FText::AsNumber(OutputResolution.X));
+		OutArguments.Add(TEXT("output_height"), FText::AsNumber(OutputResolution.Y));
+	}
+}
\ No newline at end of file
diff --git a/Engine/Plugins/MovieScene/MovieRenderPipeline/Source/MovieRenderPipelineCore/Private/MoviePipelineQueue.cpp b/Engine/Plugins/MovieScene/MovieRenderPipeline/Source/MovieRenderPipelineCore/Private/MoviePipelineQueue.cpp
new file mode 100644
--- /dev/null
+++ b/Engine/Plugins/MovieScene/MovieRenderPipeline/Source/MovieRenderPipelineCore/Private/MoviePipelineQueue.cpp
@@ -0,0 +1,6 @@
+// Copyright 1998-2019 Epic Games, Inc. All Rights Reserved.
+
+#include "MoviePipelineQueue.h"
+#include "MoviePipelineOutputSetting.h"
+#include "MoviePipelineSetting.h"
+
