Copy up from Dev-Editor [REDACTED_HANDLE]

diff --git a/Engine/Plugins/MovieScene/MovieRenderPipeline/Source/MovieRenderPipelineCore/Private/MoviePipelineRendering.cpp b/Engine/Plugins/MovieScene/MovieRenderPipeline/Source/MovieRenderPipelineCore/Private/MoviePipelineRendering.cpp
new file mode 100644
--- /dev/null
+++ b/Engine/Plugins/MovieScene/MovieRenderPipeline/Source/MovieRenderPipelineCore/Private/MoviePipelineRendering.cpp
@@ -0,0 +1,552 @@
+// Copyright 1998-2019 Epic Games, Inc. All Rights Reserved.
+
+#include "MoviePipeline.h"
+#include "MovieRenderPipelineCoreModule.h"
+#include "Misc/FrameRate.h"
+#include "MoviePipelineOutputBase.h"
+#include "MoviePipelineAccumulationSetting.h"
+#include "MoviePipelineShotConfig.h"
+#include "MoviePipelineRenderPass.h"
+#include "MoviePipelineOutputBuilder.h"
+#include "RenderingThread.h"
+#include "MoviePipelineOutputSetting.h"
+#include "MoviePipelineConfigBase.h"
+#include "MoviePipelineMasterConfig.h"
+#include "Math/Halton.h"
+#include "ImageWriteTask.h"
+#include "ImageWriteQueue.h"
+#include "Engine/World.h"
+#include "GameFramework/PlayerController.h"
+#include "MoviePipelineHighResSetting.h"
+#include "Modules/ModuleManager.h"
+#include "MoviePipelineCameraSetting.h"
+
+// For flushing async systems
+#include "RendererInterface.h"
+#include "LandscapeProxy.h"
+#include "EngineModule.h"
+#include "DistanceFieldAtlas.h"
+#include "ShaderCompiler.h"
+#include "EngineUtils.h"
+
+#define LOCTEXT_NAMESPACE "MoviePipeline"
+
+static TArray<UMoviePipelineRenderPass*> GetAllRenderPasses(const UMoviePipelineMasterConfig* InMasterConfig, const FMoviePipelineShotInfo& InShot)
+{
+	TArray<UMoviePipelineRenderPass*> RenderPasses;
+
+	// Master Configuration first.
+	RenderPasses.Append(InMasterConfig->FindSettings<UMoviePipelineRenderPass>());
+
+	// And then any additional passes requested by the shot.
+	if (InShot.ShotOverrideConfig != nullptr)
+	{
+		RenderPasses.Append(InShot.ShotOverrideConfig->FindSettings<UMoviePipelineRenderPass>());
+	}
+
+	return RenderPasses;
+}
+
+
+bool GetAnyOutputWantsAlpha(UMoviePipelineConfigBase* InConfig)
+{
+	TArray<UMoviePipelineOutputBase*> OutputSettings = InConfig->FindSettings<UMoviePipelineOutputBase>();
+
+	for (const UMoviePipelineOutputBase* Output : OutputSettings)
+	{
+		if (Output->IsAlphaSupported())
+		{
+			return true;
+		}
+	}
+
+	return false;
+}
+
+void UMoviePipeline::SetupRenderingPipelineForShot(FMoviePipelineShotInfo& Shot)
+{
+	/*
+	* To support tiled rendering we take the final effective resolution and divide
+	* it by the number of tiles to find the resolution of each render target. To 
+	* handle non-evenly divisible numbers/resolutions we may oversize the targets
+	* by a few pixels and then take the center of the resulting image when interlacing
+	* to produce the final image at the right resolution. For example:
+	*
+	* 1920x1080 in 7x7 tiles gives you 274.29x154.29. We ceiling this to set the resolution
+	* of the render pass to 275x155 which will give us a final interleaved image size of
+	* 1925x1085. To ensure that the image matches a non-scaled one we take the center out.
+	* LeftOffset = floor((1925-1920)/2) = 2
+	* RightOffset = (1925-1920-LeftOffset)
+	*/
+	UMoviePipelineAccumulationSetting* AccumulationSettings = FindOrAddSetting<UMoviePipelineAccumulationSetting>(Shot);
+	UMoviePipelineHighResSetting* HighResSettings = FindOrAddSetting<UMoviePipelineHighResSetting>(Shot);
+	UMoviePipelineOutputSetting* OutputSettings = GetPipelineMasterConfig()->FindSetting<UMoviePipelineOutputSetting>();
+	check(OutputSettings);
+
+	bool bAnyOutputWantsAlpha = GetAnyOutputWantsAlpha(GetPipelineMasterConfig());
+
+	FIntPoint BackbufferTileCount = FIntPoint(HighResSettings->TileCount, HighResSettings->TileCount);
+	
+	// Figure out how big each sub-region (tile) is.
+	FIntPoint BackbufferResolution = FIntPoint(
+		FMath::CeilToInt(OutputSettings->OutputResolution.X / HighResSettings->TileCount),
+		FMath::CeilToInt(OutputSettings->OutputResolution.Y / HighResSettings->TileCount));
+
+	// Then increase each sub-region by the overlap amount.
+	BackbufferResolution = HighResSettings->CalculatePaddedBackbufferSize(BackbufferResolution);
+
+	// Note how many tiles we wish to render with.
+	BackbufferTileCount = FIntPoint(HighResSettings->TileCount, HighResSettings->TileCount);
+
+	// Initialize our render pass. This is a copy of the settings to make this less coupled to the Settings UI.
+	MoviePipeline::FMoviePipelineRenderPassInitSettings RenderPassInitSettings;
+	RenderPassInitSettings.BackbufferResolution = BackbufferResolution;
+	RenderPassInitSettings.TileCount = BackbufferTileCount;
+	RenderPassInitSettings.bAccumulateAlpha = bAnyOutputWantsAlpha;
+
+	// Code expects at least a 1x1 tile.
+	ensure(RenderPassInitSettings.TileCount.X > 0 && RenderPassInitSettings.TileCount.Y > 0);
+
+	// Now we need to look at all of the desired passes and find a unique set of actual engine passes that need
+	// to be rendered. This allows us to have multiple passes that re-use one render from the engine for efficiency.
+	TSet<FMoviePipelinePassIdentifier> RequiredEnginePasses;
+
+	for (UMoviePipelineRenderPass* RenderPass : GetAllRenderPasses(GetPipelineMasterConfig(), Shot))
+	{
+		RenderPass->GetRequiredEnginePasses(RequiredEnginePasses);
+	}
+
+	// There shouldn't be any render passes active from previous shots by now. The system should have flushed/stalled between
+	// them to complete using resources before switching passes.
+	check(ActiveRenderPasses.Num() == 0);
+
+	// Instantiate a new instance of every engine render pass we know how to use.
+	TArray<TSharedPtr<MoviePipeline::FMoviePipelineEnginePass>> RenderPasses;
+
+	FMovieRenderPipelineCoreModule& CoreModule = FModuleManager::Get().LoadModuleChecked<FMovieRenderPipelineCoreModule>("MovieRenderPipelineCore");
+	for (const FOnCreateEngineRenderPass& PassCreationDelegate : CoreModule.GetEngineRenderPasses())
+	{
+		TSharedRef<MoviePipeline::FMoviePipelineEnginePass> PassInstance = PassCreationDelegate.Execute();
+		if (RequiredEnginePasses.Contains(PassInstance->PassIdentifier))
+		{
+			ActiveRenderPasses.Add(PassInstance);
+			RequiredEnginePasses.Remove(PassInstance->PassIdentifier);
+		}
+	}
+
+	for(const FMoviePipelinePassIdentifier& RemainingPass : RequiredEnginePasses)
+	{
+		UE_LOG(LogMovieRenderPipeline, Warning, TEXT("Pass \"%d\" was listed as a required engine render pass but was not found. Did you forget to register it with the module?"), *RemainingPass.Name);
+		OnMoviePipelineErrored().Broadcast(this, true, LOCTEXT("MissingEnginePass", "A Render Pass specified an invalid Pass Identifier. Aborting Render. Check the log for more information."));
+	}
+
+	// Initialize each of the engine render passes.
+	for (TSharedPtr<MoviePipeline::FMoviePipelineEnginePass> EnginePass : ActiveRenderPasses)
+	{
+		EnginePass->Setup(MakeWeakObjectPtr(this), RenderPassInitSettings);
+	}
+
+	// We can now initialize the output passes and provide them a reference to the engine passes to get data from.
+	int32 NumOutputPasses = 0;
+	for (UMoviePipelineRenderPass* RenderPass : GetAllRenderPasses(GetPipelineMasterConfig(), Shot))
+	{
+		RenderPass->Setup(ActiveRenderPasses);
+		NumOutputPasses++;
+	}
+
+	UE_LOG(LogMovieRenderPipeline, Log, TEXT("Finished setting up rendering for shot. Shot has %d Engine Passes and %d Output Passes."), ActiveRenderPasses.Num(), NumOutputPasses);
+}
+
+void UMoviePipeline::TeardownRenderingPipelineForShot(FMoviePipelineShotInfo& Shot)
+{
+	// Master Configuration first.
+	for (UMoviePipelineRenderPass* RenderPass : GetAllRenderPasses(GetPipelineMasterConfig(), Shot))
+	{
+		RenderPass->Teardown();
+	}
+
+	for (TSharedPtr<MoviePipeline::FMoviePipelineEnginePass> EnginePass : ActiveRenderPasses)
+	{
+		EnginePass->Teardown();
+	}
+
+	ActiveRenderPasses.Empty();
+}
+
+void UMoviePipeline::RenderFrame()
+{
+	// Flush built in systems before we render anything. This maximizes the likelihood that the data is prepared for when
+	// the render thread uses it.
+	FlushAsyncEngineSystems();
+
+	FMoviePipelineShotInfo& CurrentShot = ShotList[CurrentShotIndex];
+	FMoviePipelineCameraCutInfo& CurrentCameraCut = CurrentShot.GetCurrentCameraCut();
+
+	// We render a frame during motion blur fixes (but don't try to submit it for output) to allow seeding
+	// the view histories with data. We want to do this all in one tick so that the motion is correct.
+	const bool bIsHistoryOnlyFrame = CurrentCameraCut.State == EMovieRenderShotState::MotionBlur;
+	const bool bIsRenderFrame = CurrentCameraCut.State == EMovieRenderShotState::Rendering;
+
+	if(!(bIsHistoryOnlyFrame || bIsRenderFrame))
+	{
+		UE_LOG(LogMovieRenderPipeline, VeryVerbose, TEXT("[%d] Skipping RenderFrame() call due to state being ineligable for rendering. State: %d"), GFrameCounter, CurrentCameraCut.State);
+		return;
+	}
+	
+	// To produce a frame from the movie pipeline we may render many frames over a period of time, additively collecting the results
+	// together before submitting it for writing on the last result - this is referred to as an "output frame". The 1 (or more) samples
+	// that make up each output frame are referred to as "sample frames". Within each sample frame, we may need to render the scene many
+	// times. In order to support ultra-high-resolution rendering (>16k) movie pipelines support building an output frame out of 'tiles'. 
+	// Each tile renders the entire viewport with a small offset which causes different samples to be picked for each final pixel. These
+	// 'tiles' are then interleaved together (on the CPU) to produce a higher resolution result. For each tile, we can render a number
+	// of jitters that get added together to produce a higher quality single frame. This is useful for cases where you may not want any 
+	// motion (such as trees fluttering in the wind) but you do want high quality anti-aliasing on the edges of the pixels. Finally,
+	// the outermost loop (which is not represented here) is accumulation over time which happens over multiple engine ticks.
+	// 
+	// In short, for each output frame, for each accumulation frame, for each tile X/Y, for each jitter, we render a pass. This setup is
+	// designed to maximize the likely hood of deterministic rendering and that different passes line up with each other.
+	UMoviePipelineAccumulationSetting* AccumulationSettings = FindOrAddSetting<UMoviePipelineAccumulationSetting>(CurrentShot);
+	UMoviePipelineCameraSetting* CameraSettings = FindOrAddSetting<UMoviePipelineCameraSetting>(CurrentShot);
+	UMoviePipelineHighResSetting* HighResSettings = FindOrAddSetting<UMoviePipelineHighResSetting>(CurrentShot);
+	UMoviePipelineOutputSetting* OutputSettings = GetPipelineMasterConfig()->FindSetting<UMoviePipelineOutputSetting>();
+	check(AccumulationSettings);
+	check(CameraSettings);
+	check(HighResSettings);
+	check(OutputSettings);
+
+	FIntPoint TileCount = FIntPoint(HighResSettings->TileCount, HighResSettings->TileCount);
+	FIntPoint OutputResolution = OutputSettings->OutputResolution;
+
+	int32 NumSpatialSamples = AccumulationSettings->SpatialSampleCount;
+	int32 NumTemporalSamples = CameraSettings->TemporalSampleCount;
+	ensure(TileCount.X > 0 && TileCount .Y> 0 && NumSpatialSamples > 0);
+
+	FrameInfo.PrevViewLocation = FrameInfo.CurrViewLocation;
+	FrameInfo.PrevViewRotation = FrameInfo.CurrViewRotation;
+
+	APlayerController* LocalPlayerController = GetWorld()->GetFirstPlayerController();
+	LocalPlayerController->GetPlayerViewPoint(FrameInfo.CurrViewLocation, FrameInfo.CurrViewRotation);
+
+	// if it is the first frame, then motion blur is disabled, and we just use the current frame camera position/rotation
+	if (bIsHistoryOnlyFrame)
+	{
+		FrameInfo.PrevViewLocation = FrameInfo.CurrViewLocation;
+		FrameInfo.PrevViewRotation = FrameInfo.CurrViewRotation;
+	}
+
+	TArray<UMoviePipelineRenderPass*> InputBuffers = GetAllRenderPasses(GetPipelineMasterConfig(), CurrentShot);
+
+	// If this is the first sample for a new frame, we want to notify the output builder that it should expect data to accumulate for this frame.
+	if (CachedOutputState.TemporalSampleIndex == 0)
+	{
+		// This happens before any data is queued for this frame.
+		FMoviePipelineMergerOutputFrame& OutputFrame = OutputBuilder->QueueOutputFrame_GameThread(CachedOutputState);
+
+		// Now we need to go through all passes and get any identifiers from them of what this output frame should expect.
+		for (UMoviePipelineRenderPass* RenderPass : InputBuffers)
+		{
+			RenderPass->GatherOutputPasses(OutputFrame.ExpectedRenderPasses);
+		}
+	}
+
+	for (int32 TileY = 0; TileY < TileCount.Y; TileY++)
+	{
+		for (int32 TileX = 0; TileX < TileCount.X; TileX++)
+		{
+			/*
+			* Two different features need to move the viewport and need to work in combination with each other.
+			* Ultra-high resolution rendering and single-frame anti-aliasing both work by moving the viewport
+			* by a small amount to create new samples for each render pass which generates more usable data.
+			*
+			* High resolution rendering is referred to as "tiles" but it is not actually a tile based render.
+			* Splitting the view into many smaller rectangles and using an off-center camera projection on them
+			* wouldn't work due to post-processing effects (such as vignette) as you would end up with many copies
+			* of the effect on your final image. Instead, it renders the entire screen for every 'tile', but shifts
+			* the camera position a portion of a pixel to cause new, unique data to be captured in the next tile.
+			*
+			* Single-frame anti-aliasing works in a similar way, by shifting the view by less than a pixel for each
+			* sample in the anti-aliasing. While the high-resolution rendering is a specific grid pattern, the
+			* anti-aliasing uses a more nuanced technique for picking offsets. It uses a deterministic Halton
+			* sequence to pick the offset, weighted by a Gaussian curve to mostly focus on samples around the
+			* center, but occasionally picking an offset which would be a sample outside of it's current sub-pixel.
+			*
+			* In the example below we're looking at one render target pixel, split into a 2x2 tiling. To gain
+			* 2x resolution, we shift the viewport so that the center of the pixel is no longer 'o', but now
+			* 'v'. Since this shift is less than a whole pixel we can interlace the resulting images and they
+			* won't be duplicates-but-shifted-by-one of each other. This is how we end up with new unique data.
+			*
+			* |-------------|-------------|
+			* |        a    |             |
+			* |    a        |             |
+			* |      v      |      v      |
+			* |   a         |             |
+			* |a       a    |             |
+			* |-------------o-------------|
+			* |             |             |
+			* |             |             |
+			* |      v      |      v      |
+			* |             |             |
+			* |             |             |
+			* |-------------|-------------|
+			*
+			* Now that a given render target pixel |v| has been shifted to it's correct location, we need to randomly
+			* sample within that box (perhaps with a slight overshoot) for the anti-aliasing samples. These samples
+			* ('a') are added together to create one |v| pixel and then only the |v| pixels are interleaved.
+			* Interleaving is performed on the CPU to support > 16k output sizes which is the maximum texture resolution
+			* for most GPUs.
+			*/
+
+			// Calculate the sub-pixel shift to get this pass into the right portion of the main pixel. Assuming the
+			// pixel coordinates are [0-1], the size of a sub-pixel (for tiles) is 1/NumTiles. To get to the right
+			// 'quadrant' we need to take our current tile index and multiply it by the size of each tile, so
+			// for the top right quadrant (in above example) size=(1/2), offset=(size*1). Finally we add half
+			// of the size to put us in the center of that pixel.
+			double TilePixelSizeFractionX;
+			double TilePixelSizeFractionY; 
+			double TileShiftX; 
+			double TileShiftY;
+
+			{
+				TileShiftX = 0.5f;
+				TileShiftY = 0.5f;
+
+				TilePixelSizeFractionX = 1.0f;
+				TilePixelSizeFractionY = 1.0f;
+			}
+
+			// If they're trying to seed the histories, then we'll just override the NumSpatialSamples since they don't need all samples as
+			// this isn't output.
+			if (bIsHistoryOnlyFrame)
+			{
+				NumSpatialSamples = 1;
+			}
+
+			// Now we want to render a user-configured number of spatial jitters to come up with the final output for this tile. 
+			for (int32 SpatialSample = 0; SpatialSample < NumSpatialSamples; SpatialSample++)
+			{
+				// Count this as a sample rendered for the current work.
+				CurrentCameraCut.CurrentWorkInfo.NumSamples++;
+
+				float OffsetX = 0.5f;
+				float OffsetY = 0.5f; 
+
+				// Only jitter when required, as jitter requires TAA to be off. 
+				// Also, we want jitter if either time samples or spatial samples are > 1, otherwise not.
+				if (!bIsHistoryOnlyFrame &&
+					(NumSpatialSamples > 1 || NumTemporalSamples > 1))
+				{
+					OffsetX = Halton((SpatialSample + 1) + (CachedOutputState.TemporalSampleIndex * NumSpatialSamples), 2);
+					OffsetY = Halton((SpatialSample + 1) + (CachedOutputState.TemporalSampleIndex * NumSpatialSamples), 3);
+				}
+
+				// Our spatial sample can only move within a +/- range within the tile we're in.
+				double SpatialShiftX = FMath::GetMappedRangeValueClamped(FVector2D(0, 1), FVector2D(-TilePixelSizeFractionX / 2.f, TilePixelSizeFractionX / 2.f), OffsetX);
+				double SpatialShiftY = FMath::GetMappedRangeValueClamped(FVector2D(0, 1), FVector2D(-TilePixelSizeFractionY / 2.f, TilePixelSizeFractionY / 2.f), OffsetY);
+
+				double FinalSubPixelShiftX = - (TileShiftX + SpatialShiftX - 0.5);
+				double FinalSubPixelShiftY = (TileShiftY + SpatialShiftY - 0.5);
+
+				FIntPoint BackbufferResolution = OutputSettings->OutputResolution;
+				FIntPoint TileResolution = BackbufferResolution;
+
+				{
+					BackbufferResolution = FIntPoint(
+						FMath::CeilToInt(BackbufferResolution.X / TileCount.X),
+						FMath::CeilToInt(BackbufferResolution.Y / TileCount.Y));
+
+					// Cache tile resolution before we add padding.
+					TileResolution = BackbufferResolution;
+
+					// Apply size padding.
+					BackbufferResolution = HighResSettings->CalculatePaddedBackbufferSize(BackbufferResolution);
+				}
+
+				// We take all of the information needed to render a single sample and package it into a struct.
+				FMoviePipelineRenderPassMetrics SampleState;
+				SampleState.OutputState = CachedOutputState;
+				SampleState.SpatialShift = FVector2D((float)( FinalSubPixelShiftX) * 2.0f / BackbufferResolution.X, (float)FinalSubPixelShiftY * 2.0f / BackbufferResolution.X);
+				SampleState.TileIndexes = FIntPoint(TileX, TileY);
+				SampleState.TileCounts = TileCount;
+				SampleState.bIsHistoryOnlyFrame = bIsHistoryOnlyFrame;
+				SampleState.SpatialSampleIndex = SpatialSample;
+				SampleState.SpatialSampleCount = NumSpatialSamples;
+				SampleState.TemporalSampleIndex = CachedOutputState.TemporalSampleIndex;
+				SampleState.TemporalSampleCount = CameraSettings->TemporalSampleCount;
+				SampleState.AccumulationGamma = AccumulationSettings->AccumulationGamma;
+				SampleState.BackbufferSize = BackbufferResolution;
+				SampleState.TileSize = TileResolution;
+				SampleState.FrameInfo = FrameInfo;
+				SampleState.bWriteSampleToDisk = HighResSettings->bWriteAllSamples;
+				SampleState.ExposureCompensation = CameraSettings->bManualExposure ? CameraSettings->ExposureCompensation : TOptional<float>();
+				SampleState.TextureSharpnessBias = HighResSettings->TextureSharpnessBias;
+				{
+					SampleState.OverlappedPad = FIntPoint(FMath::CeilToInt(TileResolution.X * HighResSettings->OverlapPercentage), 
+														   FMath::CeilToInt(TileResolution.Y * HighResSettings->OverlapPercentage));
+					SampleState.OverlappedOffset = FIntPoint(TileX * TileResolution.X - SampleState.OverlappedPad.X,
+															  TileY * TileResolution.Y - SampleState.OverlappedPad.Y);
+					SampleState.OverlappedSubpixelShift = FVector2D(1.0f - OffsetX,OffsetY);
+				}
+
+				SampleState.WeightFunctionX.InitHelper(SampleState.OverlappedPad.X, SampleState.TileSize.X, SampleState.OverlappedPad.X);
+				SampleState.WeightFunctionY.InitHelper(SampleState.OverlappedPad.Y, SampleState.TileSize.Y, SampleState.OverlappedPad.Y);
+
+				// Now we can request that all of the engine passes render. The individual render passes should have already registered delegates
+				// to receive data when the engine render pass is run, so no need to run them.
+				for (TSharedPtr<MoviePipeline::FMoviePipelineEnginePass> EnginePass : ActiveRenderPasses)
+				{
+					EnginePass->RenderSample_GameThread(SampleState);
+				}
+			}
+		}
+	}
+	
+	// UE_LOG(LogMovieRenderPipeline, Warning, TEXT("[%d] Pre-FlushRenderingCommands"), GFrameCounter);
+	FlushRenderingCommands();
+	// UE_LOG(LogMovieRenderPipeline, Warning, TEXT("[%d] Post-FlushRenderingCommands"), GFrameCounter);
+}
+
+void UMoviePipeline::OnFrameCompletelyRendered(FMoviePipelineMergerOutputFrame&& OutputFrame, const TSharedRef<FImagePixelDataPayload, ESPMode::ThreadSafe> InFrameData)
+{
+	UE_LOG(LogMovieRenderPipeline, Warning, TEXT("[%d] Data required for output available! Frame: %d"), GFrameCounter, OutputFrame.FrameOutputState.OutputFrameNumber);
+	
+	for (UMoviePipelineOutputBase* OutputContainer : GetPipelineMasterConfig()->GetOutputContainers())
+	{
+		OutputContainer->OnRecieveImageDataImpl(&OutputFrame);
+	}
+}
+
+void UMoviePipeline::OnSampleRendered(TUniquePtr<FImagePixelData>&& OutputSample, const TSharedRef<FImagePixelDataPayload, ESPMode::ThreadSafe> InFrameData)
+{
+	UMoviePipelineOutputSetting* OutputSettings = GetPipelineMasterConfig()->FindSetting<UMoviePipelineOutputSetting>();
+	check(OutputSettings);
+
+	TUniquePtr<FImageWriteTask> TileImageTask = MakeUnique<FImageWriteTask>();
+
+	// Fill alpha for now 
+	/*switch (OutputSample->GetType())
+	{
+		case EImagePixelType::Color:
+		{
+			TileImageTask->PixelPreProcessors.Add(TAsyncAlphaWrite<FColor>(255));
+			break;
+		}
+		case EImagePixelType::Float16:
+		{
+
+			break;
+		}
+		case EImagePixelType::Float32:
+		{
+			TileImageTask->PixelPreProcessors.Add(TAsyncAlphaWrite<FLinearColor>(255));
+			break;
+		}
+		default:
+			check(false);
+	}*/
+
+	// JPEG output
+	TileImageTask->Format = EImageFormat::EXR;
+	TileImageTask->CompressionQuality = 100;
+
+	FString OutputName = FString::Printf(TEXT("/%s_SS_%d_TS_%d_TileX_%d_TileY_%d.%d.jpeg"),
+		*InFrameData->PassIdentifier.Name, InFrameData->SampleState.SpatialSampleIndex, InFrameData->SampleState.TemporalSampleIndex,
+		InFrameData->SampleState.TileIndexes.X, InFrameData->SampleState.TileIndexes.Y, InFrameData->OutputState.OutputFrameNumber);
+
+	FString OutputDirectory = OutputSettings->OutputDirectory.Path;
+	FString OutputPath = OutputDirectory + OutputName;
+	TileImageTask->Filename = OutputPath;
+
+	// Duplicate the data so that the Image Task can own it.
+	TileImageTask->PixelData = MoveTemp(OutputSample);
+	ImageWriteQueue->Enqueue(MoveTemp(TileImageTask));
+}
+
+
+
+void UMoviePipeline::FlushAsyncEngineSystems()
+{
+	// Flush Level Streaming. This solves the problem where levels that are not controlled
+	// by the Sequencer Level Visibility track are marked for Async Load by a gameplay system.
+	// This will register any new actors/components that were spawned during this frame. This needs 
+	// to be done before the shader compiler is flushed so that we compile shaders for any newly
+	// spawned component materials.
+	if (GetWorld())
+	{
+		GetWorld()->FlushLevelStreaming(EFlushLevelStreamingType::Full);
+	}
+
+	// Now we can flush the shader compiler. ToDo: This should probably happen right before SendAllEndOfFrameUpdates() is normally called
+	if (GShaderCompilingManager)
+	{
+		bool bDidWork = false;
+		int32 NumShadersToCompile = GShaderCompilingManager->GetNumRemainingJobs();
+		if (NumShadersToCompile > 0)
+		{
+			UE_LOG(LogMovieRenderPipeline, Log, TEXT("[%d] Starting build for %d shaders."), GFrameCounter, NumShadersToCompile);
+		}
+
+		while (GShaderCompilingManager->GetNumRemainingJobs() > 0 || GShaderCompilingManager->HasShaderJobs())
+		{
+			UE_LOG(LogMovieRenderPipeline, Log, TEXT("[%d] Waiting for %d shaders [Has Shader Jobs: %d] to finish compiling..."), GFrameCounter, GShaderCompilingManager->GetNumRemainingJobs(), GShaderCompilingManager->HasShaderJobs());
+			GShaderCompilingManager->ProcessAsyncResults(false, true);
+
+			// Sleep for 1 second and then check again. This way we get an indication of progress as this works.
+			FPlatformProcess::Sleep(1.f);
+			bDidWork = true;
+		}
+
+		if (bDidWork)
+		{
+			UE_LOG(LogMovieRenderPipeline, Log, TEXT("[%d] Done building %d shaders."), GFrameCounter, NumShadersToCompile);
+		}
+	}
+
+	// Flush the Mesh Distance Field builder as well.
+	if (GDistanceFieldAsyncQueue)
+	{
+		bool bDidWork = false;
+		int32 NumDistanceFieldsToBuild = GDistanceFieldAsyncQueue->GetNumOutstandingTasks();
+		if (NumDistanceFieldsToBuild > 0)
+		{
+			UE_LOG(LogMovieRenderPipeline, Log, TEXT("[%d] Starting build for %d mesh distance fields."), GFrameCounter, NumDistanceFieldsToBuild);
+		}
+
+		while (GDistanceFieldAsyncQueue->GetNumOutstandingTasks() > 0)
+		{
+			UE_LOG(LogMovieRenderPipeline, Log, TEXT("[%d] Waiting for %d Mesh Distance Fields to finish building..."), GFrameCounter, GDistanceFieldAsyncQueue->GetNumOutstandingTasks());
+			GDistanceFieldAsyncQueue->ProcessAsyncTasks();
+
+			// Sleep for 1 second and then check again. This way we get an indication of progress as this works.
+			FPlatformProcess::Sleep(1.f);
+			bDidWork = true;
+		}
+
+		if (bDidWork)
+		{
+			UE_LOG(LogMovieRenderPipeline, Log, TEXT("[%d] Done building %d Mesh Distance Fields."), GFrameCounter, NumDistanceFieldsToBuild);
+		}
+	}
+
+	// Flush grass
+	{
+		for (TActorIterator<ALandscapeProxy> It(GetWorld()); It; ++It)
+		{
+			ALandscapeProxy* LandscapeProxy = (*It);
+			if (LandscapeProxy)
+			{
+				TArray<FVector> CameraList;
+				LandscapeProxy->UpdateGrass(CameraList, true);
+			}
+		}
+	}
+
+	// Flush virtual texture tile calculations
+	ERHIFeatureLevel::Type FeatureLevel = GetWorld()->FeatureLevel;
+	ENQUEUE_RENDER_COMMAND(VirtualTextureSystemFlushCommand)(
+		[FeatureLevel](FRHICommandListImmediate& RHICmdList)
+	{
+		GetRendererModule().LoadPendingVirtualTextureTiles(RHICmdList, FeatureLevel);
+	});
+}
+
+#undef LOCTEXT_NAMESPACE // "MoviePipeline"
diff --git a/Engine/Plugins/MovieScene/MovieRenderPipeline/Source/MovieRenderPipelineCore/Private/MoviePipelineSetting.cpp b/Engine/Plugins/MovieScene/MovieRenderPipeline/Source/MovieRenderPipelineCore/Private/MoviePipelineSetting.cpp
new file mode 100644
--- /dev/null
+++ b/Engine/Plugins/MovieScene/MovieRenderPipeline/Source/MovieRenderPipelineCore/Private/MoviePipelineSetting.cpp
@@ -0,0 +1,28 @@
+// Copyright 1998-2019 Epic Games, Inc. All Rights Reserved.
+
+#include "MoviePipelineSetting.h"
+#include "MoviePipeline.h"
+
+UMoviePipelineSetting::UMoviePipelineSetting()
+	: bEnabled(true)
+{
+}
+void UMoviePipelineSetting::OnMoviePipelineInitialized(UMoviePipeline* InPipeline)
+{
+	CachedPipeline = MakeWeakObjectPtr(InPipeline);
+	SetupForPipelineImpl(InPipeline);
+}
+
+UMoviePipeline* UMoviePipelineSetting::GetPipeline() const
+{
+	// If this check trips then life cycles aren't as expected.
+	UMoviePipeline* OutPipeline = CachedPipeline.Get();
+	check(OutPipeline);
+
+	return OutPipeline;
+}
+
+UWorld* UMoviePipelineSetting::GetWorld() const
+{
+	return GetPipeline()->GetWorld();
+}
\ No newline at end of file
diff --git a/Engine/Plugins/MovieScene/MovieRenderPipeline/Source/MovieRenderPipelineCore/Private/MovieRenderPipelineCommandLine.cpp b/Engine/Plugins/MovieScene/MovieRenderPipeline/Source/MovieRenderPipelineCore/Private/MovieRenderPipelineCommandLine.cpp
new file mode 100644
--- /dev/null
+++ b/Engine/Plugins/MovieScene/MovieRenderPipeline/Source/MovieRenderPipelineCore/Private/MovieRenderPipelineCommandLine.cpp
@@ -0,0 +1,239 @@
+// Copyright 1998-2019 Epic Games, Inc. All Rights Reserved.
+
+#include "MovieRenderPipelineCoreModule.h"
+#include "Modules/ModuleInterface.h"
+#include "Misc/CommandLine.h"
+#include "Modules/ModuleManager.h"
+#include "MoviePipelineExecutor.h"
+#include "MoviePipeline.h"
+#include "MoviePipelineInProcessExecutor.h"
+#include "LevelSequence.h"
+#include "MoviePipelineMasterConfig.h"
+#include "MovieRenderPipelineDataTypes.h"
+
+#if WITH_EDITOR
+//#include "Editor.h"
+#endif
+
+void FMovieRenderPipelineCoreModule::InitializeCommandLineMovieRender()
+{
+	// This should be called after the engine has a chance to load all of the modules so that all of our classes are loaded.
+	if (!ensureMsgf(SequenceAssetValue.Len() > 0, TEXT("InitializeCommandLineMovieRender shouldn't be called without an asset to render!")))
+	{
+		return;
+	}
+
+#if WITH_EDITOR
+	//const bool bIsGameMode = !GEditor;
+	//if (!bIsGameMode)
+	//{
+	//	UE_LOG(LogMovieRenderPipeline, Fatal, TEXT("Command Line Renders must be performed in -game mode, otherwise use the editor ui/python and PIE. Add -game to your command line arguments."));
+	//	FPlatformMisc::RequestExitWithStatus(false, MoviePipelineErrorCodes::Critical);
+	//	return;
+	//}
+#endif
+
+	// Attempt to convert their command line arguments into the required objects.
+	ULevelSequence* LevelSequence = nullptr;
+	UMoviePipelineMasterConfig* MasterConfig = nullptr;
+	UClass* ExecutorClass = nullptr;
+	UClass* PipelineClass = nullptr;
+
+	uint8 ReturnCode = GetMovieRenderClasses(SequenceAssetValue, SettingsAssetValue, MoviePipelineLocalExecutorClassType, MoviePipelineClassType, 
+											 /*Out*/ LevelSequence, /*Out*/ MasterConfig, /*Out*/ ExecutorClass, /*Out*/ PipelineClass);
+	if (!ensureMsgf(LevelSequence && MasterConfig && ExecutorClass && PipelineClass, TEXT("There was a failure parsing the command line and a movie render cannot be started. Check the log for more details.")))
+	{
+		// Take the failure return code from the detection of our command line arguments.
+		FPlatformMisc::RequestExitWithStatus(/*Force*/ false, /*ReturnCode*/ ReturnCode);
+		return;
+	}
+	else
+	{
+		UE_LOG(LogMovieRenderPipeline, Log, TEXT("Successfully detected and loaded required movie arguments. Rendering will begin once the map is loaded."));
+		UE_LOG(LogMovieRenderPipeline, Log, TEXT("LevelSequence: %s Config: %s ExecutorClass: %s PipelineClass: %s"), *LevelSequence->GetPathName(), *MasterConfig->GetPathName(), *ExecutorClass->GetPathName(), *PipelineClass->GetPathName());
+	}
+
+	// By this time, we know what assets you want to render, how to process the array of assets, and what runs an individual render. First we will create an executor.
+	// UMoviePipelineExecutorBase* ExecutorBase = NewObject<UMoviePipelineExecutorBase>(GetTransientPackage(), ExecutorClass);
+	// 
+	// // We add the Executor to the root set. It will own all of the configuration data so this keeps it nicely in memory until finished,
+	// // and means we only have to add/remove one thing from the root set, everything else uses normal outer ownership.
+	// ExecutorBase->AddToRoot();
+	// ExecutorBase->OnExecutorFinished().AddRaw(this, &FMovieRenderPipelineCoreModule::OnCommandLineMovieRenderCompleted);
+	// ExecutorBase->OnExecutorErrored().AddRaw(this, &FMovieRenderPipelineCoreModule::OnCommandLineMovieRenderErrored);
+	// ExecutorBase->SetMoviePipelineClass(PipelineClass);
+	// 
+	// TArray<FMoviePipelineExecutorJobPrev> Jobs;
+	// Jobs.Add(FMoviePipelineExecutorJobPrev(SequenceAssetValue, MasterConfig));
+	// 
+	// ExecutorBase->Execute(Jobs);
+}
+
+void FMovieRenderPipelineCoreModule::OnCommandLineMovieRenderCompleted(UMoviePipelineExecutorBase* InExecutor, bool bSuccess)
+{
+	if (InExecutor)
+	{
+		InExecutor->RemoveFromRoot();
+	}
+
+	// Return a success code. Ideally any errors detected during rendering will return different codes.
+	FPlatformMisc::RequestExitWithStatus(/*Force*/ false, /*ReturnCode*/ 0);
+}
+
+void FMovieRenderPipelineCoreModule::OnCommandLineMovieRenderErrored(UMoviePipelineExecutorBase* InExecutor, UMoviePipeline* InPipelineWithError, bool bIsFatal, FText ErrorText)
+{
+
+}
+
+
+UClass* GetLocalExecutorClass(const FString& MoviePipelineLocalExecutorClassType, const FString ExecutorClassFormatString)
+{
+	if (MoviePipelineLocalExecutorClassType.Len() > 0)
+	{
+		UMoviePipelineExecutorBase* LoadedExecutor = LoadObject<UMoviePipelineExecutorBase>(GetTransientPackage(), *MoviePipelineLocalExecutorClassType);
+		if (LoadedExecutor)
+		{
+			UE_LOG(LogMovieRenderPipeline, Log, TEXT("Loaded explicitly specified Movie Pipeline Executor %s."), *MoviePipelineLocalExecutorClassType);
+			return LoadedExecutor->GetClass();
+		}
+		else
+		{
+			// They explicitly specified an object, but that couldn't be loaded so it's an error.
+			UE_LOG(LogMovieRenderPipeline, Fatal, TEXT("Failed to load specified Local Executor class. Executor Class: %s"), *MoviePipelineLocalExecutorClassType);
+			UE_LOG(LogMovieRenderPipeline, Warning, TEXT("%s"), *ExecutorClassFormatString);
+			return nullptr;
+		}
+	}
+	else
+	{
+		// Fall back to our provided one. This is okay because it doesn't come from a user preference,
+		// since it needs an different executor than the editor provided New Process
+		UE_LOG(LogMovieRenderPipeline, Log, TEXT("Using default Movie Pipeline Executor %s. See '-MoviePipelineLocalExecutorClass' if you need to override this."), *MoviePipelineLocalExecutorClassType);
+		return UMoviePipelineInProcessExecutor::StaticClass();
+	}
+}
+
+UClass* GetMoviePipelineClass(const FString& MoviePipelineLocalExecutorClassType, const FString ExecutorClassFormatString)
+{
+	if (MoviePipelineLocalExecutorClassType.Len() > 0)
+	{
+		UMoviePipelineExecutorBase* LoadedExecutor = LoadObject<UMoviePipelineExecutorBase>(GetTransientPackage(), *MoviePipelineLocalExecutorClassType);
+		if (LoadedExecutor)
+		{
+			UE_LOG(LogMovieRenderPipeline, Log, TEXT("Loaded explicitly specified Movie Pipeline Executor %s."), *MoviePipelineLocalExecutorClassType);
+			return LoadedExecutor->GetClass();
+		}
+		else
+		{
+			// They explicitly specified an object, but that couldn't be loaded so it's an error.
+			UE_LOG(LogMovieRenderPipeline, Fatal, TEXT("Failed to load specified Local Executor class. Executor Class: %s"), *MoviePipelineLocalExecutorClassType);
+			UE_LOG(LogMovieRenderPipeline, Warning, TEXT("%s"), *ExecutorClassFormatString);
+			return nullptr;
+		}
+	}
+	else
+	{
+		UE_LOG(LogMovieRenderPipeline, Log, TEXT("Using default Movie Pipeline %s. See '-MoviePipelineClass' if you need to override this."), *MoviePipelineLocalExecutorClassType);
+		return UMoviePipeline::StaticClass();
+	}
+}
+
+
+bool FMovieRenderPipelineCoreModule::IsTryingToRenderMovieFromCommandLine(FString& OutSequenceAssetPath, FString& OutConfigAssetPath, FString& OutExecutorType, FString& OutPipelineType) const
+{
+	// Look to see if they've specified a Movie Pipeline to run. This should be in the format:
+	// "/Game/MySequences/MySequence.MySequence"
+	FParse::Value(FCommandLine::Get(), TEXT("-MoviePipeline="), OutSequenceAssetPath);
+
+	// Look to see if they've specified a configuration to use. This should be in the format:
+	// "/Game/MyRenderSettings/MyHighQualitySetting.MyHighQualitySetting" or an absolute path 
+	// to a exported json file.
+	FParse::Value(FCommandLine::Get(), TEXT("-MoviePipelineConfig="), OutConfigAssetPath);
+
+	// The user may want to override the executor. By default, we use the one specified in the Project
+	// Settings, but we allow them to override it on the command line (for render farms, etc.) in case
+	// the one used when a human is running the box isn't appropriate. This should be in the format:
+	// "/Script/ModuleName.ClassNameNoUPrefix"
+	FParse::Value(FCommandLine::Get(), TEXT("-MoviePipelineLocalExecutorClass="), OutExecutorType);
+
+	// The user may want to override the Movie Pipeline itself. By default we use the one specified in the Project
+	// Setting, but we allow overriding it for consistency's sake really. This should be in the format:
+	// "/Script/ModuleName.ClassNameNoUPrefix"
+	FParse::Value(FCommandLine::Get(), TEXT("-MoviePipelineClass="), OutPipelineType);
+
+	return OutSequenceAssetPath.Len() > 0;
+}
+
+uint8 FMovieRenderPipelineCoreModule::GetMovieRenderClasses(const FString& InSequenceAssetPath, const FString& InConfigAssetPath, const FString& InExecutorType, const FString& InPipelineType,
+	ULevelSequence*& OutSequence, UMoviePipelineMasterConfig*& OutConfig, UClass*& OutExecutorClass, UClass*& OutPipelineClass) const
+{
+	if (!ensureMsgf(InSequenceAssetPath.Len() > 0, TEXT("GetMovieRenderClasses should not be called without an sequence asset path specified!")))
+	{
+		return MoviePipelineErrorCodes::NoAsset;
+	}
+
+	// Store off the messages that print the expected format since they're printed in a couple places.
+	const FString SequenceAssetFormatString = TEXT("Movie Pipeline Asset should be specified in the format '-MoviePipeline=\"/Game/MySequences/MyPipelineAsset\", or a relative (to executable)/absolute path to a JSON file in the format '-MoviePipeline=\"D:/MyMoviePipelineConfig.json\".");
+	const FString ConfigAssetFormatString	= TEXT("Pipeline Config Asset should be specified in the format '-MoviePipelineConfig=\"/Game/MyRenderSettings/MyHighQualitySetting.MyHighQualitySetting\"'");
+	const FString PipelineClassFormatString = TEXT("Movie Pipeline Class should be specified in the format '-MoviePipelineClass=\"/Script/ModuleName.ClassNameNoUPrefix\"'");
+	const FString ExecutorClassFormatString = TEXT("Pipeline Executor Class should be specified in the format '-MoviePipelineLocalExecutorClass=\"/Script/ModuleName.ClassNameNoUPrefix\"'");
+
+	// They need both an asset and a configuration so early out if we don't have that.
+	if (InConfigAssetPath.Len() == 0)
+	{
+		UE_LOG(LogMovieRenderPipeline, Fatal, TEXT("No Movie Pipeline Config was specified to use as settings."));
+		UE_LOG(LogMovieRenderPipeline, Warning, TEXT("%s"), *ConfigAssetFormatString);
+		return MoviePipelineErrorCodes::NoConfig;
+	}
+
+	OutSequence = nullptr;
+	OutConfig = nullptr;
+	OutExecutorClass = nullptr;
+	OutPipelineClass = nullptr;
+
+	// We try to detect executor/class overrides first because it'll affect deserialization of JSON based pipeline assets.
+	OutExecutorClass = GetLocalExecutorClass(MoviePipelineLocalExecutorClassType, ExecutorClassFormatString);
+	OutPipelineClass = GetMoviePipelineClass(MoviePipelineClassType, PipelineClassFormatString);
+
+	if (!ensureMsgf(OutPipelineClass && OutExecutorClass, TEXT("Attempted to render a movie pipeline but could not load executor or pipeline class, nor fall back to defaults.")))
+	{
+		return MoviePipelineErrorCodes::Critical;
+	}
+
+	// Locate and load the level sequence they wish to render.
+	{
+		// Convert it to a soft object path and use that load to ensure it follows redirectors, etc.
+		FSoftObjectPath AssetPath = FSoftObjectPath(InSequenceAssetPath);
+		OutSequence = CastChecked<ULevelSequence>(AssetPath.TryLoad());
+
+		if (!OutSequence)
+		{
+			UE_LOG(LogMovieRenderPipeline, Fatal, TEXT("Failed to find Movie Pipeline sequence asset to render. Please note that the /Content/ part of the on-disk structure is omitted. Looked for: %s"), *InSequenceAssetPath);
+			UE_LOG(LogMovieRenderPipeline, Warning, TEXT("%s"), *SequenceAssetFormatString);
+			return MoviePipelineErrorCodes::NoAsset;
+		}
+	}
+
+	// Now look for the configuration file to see how to render it.
+	if (InConfigAssetPath.StartsWith("/Game/"))
+	{
+		// Convert it to a soft object path and use that load to ensure it follows redirectors, etc.
+		FSoftObjectPath AssetPath = FSoftObjectPath(InConfigAssetPath);
+		OutConfig = CastChecked<UMoviePipelineMasterConfig>(AssetPath.TryLoad());
+
+		if (!OutConfig)
+		{
+			UE_LOG(LogMovieRenderPipeline, Fatal, TEXT("Failed to find Pipeline Configuration asset to render. Please note that the /Content/ part of the on-disk structure is omitted. Looked for: %s"), *InConfigAssetPath);
+			UE_LOG(LogMovieRenderPipeline, Warning, TEXT("%s"), *ConfigAssetFormatString);
+			return MoviePipelineErrorCodes::NoConfig;
+		}
+	}
+	else
+	{
+		// If they didn't pass a path that started with /Game/, we'll try to see if it is a JSON file we can understand.
+		UE_LOG(LogMovieRenderPipeline, Fatal, TEXT("Not supported yet"));
+		return MoviePipelineErrorCodes::Critical;
+	}
+
+	return MoviePipelineErrorCodes::Success;
+}
diff --git a/Engine/Plugins/MovieScene/MovieRenderPipeline/Source/MovieRenderPipelineCore/Private/MovieRenderPipelineCoreModule.cpp b/Engine/Plugins/MovieScene/MovieRenderPipeline/Source/MovieRenderPipelineCore/Private/MovieRenderPipelineCoreModule.cpp
new file mode 100644
--- /dev/null
+++ b/Engine/Plugins/MovieScene/MovieRenderPipeline/Source/MovieRenderPipelineCore/Private/MovieRenderPipelineCoreModule.cpp
@@ -0,0 +1,38 @@
+// Copyright 1998-2019 Epic Games, Inc. All Rights Reserved.
+
+#include "MovieRenderPipelineCoreModule.h"
+#include "Modules/ModuleInterface.h"
+#include "Misc/CoreDelegates.h"
+#include "Modules/ModuleManager.h"
+#include "MovieRenderPipelineDataTypes.h"
+
+void FMovieRenderPipelineCoreModule::StartupModule()
+{
+	// Look to see if they supplied arguments on the command line indicating they wish to render a movie.
+	if (IsTryingToRenderMovieFromCommandLine(SequenceAssetValue, SettingsAssetValue, MoviePipelineLocalExecutorClassType, MoviePipelineClassType))
+	{
+		UE_LOG(LogMovieRenderPipeline, Log, TEXT("Detected that the user intends to render a movie. Waiting until engine loop init is complete to ensure "))
+		// Register a hook to wait until the engine has finished loading to increase the likelihood that the desired classes are loaded.
+		FCoreDelegates::OnPostEngineInit.AddRaw(this, &FMovieRenderPipelineCoreModule::InitializeCommandLineMovieRender);
+	}
+}
+
+void FMovieRenderPipelineCoreModule::ShutdownModule()
+{
+}
+
+FDelegateHandle FMovieRenderPipelineCoreModule::RegisterEngineRenderPass(FOnCreateEngineRenderPass InOnCreateEngineRenderPass)
+{
+	EngineRenderPassDelegates.Add(InOnCreateEngineRenderPass);
+	FDelegateHandle Handle = EngineRenderPassDelegates.Last().GetHandle();
+
+	return Handle;
+}
+
+void FMovieRenderPipelineCoreModule::UnregisterEngineRenderPass(FDelegateHandle InHandle)
+{
+	EngineRenderPassDelegates.RemoveAll([=](const FOnCreateEngineRenderPass& Delegate) { return Delegate.GetHandle() == InHandle; });
+}
+
+IMPLEMENT_MODULE(FMovieRenderPipelineCoreModule, MovieRenderPipelineCore);
+DEFINE_LOG_CATEGORY(LogMovieRenderPipeline); 
diff --git a/Engine/Plugins/MovieScene/MovieRenderPipeline/Source/MovieRenderPipelineCore/Public/MoviePipeline.h b/Engine/Plugins/MovieScene/MovieRenderPipeline/Source/MovieRenderPipelineCore/Public/MoviePipeline.h
new file mode 100644
--- /dev/null
+++ b/Engine/Plugins/MovieScene/MovieRenderPipeline/Source/MovieRenderPipelineCore/Public/MoviePipeline.h
@@ -0,0 +1,307 @@
+// Copyright 1998-2019 Epic Games, Inc. All Rights Reserved.
+#pragma once
+
+#include "UObject/Object.h"
+#include "Engine/EngineCustomTimeStep.h"
+#include "MovieRenderPipelineDataTypes.h"
+#include "MovieSceneTimeController.h"
+#include "MoviePipeline.generated.h"
+
+// Forward Declares
+class UMoviePipelineMasterConfig;
+class ULevelSequence;
+class UMovieSceneSequencePlayer;
+class UMoviePipelineCustomTimeStep;
+class UEngineCustomTimeStep;
+class ALevelSequenceActor;
+class UMovieRenderDebugWidget;
+class FMovieRenderViewport;
+class FMovieRenderViewportClient;
+struct FImagePixelPipe;
+struct FMoviePipelineTimeController;
+class FMoviePipelineOutputMerger;
+class IImageWriteQueue;
+class UMoviePipelineExecutorJob;
+class UMoviePipelineSetting;
+
+
+
+DECLARE_MULTICAST_DELEGATE_OneParam(FMoviePipelineFinished, UMoviePipeline*);
+DECLARE_MULTICAST_DELEGATE_ThreeParams(FMoviePipelineErrored, UMoviePipeline* /*Pipeline*/, bool /*bIsFatal*/, FText /*ErrorText*/);
+
+UCLASS(Blueprintable)
+class MOVIERENDERPIPELINECORE_API UMoviePipeline : public UObject
+{
+	GENERATED_BODY()
+	
+public:
+	UMoviePipeline();
+
+	/** 
+	* Initialize the movie pipeline with the specified settings. This kicks off the rendering process. 
+	
+	* @param InJob	- This contains settings and sequence to render this Movie Pipeline with.
+	*/
+	void Initialize(UMoviePipelineExecutorJob* InJob);
+
+	/** 
+	* Call to shut down the pipeline. This flushes any outstanding file writes and unregisters all delegates.
+	* This will not call the OnMoviePipelineFinished() delegate as it assumes the external party already knows
+	* it is "Finished". 
+	*/
+	void Shutdown();
+
+	/** 
+	* Called when we have completely finished this pipeline. This means that all frames have been rendered,
+	* all files written to disk, and any post-finalize exports have finished. This Pipeline will call
+	* Shutdown() on itself before calling this delegate to ensure we've unregistered from all delegates
+	* and are no longer trying to do anything (even if we still exist).
+	*/
+	FMoviePipelineFinished& OnMoviePipelineFinished()
+	{
+		return OnMoviePipelineFinishedDelegate;
+	}
+
+	/**
+	* Called when there was an error during the rendering of this movie pipeline (such as missing sequence, i/o failure, etc.)
+	*/
+	FMoviePipelineErrored& OnMoviePipelineErrored()
+	{
+		return OnMoviePipelineErroredDelegate;
+	}
+
+	/**
+	* Get the Master Configuration used to render this shot. This contains the global settings for the shot, as well as per-shot
+	* configurations which can contain their own settings.
+	*/
+	UFUNCTION(BlueprintPure, Category = "Movie Render Pipeline")
+	UMoviePipelineMasterConfig* GetPipelineMasterConfig() const;
+public:
+
+	UFUNCTION(BlueprintPure, Category = "Movie Render Pipeline")
+	ULevelSequence* GetTargetSequence() const { return TargetSequence; }
+
+	UFUNCTION(BlueprintPure, Category = "Movie Render Pipeline")
+	const TArray<FMoviePipelineShotInfo>& GetShotList() const { return ShotList; }
+
+	UFUNCTION(BlueprintPure, Category = "Movie Render Pipeline")
+	int32 GetCurrentShotIndex() const { return CurrentShotIndex; }
+
+	UMoviePipelineExecutorJob* GetCurrentJob() const { return CurrentJob; }
+public:
+	void OnFrameCompletelyRendered(FMoviePipelineMergerOutputFrame&& OutputFrame, const TSharedRef<FImagePixelDataPayload, ESPMode::ThreadSafe> InFrameData);
+	void OnSampleRendered(TUniquePtr<FImagePixelData>&& OutputSample, const TSharedRef<FImagePixelDataPayload, ESPMode::ThreadSafe> InFrameData);
+private:
+
+	/** Instantiate our Debug UI Widget and initialize it to ourself. */
+	void LoadDebugWidget();
+	
+	/** Called before the Engine ticks for the given frame. We use this to calculate delta times that the frame should use. */
+	void OnEngineTickBeginFrame();
+	
+	/** Called after the Engine has ticked for a given frame. Everything in the world has been updated by now so we can submit things to render. */
+	void OnEngineTickEndFrame();
+
+	void ValidateSequenceAndSettings();
+
+
+	/** Runs the per-tick logic when doing the ProducingFrames state. */
+	void TickProducingFrames();
+
+	bool ProcessEndOfCameraCut(FMoviePipelineShotInfo &CurrentShot, FMoviePipelineCameraCutInfo &CurrentCameraCut);
+
+
+	/** Called once when first moving to the Finalize state. */
+	void BeginFinalize();
+
+	/** 
+	* Runs the per-tick logic when doing the Finalize state.
+	* @param bInForceFinish		If true, this function will not return until all Output Containers say they have finalized.
+	*/
+	void TickFinalizeOutputContainers(const bool bInForceFinish);
+	/** 
+	* Runs the per-tick logic when doing the Export state. This is spread over multiple ticks to allow non-blocking background 
+	* processes (such as extra encoding) 
+	*/
+	void TickPostFinalizeExport();
+
+	/** Return true if we should early out of the TickProducingFrames function. Decrements the remaining number of steps when false. */
+	bool DebugFrameStepPreTick();
+	/** Returns true if we are idling because of debug frame stepping. */
+	bool IsDebugFrameStepIdling() const;
+
+	/** Debugging/Information. Don't hinge any logic on this as it will get called multiple times per frame in some cases. */
+	void OnSequenceEvaluated(const UMovieSceneSequencePlayer& Player, FFrameTime CurrentTime, FFrameTime PreviousTime);
+
+	/** Set up per-shot state for the specific shot, tearing down old state (if it exists) */
+	void InitializeShot(FMoviePipelineShotInfo& InShot);
+	void TeardownShot(FMoviePipelineShotInfo& InShot);
+
+	/** Initialize the rendering pipeline for the given shot. This should not get called if rendering work is still in progress for a previous shot. */
+	void SetupRenderingPipelineForShot(FMoviePipelineShotInfo& InShot);
+	/** Deinitialize the rendering pipeline for the given shot. */
+	void TeardownRenderingPipelineForShot(FMoviePipelineShotInfo& InShot);
+
+	/** Flush any async resources in the engine that need to be finalized before submitting anything to the GPU, ie: Streaming Levels and Shaders */
+	void FlushAsyncEngineSystems();
+
+	template<typename SettingType>
+	SettingType* FindOrAddSetting(const FMoviePipelineShotInfo& InShot) const
+	{
+		return (SettingType*)FindOrAddSetting(SettingType::StaticClass(), InShot);
+	}
+
+	UMoviePipelineSetting* FindOrAddSetting(TSubclassOf<UMoviePipelineSetting> InSetting, const FMoviePipelineShotInfo& InShot) const;
+	
+	/** 
+	* Renders the next frame in the Pipeline. This updates/ticks all scene view render states
+	* and produces data. This may not result in an output frame due to multiple renders 
+	* accumulating together to produce an output. frame.
+	* Should not be called if we're idling (debug), not initialized yet, or finalizing/exporting. 
+	*/
+	void RenderFrame();
+
+	/** Allow any Settings to modify the (already duplicated) sequence. This allows inserting automatic pre-roll, etc. */
+	void ModifySequenceViaExtensions(ULevelSequence* InSequence);
+
+	/** Calculate the expected amount of total work that this Movie Pipeline is expected to do. */
+	FMoviePipelineWorkInfo CalculateExpectedOutputMetrics();
+private:
+	/** Initialize a new Level Sequence Actor to evaluate our target sequence. Disables any existing Level Sequences pointed at our original sequence. */
+	void InitializeLevelSequenceActor(ULevelSequence* OriginalLevelSequence, ULevelSequence* InSequenceToApply);
+
+	/** This converts the sequence into a Shot List and expands bounds.*/
+	TArray<FMoviePipelineShotInfo> BuildShotListFromSequence(const ULevelSequence* InSequence);
+
+	/** 
+	* Modifies the TargetSequence to ensure that only the specified Shot has it's associated Cinematic Shot Section enabled.
+	* This way when Handle Frames are enabled and the sections are expanded, we don't end up evaluating the previous shot. 
+	*/
+	void SetSoloShot(const FMoviePipelineShotInfo& InShot);
+
+	/* Expands the specified shot (and contained camera cuts)'s ranges for the given settings. */
+	void ExpandShot(FMoviePipelineShotInfo& InShot);
+
+	/** Calculates lots of useful numbers used in timing based off of the current shot. These are constant for a given shot. */
+	MoviePipeline::FFrameConstantMetrics CalculateShotFrameMetrics(const FMoviePipelineShotInfo& InShot) const;
+
+	/** It can be useful to know where the data we're generating was relative to the original Timeline, so this calculates that. */
+	void CalculateFrameNumbersForOutputState(const MoviePipeline::FFrameConstantMetrics& InFrameMetrics, const FMoviePipelineCameraCutInfo& InCameraCut, FMoviePipelineFrameOutputState& InOutOutputState) const;
+
+private:
+	/** Custom TimeStep used to drive the engine while rendering. */
+	UPROPERTY(Transient, Instanced)
+	UMoviePipelineCustomTimeStep* CustomTimeStep;
+
+	/** Custom Time Controller for the Sequence Player, used to match Custom TimeStep without any floating point accumulation errors. */
+	TSharedPtr<FMoviePipelineTimeController> CustomSequenceTimeController;
+
+	/** Hold a reference to the existing custom time step (if any) so we can restore it after we're done using our custom one. */
+	UPROPERTY(Transient)
+	UEngineCustomTimeStep* CachedPrevCustomTimeStep;
+
+	/** This is our duplicated sequence that we're rendering. This will get modified throughout the rendering process. */
+	UPROPERTY(Transient)
+	ULevelSequence* TargetSequence;
+
+	/** The Level Sequence Actor we spawned to play our TargetSequence. */
+	UPROPERTY(Transient)
+	ALevelSequenceActor* LevelSequenceActor;
+
+	/** The Debug UI Widget that is spawned and placed on the player UI */
+	UPROPERTY(Transient)
+	UMovieRenderDebugWidget* DebugWidget;
+
+	/** A list of all of the shots we are going to render out from this sequence. */
+	UPROPERTY(Transient, BlueprintReadOnly, meta = (AllowPrivateAccess = true), Category = "Movie Render Pipeline")
+	TArray<FMoviePipelineShotInfo> ShotList;
+
+	/** What state of the overall flow are we in? See enum for specifics. */
+	UPROPERTY(Transient, BlueprintReadOnly, meta = (AllowPrivateAccess = true), Category = "Movie Render Pipeline")
+	EMovieRenderPipelineState PipelineState;
+
+	/** What is the index of the shot we are working on. -1 if not initialized, may be greater than ShotList.Num() if we've reached the end. */
+	UPROPERTY(Transient, BlueprintReadOnly, meta = (AllowPrivateAccess = true), Category = "Movie Render Pipeline")
+	int32 CurrentShotIndex;
+
+	/** The time (in UTC) that Initialize was called. Used to track elapsed time. */
+	UPROPERTY(Transient, BlueprintReadOnly, meta = (AllowPrivateAccess = true), Category = "Movie Render Pipeline")
+	FDateTime InitializationTime;
+
+	UPROPERTY(Transient, BlueprintReadOnly, meta = (AllowPrivateAccess = true), Category = "Movie Render Pipeline")
+	FMoviePipelineFrameOutputState CachedOutputState;
+
+	/** 
+	* Cache the overall expected work for the entire Movie Pipeline. Used to compare current progress against total for progress.
+	* Current Progress is tracked by individual shots.
+	*/
+	FMoviePipelineWorkInfo TotalExpectedWork;
+
+	/** 
+	* Have we hit the callback for the BeginFrame at least once? This solves an issue where the delegates
+	* get registered mid-frame so you end up calling EndFrame before BeginFrame which is undesirable.
+	*/
+	bool bHasRunBeginFrameOnce;
+	/** Should we pause the game at the end of the frame? Used to implement frame step debugger. */
+	bool bPauseAtEndOfFrame;
+
+	/** When using temporal sub-frame stepping common counts (such as 3) don't result in whole ticks. We keep track of how many ticks we lose so we can add them the next time there's a chance. */
+	float AccumulatedTickSubFrameDeltas;
+
+	/** Called when we have completely finished. This object will call Shutdown before this and stop ticking. */
+	FMoviePipelineFinished OnMoviePipelineFinishedDelegate;
+
+	/** Called when there is a warning/error that the user should pay attention to.*/
+	FMoviePipelineErrored OnMoviePipelineErroredDelegate;
+
+	/**
+	 * We have to apply camera motion vectors manually. So we keep the current and previous fram'es camera view and rotation.
+	 * Then we render a sequence of the same movement, and update after running the game sim.
+	 **/
+	MoviePipeline::FMoviePipelineFrameInfo FrameInfo;
+
+public:
+	/** A list of engine passes which need to be run each frame to generate required content for all the movie render passes. */
+	TArray<TSharedPtr<MoviePipeline::FMoviePipelineEnginePass>> ActiveRenderPasses;
+
+	/** This gathers all of the produced data for an output frame (which may come in async many frames later) before passing them onto the Output Containers. */
+	TSharedPtr<FMoviePipelineOutputMerger, ESPMode::ThreadSafe> OutputBuilder;
+
+	/** A debug image sequence writer in the event they want to dump every sample generated on its own. */
+	IImageWriteQueue* ImageWriteQueue;
+
+private:
+	/** Keep track of which job we're working on. This holds our Configuration + which shots we're supposed to render from it. */
+	UPROPERTY(Transient)
+	UMoviePipelineExecutorJob* CurrentJob;
+};
+
+UCLASS()
+class UMoviePipelineCustomTimeStep : public UEngineCustomTimeStep
+{
+	GENERATED_BODY()
+
+public:
+	// UEngineCustomTimeStep Interface
+	virtual bool Initialize(UEngine* InEngine) override { return true; }
+	virtual void Shutdown(UEngine* InEngine) override {}
+	virtual bool UpdateTimeStep(UEngine* InEngine) override;
+	virtual ECustomTimeStepSynchronizationState GetSynchronizationState() const override { return ECustomTimeStepSynchronizationState::Synchronized; }
+	// ~UEngineCustomTimeStep Interface
+
+	void SetCachedFrameTiming(const MoviePipeline::FFrameTimeStepCache& InTimeCache) { TimeCache = InTimeCache; }
+
+private:
+	/** We don't do any thinking on our own, instead we just spit out the numbers stored in our time cache. */
+	MoviePipeline::FFrameTimeStepCache TimeCache;
+};
+
+struct FMoviePipelineTimeController : public FMovieSceneTimeController
+{
+	virtual FFrameTime OnRequestCurrentTime(const FQualifiedFrameTime& InCurrentTime, float InPlayRate) override;
+	void SetCachedFrameTiming(const FQualifiedFrameTime& InTimeCache) { TimeCache = InTimeCache; }
+
+private:
+	/** Simply store the number calculated and return it when requested. */
+	FQualifiedFrameTime TimeCache;
+};
\ No newline at end of file
diff --git a/Engine/Plugins/MovieScene/MovieRenderPipeline/Source/MovieRenderPipelineCore/Public/MoviePipelineAccumulationSetting.h b/Engine/Plugins/MovieScene/MovieRenderPipeline/Source/MovieRenderPipelineCore/Public/MoviePipelineAccumulationSetting.h
new file mode 100644
--- /dev/null
+++ b/Engine/Plugins/MovieScene/MovieRenderPipeline/Source/MovieRenderPipelineCore/Public/MoviePipelineAccumulationSetting.h
@@ -0,0 +1,51 @@
+// Copyright 1998-2019 Epic Games, Inc. All Rights Reserved.
+#pragma once
+
+#include "MoviePipelineSetting.h"
+#include "MovieRenderPipelineDataTypes.h"
+#include "MoviePipelineAccumulationSetting.generated.h"
+
+
+UCLASS(Blueprintable)
+class MOVIERENDERPIPELINECORE_API UMoviePipelineAccumulationSetting : public UMoviePipelineSetting
+{
+	GENERATED_BODY()
+public:
+	UMoviePipelineAccumulationSetting()
+		: SpatialSampleCount(1)
+		, HandleFrameCount(0)
+		, WarmUpFrameCount(0)
+		, AccumulationGamma(1.f)
+	{
+	}
+
+public:
+#if WITH_EDITOR
+	virtual FText GetDisplayText() const override { return NSLOCTEXT("MovieRenderPipeline", "AccumulationSettingDisplayName", "Subsample Accumulation"); }
+#endif
+protected:
+	virtual bool IsValidOnShots() const override { return true; }
+	virtual bool IsValidOnMaster() const override { return true; }
+	
+public:
+
+	/**
+	* How many frames should we accumulate together before contributing to one overall sample. This lets you
+	* increase the anti-aliasing quality of an sample, or have high quality anti-aliasing if you don't want
+	* any motion blur due to accumulation over time in SampleCount.
+	*/
+	UPROPERTY(EditAnywhere, BlueprintReadWrite, meta = (UIMin = 1, ClampMin = 1), Category = "Movie Pipeline")
+	int32 SpatialSampleCount;
+
+	UPROPERTY(EditAnywhere, BlueprintReadWrite, meta = (UIMin = 1, ClampMin = 1), Category = "Movie Pipeline")
+	int32 HandleFrameCount;
+
+	UPROPERTY(EditAnywhere, BlueprintReadWrite, meta = (UIMin = 1, ClampMin = 1), Category = "Movie Pipeline")
+	int32 WarmUpFrameCount;
+
+	/**
+	* For advanced users, the gamma space to apply accumulation in. During accumulation, pow(x,AccumulationGamma) 
+	* is applied and pow(x,1/AccumulationGamma) is applied after accumulation is finished
+	*/
+	float AccumulationGamma;
+};
\ No newline at end of file
diff --git a/Engine/Plugins/MovieScene/MovieRenderPipeline/Source/MovieRenderPipelineCore/Public/MoviePipelineBlueprintLibrary.h b/Engine/Plugins/MovieScene/MovieRenderPipeline/Source/MovieRenderPipelineCore/Public/MoviePipelineBlueprintLibrary.h
new file mode 100644
--- /dev/null
+++ b/Engine/Plugins/MovieScene/MovieRenderPipeline/Source/MovieRenderPipelineCore/Public/MoviePipelineBlueprintLibrary.h
@@ -0,0 +1,95 @@
+// Copyright 1998-2019 Epic Games, Inc. All Rights Reserved.
+
+#pragma once
+
+#include "Kismet/BlueprintFunctionLibrary.h"
+#include "MovieRenderPipelineDataTypes.h"
+
+#include "MoviePipelineBlueprintLibrary.generated.h"
+
+// Forward Declare
+class UMoviePipeline;
+class UMovieSceneSequence;
+
+UCLASS()
+class MOVIERENDERPIPELINECORE_API UMoviePipelineBlueprintLibrary : public UBlueprintFunctionLibrary
+{
+	GENERATED_BODY()
+
+public:
+	/**
+	* Takes a given Output Directory and FileName which contain {formatStrings} and convert those format strings
+	* to their actual values as best as possible. Merges the result into one final filepath string.
+	*/
+	// UFUNCTION(BlueprintPure, Category = "Movie Render Pipeline")
+	// static FString GetCameraCutCounts(const UMoviePipeline* InPipeline, const FMoviePipelineFrameOutputState& OutputState, const FDirectoryPath& InDirectoryFormatString, const FString& InFileNameFormatString);
+
+	/** 
+	* Get the total number of Camera Cuts in a Movie Pipeline, and an index of which one is being processed.
+	* Note: Index doesn't index into any arrays (since internally we store CameraCuts inside of Shots)
+	*/
+	UFUNCTION(BlueprintPure, Category = "Movie Render Pipeline")
+	static void GetCameraCutCounts(const UMoviePipeline* InPipeline, int32& OutTotalCuts, int32& OutCurrentCutIndex);
+	
+	/** 
+	* Returns the number of expected frames to be produced by Initial Range + Handle Frames given a Frame Rate.
+	* This will be inaccurate when PlayRate tracks are involved.
+	*/
+	UFUNCTION(BlueprintPure, Category = "Movie Render Pipeline")
+	static int32 GetOutputFrameCountEstimate(const FMoviePipelineCameraCutInfo& InCameraCut);
+	
+	/**
+	* Returns the expected number of frames different frames that will be submitted for rendering. This is
+	* the number of output frames * temporal samples. This will be inaccurate when PlayRate tracks are involved.
+	*/
+	UFUNCTION(BlueprintPure, Category = "Movie Render Pipeline")
+	static int32 GetTemporalFrameCountEstimate(const FMoviePipelineCameraCutInfo& InCameraCut);
+	
+	/** 
+	* Returns misc. utility frame counts (Warm Up + MotionBlur Fix) since these are outside the ranges.
+	*/
+	UFUNCTION(BlueprintPure, Category = "Movie Render Pipeline")
+	static int32 GetUtilityFrameCountEstimate(const FMoviePipelineCameraCutInfo& InCameraCut);
+	
+	/**
+	* Returns the number of samples submitted to the GPU, optionally counting samples for the various parts.
+	* This will be inaccurate when PlayRate tracks are involved.
+	*/
+	UFUNCTION(BlueprintPure, Category = "Movie Render Pipeline")
+	static int32 GetSampleCountEstimate(const FMoviePipelineCameraCutInfo& InCameraCut, const bool bIncludeWarmup = true, const bool bIncludeMotionBlur = true);
+	
+	/**
+	* Duplicates the specified sequence using a medium depth copy. Standard duplication will only duplicate
+	* the top level Sequence (since shots and sub-sequences are other standalone assets) so this function
+	* recursively duplicates the given sequence, shot and subsequence and then fixes up the references to
+	* point to newly duplicated sequences.
+	*
+	* This does not duplicate any assets that the sequence points to outside of Shots/Subsequences.
+	*
+	* @param	Outer		- The Outer of the newly duplicated object. Leave null for TransientPackage();
+	* @param	InSequence	- The sequence to recursively duplicate.
+	* @return				- The duplicated sequence, or null if no sequence was provided to duplicate.
+	*/
+	UFUNCTION(BlueprintCallable, Category = "Movie Render Pipeline")
+	static UMovieSceneSequence* DuplicateSequence(UObject* Outer, UMovieSceneSequence* InSequence);
+
+	// UFUNCTION(BlueprintPure, Category = "Movie Render Pipeline")
+	// 	FMoviePipelineShotInfo GetCurrentShotSnapshot() const;
+	// 
+	// UFUNCTION(BlueprintPure, Category = "Movie Render Pipeline")
+	// 	FMoviePipelineCameraCutInfo GetCurrentCameraCutSnapshot() const;
+	// 
+	// UFUNCTION(BlueprintPure, Category = "Movie Render Pipeline")
+	// 	FMoviePipelineFrameOutputState GetOutputStateSnapshot() const;
+
+	/**
+	* Returns an estimate based on the average time taken to render all previous frames.
+	* Will be incorrect if PlayRate tracks are in use. Will be inaccurate when different
+	* shots take significantly different amounts of time to render.
+	*
+	* @param OutTimespan: The returned estimated time left. Will be default initialized if there is no estimate.
+	* @return True if we can make a reasonable estimate, false otherwise (ie: no rendering has been done to estimate).
+	*/
+	// UFUNCTION(BlueprintPure, Category = "Movie Render Pipeline")
+	// bool GetRemainingTimeEstimate(FTimespan& OutTimespan) const;
+};
\ No newline at end of file
diff --git a/Engine/Plugins/MovieScene/MovieRenderPipeline/Source/MovieRenderPipelineCore/Public/MoviePipelineConfigBase.h b/Engine/Plugins/MovieScene/MovieRenderPipeline/Source/MovieRenderPipelineCore/Public/MoviePipelineConfigBase.h
new file mode 100644
--- /dev/null
+++ b/Engine/Plugins/MovieScene/MovieRenderPipeline/Source/MovieRenderPipelineCore/Public/MoviePipelineConfigBase.h
@@ -0,0 +1,151 @@
+// Copyright 1998-2019 Epic Games, Inc. All Rights Reserved.
+#pragma once
+
+#include "UObject/Object.h"
+#include "Templates/SubclassOf.h"
+#include "MoviePipelineSetting.h"
+#include "MoviePipelineConfigBase.generated.h"
+
+
+UCLASS(BlueprintType, Abstract)
+class MOVIERENDERPIPELINECORE_API UMoviePipelineConfigBase : public UObject
+{
+	GENERATED_BODY()
+
+public:
+	UMoviePipelineConfigBase()
+	{
+		SettingsSerialNumber = -1;
+	}
+
+public:
+	/**
+	* Removes the specific instance from our Setting list.
+	*/
+	UFUNCTION(BlueprintCallable, Category = "Movie Render Pipeline")
+	void RemoveSetting(UMoviePipelineSetting* InSetting);
+
+	virtual void CopyFrom(UMoviePipelineConfigBase* InConfig);
+
+	int32 GetSettingsSerialNumber() const { return SettingsSerialNumber; }
+
+	/**
+	* Returns an array of all settings in this config.
+	*/
+	virtual TArray<UMoviePipelineSetting*> GetSettings() const { return Settings; }
+
+public:
+
+	/**
+	* Find a setting of a particular type for this config.
+	* @param InClass - Class that you wish to find the setting object for.
+	* @return An instance of this class if it already exists as a setting on this config, otherwise null.
+	*/
+	UFUNCTION(BlueprintPure, meta = (DeterminesOutputType = "InClass"), Category = "Movie Render Pipeline")
+	UMoviePipelineSetting* FindSettingByClass(TSubclassOf<UMoviePipelineSetting> InClass) const
+	{
+		TArray<UMoviePipelineSetting*> AllSettings = GetSettings();
+		UMoviePipelineSetting* const* Found = AllSettings.FindByPredicate([InClass](UMoviePipelineSetting* In) { return In && In->GetClass() == InClass; });
+		return Found ? CastChecked<UMoviePipelineSetting>(*Found) : nullptr;
+	}
+
+	/**
+	* Finds a setting of a particular type for this pipeline config, adding it if it doesn't already exist.
+	* @param InClass - Class you wish to find or create the setting object for.
+	* @return An instance of this class as a setting on this config.
+	*/
+	UFUNCTION(BlueprintCallable, meta = (DeterminesOutputType = "InClass"), Category = "Movie Render Pipeline")
+	UMoviePipelineSetting* FindOrAddSettingByClass(TSubclassOf<UMoviePipelineSetting> InClass)
+	{
+		UMoviePipelineSetting* Found = FindSettingByClass(InClass);
+		if (!Found)
+		{
+			Modify();
+			
+			Found = NewObject<UMoviePipelineSetting>(this, InClass);
+
+			if (CanSettingBeAdded(Found))
+			{
+				Settings.Add(Found);
+				++SettingsSerialNumber;
+			}
+			else
+			{
+				FFrame::KismetExecutionMessage(*FString::Printf(TEXT("Setting %d is not compatible with this Config Type and was not added."), *InClass->GetName()), ELogVerbosity::Error);
+				return nullptr;
+			}
+		}
+
+		return Found;
+	}
+
+	/**
+	* Finds an array of all settings matching the specific type, including inherited classes.
+	*/
+	template<typename SettingType>
+	TArray<SettingType*> FindSettings() const
+	{
+		TArray<SettingType*> FoundSettings;
+
+		TArray<UMoviePipelineSetting*> AllSettings = GetSettings();
+		for (UMoviePipelineSetting* Setting : AllSettings)
+		{
+			if (Setting->GetClass()->IsChildOf<SettingType>())
+			{
+				FoundSettings.Add(Cast<SettingType>(Setting));
+			}
+		}
+
+		return FoundSettings;
+	}
+
+	/**
+	* Find a setting of a particular type for this config
+	*/
+	template<typename SettingType>
+	SettingType* FindSetting() const
+	{
+		UClass* PredicateClass = SettingType::StaticClass();
+
+		TArray<UMoviePipelineSetting*> AllSettings = GetSettings();
+		UMoviePipelineSetting* const* Found = AllSettings.FindByPredicate([PredicateClass](UMoviePipelineSetting* In) { return In && In->GetClass() == PredicateClass; });
+		return Found ? CastChecked<SettingType>(*Found) : nullptr;
+	}
+
+	/**
+	 * Find a setting of a particular type for this config instance, adding one if it was not found.
+	 */
+	template<typename SettingType>
+	SettingType* FindOrAddSetting()
+	{
+		SettingType* Found = FindSetting<SettingType>();
+		if (!Found)
+		{
+			Modify();
+
+			Found = NewObject<SettingType>(this);
+			if (CanSettingBeAdded(Found))
+			{
+				Settings.Add(Found);
+				++SettingsSerialNumber;
+			}
+			else
+			{
+				FFrame::KismetExecutionMessage(*FString::Printf(TEXT("Setting %d is not compatible with this Config Type and was not added."), *Found->GetName()), ELogVerbosity::Error);
+				return nullptr;
+			}
+		}
+		return Found;
+	}
+
+public:
+	virtual bool CanSettingBeAdded(const UMoviePipelineSetting* InSetting) const PURE_VIRTUAL( UMoviePipelineConfigBase::CanSettingBeAdded, return false; );
+
+protected:
+	/** Array of settings classes that affect various parts of the output pipeline. */
+	UPROPERTY(VisibleAnywhere, Instanced, Category = "Movie Pipeline")
+	TArray<UMoviePipelineSetting*> Settings;
+
+private:
+	int32 SettingsSerialNumber;
+};
\ No newline at end of file
diff --git a/Engine/Plugins/MovieScene/MovieRenderPipeline/Source/MovieRenderPipelineCore/Public/MoviePipelineInProcessExecutor.h b/Engine/Plugins/MovieScene/MovieRenderPipeline/Source/MovieRenderPipelineCore/Public/MoviePipelineInProcessExecutor.h
new file mode 100644
--- /dev/null
+++ b/Engine/Plugins/MovieScene/MovieRenderPipeline/Source/MovieRenderPipelineCore/Public/MoviePipelineInProcessExecutor.h
@@ -0,0 +1,31 @@
+// Copyright 1998-2019 Epic Games, Inc. All Rights Reserved.
+#pragma once
+
+#include "MoviePipelineLinearExecutor.h"
+#include "MoviePipelineInProcessExecutor.generated.h"
+
+/**
+* This executor implementation can process an array of movie pipelines and
+* run them inside the currently running process. This is intended for usage
+* outside of the editor (ie. -game mode) as it will take over the currently
+* running world/game instance instead of launching a new world instance like 
+* the editor only PIE.
+*/
+UCLASS(Blueprintable)
+class MOVIERENDERPIPELINECORE_API UMoviePipelineInProcessExecutor : public UMoviePipelineLinearExecutorBase
+{
+	GENERATED_BODY()
+
+public:
+	UMoviePipelineInProcessExecutor()
+		: UMoviePipelineLinearExecutorBase()
+	{
+	}
+
+protected:
+	virtual void Start(const UMoviePipelineExecutorJob* InJob) override;
+
+private:
+	void OnMapLoadFinished(UWorld* NewWorld);
+	void OnMoviePipelineFinished(UMoviePipeline* InMoviePipeline);
+};
\ No newline at end of file
diff --git a/Engine/Plugins/MovieScene/MovieRenderPipeline/Source/MovieRenderPipelineCore/Public/MoviePipelineLinearExecutor.h b/Engine/Plugins/MovieScene/MovieRenderPipeline/Source/MovieRenderPipelineCore/Public/MoviePipelineLinearExecutor.h
new file mode 100644
--- /dev/null
+++ b/Engine/Plugins/MovieScene/MovieRenderPipeline/Source/MovieRenderPipelineCore/Public/MoviePipelineLinearExecutor.h
@@ -0,0 +1,57 @@
+// Copyright 1998-2019 Epic Games, Inc. All Rights Reserved.
+#pragma once
+
+#include "MoviePipelineExecutor.h"
+#include "MovieRenderPipelineDataTypes.h"
+#include "MoviePipelineLinearExecutor.generated.h"
+
+class UMoviePipeline;
+class UMoviePipelineQueue;
+
+/**
+* This is an abstract base class designed for executing an array of movie pipelines in linear
+* fashion. It is generally the case that you only want to execute one at a time on a local machine
+* and a different executor approach should be taken for a render farm that distributes out jobs
+* to many different machines.
+*/
+UCLASS(Blueprintable, Abstract)
+class MOVIERENDERPIPELINECORE_API UMoviePipelineLinearExecutorBase : public UMoviePipelineExecutorBase
+{
+	GENERATED_BODY()
+	
+public:
+	UMoviePipelineLinearExecutorBase()
+		: UMoviePipelineExecutorBase()
+		, CurrentPipelineIndex(0)
+	{
+	}
+
+protected:
+	virtual void ExecuteImpl(UMoviePipelineQueue* InPipelineQueue) override;
+
+	virtual void StartPipelineByIndex(int32 InPipelineIndex);
+	virtual void Start(const UMoviePipelineExecutorJob* InJob) {}
+
+public:
+	virtual void OnIndividualPipelineFinished(UMoviePipeline* /* FinishedPipeline */);
+	virtual void OnExecutorFinishedImpl();
+	virtual void OnPipelineErrored(UMoviePipeline* InPipeline, bool bIsFatal, FText ErrorText);
+protected:
+	
+	/** List of Pipeline Configs we've been asked to execute. */
+	UPROPERTY(Transient)
+	UMoviePipelineQueue* Queue;
+
+	/** A Movie Pipeline that has been spawned and is running (if any) */
+	UPROPERTY(Transient)
+	UMoviePipeline* ActiveMoviePipeline;
+
+	/** Which Pipeline Config are we currently working on. */
+	int32 CurrentPipelineIndex;
+
+private:
+	/** Used to track total processing duration. */
+	FDateTime InitializationTime;
+};
+
+
diff --git a/Engine/Plugins/MovieScene/MovieRenderPipeline/Source/MovieRenderPipelineCore/Public/MoviePipelineMasterConfig.h b/Engine/Plugins/MovieScene/MovieRenderPipeline/Source/MovieRenderPipelineCore/Public/MoviePipelineMasterConfig.h
new file mode 100644
--- /dev/null
+++ b/Engine/Plugins/MovieScene/MovieRenderPipeline/Source/MovieRenderPipelineCore/Public/MoviePipelineMasterConfig.h
@@ -0,0 +1,74 @@
+// Copyright 1998-2019 Epic Games, Inc. All Rights Reserved.
+#pragma once
+
+#include "MoviePipelineConfigBase.h"
+#include "Misc/FrameRate.h"
+
+#include "MoviePipelineMasterConfig.generated.h"
+
+
+// Forward Declares
+class ULevelSequence;
+class UMoviePipelineSetting;
+class UMoviePipelineRenderPass;
+class UMoviePipelineOutputBase;
+class UMoviePipelineShotConfig;
+class UMoviePipelineOutputSetting;
+class UMoviePipelineExecutorJob;
+
+/**
+* This class describes the main configuration for a Movie Render Pipeline.
+* Only settings that apply to the entire output should be stored here,
+* anything that is changed on a per-shot basis should be stored inside of 
+* UMovieRenderShotConfig instead.
+*
+* THIS CLASS SHOULD BE IMMUTABLE ONCE PASSED TO THE PIPELINE FOR PROCESSING.
+* (Otherwise you will be modifying the instance that exists in the UI)
+*/
+UCLASS(Blueprintable)
+class MOVIERENDERPIPELINECORE_API UMoviePipelineMasterConfig : public UMoviePipelineConfigBase
+{
+	GENERATED_BODY()
+	
+public:
+	UMoviePipelineMasterConfig();
+
+public:
+	TArray<UMoviePipelineOutputBase*> GetOutputContainers() const;
+	virtual TArray<UMoviePipelineSetting*> GetSettings() const override;
+	virtual void CopyFrom(UMoviePipelineConfigBase* InConfig) override;
+
+public:
+
+	/** Returns a pointer to the config specified for the shot, otherwise the default for this pipeline. */
+	UMoviePipelineShotConfig* GetConfigForShot(const FString& ShotName) const;
+
+	void GetFilenameFormatArguments(FFormatNamedArguments& OutArguments, const UMoviePipelineExecutorJob* InJob) const;
+
+
+	/**
+	* Returns the frame rate override from the Master Configuration (if any) or the Sequence frame rate if no override is specified.
+	* This should be treated as the actual output framerate of the overall Pipeline.
+	*/
+	UFUNCTION(BlueprintCallable, Category = "Movie Render Pipeline")
+	FFrameRate GetEffectiveFrameRate(const ULevelSequence* InSequence) const;
+protected:
+	virtual bool CanSettingBeAdded(const UMoviePipelineSetting* InSetting) const override
+	{
+		check(InSetting);
+		return InSetting->IsValidOnMaster();
+	}
+public:
+	
+	/** The default shot-setup to use for any shot that doesn't a specific implementation. This is required! */
+	UPROPERTY(EditAnywhere, Instanced, BlueprintReadWrite, Category = "Movie Render Pipeline")
+	UMoviePipelineShotConfig* DefaultShotConfig;
+	
+	/** A mapping of Shot Name -> Shot Config to use for rendering specific shots with specific configs. */
+	UPROPERTY(VisibleAnywhere, Instanced, BlueprintReadOnly, Category = "Movie Render Pipeline")
+	TMap<FString, UMoviePipelineShotConfig*> PerShotConfigMapping;
+
+private:
+	UPROPERTY(Transient, Instanced)
+	UMoviePipelineOutputSetting* OutputSetting;
+};
\ No newline at end of file
diff --git a/Engine/Plugins/MovieScene/MovieRenderPipeline/Source/MovieRenderPipelineCore/Public/MoviePipelineRenderPass.h b/Engine/Plugins/MovieScene/MovieRenderPipeline/Source/MovieRenderPipelineCore/Public/MoviePipelineRenderPass.h
new file mode 100644
--- /dev/null
+++ b/Engine/Plugins/MovieScene/MovieRenderPipeline/Source/MovieRenderPipelineCore/Public/MoviePipelineRenderPass.h
@@ -0,0 +1,51 @@
+// Copyright 1998-2019 Epic Games, Inc. All Rights Reserved.
+#pragma once
+
+#include "UObject/Object.h"
+#include "MovieRenderPipelineDataTypes.h"
+#include "MoviePipelineSetting.h"
+#include "ImageWriteStream.h"
+#include "MoviePipelineRenderPass.generated.h"
+
+UCLASS(Blueprintable, Abstract)
+class MOVIERENDERPIPELINECORE_API UMoviePipelineRenderPass : public UMoviePipelineSetting
+{
+	GENERATED_BODY()
+public:
+	void Setup(TArray<TSharedPtr<MoviePipeline::FMoviePipelineEnginePass>>& InEnginePasses)
+	{
+		SetupImpl(InEnginePasses);
+	}
+
+	void Teardown()
+	{
+		TeardownImpl();
+	}
+
+	/** An array of identifiers for the output buffers expected as a result of this render pass. */
+	void GatherOutputPasses(TArray<FMoviePipelinePassIdentifier>& ExpectedRenderPasses)
+	{
+		GatherOutputPassesImpl(ExpectedRenderPasses);
+	}
+
+	/** The required engine passes this pass needs to operate. */
+	void GetRequiredEnginePasses(TSet<FMoviePipelinePassIdentifier>& RequiredEnginePasses)
+	{
+		GetRequiredEnginePassesImpl(RequiredEnginePasses);
+	}
+
+protected:
+	virtual bool IsValidOnShots() const override { return true; }
+	virtual bool IsValidOnMaster() const override { return true; }
+#if WITH_EDITOR
+	virtual FText GetCategoryText() const override { return NSLOCTEXT("MovieRenderPipeline", "RenderingCategoryName_Text", "Rendering"); }
+#endif
+protected:
+	virtual void GetRequiredEnginePassesImpl(TSet<FMoviePipelinePassIdentifier>& RequiredEnginePasses) {}
+
+	virtual void SetupImpl(TArray<TSharedPtr<MoviePipeline::FMoviePipelineEnginePass>>& InEnginePasses) {}
+
+	virtual void TeardownImpl() {}
+
+	virtual void GatherOutputPassesImpl(TArray<FMoviePipelinePassIdentifier>& ExpectedRenderPasses) {}
+};
\ No newline at end of file
diff --git a/Engine/Plugins/MovieScene/MovieRenderPipeline/Source/MovieRenderPipelineCore/Public/MoviePipelineSetting.h b/Engine/Plugins/MovieScene/MovieRenderPipeline/Source/MovieRenderPipelineCore/Public/MoviePipelineSetting.h
new file mode 100644
--- /dev/null
+++ b/Engine/Plugins/MovieScene/MovieRenderPipeline/Source/MovieRenderPipelineCore/Public/MoviePipelineSetting.h
@@ -0,0 +1,89 @@
+// Copyright 1998-2019 Epic Games, Inc. All Rights Reserved.
+#pragma once
+
+#include "UObject/Object.h"
+#include "Internationalization/Text.h"
+#include "MoviePipelineSetting.generated.h"
+
+class UMoviePipeline;
+struct FSlateBrush;
+class UMoviePipelineExecutorJob;
+
+/**
+* A base class for all Movie Render Pipeline settings.
+*/
+UCLASS(Blueprintable, Abstract)
+class MOVIERENDERPIPELINECORE_API UMoviePipelineSetting : public UObject
+{
+	GENERATED_BODY()
+		
+public:
+	UMoviePipelineSetting();
+
+	/**
+	* This is called once on a setting when the movie pipeline is first set up. If the setting
+	* only exists as part of a shot override, it will be called once when the shot is initialized.
+	* 
+	* This should be used for resource initialization if needed, if you need to change settings
+	* see shot-related callbacks so that they work properly with shot-overrides.
+	*/
+	void OnMoviePipelineInitialized(UMoviePipeline* InPipeline);
+
+	/**
+	* This is called once on a setting when the movie pipeline is shut down. If the setting
+	* only exists as part of a shot override, it will be called once the shot is finished.
+	*
+	* This should be used for releasing resources if needed, if you need to change settings
+	* see shot-related callbacks so that they work properly with shot-overrides.
+	*/
+	void OnMoviePipelineShutdown(UMoviePipeline* InPipeline) { TeardownForPipelineImpl(InPipeline); }
+
+	// UObject Interface
+	virtual UWorld* GetWorld() const override;
+	// ~UObject Interface
+
+protected:
+	UMoviePipeline* GetPipeline() const;
+
+	virtual void SetupForPipelineImpl(UMoviePipeline* InPipeline) {}
+	virtual void TeardownForPipelineImpl(UMoviePipeline* InPipeline) {}
+	
+public:
+#if WITH_EDITOR
+	/** Warning: This gets called on the CDO of the object */
+	virtual FText GetDisplayText() const { return this->GetClass()->GetDisplayNameText(); }
+	/** Warning: This gets called on the CDO of the object */
+	virtual FText GetCategoryText() const { return NSLOCTEXT("MovieRenderPipeline", "DefaultCategoryName_Text", "Default"); }
+
+	/** Return a string to show in the footer of the details panel. Will be combined with other selected settings. */
+	virtual FText GetFooterText(UMoviePipelineExecutorJob* InJob) const { return FText(); }
+#endif
+	/** Can this configuration setting be added to shots? If not, it will throw an error when trying to add it to a shot config. */
+	virtual bool IsValidOnShots() const PURE_VIRTUAL(UMoviePipelineSetting::IsValidOnShots, return false; );
+	/** Can this configuration setting be added to the master configuration? If not, it will throw an error when trying to add it to the master configuration. */
+	virtual bool IsValidOnMaster() const PURE_VIRTUAL(UMoviePipelineSetting::IsValidOnMaster, return false; );
+
+	virtual void GetFilenameFormatArguments(FFormatNamedArguments& OutArguments, const UMoviePipelineExecutorJob* InJob) const {}
+
+	/** Should the Pipeline automatically create an instance of this under the hood so calling code can rely on it existing? */
+	virtual bool IsRequired() const { return false; }
+	
+	/** Can only one of these settings objects be active in a valid pipeline? */
+	virtual bool IsSolo() const { return true; }
+	
+	/** Is this setting valid? Return false and add a reason it's not valid to the array if not. */
+	virtual bool ValidatePipeline(TArray<FText>& OutValidationErrors) const { return true; }
+	
+	/** What icon should this setting use when displayed in the tree list. */
+	const FSlateBrush* GetDisplayIcon() { return nullptr; }
+	
+	/** What tooltip should be displayed for this setting when hovered in the tree list? */
+	FText GetDescriptionText() { return FText(); }
+	
+	/** Is this setting currently enabled? Disabled settings are like they never existed. */
+	bool bEnabled;
+
+private:
+	UPROPERTY(Transient)
+	TWeakObjectPtr<UMoviePipeline> CachedPipeline;
+};
\ No newline at end of file
diff --git a/Engine/Plugins/MovieScene/MovieRenderPipeline/Source/MovieRenderPipelineCore/Public/MoviePipelineShotConfig.h b/Engine/Plugins/MovieScene/MovieRenderPipeline/Source/MovieRenderPipelineCore/Public/MoviePipelineShotConfig.h
new file mode 100644
--- /dev/null
+++ b/Engine/Plugins/MovieScene/MovieRenderPipeline/Source/MovieRenderPipelineCore/Public/MoviePipelineShotConfig.h
@@ -0,0 +1,30 @@
+// Copyright 1998-2019 Epic Games, Inc. All Rights Reserved.
+#pragma once
+
+#include "MoviePipelineConfigBase.h"
+
+#include "MoviePipelineShotConfig.generated.h"
+
+
+// Forward Declares
+class UMoviePipelineRenderPass;
+
+UCLASS(Blueprintable)
+class MOVIERENDERPIPELINECORE_API UMoviePipelineShotConfig : public UMoviePipelineConfigBase
+{
+	GENERATED_BODY()
+	
+public:
+	UMoviePipelineShotConfig()
+	{
+	}
+	
+
+protected:
+	virtual bool CanSettingBeAdded(const UMoviePipelineSetting* InSetting) const override
+	{
+		check(InSetting);
+		return InSetting->IsValidOnShots();
+	}
+
+};
\ No newline at end of file
diff --git a/Engine/Plugins/MovieScene/MovieRenderPipeline/Source/MovieRenderPipelineCore/Public/MovieRenderPipelineCoreModule.h b/Engine/Plugins/MovieScene/MovieRenderPipeline/Source/MovieRenderPipelineCore/Public/MovieRenderPipelineCoreModule.h
new file mode 100644
--- /dev/null
+++ b/Engine/Plugins/MovieScene/MovieRenderPipeline/Source/MovieRenderPipelineCore/Public/MovieRenderPipelineCoreModule.h
@@ -0,0 +1,67 @@
+// Copyright 1998-2019 Epic Games, Inc. All Rights Reserved.
+
+#pragma once
+
+#include "CoreMinimal.h"
+#include "Modules/ModuleInterface.h"
+
+// Forward Declare
+class UMoviePipelineExecutorBase;
+class UMoviePipelineMasterConfig;
+class ULevelSequence;
+class UMoviePipeline;
+
+namespace MoviePipeline { struct FMoviePipelineEnginePass; }
+
+namespace MoviePipelineErrorCodes
+{
+	/** Everything completed as expected or we (unfortunately) couldn't detect the error. */
+	constexpr uint8 Success = 0;
+	/** Fallback for any generic critical failure. This should be used for "Unreal concepts aren't working as expected" severity errors. */
+	constexpr uint8 Critical = 1;
+	/** The specified level sequence asset could not be found. Check the logs for details on what it looked for. */
+	constexpr uint8 NoAsset = 2;
+	/** The specified pipeline configuration asset could not be found. Check the logs for details on what it looked for. */
+	constexpr uint8 NoConfig = 3;
+
+}
+/** A delegate which will create an engine render pass for the Curve Editor. This declares a new engine pass which multiple Pipeline Render Passes can share to reduce re-renders. */
+DECLARE_DELEGATE_RetVal(TSharedRef<MoviePipeline::FMoviePipelineEnginePass>, FOnCreateEngineRenderPass);
+
+class MOVIERENDERPIPELINECORE_API FMovieRenderPipelineCoreModule : public IModuleInterface
+{
+public:
+	virtual void StartupModule() override;
+	virtual void ShutdownModule() override;
+
+	void UnregisterEngineRenderPass(FDelegateHandle InHandle);
+	FDelegateHandle RegisterEngineRenderPass(FOnCreateEngineRenderPass InOnCreateEngineRenderPass);
+
+	TArrayView<const FOnCreateEngineRenderPass> GetEngineRenderPasses() const
+	{
+		return EngineRenderPassDelegates;
+	}
+private:
+	bool IsTryingToRenderMovieFromCommandLine(FString& OutSequenceAssetPath, FString& OutConfigAssetPath, FString& OutExecutorType, FString& OutPipelineType) const;
+	void InitializeCommandLineMovieRender();
+
+	void OnCommandLineMovieRenderCompleted(UMoviePipelineExecutorBase* InExecutor, bool bSuccess);
+	void OnCommandLineMovieRenderErrored(UMoviePipelineExecutorBase* InExecutor, UMoviePipeline* InPipelineWithError, bool bIsFatal, FText ErrorText);
+
+	uint8 GetMovieRenderClasses(const FString& InAssetPath, const FString& InConfigAssetPath, const FString& InExecutorType, const FString& InPipelineType, 
+								ULevelSequence*& OutAsset, UMoviePipelineMasterConfig*& OutConfig, UClass*& OutExecutorClass, UClass*& OutPipelineClass) const;
+
+private:
+	TSharedRef<struct MoviePipeline::FMoviePipelineEnginePass> CreateDeferredEnginePass();
+	FDelegateHandle DeferredEnginePassHandle;
+private:
+	TArray<FOnCreateEngineRenderPass> EngineRenderPassDelegates;
+
+private:
+	FString MoviePipelineLocalExecutorClassType;
+	FString MoviePipelineClassType;
+	FString SequenceAssetValue;
+	FString SettingsAssetValue;
+};
+
+MOVIERENDERPIPELINECORE_API DECLARE_LOG_CATEGORY_EXTERN(LogMovieRenderPipeline, Log, All);
\ No newline at end of file
diff --git a/Engine/Plugins/MovieScene/MovieRenderPipeline/Source/MovieRenderPipelineCore/Public/MovieRenderPipelineDataTypes.cpp b/Engine/Plugins/MovieScene/MovieRenderPipeline/Source/MovieRenderPipelineCore/Public/MovieRenderPipelineDataTypes.cpp
new file mode 100644
--- /dev/null
+++ b/Engine/Plugins/MovieScene/MovieRenderPipeline/Source/MovieRenderPipelineCore/Public/MovieRenderPipelineDataTypes.cpp
@@ -0,0 +1,232 @@
+// Copyright 1998-2019 Epic Games, Inc. All Rights Reserved.
+
+#include "MovieRenderPipelineDataTypes.h"
+#include "MovieRenderPipelineCoreModule.h"
+#include "Sections/MovieSceneCinematicShotSection.h"
+#include "MoviePipelineAccumulationSetting.h"
+#include "MoviePipeline.h"
+
+FString FMoviePipelineShotInfo::GetDisplayName() const
+{
+	if (CinematicShotSection.Get())
+	{
+		return CinematicShotSection->GetShotDisplayName();
+	}
+
+	return TEXT("Unnamed");
+}
+
+FFrameNumber FMoviePipelineCameraCutInfo::GetOutputFrameCountEstimate() const
+{
+	// TotalRange is stored in Tick Resolution, so we convert 1 frame of Frame Rate to the number of ticks.
+	FFrameNumber OneFrameInTicks = FFrameRate::TransformTime(FFrameTime(FFrameNumber(1)), CachedFrameRate, CachedTickResolution).FloorToFrame();
+
+	// Find out how many ticks long our total output range is.
+	FFrameNumber TotalOutputRangeTicks = TotalOutputRange.Size<FFrameNumber>();
+	int32 NumFrames = FMath::CeilToInt(TotalOutputRangeTicks.Value / (double)OneFrameInTicks.Value);
+
+	return FFrameNumber(NumFrames);
+}
+
+FFrameNumber FMoviePipelineCameraCutInfo::GetTemporalFrameCountEstimate() const
+{
+	FFrameNumber OutputEstimate = GetOutputFrameCountEstimate();
+	FFrameNumber NumRenderedSamples = OutputEstimate * NumTemporalSamples;
+
+	return NumRenderedSamples;
+}
+
+FFrameNumber FMoviePipelineCameraCutInfo::GetUtilityFrameCountEstimate() const
+{
+	return FFrameNumber(NumWarmUpFrames + (bAccurateFirstFrameHistory ? 1 : 0));
+}
+
+FFrameNumber FMoviePipelineCameraCutInfo::GetSampleCountEstimate(const bool bIncludeWarmup, const bool bIncludeMotionBlur) const
+{
+	int32 TotalSamples = 0;
+	
+	// Warm Up Frames currently don't submit anything to the GPU.
+	if (bIncludeWarmup)
+	{
+		TotalSamples += 0;
+	}
+
+	// Motion blur only adds one sample (overrides spatial samples, only executes once before first temporal sample)
+	if (bIncludeMotionBlur)
+	{
+		TotalSamples += 1;
+	}
+
+	// Only the main frames are rendered with Temporal sub-sampling and Spatial sub-sampling.
+	TotalSamples += GetOutputFrameCountEstimate().Value * NumTemporalSamples * NumSpatialSamples;
+
+	return FFrameNumber(TotalSamples);
+}
+
+void FMoviePipelineCameraCutInfo::SetNextState(const EMovieRenderShotState InCurrentState)
+{
+	switch (InCurrentState)
+	{
+		// This may get called multiple times so we just do nothing until it's appropriate to move on from WarmingUp.
+	case EMovieRenderShotState::WarmingUp:
+		// Warming Up can jump directly to either Rendering or to MotionBlur depending on if fixes are applied.
+		if (NumWarmUpFramesRemaining == 0)
+		{
+			if (bAccurateFirstFrameHistory)
+			{
+				UE_LOG(LogMovieRenderPipeline, Verbose, TEXT("[%d] Shot WarmUp finished. Setting state to MotionBlur."), GFrameCounter);
+				State = EMovieRenderShotState::MotionBlur;
+			}
+			else
+			{
+				UE_LOG(LogMovieRenderPipeline, Verbose, TEXT("[%d] Shot WarmUp finished. Setting state to Rendering due to no MotionBlur pre-frames."), GFrameCounter);
+				State = EMovieRenderShotState::Rendering;
+			}
+		}
+		break;
+		// This should only be called once with the Uninitalized state.
+	case EMovieRenderShotState::Uninitialized:
+		// Uninitialized can either jump to WarmUp, MotionBlur, or straight to Rendering if no fixes are applied.
+		if (NumWarmUpFramesRemaining > 0)
+		{
+			UE_LOG(LogMovieRenderPipeline, Log, TEXT("[%d] Initialization set state to WarmingUp due to having %d warm up frames."), GFrameCounter, NumWarmUpFramesRemaining);
+			State = EMovieRenderShotState::WarmingUp;
+		}
+		// If they didn't want any warm-up frames we'll still check to see if they want to fix motion blur on the first frame.
+		else if (bAccurateFirstFrameHistory)
+		{
+			UE_LOG(LogMovieRenderPipeline, Log, TEXT("[%d] Initialization set state to MotionBlur due to having no warm up frames."), GFrameCounter);
+			State = EMovieRenderShotState::MotionBlur;
+		}
+		else
+		{
+			UE_LOG(LogMovieRenderPipeline, Verbose, TEXT("[%d] Initialization set state to Rendering due to no MotionBlur pre-frames."), GFrameCounter);
+			State = EMovieRenderShotState::Rendering;
+		}
+		break;
+	}
+
+}
+
+FMoviePipelineWorkInfo FMoviePipelineCameraCutInfo::GetTotalWorkEstimate() const
+{
+	FMoviePipelineWorkInfo WorkEstimate;
+
+	// Initial Range + Handle Frames
+	FFrameNumber OutputFrameCount = GetOutputFrameCountEstimate();
+	WorkEstimate.NumOutputFrames += OutputFrameCount.Value;
+
+	// Warm Up (no submission to gpu), Motion Blur
+	FFrameNumber UtilityFrameCount = GetUtilityFrameCountEstimate();
+	WorkEstimate.NumUtilityFrames += UtilityFrameCount.Value;
+
+	// OutputFrameCount * TemporalSampleCount (since each output frame gets time sliced)
+	FFrameNumber SubsampledOutputFrameCount = GetTemporalFrameCountEstimate();
+	WorkEstimate.NumOutputFramesWithSubsampling += SubsampledOutputFrameCount.Value;
+
+	// How many different samples are we submitting to the GPU?
+	FFrameNumber TotalSampleCount = GetSampleCountEstimate(true, true) * NumTiles.X * NumTiles.Y;
+	WorkEstimate.NumSamples += TotalSampleCount.Value;
+
+	// How many images will we produce to create the high-res tile.
+	WorkEstimate.NumTiles = NumTiles.X * NumTiles.Y;
+
+	return WorkEstimate;
+}
+
+
+bool MoviePipeline::FTileWeight1D::operator==(const FTileWeight1D& Rhs) const
+{
+	const bool bIsSame =
+		(X0 == Rhs.X0) &&
+		(X1 == Rhs.X1) &&
+		(X2 == Rhs.X2) &&
+		(X3 == Rhs.X3);
+
+	return bIsSame;
+}
+
+void MoviePipeline::FTileWeight1D::InitHelper(int32 PadLeft, int32 SizeCenter, int32 PadRight)
+{
+	check(PadLeft >= 0);
+	check(SizeCenter > 0);
+	check(PadRight >= 0);
+
+	int32 Size = PadLeft + SizeCenter + PadRight;
+	int32 Midpoint = PadLeft + SizeCenter / 2;
+
+	X0 = PadLeft / 2;
+	X1 = (3 * PadLeft) / 2;
+	X2 = Size - ((3 * PadLeft) / 2);
+	X3 = Size - (PadRight / 2);
+
+	X1 = FMath::Min<int32>(X1, Midpoint);
+	X2 = FMath::Max<int32>(X2, Midpoint);
+}
+
+float MoviePipeline::FTileWeight1D::CalculateWeight(int32 Pixel) const
+{
+	// the order of the if is important, because if X0=X1, then
+	// we take the Pixel < X0 path, not the Pixel < X1 path which would cause a
+	// divide by zero. Same for X2 and X3.
+	if (Pixel < X0)
+	{
+		return 0.0f;
+	}
+	else if (Pixel < X1)
+	{
+		// note that it's impossible for X0==X1, because the X0 if is earlier
+		return float(Pixel - X0) / float(X1 - X0);
+	}
+	else if (Pixel < X2)
+	{
+		return 1.0f;
+	}
+	else if (Pixel < X3)
+	{
+		// note that it's impossible for X2==X3
+		return float(X3 - Pixel) / float(X3 - X2);
+	}
+
+	return 0.0f;
+}
+
+void MoviePipeline::FTileWeight1D::CalculateArrayWeight(TArray<float>& WeightData, int Size) const
+{
+	WeightData.SetNum(Size);
+	for (int32 Index = 0; Index < X0; Index++)
+	{
+		WeightData[Index] = 0.0f;
+	}
+
+	float ScaleLhs = (X0 != X1) ? 1.0f / float(X1 - X0) : 0.0f;
+	float ScaleRhs = (X2 != X3) ? 1.0f / float(X3 - X2) : 0.0f;
+
+	for (int32 Index = X0; Index < X1; Index++)
+	{
+		float W = FMath::Clamp<float>(((float(Index) + .5f) - float(X0)) * ScaleLhs, 0.0f, 1.0f);
+		WeightData[Index] = W;
+	}
+
+	for (int32 Index = X1; Index < X2; Index++)
+	{
+		WeightData[Index] = 1.0f;
+	}
+
+	for (int32 Index = X2; Index < X3; Index++)
+	{
+		float W = FMath::Clamp<float>(1.0f - ((float(Index) + .5f) - float(X2)) * ScaleRhs, 0.0f, 1.0f);
+		WeightData[Index] = W;
+	}
+
+	for (int32 Index = X3; Index < Size; Index++)
+	{
+		WeightData[Index] = 0.0f;
+	}
+}
+
+
+UMoviePipeline* MoviePipeline::FMoviePipelineEnginePass::GetPipeline() const
+{
+	return OwningPipeline.Get();
+}
