Copy up from Dev-Editor [REDACTED_HANDLE]

diff --git a/Engine/Plugins/MovieScene/MovieRenderPipeline/Source/MovieRenderPipelineCore/Public/MovieRenderPipelineDataTypes.h b/Engine/Plugins/MovieScene/MovieRenderPipeline/Source/MovieRenderPipelineCore/Public/MovieRenderPipelineDataTypes.h
new file mode 100644
--- /dev/null
+++ b/Engine/Plugins/MovieScene/MovieRenderPipeline/Source/MovieRenderPipelineCore/Public/MovieRenderPipelineDataTypes.h
@@ -0,0 +1,907 @@
+// Copyright 1998-2019 Epic Games, Inc. All Rights Reserved.
+#pragma once
+
+#include "Misc/Timecode.h"
+#include "ImagePixelData.h"
+#include "MovieRenderPipelineDataTypes.generated.h"
+
+class UMovieSceneCinematicShotSection;
+
+template<typename ElementType> class TRange;
+template<class T, class TWeakObjectPtrBase> struct TWeakObjectPtr;
+
+class UWorld;
+class ULevelSequence;
+class UMovieSceneCinematicShotSection;
+class UMoviePipelineMasterConfig;
+class UMoviePipelineShotConfig;
+class UMovieSceneCameraCutSection;
+class UMoviePipelineRenderPass;
+struct FImagePixelData;
+struct FImageOverlappedAccumulator;
+class FMoviePipelineOutputMerger;
+class FRenderTarget;
+class UMoviePipeline;
+
+
+/**
+* What is the current overall state of the Pipeline? States are processed in order from first to
+* last and all states will be hit (though there is no guarantee the state will not be transitioned
+* away from on the same frame it entered it). Used to help track overall progress and validate
+* code flow.
+*/
+UENUM(BlueprintType)
+enum class EMovieRenderPipelineState : uint8
+{
+	/** The pipeline has not been initialized yet. Only valid operation is to call Initialize. */
+	Uninitialized = 0,
+	/** The pipeline has been initialized and is now controlling time and working on producing frames. */
+	ProducingFrames = 1,
+	/** All desired frames have been produced. Audio is already finalized. Outputs have a chance to finish long processing tasks. */
+	Finalize = 2,
+	/** All outputs have finished writing to disk or otherwise processing. Additional exports that may have needed information about the produced file can now be run. */
+	Export = 3,
+	/** The pipeline has been shut down. It is an error to shut it down again. */
+	Shutdown = 4,
+};
+
+/**
+* What is the current state of a shot? States are processed in order from first to last but not
+* all states are required, ie: WarmUp and MotionBlur can be disabled and the shot will never
+* pass through this state.
+*/
+UENUM(BlueprintType)
+enum class EMovieRenderShotState : uint8
+{
+	/** The shot has not been initialized yet.*/
+	Uninitialized = 0,
+
+	/** The shot is warming up. Engine ticks are passing but no frames are being produced. */
+	WarmingUp = 1,
+	/*
+	* The shot is doing additional pre-roll for motion blur. No frames are being produced,
+	* but the rendering pipeline is being run to seed histories.
+	*/
+	MotionBlur = 2,
+	/*
+	* The shot is working on producing frames and may be currently doing a sub-frame or
+	* a whole frame.
+	*/
+	Rendering = 3,
+	/*
+	* The shot has produced all frames it will produce. No more evaluation should be
+	* done for this shot once it reaches this state.
+	*/
+	Finished = 4
+};
+
+struct FMoviePipelinePassIdentifier
+{
+	FMoviePipelinePassIdentifier()
+	{}
+
+	FMoviePipelinePassIdentifier(const FString& InPassName)
+		: Name(InPassName)
+	{
+	}
+
+	bool operator == (const FMoviePipelinePassIdentifier& InRHS) const
+	{
+		return Name == InRHS.Name;
+	}
+
+	bool operator != (const FMoviePipelinePassIdentifier& InRHS) const
+	{
+		return !(*this == InRHS);
+	}
+
+	friend uint32 GetTypeHash(FMoviePipelinePassIdentifier OutputState)
+	{
+		return GetTypeHash(OutputState.Name);
+	}
+
+public:
+	FString Name;
+};
+
+namespace MoviePipeline
+{
+	/**
+	* Frame info needed for creating the FSceneView for each tile.
+	*/
+	struct FMoviePipelineFrameInfo
+	{
+	public:
+		FMoviePipelineFrameInfo() {}
+		virtual ~FMoviePipelineFrameInfo() {};
+
+		/**
+		 * We have to apply camera motion vectors manually. So we keep the current and previous fram'es camera view and rotation.
+		 * Then we render a sequence of the same movement, and update after running the game sim.
+		 **/
+
+		/** Current frame camera view location **/
+		FVector CurrViewLocation;
+
+		/** Current frame camera view rotation **/
+		FRotator CurrViewRotation;
+
+		/** Previous frame camera view location **/
+		FVector PrevViewLocation;
+
+		/** Previous frame camera view rotation **/
+		FRotator PrevViewRotation;
+	};
+
+	/**
+	* Utility function for the Pipeline to calculate the time metrics for the current
+	* frame that can then be used by the custom time step to avoid custom time logic.
+	*/
+	struct FFrameTimeStepCache
+	{
+		FFrameTimeStepCache()
+			: DeltaTime(0.0)
+		{}
+
+		FFrameTimeStepCache(double InDeltaTime)
+			: DeltaTime(InDeltaTime)
+		{}
+
+		double DeltaTime;
+	};
+
+	struct FOutputFrameData
+	{
+		FIntPoint Resolution;
+
+		TArray<FColor> ColorBuffer;
+
+		FString PassName;
+	};
+
+	struct MOVIERENDERPIPELINECORE_API FTileWeight1D
+	{
+	public:
+		/**
+		 * For tiles, this struct stores the weight function in one dimension.
+		 *
+		 *                X1           X2
+		 * 1.0              |------------|
+		 *                 /              \
+		 *                /                \
+		 *               /                  \
+		 * 0.0 |--------|                    |----------------|
+		 *     0       X0                    X3             Size
+		 *
+		 *  For a tile, the weight goes from:
+		 *  [0 ,X0]    : 0.0
+		 *  [X0,X1]    : 0.0 to 1.0
+		 *  [X1,X2]    : 1.0
+		 *  [X2,X3]    : 1.0 to 0.0
+		 *  [X3,SizeX] : 0.0
+		 **/
+
+		/** Constructor. Sets to invalid values and must be initialized to reasonable values before using. **/
+		FTileWeight1D()
+		{
+			X0 = 0;
+			X1 = 0;
+			X2 = 0;
+			X3 = 0;
+		}
+
+		/** Make sure this map has valid values. **/
+		void CheckValid() const
+		{
+			check(0 <= X0);
+			check(X0 <= X1);
+			check(X1 <= X2);
+			check(X2 <= X3);
+		}
+
+		/** Is equal operator.**/
+		bool operator==(const FTileWeight1D& Rhs) const;
+
+		/**
+		 * The full tile is of size PadLeft + SizeCenter + PadRight
+		 *
+		 *  |------PadLeft-------|------------SizeCenter-----------|------PadRight-----|
+		 *  
+		 * This function puts X0 in the middle of PadLeft, X3 in the middle of PadRight.
+		 * And X1 is "reflected" around PadLeft, and X2 is "reflected" around PadRight.
+		 *
+		 *  |------PadLeft-------|------------SizeCenter-----------|------PadRight-----|
+		 *           X0                  X1                  X2              X3
+		 * Also note that PadLeft and PadRight are allowed to be zero, but SizeCenter must be > 0;
+		 **/
+		void InitHelper(int32 PadLeft, int32 SizeCenter, int32 PadRight);
+
+		float CalculateWeight(int32 Pixel) const;
+
+		void CalculateArrayWeight(TArray<float>& WeightData, int Size) const;
+
+		int32 X0;
+		int32 X1;
+		int32 X2;
+		int32 X3;
+	};
+
+	struct FFrameConstantMetrics
+	{
+		/** What is the tick resolution fo the sequence */
+		FFrameRate TickResolution;
+		/** What is the effective frame rate of the output */
+		FFrameRate FrameRate;
+		/** How many ticks per output frame. */
+		FFrameTime TicksPerOutputFrame;
+		/** How many ticks per sub-frame sample. */
+		FFrameTime TicksPerSample;
+		/** How many ticks while the shutter is closed */
+		FFrameTime TicksWhileShutterClosed;
+		/** How many ticks while the shutter is opened */
+		FFrameTime TicksWhileShutterOpen;
+
+		/** Fraction of the output frame that the accumulation frames should cover. */
+		double ShutterAnglePercentage;
+
+		/** Fraction of the output frame that the shutter is closed (that accumulation frames should NOT cover). */
+		double ShutterClosedFraction;
+
+		/**
+		* Our time samples are offset by this many ticks. For a given output frame N, we can bias the time N represents
+		* to be before, during or after the data stored on time N.
+		*/
+		FFrameTime ShutterOffsetTicks;
+
+		/** 
+		* How many ticks do we have to go to center our motion blur? This makes it so the blurred area matches the object
+		* position during the time it represents, and not before.
+		*/
+		FFrameTime MotionBlurCenteringOffsetTicks;
+
+		FFrameTime GetFinalEvalTime(const FFrameNumber InTime) const
+		{
+			// We just use a consistent offset from the given time to take motion blur centering and
+			// shutter timing offsets into account. Written here just to consolidate various usages.
+			return FFrameTime(InTime) + MotionBlurCenteringOffsetTicks + ShutterOffsetTicks;
+		}
+
+
+	};
+}
+
+USTRUCT(BlueprintType)
+struct FMoviePipelineWorkInfo
+{
+	GENERATED_BODY()
+
+public:
+	FMoviePipelineWorkInfo()
+		: NumCameraCuts(0)
+		, NumOutputFrames(0)
+		, NumUtilityFrames(0)
+		, NumOutputFramesWithSubsampling(0)
+		, NumSamples(0)
+		, NumTiles(0)
+	{
+	}
+
+	FString ToDisplayString()
+	{
+		return FString::Printf(TEXT("Camera Cuts: %d Output Frames: %d SubsampledOutputFrames: %d TotalSamples: %d UtilityFrames: %d TileCount: %d"),
+			NumCameraCuts, NumOutputFramesWithSubsampling, NumOutputFramesWithSubsampling, NumSamples, NumUtilityFrames, NumTiles);
+	}
+
+	/** Either the current camera cut number (not index) or the total number of camera cuts. */
+	int32 NumCameraCuts;
+	/** How many frames do we expect to see written to disk. */
+	int32 NumOutputFrames;
+	/** How many extra engine ticks are run for utility reasons that don't produce any output. */
+	int32 NumUtilityFrames;
+	/** OutputFrames * TemporalSamples. This many engine frames are passed during the production of all Output Frames. */
+	int32 NumOutputFramesWithSubsampling;
+	/** NumOutputFrames * NumTemporalSamples * SpatialSamples * NumTiles */
+	int32 NumSamples;
+	/** Either the current tile number (not index) or the total number of tiles */
+	int32 NumTiles;
+
+	FMoviePipelineWorkInfo& operator += (const FMoviePipelineWorkInfo& InRHS)
+	{
+		NumCameraCuts += InRHS.NumCameraCuts;
+		NumOutputFrames += InRHS.NumOutputFrames;
+		NumUtilityFrames += InRHS.NumUtilityFrames;
+		NumOutputFramesWithSubsampling += InRHS.NumOutputFramesWithSubsampling;
+		NumSamples += InRHS.NumSamples;
+		NumTiles += InRHS.NumTiles;
+
+		return *this;
+	}
+
+	friend FMoviePipelineWorkInfo operator+ (const FMoviePipelineWorkInfo& InLHS, const FMoviePipelineWorkInfo& InRHS)
+	{
+		FMoviePipelineWorkInfo Sum = InLHS;
+		Sum += InRHS;
+		return Sum;
+	}
+
+	bool operator == (const FMoviePipelineWorkInfo& InRHS) const
+	{
+		return	NumCameraCuts == InRHS.NumCameraCuts
+				&& NumOutputFrames == InRHS.NumOutputFrames
+				&& NumUtilityFrames == InRHS.NumUtilityFrames
+				&& NumOutputFramesWithSubsampling == InRHS.NumOutputFramesWithSubsampling
+				&& NumSamples == InRHS.NumSamples
+				&& NumTiles == InRHS.NumTiles;
+	}
+
+	bool operator != (const FMoviePipelineWorkInfo& InRHS) const
+	{
+		return !(*this == InRHS);
+	}
+
+
+};
+
+
+
+UENUM(BlueprintType)
+enum class EMoviePipelineShutterTiming : uint8
+{
+	FrameOpen,
+	FrameCenter,
+	FrameClose
+};
+
+USTRUCT(BlueprintType)
+struct FMoviePipelineCameraCutInfo
+{
+	GENERATED_BODY()
+public:
+	FMoviePipelineCameraCutInfo()
+		: OriginalRange(TRange<FFrameNumber>::Empty())
+		, TotalOutputRange(TRange<FFrameNumber>::Empty())
+		, NumWarmUpFrames(0)
+		, bAccurateFirstFrameHistory(true)
+		, NumTemporalSamples(0)
+		, NumSpatialSamples(0)
+		, NumTiles(0, 0)
+		, State(EMovieRenderShotState::Uninitialized)
+		, CurrentTick(FFrameNumber(0))
+		, bHasEvaluatedMotionBlurFrame(false)
+		, bHasNotifiedFrameProductionStarted(false)
+		, NumWarmUpFramesRemaining(0)
+	{}
+
+	bool IsInitialized() const { return State != EMovieRenderShotState::Uninitialized; }
+	void SetNextState(const EMovieRenderShotState InCurrentState);
+
+	/** 
+	* Returns the number of expected frames to be produced by Initial Range + Handle Frames given a Frame Rate.
+	* This will be inaccurate when PlayRate tracks are involved.
+	*/
+	FFrameNumber GetOutputFrameCountEstimate() const;
+
+	/**
+	* Returns the expected number of frames different frames that will be submitted for rendering. This is
+	* the number of output frames * temporal samples. This will be inaccurate when PlayRate tracks are involved.
+	*/
+	FFrameNumber GetTemporalFrameCountEstimate() const;
+
+	/** 
+	* Returns misc. utility frame counts (Warm Up + MotionBlur Fix) since these are outside the ranges.
+	*/
+	FFrameNumber GetUtilityFrameCountEstimate() const;
+
+	/**
+	* Returns the number of samples submitted to the GPU, optionally counting samples for the various parts.
+	* This will be inaccurate when PlayRate tracks are involved.
+	*/
+	FFrameNumber GetSampleCountEstimate(const bool bIncludeWarmup = true, const bool bIncludeMotionBlur = true) const;
+	
+	/**
+	* Summarizes the total expected amount of work into one easy to access struct.
+	*/
+	FMoviePipelineWorkInfo GetTotalWorkEstimate() const;
+
+public:
+	/** The original non-modified range for this shot that will be rendered. */
+	TRange<FFrameNumber> OriginalRange;
+
+	/** The range for this shot including handle frames that will be rendered. */
+	TRange<FFrameNumber> TotalOutputRange;
+
+	/** How many warm up frames should this shot wait for. */
+	UPROPERTY(BlueprintReadOnly, Category = "Movie Render Pipeline")
+	int32 NumWarmUpFrames;
+
+	/** Should we evaluate/render an extra frame at the start of this shot to show correct motion blur on the first frame? */
+	UPROPERTY(BlueprintReadOnly, Category = "Movie Render Pipeline")
+	bool bAccurateFirstFrameHistory;
+
+	UPROPERTY(BlueprintReadOnly, Transient, Category = "Movie Render Pipeline")
+	TWeakObjectPtr<UMovieSceneCameraCutSection> CameraCutSection;
+
+	/** How many temporal samples is each frame broken up into? */
+	UPROPERTY(BlueprintReadOnly, Category = "Movie Render Pipeline")
+	int32 NumTemporalSamples;
+
+	/** For each temporal sample, how many spatial samples do we take? */
+	UPROPERTY(BlueprintReadOnly, Category = "Movie Render Pipeline")
+	int32 NumSpatialSamples;
+
+	/** How many image tiles are going to be rendered per temporal frame. */
+	UPROPERTY(BlueprintReadOnly, Category = "Movie Render Pipeline")
+	FIntPoint NumTiles;
+
+	/** Display name for UI Purposes. */
+	UPROPERTY(BlueprintReadOnly, Category = "Movie Render Pipeline")
+	FString DisplayName;
+
+	/** Cached Frame Rate these are being rendered at. Simplifies some APIs. */
+	FFrameRate CachedFrameRate;
+
+	/** Cached Tick Resolution our numbers are in. Simplifies some APIs. */
+	FFrameRate CachedTickResolution;
+public:
+	/** The current state of processing this Shot is in. Not all states will be passed through. */
+	UPROPERTY(BlueprintReadOnly, Category = "Movie Render Pipeline")
+	EMovieRenderShotState State;
+
+	UPROPERTY(BlueprintReadOnly, Category = "Movie Render Pipeline")
+	FFrameNumber CurrentTick;
+
+	/** Metrics - How much work has been done for this particular shot. */
+	UPROPERTY(BlueprintReadOnly, Category = "Movie Render Pipeline")
+	FMoviePipelineWorkInfo CurrentWorkInfo;
+
+	/** 
+	* Metrics - How much work is there estimated to do for this shot? Can be more than estimated total due to PlayRate tracks. 
+	* Cached after setup so that Blueprints can read the variable directly due to no support for UFunctions on Structs.
+	*/
+	UPROPERTY(BlueprintReadOnly, Category = "Movie Render Pipeline")
+	FMoviePipelineWorkInfo TotalWorkInfo;
+
+public:
+	/** Have we evaluated the motion blur frame? Only used if bEvaluateMotionBlurOnFirstFrame is set */
+	UPROPERTY(BlueprintReadOnly, Category = "Movie Render Pipeline")
+	bool bHasEvaluatedMotionBlurFrame;
+
+	/** Have we notified the output containers that we're about to begin producing frames? */
+	UPROPERTY(BlueprintReadOnly, Category = "Movie Render Pipeline")
+	bool bHasNotifiedFrameProductionStarted;
+
+	/** How many warm up frames are left to process for this shot. May always be zero. */
+	UPROPERTY(BlueprintReadOnly, Category = "Movie Render Pipeline")
+	int32 NumWarmUpFramesRemaining;
+
+	bool operator == (const FMoviePipelineCameraCutInfo& InRHS) const
+	{
+		return
+			OriginalRange == InRHS.OriginalRange &&
+			TotalOutputRange == InRHS.TotalOutputRange &&
+			State == InRHS.State &&
+			NumWarmUpFramesRemaining == InRHS.NumWarmUpFramesRemaining &&
+			bAccurateFirstFrameHistory == InRHS.bAccurateFirstFrameHistory &&
+			bHasEvaluatedMotionBlurFrame == InRHS.bHasEvaluatedMotionBlurFrame &&
+			bHasNotifiedFrameProductionStarted == InRHS.bHasNotifiedFrameProductionStarted &&
+			CurrentTick == InRHS.CurrentTick &&
+			CameraCutSection == InRHS.CameraCutSection;
+	}
+
+	bool operator != (const FMoviePipelineCameraCutInfo& InRHS) const
+	{
+		return !(*this == InRHS);
+	}
+};
+
+/**
+* Pre-calculated information about a shot we are going to produce. This lets us build
+* the expected output at the start of the process and just read from it later. Having
+* all information in advanced aids in debugging and visualization of progress.
+*/
+USTRUCT(BlueprintType)
+struct FMoviePipelineShotInfo
+{
+	GENERATED_BODY()
+public:
+	FMoviePipelineShotInfo()
+		: NumHandleFrames(0)
+		, OriginalRange(TRange<FFrameNumber>::Empty())
+		, TotalOutputRange(TRange<FFrameNumber>::Empty())
+		, HandleFrameRangeStart(TRange<FFrameNumber>::Empty())
+		, HandleFrameRangeEnd(TRange<FFrameNumber>::Empty())
+		, ShotOverrideConfig(nullptr)
+		, CinematicShotSection(nullptr)
+		, CurrentCameraCutIndex(0)
+	{}
+
+	FString GetDisplayName() const;
+
+public:
+	/** How many handle frames (in Display Rate of the Master Sequence) */
+	int32 NumHandleFrames;
+
+	/** The original non-modified (overall) range for this shot that will be rendered. */
+	TRange<FFrameNumber> OriginalRange;
+
+	/** The range for this shot including handle frames that will be rendered. */
+	TRange<FFrameNumber> TotalOutputRange;
+
+	/** The range of time that represents the handle frames (if any) before the shot. */
+	TRange<FFrameNumber> HandleFrameRangeStart;
+
+	/** The range of time that represents the handle frames (if any) after the shot. */
+	TRange<FFrameNumber> HandleFrameRangeEnd;
+	
+	UPROPERTY(Transient, BlueprintReadOnly, Category = "Movie Render Pipeline")
+	UMoviePipelineShotConfig* ShotOverrideConfig;
+
+	UPROPERTY(BlueprintReadOnly, Transient, Category = "Movie Render Pipeline")
+	TWeakObjectPtr<UMovieSceneCinematicShotSection> CinematicShotSection;
+
+	UPROPERTY(BlueprintReadOnly, Category = "Movie Render Pipeline")
+	TArray<FMoviePipelineCameraCutInfo> CameraCuts;
+
+	UPROPERTY(BlueprintReadOnly, Category = "Movie Render Pipeline")
+	int32 CurrentCameraCutIndex;
+
+	FMoviePipelineCameraCutInfo& GetCurrentCameraCut()
+	{
+		check(CurrentCameraCutIndex >= 0 && CurrentCameraCutIndex < CameraCuts.Num());
+		return CameraCuts[CurrentCameraCutIndex];
+	}
+
+	bool SetNextShotActive()
+	{
+		if (CurrentCameraCutIndex + 1 < CameraCuts.Num() - 1)
+		{
+			CurrentCameraCutIndex++;
+			return false;
+		}
+
+		// We're on the last camera cut, we can't make another one active.
+		return true;
+	}
+
+	bool operator == (const FMoviePipelineShotInfo& InRHS) const
+	{ 
+		return
+			NumHandleFrames == InRHS.NumHandleFrames &&
+			OriginalRange == InRHS.OriginalRange &&
+			TotalOutputRange == InRHS.TotalOutputRange &&
+			HandleFrameRangeStart == InRHS.HandleFrameRangeStart &&
+			HandleFrameRangeEnd == InRHS.HandleFrameRangeEnd &&
+			ShotOverrideConfig == InRHS.ShotOverrideConfig &&
+			CinematicShotSection == InRHS.CinematicShotSection &&
+			CameraCuts == InRHS.CameraCuts &&
+			CurrentCameraCutIndex == InRHS.CurrentCameraCutIndex;
+	}
+
+	bool operator != (const FMoviePipelineShotInfo& InRHS) const
+	{
+		return !(*this == InRHS);
+	}
+};
+
+/**
+* The Tick/Render loops are decoupled from the actual desired output. In some cases
+* we may render but not desire an output (such as filling temporal histories) or
+* we may be accumulating the results into a target which does not produce an output.
+* Finally, we may be running the Tick/Render loop but not wanting to do anything!
+* This is the case when we are using the frame-step debugging, the engine is still
+* live and running, but we don't want to advance the output process forwards at all.
+*/
+USTRUCT(BlueprintType)
+struct FMoviePipelineFrameOutputState
+{
+	GENERATED_BODY()
+public:
+	FMoviePipelineFrameOutputState()
+		: OutputFrameNumber(-1)
+		, TemporalSampleIndex(-1)
+		, MotionBlurFraction(0.5f)
+		, FrameDeltaTime(0)
+		, WorldSeconds(0)
+		, bWasAffectedByTimeDilation(false)
+	{
+		// There is other code which relies on this starting on -1.
+		check(OutputFrameNumber == -1);
+
+		// Likewise, the code assumes you start on sample index 0.
+		check(TemporalSampleIndex == -1);
+	}
+
+	/**
+	* The expected output frame count that the render is working towards creating.
+	* This number accurately tracks the number of frames we have produced even if
+	* the file written to disk uses a different number (due to relative frame numbers
+	* or offset frames being added.
+	*/
+	UPROPERTY(BlueprintReadOnly, Transient, Category = "Movie Render Pipeline")
+	int32 OutputFrameNumber;
+
+	/** Which sub-frame are we on when using Accumulation Frame rendering. */
+	UPROPERTY(BlueprintReadOnly, Transient, Category = "Movie Render Pipeline")
+	int32 TemporalSampleIndex;
+
+	/** A 0-1 fraction of how much motion blur should be applied. This is tied to the shutter angle but more complicated when using sub-sampling */
+	UPROPERTY(BlueprintReadOnly, Transient, Category = "Movie Render Pipeline")
+	float MotionBlurFraction;
+
+	/** The delta time in seconds used for this frame. */
+	double FrameDeltaTime;
+
+	/** The total elapsed time since the pipeline started.*/
+	double WorldSeconds;
+	
+	/** The closest frame number (in Display Rate) on the Sequence. May be duplicates in the case of different output framerate or Play Rate tracks. */
+	FFrameNumber SourceFrameNumber;
+
+	/** The closest time code version of the SourceFrameNumber on the Sequence. May be a duplicate in the case of different output framerate or Play Rate tracks. */
+	FTimecode SourceTimeCode;
+
+	/** 
+	* The closest frame number (in Display Rate) on the Sequence adjusted for the effective output rate. These numbers will not line up with the frame
+	* in the source Sequence if the output frame rate differs from the Sequence display rate. May be a duplicate in the event of Play Rate tracks.
+	*/
+	FFrameNumber EffectiveFrameNumber;
+
+	/** The closest time code version of the EffectiveFrameNumber. May be a duplicate in the event of Play Rate tracks. */
+	FTimecode EffectiveTimeCode;
+
+	/**
+	* If true, there was a non-1.0 Time Dilation in effect when this frame was produced. This indicates that there 
+	* may be duplicate frame Source/Effective frame numbers as they find the closest ideal time to the current.
+	*/
+	bool bWasAffectedByTimeDilation;
+
+	bool operator == (const FMoviePipelineFrameOutputState& InRHS) const
+	{
+		return
+			OutputFrameNumber == InRHS.OutputFrameNumber &&
+			// TemporalSampleIndex == InRHS.TemporalSampleIndex &&
+			// FMath::IsNearlyEqual(MotionBlurFraction,InRHS.MotionBlurFraction) &&
+			// FMath::IsNearlyEqual(FrameDeltaTime, InRHS.FrameDeltaTime) &&
+			// FMath::IsNearlyEqual(WorldSeconds, InRHS.WorldSeconds) &&
+			SourceFrameNumber == InRHS.SourceFrameNumber &&
+			SourceTimeCode == InRHS.SourceTimeCode &&
+			EffectiveFrameNumber == InRHS.EffectiveFrameNumber &&
+			EffectiveTimeCode == InRHS.EffectiveTimeCode &&
+			bWasAffectedByTimeDilation == InRHS.bWasAffectedByTimeDilation;
+	}
+
+	bool operator != (const FMoviePipelineFrameOutputState& InRHS) const
+	{
+		return !(*this == InRHS);
+	}
+
+	friend uint32 GetTypeHash(FMoviePipelineFrameOutputState OutputState)
+	{
+		return GetTypeHash(OutputState.OutputFrameNumber);
+	}
+};
+
+/**
+* These parameters define a single sample that a render pass should render with.
+*/
+struct FMoviePipelineRenderPassMetrics
+{
+public:
+	/** How many tiles on X and Y are there total. */
+	FIntPoint TileCounts;
+
+	/** Of the TileCount, which X/Y Tile is this sample for. */
+	FIntPoint TileIndexes;
+
+	/** Get a 0-(TileCount-1) version of TileIndex. */
+	FORCEINLINE int32 GetTileIndex() const
+	{
+		return (TileIndexes.Y * TileCounts.Y) + TileIndexes.X;
+	}
+
+	FORCEINLINE int32 GetTileCount() const
+	{
+		return (TileCounts.X * TileCounts.Y);
+	}
+
+	/** How big is the back buffer for this sample? This is the final target size divided by the number of tiles with padding added. */
+	FIntPoint BackbufferSize;
+
+	/** How big is an individual tile for this sample? This is the final target size divided by the number of tiles (without padding added) */
+	FIntPoint TileSize;
+
+	/** If true, we will discard this sample after rendering. Used to get history set up correctly. */
+	bool bIsHistoryOnlyFrame;
+
+	/** If true, we will write this sample to disk (and also send it to the accumulator.) */
+	bool bWriteSampleToDisk;
+
+	/** How many spatial jitters will there be total for this particular temporal accumulation frame? */
+	int32 SpatialSampleCount;
+	
+	/** Of the SpatialSampleCount, which index is this? */
+	int32 SpatialSampleIndex;
+
+	/** How many temporal jitters will there be total for this particular output frame? */
+	int32 TemporalSampleCount;
+
+	/** Of the TemporalSampleCount, which index is this? */
+	int32 TemporalSampleIndex;
+
+	/** How many pixels are we overlapping with adjacent tiles (on each side) */
+	FIntPoint OverlappedPad;
+
+	/** How much is this sample offset, taking padding into account. */
+	FIntPoint OverlappedOffset;
+
+	/** 
+	* The gamma space to apply accumulation in. During accumulation, pow(x,AccumulationGamma) is applied
+	* and pow(x,1/AccumulationGamma) is applied after accumulation is finished. 1.0 means no change."
+	*/
+	float AccumulationGamma;
+
+	/**
+	* The amount of screen-space shift used to replace TAA Projection Jitter, modified for each spatial sample
+	* of the render.
+	*/
+	FVector2D SpatialShift;
+
+	/**
+	* If set, forces the exposure compensation on a render. Useful when doing tiled renders where auto-exposure is disabled.
+	*/
+	TOptional<float> ExposureCompensation;
+
+	/**
+	* Any additional texture mip map bias that should be added when rendering. Can be used to force extra sharpness.
+	*/
+	float TextureSharpnessBias;
+
+
+	FMoviePipelineFrameOutputState OutputState;
+
+	FVector2D OverlappedSubpixelShift;
+
+	MoviePipeline::FTileWeight1D WeightFunctionX;
+	MoviePipeline::FTileWeight1D WeightFunctionY;
+
+	MoviePipeline::FMoviePipelineFrameInfo FrameInfo;
+};
+
+namespace MoviePipeline
+{
+	struct FMoviePipelineRenderPassInitSettings
+	{
+	public:
+		FMoviePipelineRenderPassInitSettings()
+			: BackbufferResolution(0, 0)
+			, TileCount(0, 0)
+			, bAccumulateAlpha(false)
+		{
+		}
+
+	public:
+		/** This takes into account any padding needed for tiled rendering overlap. Different than the output resolution of the final image. */
+		FIntPoint BackbufferResolution;
+
+		/** How many tiles (in each direction) are we rendering with. */
+		FIntPoint TileCount;
+
+		/** If true, the alpha channel will be accumulated. This adds about 30% cost over the base to accumulation. */
+		bool bAccumulateAlpha;
+	};
+
+	struct MOVIERENDERPIPELINECORE_API FMoviePipelineEnginePass
+	{
+		FMoviePipelineEnginePass(const FMoviePipelinePassIdentifier& InPassIdentifier)
+			: PassIdentifier(InPassIdentifier)
+		{
+		}
+
+		virtual ~FMoviePipelineEnginePass()
+		{}
+
+		virtual void Setup(TWeakObjectPtr<UMoviePipeline> InOwningPipeline, const FMoviePipelineRenderPassInitSettings& InInitSettings)
+		{
+			OwningPipeline = InOwningPipeline;
+			InitSettings = InInitSettings;
+		}
+
+		virtual void RenderSample_GameThread(const FMoviePipelineRenderPassMetrics& InSampleState)
+		{
+		}
+
+		virtual void Teardown()
+		{
+		}
+
+		const FMoviePipelineRenderPassInitSettings& GetInitSettings() const { return InitSettings; }
+
+	protected:
+		UMoviePipeline* GetPipeline() const;
+
+	public:
+		/** A unique name for this engine pass. This is how an individual output pass specifies what data source it wants. */
+		FMoviePipelinePassIdentifier PassIdentifier;
+
+	protected:
+		FMoviePipelineRenderPassInitSettings InitSettings;
+
+	private:
+		TWeakObjectPtr<UMoviePipeline> OwningPipeline;
+	};
+
+	struct FSampleRenderThreadParams
+	{
+		FMoviePipelineRenderPassMetrics SampleState;
+		TSharedPtr<FImageOverlappedAccumulator, ESPMode::ThreadSafe> ImageAccumulator;
+		TSharedPtr<FMoviePipelineOutputMerger, ESPMode::ThreadSafe> OutputMerger;
+		FMoviePipelinePassIdentifier PassIdentifier;
+		bool bAccumulateAlpha;
+	};
+}
+
+struct FImagePixelDataPayload : IImagePixelDataPayload, public TSharedFromThis<FImagePixelDataPayload>
+{
+	FMoviePipelineFrameOutputState OutputState;
+	FMoviePipelineRenderPassMetrics SampleState;
+
+	FMoviePipelinePassIdentifier PassIdentifier;
+
+	/** Is this the first tile of an image and we should start accumulating? */
+	FORCEINLINE bool IsFirstTile() const
+	{
+		return SampleState.TileIndexes.X == 0 && SampleState.TileIndexes.Y == 0 && SampleState.SpatialSampleIndex == 0;
+	}
+
+	/** Is this the last tile of an image and we should finish accumulating? */
+	FORCEINLINE bool IsLastTile() const
+	{
+		return SampleState.TileIndexes.X == SampleState.TileCounts.X - 1 &&
+			   SampleState.TileIndexes.Y == SampleState.TileCounts.Y - 1 &&
+			   SampleState.SpatialSampleIndex == SampleState.SpatialSampleCount - 1;
+	}
+
+	FORCEINLINE bool IsFirstTemporalSample() const
+	{
+		return SampleState.TemporalSampleIndex == 0;
+	}
+
+	FORCEINLINE bool IsLastTemporalSample() const
+	{
+		return SampleState.TemporalSampleIndex == SampleState.TemporalSampleCount - 1;
+	}
+};
+
+struct MOVIERENDERPIPELINECORE_API FMoviePipelineMergerOutputFrame
+{
+public:
+	FMoviePipelineMergerOutputFrame() {}
+	virtual ~FMoviePipelineMergerOutputFrame() {};
+
+	FMoviePipelineMergerOutputFrame& operator=(FMoviePipelineMergerOutputFrame&& InOther)
+	{
+		FrameOutputState = InOther.FrameOutputState;
+		ExpectedRenderPasses = InOther.ExpectedRenderPasses;
+		ImageOutputData = MoveTemp(InOther.ImageOutputData);
+	
+		return *this;
+	}
+	FMoviePipelineMergerOutputFrame(FMoviePipelineMergerOutputFrame&& InOther)
+		: FrameOutputState(InOther.FrameOutputState)
+		, ExpectedRenderPasses(InOther.ExpectedRenderPasses)
+		, ImageOutputData(MoveTemp(InOther.ImageOutputData))
+	{
+	}
+
+private:
+	// Explicitly delete copy operators since we own unique data.
+	// void operator=(const FMoviePipelineMergerOutputFrame&);
+	// FMoviePipelineMergerOutputFrame(const FMoviePipelineMergerOutputFrame&);
+
+public:
+	FMoviePipelineFrameOutputState FrameOutputState;
+
+	TArray<FMoviePipelinePassIdentifier> ExpectedRenderPasses;
+
+	TMap<FMoviePipelinePassIdentifier, TUniquePtr<FImagePixelData>> ImageOutputData;
+};
\ No newline at end of file
diff --git a/Engine/Plugins/MovieScene/MovieRenderPipeline/Source/MovieRenderPipelineEditor/Private/MoviePipelineEditorBlueprintLibrary.cpp b/Engine/Plugins/MovieScene/MovieRenderPipeline/Source/MovieRenderPipelineEditor/Private/MoviePipelineEditorBlueprintLibrary.cpp
new file mode 100644
--- /dev/null
+++ b/Engine/Plugins/MovieScene/MovieRenderPipeline/Source/MovieRenderPipelineEditor/Private/MoviePipelineEditorBlueprintLibrary.cpp
@@ -0,0 +1,52 @@
+// Copyright 1998-2019 Epic Games, Inc. All Rights Reserved.
+
+#include "MoviePipelineEditorBlueprintLibrary.h"
+#include "UObject/Package.h"
+#include "Misc/PackageName.h"
+#include "AssetRegistryModule.h"
+#include "FileHelpers.h"
+#include "ObjectTools.h"
+#include "PackageTools.h"
+
+#define LOCTEXT_NAMESPACE "MoviePipelineEditorBlueprintLibrary"
+
+bool UMoviePipelineEditorBlueprintLibrary::ExportConfigToAsset(const UMoviePipelineMasterConfig* InConfig, const FString& InPackagePath, const FString& InFileName, const bool bInSaveAsset, UMoviePipelineMasterConfig*& OutAsset, FText& OutErrorReason)
+{
+	if(!InConfig)
+	{
+		OutErrorReason = LOCTEXT("CantExportNullConfigToPackage", "Can't export a null configuration to a package.");
+		return false;
+	}
+	
+	FString FixedAssetName = ObjectTools::SanitizeObjectName(InFileName);
+	FString NewPackageName = FPackageName::GetLongPackagePath(InPackagePath) + TEXT("/") + FixedAssetName;
+
+	if (!FPackageName::IsValidLongPackageName(NewPackageName, false, &OutErrorReason))
+	{
+		return false;
+	}
+
+	UPackage* NewPackage = CreatePackage(nullptr, *NewPackageName);
+	NewPackage->AddToRoot();
+	
+	// Duplicate the provided config into this package.
+	UMoviePipelineMasterConfig* NewConfig = Cast<UMoviePipelineMasterConfig>(StaticDuplicateObject(InConfig, NewPackage, FName(*InFileName), RF_NoFlags));
+	NewConfig->SetFlags(RF_Public | RF_Transactional | RF_Standalone);
+	NewConfig->MarkPackageDirty();
+
+	// Mark it so it shows up in the Content Browser immediately
+	FAssetRegistryModule::AssetCreated(NewConfig);
+
+	// If they want to save, ask them to save (and add to version control)
+	if (bInSaveAsset)
+	{
+		TArray<UPackage*> Packages;
+		Packages.Add(NewConfig->GetOutermost());
+
+		return UEditorLoadingAndSavingUtils::SavePackages(Packages, true);
+	}
+
+	return true;
+}
+
+#undef LOCTEXT_NAMESPACE // "MoviePipelineEditorBlueprintLibrary"
\ No newline at end of file
diff --git a/Engine/Plugins/MovieScene/MovieRenderPipeline/Source/MovieRenderPipelineEditor/Private/MoviePipelinePIEExecutor.cpp b/Engine/Plugins/MovieScene/MovieRenderPipeline/Source/MovieRenderPipelineEditor/Private/MoviePipelinePIEExecutor.cpp
new file mode 100644
--- /dev/null
+++ b/Engine/Plugins/MovieScene/MovieRenderPipeline/Source/MovieRenderPipelineEditor/Private/MoviePipelinePIEExecutor.cpp
@@ -0,0 +1,154 @@
+// Copyright 1998-2019 Epic Games, Inc. All Rights Reserved.
+#include "MoviePipelinePIEExecutor.h"
+#include "MoviePipelineMasterConfig.h"
+#include "MoviePipelineShotConfig.h"
+#include "MoviePipeline.h"
+#include "MoviePipelineGameOverrideSetting.h"
+#include "MovieRenderPipelineCoreModule.h"
+#include "Framework/Application/SlateApplication.h"
+#include "MovieRenderPipelineDataTypes.h"
+#include "MovieRenderPipelineSettings.h"
+#include "Editor/EditorEngine.h"
+#include "Editor.h"
+#include "Settings/LevelEditorPlaySettings.h"
+#include "MoviePipelineQueue.h"
+
+#define LOCTEXT_NAMESPACE "MoviePipelinePIEExecutor"
+
+FText UMoviePipelinePIEExecutor::GetWindowTitle(const int32 InConfigIndex, const int32 InNumConfigs) const
+{
+	FNumberFormattingOptions PercentFormatOptions;
+	PercentFormatOptions.MinimumIntegralDigits = 2;
+	PercentFormatOptions.MaximumIntegralDigits = 3;
+
+	FText TitleFormatString = LOCTEXT("MoviePreviewWindowTitleFormat", "Movie Pipeline Render (Preview) [{CurrentCount}/{TotalCount} Total] {PercentComplete}% Completed.");
+	FText WindowTitle = FText::FormatNamed(TitleFormatString, TEXT("CurrentCount"), FText::AsNumber(InConfigIndex + 1), TEXT("TotalCount"), FText::AsNumber(InNumConfigs), TEXT("PercentComplete"), FText::AsNumber(12.f, &PercentFormatOptions));
+	return WindowTitle;
+}
+
+void UMoviePipelinePIEExecutor::Start(const UMoviePipelineExecutorJob* InJob)
+{
+	Super::Start(InJob);
+
+	// Create a Slate window to hold our preview.
+	TSharedRef<SWindow> CustomWindow = SNew(SWindow)
+		.Title_UObject(this, &UMoviePipelinePIEExecutor::GetWindowTitle, CurrentPipelineIndex, Queue->GetJobs().Num())
+		.ClientSize(FVector2D(1280, 720))
+		.AutoCenter(EAutoCenter::PrimaryWorkArea)
+		.UseOSWindowBorder(true)
+		.FocusWhenFirstShown(false)
+		.ActivationPolicy(EWindowActivationPolicy::Never)
+		.HasCloseButton(true)
+		.SupportsMaximize(false)
+		.SupportsMinimize(true)
+		.SizingRule(ESizingRule::UserSized);
+
+	FSlateApplication::Get().AddWindow(CustomWindow);
+
+	// Initialize our own copy of the Editor Play settings which we will adjust defaults on.
+	ULevelEditorPlaySettings* PlayInEditorSettings = NewObject<ULevelEditorPlaySettings>();
+	PlayInEditorSettings->SetPlayNetMode(EPlayNetMode::PIE_Standalone);
+	PlayInEditorSettings->SetPlayNumberOfClients(1);
+	PlayInEditorSettings->bLaunchSeparateServer = false;
+
+	FRequestPlaySessionParams Params;
+	Params.EditorPlaySettings = PlayInEditorSettings;
+	Params.CustomPIEWindow = CustomWindow;
+
+	UMoviePipelineGameOverrideSetting* GameOverrides = InJob->GetConfiguration()->FindSetting<UMoviePipelineGameOverrideSetting>();
+	if (GameOverrides)
+	{
+		Params.GameModeOverride = GameOverrides->GameModeOverride;
+	}
+
+	// Kick off an async request to start a play session. This won't happen until the next frame.
+	GEditor->RequestPlaySession(Params);
+
+	// Listen for PIE startup since there's no current way to pass a delegate through the request.
+	FEditorDelegates::PostPIEStarted.AddUObject(this, &UMoviePipelinePIEExecutor::OnPIEStartupFinished);
+}
+
+void UMoviePipelinePIEExecutor::OnPIEStartupFinished(bool)
+{
+	// Immediately un-bind our delegate so that we don't catch all PIE startup requests in the future.
+	FEditorDelegates::PostPIEStarted.RemoveAll(this);
+
+
+	// Hack to find out the PIE world since it is not provided by the delegate.
+	UWorld* ExecutingWorld = nullptr;
+	
+	for (const FWorldContext& Context : GEngine->GetWorldContexts())
+	{
+		if (Context.WorldType == EWorldType::PIE)
+		{
+			ExecutingWorld = Context.World();
+		}
+	}
+	
+	check(ExecutingWorld);
+
+	// Allow the user to have overridden which Pipeline is actually run. This is an unlikely scenario but allows
+	// the user to create their own implementation while still re-using the rest of our UI and infrastructure.
+	const UMovieRenderPipelineProjectSettings* ProjectSettings = GetDefault<UMovieRenderPipelineProjectSettings>();
+	TSubclassOf<UMoviePipeline> PipelineClass = ProjectSettings->DefaultPipeline;
+
+	// This Pipeline belongs to the world being created so that they have context for things they execute.
+	ActiveMoviePipeline = NewObject<UMoviePipeline>(ExecutingWorld, PipelineClass);
+	
+	ActiveMoviePipeline->Initialize(Queue->GetJobs()[CurrentPipelineIndex]);
+
+	// Listen for when the pipeline thinks it has finished.
+	ActiveMoviePipeline->OnMoviePipelineFinished().AddUObject(this, &UMoviePipelinePIEExecutor::OnPIEMoviePipelineFinished);
+	ActiveMoviePipeline->OnMoviePipelineErrored().AddUObject(this, &UMoviePipelinePIEExecutor::OnPipelineErrored);
+
+	// Listen for PIE shutdown in case the user hits escape to close it. 
+	FEditorDelegates::EndPIE.AddUObject(this, &UMoviePipelinePIEExecutor::OnPIEEnded);
+}
+
+void UMoviePipelinePIEExecutor::OnPIEMoviePipelineFinished(UMoviePipeline* InMoviePipeline)
+{
+	// Unsubscribe to the EndPIE event so we don't think the user canceled it.
+	FEditorDelegates::EndPIE.RemoveAll(this);
+
+	if (ActiveMoviePipeline)
+	{
+		// Unsubscribe in the event that it gets called twice we don't have issues.
+		ActiveMoviePipeline->OnMoviePipelineFinished().RemoveAll(this);
+	}
+
+	// The End Play will happen on the next frame.
+	GEditor->RequestEndPlayMap();
+
+	// Delay for one frame so that PIE can finish shut down. It's not a huge fan of us starting up on the same frame.
+	FTimerHandle Handle;
+	GEditor->GetTimerManager()->SetTimer(Handle, FTimerDelegate::CreateUObject(this, &UMoviePipelinePIEExecutor::DelayedFinishNotification), 1.f, false);
+}
+
+void UMoviePipelinePIEExecutor::OnPIEEnded(bool)
+{
+	// If the movie pipeline finishes naturally we unsubscribe from the EndPIE event.
+	// This means that if this gets called, the user has hit escape and wants to cancel.
+	if (ActiveMoviePipeline)
+	{
+		// This will flush any outstanding work on the movie pipeline (file writes) immediately
+		ActiveMoviePipeline->Shutdown();
+
+		// Broadcast that we finished.
+		// ToDo: bAnyJobHadFatalError
+		DelayedFinishNotification();
+	}
+}
+
+void UMoviePipelinePIEExecutor::DelayedFinishNotification()
+{
+	UMoviePipeline* MoviePipeline = ActiveMoviePipeline;
+	
+	// Null these out now since OnIndividualPipelineFinished might invoke something that causes a GC
+	// and we want them to go away with the GC.
+	ActiveMoviePipeline = nullptr;
+	
+	// Now that another frame has passed and we should be OK to start another PIE session, notify our owner.
+	OnIndividualPipelineFinished(MoviePipeline);
+}
+
+#undef LOCTEXT_NAMESPACE // "MoviePipelinePIEExecutor"
diff --git a/Engine/Plugins/MovieScene/MovieRenderPipeline/Source/MovieRenderPipelineEditor/Private/MovieRenderPipelineEditorModule.cpp b/Engine/Plugins/MovieScene/MovieRenderPipeline/Source/MovieRenderPipelineEditor/Private/MovieRenderPipelineEditorModule.cpp
new file mode 100644
--- /dev/null
+++ b/Engine/Plugins/MovieScene/MovieRenderPipeline/Source/MovieRenderPipelineEditor/Private/MovieRenderPipelineEditorModule.cpp
@@ -0,0 +1,158 @@
+// Copyright 1998-2019 Epic Games, Inc. All Rights Reserved.
+
+#include "MovieRenderPipelineEditorModule.h"
+#include "MovieRenderPipelineCoreModule.h"
+#include "MoviePipelineExecutor.h"
+#include "WorkspaceMenuStructureModule.h"
+#include "WorkspaceMenuStructure.h"
+#include "Widgets/SMoviePipelineConfigTabContent.h"
+#include "Widgets/SMoviePipelineQueueTabContent.h"
+#include "ISettingsModule.h"
+#include "MovieRenderPipelineSettings.h"
+#include "HAL/IConsoleManager.h"
+#include "MoviePipelineExecutor.h"
+#include "Widgets/Docking/SDockTab.h"
+#include "Framework/Docking/TabManager.h"
+#include "EditorStyleSet.h"
+#include "Misc/PackageName.h"
+#include "UObject/UObjectGlobals.h"
+#include "UObject/Package.h"
+#include "Templates/SubclassOf.h"
+#include "MovieRenderPipelineStyle.h"
+
+#define LOCTEXT_NAMESPACE "FMovieRenderPipelineEditorModule"
+
+FName IMovieRenderPipelineEditorModule::MoviePipelineQueueTabName = "MoviePipelineQueue";
+FText IMovieRenderPipelineEditorModule::MoviePipelineQueueTabLabel = LOCTEXT("MovieRenderQueueTab_Label", "Movie Render Queue");
+FName IMovieRenderPipelineEditorModule::MoviePipelineConfigEditorTabName = "MovieRenderPipeline";
+FText IMovieRenderPipelineEditorModule::MoviePipelineConfigEditorTabLabel = LOCTEXT("MovieRenderPipelineTab_Label", "Movie Render Pipeline");
+										 
+namespace
+{
+	static TSharedRef<SDockTab> SpawnMovieRenderPipelineTab(const FSpawnTabArgs& SpawnTabArgs)
+	{
+		TSharedRef<SDockTab> NewTab = SNew(SDockTab)
+			.TabRole(ETabRole::NomadTab)
+			[
+				SNew(SMoviePipelineConfigTabContent)
+			];
+		
+		return NewTab;
+	}
+
+	static TSharedRef<SDockTab> SpawnMoviePipelineQueueTab(const FSpawnTabArgs& InSpawnTabArgs)
+	{
+		TSharedRef<SDockTab> NewTab = SNew(SDockTab)
+			.TabRole(ETabRole::MajorTab)
+			[
+				SNew(SMoviePipelineQueueTabContent)
+			];
+
+		return NewTab;
+	}
+
+	static void RegisterTabImpl()
+	{
+		// FTabSpawnerEntry& MRPConfigTabSpawner = FGlobalTabmanager::Get()->RegisterNomadTabSpawner(IMovieRenderPipelineEditorModule::MoviePipelineConfigEditorTabName, FOnSpawnTab::CreateStatic(SpawnMovieRenderPipelineTab));
+		// 
+		// MRPConfigTabSpawner
+		// 	.SetGroup(WorkspaceMenu::GetMenuStructure().GetLevelEditorCinematicsCategory())
+		// 	.SetDisplayName(IMovieRenderPipelineEditorModule::MoviePipelineConfigEditorTabLabel)
+		// 	.SetTooltipText(LOCTEXT("MovieRenderPipelineConfigTab_Tooltip", "Open the Movie Render Config UI for creating and editing presets."))
+		// 	.SetIcon(FSlateIcon(FEditorStyle::GetStyleSetName(), "ContentBrowser.TabIcon"));
+
+		FTabSpawnerEntry& MRPQueueTabSpawner = FGlobalTabmanager::Get()->RegisterNomadTabSpawner(IMovieRenderPipelineEditorModule::MoviePipelineQueueTabName, FOnSpawnTab::CreateStatic(SpawnMoviePipelineQueueTab));
+
+		MRPQueueTabSpawner
+			.SetGroup(WorkspaceMenu::GetMenuStructure().GetLevelEditorCinematicsCategory())
+			.SetDisplayName(IMovieRenderPipelineEditorModule::MoviePipelineQueueTabLabel)
+			.SetTooltipText(LOCTEXT("MovieRenderPipelineQueueTab_Tooltip", "Open the Movie Render Queue to render Sequences to disk at a higher quality than realtime allows."))
+			.SetIcon(FSlateIcon(FEditorStyle::GetStyleSetName(), "ContentBrowser.TabIcon"));
+	}
+}
+
+void FMovieRenderPipelineEditorModule::PerformTestPipelineRender(const TArray<FString>& Args)
+{
+	if (Args.Num() == 0)
+	{
+		UE_LOG(LogMovieRenderPipeline, Warning, TEXT("No arguments specified. Specify the path to a sequence asset to render."));
+		return;
+	}
+
+	FSoftObjectPath SequencePath = FSoftObjectPath(FPackageName::ExportTextPathToObjectPath(Args[0]));
+
+	const UMovieRenderPipelineProjectSettings* ProjectSettings = GetDefault<UMovieRenderPipelineProjectSettings>();
+	TSubclassOf<UMoviePipelineExecutorBase> ExecutorClass = ProjectSettings->DefaultLocalExecutor;
+
+	check(false);
+	// Create an instance of the Executor and add it to the root so it never dies.
+	// Executor = NewObject<UMoviePipelineExecutorBase>(GetTransientPackage(), ExecutorClass);
+	// Executor->AddToRoot();
+	// Executor->OnExecutorFinished().AddRaw(this, &FMovieRenderPipelineEditorModule::OnTestPipelineExecutorFinished);
+	// 
+	// TArray<UMoviePipelineMasterConfig*> Pipelines = GenerateTestPipelineConfigs();
+	// 
+	// TArray<FMoviePipelineExecutorJobPrev> Jobs;
+	// for(UMoviePipelineMasterConfig* Pipeline : Pipelines)
+	// {
+	// 	Jobs.Add(FMoviePipelineExecutorJobPrev(SequencePath, Pipeline));
+	// }
+	// 
+	// Executor->Execute(Jobs);
+}
+
+void FMovieRenderPipelineEditorModule::OnTestPipelineExecutorFinished(UMoviePipelineExecutorBase* InExecutor, bool bSuccess)
+{
+	if (Executor)
+	{
+		Executor->RemoveFromRoot();
+		Executor = nullptr;
+	}
+}
+
+void FMovieRenderPipelineEditorModule::RegisterSettings()
+{
+	ISettingsModule& SettingsModule = FModuleManager::LoadModuleChecked<ISettingsModule>("Settings");
+
+	SettingsModule.RegisterSettings("Project", "Plugins", "Movie Render Pipeline",
+		LOCTEXT("ProjectSettings_Label", "Movie Render Pipeline"),
+		LOCTEXT("ProjectSettings_Description", "Configure project-wide defaults for the movie render pipeline."),
+		GetMutableDefault<UMovieRenderPipelineProjectSettings>()
+	);
+
+	IConsoleManager::Get().RegisterConsoleCommand(
+		TEXT("MovieRenderPipeline.TestRenderSequence"),
+		TEXT("Renders the specified sequence asset using the default executor and a pre-made configuration."),
+		FConsoleCommandWithArgsDelegate::CreateRaw(this, &FMovieRenderPipelineEditorModule::PerformTestPipelineRender),
+		ECVF_Default
+	);
+}
+
+void FMovieRenderPipelineEditorModule::UnregisterSettings()
+{
+	ISettingsModule* SettingsModule = FModuleManager::GetModulePtr<ISettingsModule>("Settings");
+	if (SettingsModule)
+	{
+		SettingsModule->UnregisterSettings("Project", "Plugins", "Movie Render Pipeline");
+	}
+
+	IConsoleManager::Get().UnregisterConsoleObject(TEXT("MovieRenderPipeline.TestRenderSequence"));
+}
+
+void FMovieRenderPipelineEditorModule::StartupModule()
+{
+	// Initialize our custom style
+	FMovieRenderPipelineStyle::Get();
+
+	RegisterTabImpl();
+	RegisterSettings();
+}
+
+void FMovieRenderPipelineEditorModule::ShutdownModule()
+{
+	UnregisterSettings();
+}
+
+#undef LOCTEXT_NAMESPACE
+
+IMPLEMENT_MODULE(FMovieRenderPipelineEditorModule, MovieRenderPipelineEditor)
\ No newline at end of file
diff --git a/Engine/Plugins/MovieScene/MovieRenderPipeline/Source/MovieRenderPipelineEditor/Private/MovieRenderPipelineEditorModule.h b/Engine/Plugins/MovieScene/MovieRenderPipeline/Source/MovieRenderPipelineEditor/Private/MovieRenderPipelineEditorModule.h
new file mode 100644
--- /dev/null
+++ b/Engine/Plugins/MovieScene/MovieRenderPipeline/Source/MovieRenderPipelineEditor/Private/MovieRenderPipelineEditorModule.h
@@ -0,0 +1,29 @@
+// Copyright 1998-2019 Epic Games, Inc. All Rights Reserved.
+
+#pragma once
+
+#include "IMovieRenderPipelineEditorModule.h"
+class UMoviePipelineShotConfig;
+class UMoviePipelineMasterConfig;
+class UMoviePipelineExecutorBase;
+
+class FMovieRenderPipelineEditorModule : public IMovieRenderPipelineEditorModule
+{
+public:
+
+	/** IModuleInterface implementation */
+	virtual void StartupModule() override;
+	virtual void ShutdownModule() override;
+private:
+	void RegisterSettings();
+	void UnregisterSettings();
+	void PerformTestPipelineRender(const TArray<FString>& Args);
+	void OnTestPipelineExecutorFinished(UMoviePipelineExecutorBase* InExecutor, bool bSuccess);
+
+private:
+	TArray<UMoviePipelineMasterConfig*> GenerateTestPipelineConfigs();
+private:
+	class UMoviePipeline* ActiveMoviePipeline;
+	class UMoviePipelineMasterConfig* PipelineConfig;
+	UMoviePipelineExecutorBase* Executor;
+};
diff --git a/Engine/Plugins/MovieScene/MovieRenderPipeline/Source/MovieRenderPipelineEditor/Private/MovieRenderPipelineEditorTest.cpp b/Engine/Plugins/MovieScene/MovieRenderPipeline/Source/MovieRenderPipelineEditor/Private/MovieRenderPipelineEditorTest.cpp
new file mode 100644
--- /dev/null
+++ b/Engine/Plugins/MovieScene/MovieRenderPipeline/Source/MovieRenderPipelineEditor/Private/MovieRenderPipelineEditorTest.cpp
@@ -0,0 +1,110 @@
+// Copyright 1998-2019 Epic Games, Inc. All Rights Reserved.
+#include "MovieRenderPipelineEditorModule.h"
+#include "MoviePipelineShotConfig.h"
+#include "MoviePipelineGameOverrideSetting.h"
+#include "MoviePipelineAccumulationSetting.h"
+#include "MoviePipelineImageSequenceOutput.h"
+#include "MoviePipelineAudioOutput.h"
+#include "Misc/Paths.h"
+#include "MoviePipelineMasterConfig.h"
+#include "MoviePipelineOutputSetting.h"
+#include "MoviePipelineHighResSetting.h"
+#include "MoviePipelineDeferredPasses.h"
+#include "MoviePipelineCameraSetting.h"
+#include "GameFramework/GameModeBase.h"
+
+UMoviePipelineShotConfig* GenerateTestShotConfig(UObject* InOwner, int32 InSampleCount,
+	int32 InShutterAngle, EMoviePipelineShutterTiming InFrameTiming, int32 InTileCount,
+	int32 InSpatialSampleCount, float PadRatioX, float PadRatioY, float AccumulationGamma)
+{
+	UMoviePipelineShotConfig* OutConfig = NewObject<UMoviePipelineShotConfig>(InOwner);
+	OutConfig->FindOrAddSetting<UMoviePipelineDeferredPassBase>();
+
+	UMoviePipelineAccumulationSetting* Accumulation = OutConfig->FindOrAddSetting< UMoviePipelineAccumulationSetting>();
+	Accumulation->SpatialSampleCount = InSpatialSampleCount;
+	Accumulation->AccumulationGamma = AccumulationGamma;
+
+	{
+		UMoviePipelineCameraSetting* CameraSetting = OutConfig->FindOrAddSetting<UMoviePipelineCameraSetting>();
+		CameraSetting->TemporalSampleCount = InSampleCount;
+		CameraSetting->CameraShutterAngle = InShutterAngle;
+		CameraSetting->ShutterTiming = InFrameTiming;
+		CameraSetting->bManualExposure = false;
+		CameraSetting->ExposureCompensation = 10;
+	}
+
+	{
+		UMoviePipelineHighResSetting* HighRes = OutConfig->FindOrAddSetting<UMoviePipelineHighResSetting>();
+		HighRes->TileCount = InTileCount;
+		HighRes->OverlapPercentage = PadRatioX;
+		// HighRes->bWriteAllSamples = true;
+		// HighRes->OverlapPercentage = 0.f;
+	}
+
+	return OutConfig;
+}
+
+TArray<UMoviePipelineMasterConfig*> FMovieRenderPipelineEditorModule::GenerateTestPipelineConfigs()
+{
+	// int32 ShutterAngles[] = { 180, 360 };
+	// EMoviePipelineShutterTiming ShutterTimings[] = { EMoviePipelineShutterTiming::FrameOpen, EMoviePipelineShutterTiming::FrameCenter, EMoviePipelineShutterTiming::FrameClose };
+	// FString ShutterTimingNames[] = { TEXT("FOpen"), TEXT("FCenter"), TEXT("FClose") };
+	// int32 TemporalSampleCounts[] = { 1, 5 };
+
+	int32 ShutterAngles[] = { 180 };
+	EMoviePipelineShutterTiming ShutterTimings[] = { EMoviePipelineShutterTiming::FrameCenter };
+	FString ShutterTimingNames[] = { TEXT("FCenter") };
+	int32 TemporalSampleCounts[] = { 1 };
+
+	TArray<UMoviePipelineMasterConfig*> Configs;
+
+	for (int32 ShutterAngleIndex = 0; ShutterAngleIndex < UE_ARRAY_COUNT(ShutterAngles); ShutterAngleIndex++)
+	{
+		for (int32 ShutterTimingIndex = 0; ShutterTimingIndex < UE_ARRAY_COUNT(ShutterTimings); ShutterTimingIndex++)
+		{
+			for (int32 TemporalSampleCountIndex = 0; TemporalSampleCountIndex < UE_ARRAY_COUNT(TemporalSampleCounts); TemporalSampleCountIndex++)
+			{
+				// Build a folder string to place this particular test in.
+				FString DirectoryName = FString::Printf(TEXT("SA_%d_ST_%s_SC_%d"), ShutterAngles[ShutterAngleIndex], *ShutterTimingNames[ShutterTimingIndex], TemporalSampleCounts[TemporalSampleCountIndex]);
+
+				UMoviePipelineMasterConfig* OutPipeline = NewObject<UMoviePipelineMasterConfig>(GetTransientPackage());
+				UMoviePipelineOutputSetting* OutputSetting = OutPipeline->FindOrAddSetting<UMoviePipelineOutputSetting>();
+
+				int32 SizeX = 1920;
+				int32 SizeY = 1080;
+				int32 TileX = 4;
+				int32 TileY = 4;
+				int32 TestNumSamples = 8;
+				float PadRatioX = 0.5f;
+				float PadRatioY = 0.5f;
+				float AccumulationGamma = 1.0f;
+
+				OutputSetting->OutputResolution = FIntPoint(SizeX*TileX, SizeY*TileY);
+				OutputSetting->OutputDirectory.Path = FPaths::ProjectSavedDir() / TEXT("/VideoCaptures/") / *DirectoryName;
+
+				const int32 NumSpatialSamples = TestNumSamples;
+				const int32 NumTiles = TileX;
+
+				UMoviePipelineShotConfig* DefaultConfig = GenerateTestShotConfig(OutPipeline,
+					TemporalSampleCounts[TemporalSampleCountIndex],
+					ShutterAngles[ShutterAngleIndex],
+					ShutterTimings[ShutterTimingIndex],
+					NumTiles, NumSpatialSamples, PadRatioX, PadRatioY, AccumulationGamma);
+
+				OutPipeline->DefaultShotConfig = DefaultConfig;
+
+				// Add some Outputs
+				OutPipeline->FindOrAddSetting<UMoviePipelineImageSequenceOutput_EXR>();
+				OutPipeline->FindOrAddSetting<UMoviePipelineImageSequenceOutput_PNG>();
+
+				UMoviePipelineGameOverrideSetting* GamemodeOverride = OutPipeline->FindOrAddSetting<UMoviePipelineGameOverrideSetting>();
+				GamemodeOverride->GameModeOverride = AGameModeBase::StaticClass();
+
+				// Add it to our list to be processed.
+				Configs.Add(OutPipeline);
+			}
+		}
+	}
+
+	return Configs;
+}
\ No newline at end of file
diff --git a/Engine/Plugins/MovieScene/MovieRenderPipeline/Source/MovieRenderPipelineEditor/Private/Widgets/SMoviePipelineConfigEditor.cpp b/Engine/Plugins/MovieScene/MovieRenderPipeline/Source/MovieRenderPipelineEditor/Private/Widgets/SMoviePipelineConfigEditor.cpp
new file mode 100644
--- /dev/null
+++ b/Engine/Plugins/MovieScene/MovieRenderPipeline/Source/MovieRenderPipelineEditor/Private/Widgets/SMoviePipelineConfigEditor.cpp
@@ -0,0 +1,399 @@
+// Copyright 1998-2019 Epic Games, Inc. All Rights Reserved.
+
+#include "Widgets/SMoviePipelineConfigEditor.h"
+#include "Widgets/MoviePipelineWidgetConstants.h"
+#include "Widgets/SMoviePipelineConfigSettings.h"
+#include "MoviePipelineShotConfig.h"
+#include "MoviePipelineSetting.h"
+#include "MoviePipelineQueue.h"
+
+// Core includes
+#include "UObject/UObjectIterator.h"
+#include "Templates/SubclassOf.h"
+#include "ClassIconFinder.h"
+
+// AssetRegistry includes
+#include "AssetRegistryModule.h"
+#include "AssetData.h"
+ 
+// ContentBrowser includes
+#include "IContentBrowserSingleton.h"
+#include "ContentBrowserModule.h"
+ 
+// AssetTools includes
+#include "AssetToolsModule.h"
+ 
+// Slate includes
+#include "Widgets/SBoxPanel.h"
+#include "Widgets/Text/STextBlock.h"
+#include "Widgets/Input/SButton.h"
+#include "Widgets/Input/SComboButton.h"
+#include "Widgets/Layout/SSplitter.h"
+#include "Widgets/Layout/SScrollBox.h"
+#include "Widgets/Layout/SBorder.h"
+#include "Framework/MultiBox/MultiBoxBuilder.h"
+#include "Styling/SlateIconFinder.h"
+
+// EditorStyle includes
+#include "EditorStyleSet.h"
+#include "EditorFontGlyphs.h"
+#include "ScopedTransaction.h"
+#include "PropertyEditorModule.h"
+#include "IDetailsView.h"
+#include "MovieRenderPipelineStyle.h"
+#include "FrameNumberDetailsCustomization.h"
+
+#define LOCTEXT_NAMESPACE "SMoviePipelineEditor"
+
+
+TArray<UClass*> FindMoviePipelineSettingClasses()
+{
+	FAssetRegistryModule& AssetRegistryModule = FModuleManager::LoadModuleChecked<FAssetRegistryModule>(TEXT("AssetRegistry"));
+
+	TArray<FAssetData> ClassList;
+
+	FARFilter Filter;
+	Filter.ClassNames.Add(UMoviePipelineSetting::StaticClass()->GetFName());
+
+	// Include any Blueprint based objects as well, this includes things like Blutilities, UMG, and GameplayAbility objects
+	Filter.bRecursiveClasses = true;
+	AssetRegistryModule.Get().GetAssets(Filter, ClassList);
+
+	TArray<UClass*> Classes;
+
+	for (const FAssetData& Data : ClassList)
+	{
+		UClass* Class = Data.GetClass();
+		if (Class)
+		{
+			Classes.Add(Class);
+		}
+	}
+
+	for (TObjectIterator<UClass> ClassIterator; ClassIterator; ++ClassIterator)
+	{
+		if (ClassIterator->IsChildOf(UMoviePipelineSetting::StaticClass()) && !ClassIterator->HasAnyClassFlags(CLASS_Abstract | CLASS_Deprecated | CLASS_NewerVersionExists))
+		{
+			Classes.Add(*ClassIterator);
+		}
+	}
+
+	return Classes;
+}
+
+PRAGMA_DISABLE_OPTIMIZATION
+void SMoviePipelineConfigEditor::Construct(const FArguments& InArgs)
+{
+	bRequestDetailsRefresh = true;
+	PipelineConfigAttribute = InArgs._PipelineConfig;
+	OwningJobAttribute = InArgs._OwningJob;
+    
+	DetailsBox = SNew(SScrollBox);
+	DetailsBox->SetScrollBarRightClickDragAllowed(true);
+     
+	SettingsWidget = SNew(SMoviePipelineConfigSettings)
+	.OnSelectionChanged(this, &SMoviePipelineConfigEditor::OnSettingsSelectionChanged);
+    
+	CheckForNewSettingsObject();
+     
+	ChildSlot
+	[
+		SNew(SSplitter)
+		.Orientation(Orient_Horizontal)
+    
+		+ SSplitter::Slot()
+		.Value(.40f)
+		[
+			SNew(SBorder)
+			.Padding(4)
+			.BorderImage( FEditorStyle::GetBrush("ToolPanel.GroupBorder") )
+			[
+				SettingsWidget.ToSharedRef()
+			]
+		]
+    
+		+ SSplitter::Slot()
+		.Value(.60f)
+		[
+			SNew(SVerticalBox)
+			+ SVerticalBox::Slot()
+			.FillHeight(1.f)
+			[
+				DetailsBox.ToSharedRef()
+			]
+
+			+ SVerticalBox::Slot()
+			.AutoHeight()
+			[
+				SNew(SBorder)
+				.BorderImage(FEditorStyle::GetBrush("ToolPanel.GroupBorder"))
+				.Visibility(this, &SMoviePipelineConfigEditor::IsSettingFooterVisible)
+				[
+					SNew(STextBlock)
+					.AutoWrapText(true)
+					.Text(this, &SMoviePipelineConfigEditor::GetSettingsFooterText)
+				]
+			]
+		]
+	];
+}
+
+TSharedRef<SWidget> SMoviePipelineConfigEditor::MakeAddSettingButton()
+{
+	return SNew(SComboButton)
+		.ContentPadding(MoviePipeline::ButtonPadding)
+		.ButtonStyle(FMovieRenderPipelineStyle::Get(), "FlatButton.Success")
+		.OnGetMenuContent(this, &SMoviePipelineConfigEditor::OnGenerateSettingsMenu)
+		.ForegroundColor(FSlateColor::UseForeground())
+		.HasDownArrow(false)
+		.ButtonContent()
+		[
+			SNew(SHorizontalBox)
+    
+			// Plus Icon
+			+ SHorizontalBox::Slot()
+			.VAlign(VAlign_Center)
+			.AutoWidth()
+			[
+				SNew(STextBlock)
+				.TextStyle(FEditorStyle::Get(), "NormalText.Important")
+				.Font(FEditorStyle::Get().GetFontStyle("FontAwesome.10"))
+				.Text(FEditorFontGlyphs::Plus)
+			]
+    
+			// "Setting" Text
+			+ SHorizontalBox::Slot()
+			.VAlign(VAlign_Center)
+			.AutoWidth()
+			.Padding(4, 0, 0, 0)
+			[
+				SNew(STextBlock)
+				.TextStyle(FEditorStyle::Get(), "NormalText.Important")
+				.Text(LOCTEXT("AddNewSetting_Text", "Setting"))
+			]
+    
+			// Non-Default Down Caret arrow.
+			+ SHorizontalBox::Slot()
+			.VAlign(VAlign_Center)
+			.AutoWidth()
+			.Padding(4, 0, 0, 0)
+			[
+				SNew(STextBlock)
+				.TextStyle(FEditorStyle::Get(), "NormalText.Important")
+				.Font(FEditorStyle::Get().GetFontStyle("FontAwesome.10"))
+				.Text(FEditorFontGlyphs::Caret_Down)
+			]
+		];
+}
+PRAGMA_ENABLE_OPTIMIZATION
+
+void SMoviePipelineConfigEditor::CheckForNewSettingsObject()
+{
+	UMoviePipelineConfigBase* NewMoviePipeline = PipelineConfigAttribute.Get();
+	if (CachedPipelineConfig != NewMoviePipeline)
+	{
+		CachedPipelineConfig = NewMoviePipeline;
+
+		SettingsWidget->SetShotConfigObject(NewMoviePipeline);
+		bRequestDetailsRefresh = true;
+	}
+
+	UMoviePipelineExecutorJob* NewOwningJob = OwningJobAttribute.Get();
+	if (CachedOwningJob != NewOwningJob)
+	{
+		CachedOwningJob = NewOwningJob;
+		bRequestDetailsRefresh = true;
+	}
+}
+
+void SMoviePipelineConfigEditor::Tick(const FGeometry& AllottedGeometry, const double InCurrentTime, const float InDeltaTime)
+{
+	CheckForNewSettingsObject();
+	if (bRequestDetailsRefresh)
+	{
+		UpdateDetails();
+		bRequestDetailsRefresh = false;
+	}
+}
+
+TSharedRef<SWidget> SMoviePipelineConfigEditor::OnGenerateSettingsMenu()
+{
+	TSharedRef<FExtender> Extender = MakeShared<FExtender>();
+	{
+		// ULevelSequence*       LevelSequence = LevelSequenceAttribute.Get();
+		// UTakeRecorderSources* Sources = LevelSequence ? LevelSequence->FindOrAddMetaData<UTakeRecorderSources>() : nullptr;
+		// 
+		// if (Sources)
+		// {
+		// 	FTakeRecorderModule& TakeRecorderModule = FModuleManager::GetModuleChecked<FTakeRecorderModule>("TakeRecorder");
+		// 	TakeRecorderModule.PopulateSourcesMenu(Extender, Sources);
+		// }
+	}
+
+	FMenuBuilder MenuBuilder(true, nullptr, Extender);
+
+	// Put the different categories into different sections
+	TArray<UClass*> SourceClasses = FindMoviePipelineSettingClasses();
+
+	TMap<FString, TArray<UClass*>> CategorizedClasses;
+
+	for (UClass* Class : SourceClasses)
+	{
+		const UMoviePipelineSetting* SettingDefaultObject = GetDefault<UMoviePipelineSetting>(Class);
+		CategorizedClasses.FindOrAdd(SettingDefaultObject->GetCategoryText().ToString()).Add(Class);
+	}
+
+	for (TPair <FString, TArray<UClass*>> KVP : CategorizedClasses)
+	{
+		MenuBuilder.BeginSection(NAME_None, FText::FromString(KVP.Key));
+
+		Algo::SortBy(KVP.Value, &UClass::GetDisplayNameText, FText::FSortPredicate());
+
+		for (UClass* Class : KVP.Value)
+		{
+			// Get a display name for the setting from the CDO.
+			const UMoviePipelineSetting* SettingDefaultObject = GetDefault<UMoviePipelineSetting>(Class);
+
+			// Depending on the type of config we're editing, some settings may not be eligible. If this is the case, we omit them from the list.
+			bool bCanSettingBeAdded = CachedPipelineConfig->CanSettingBeAdded(SettingDefaultObject);
+			if (!bCanSettingBeAdded)
+			{
+				continue;
+			}
+
+			// If the setting already exists and it only allows a single instance, we omit them from the list.
+			bool bAllowDuplicates = true;
+			for (UMoviePipelineSetting* ExistingSetting : CachedPipelineConfig->GetSettings())
+			{
+				// If we found a setting with the same class as ours, ask the CDO if multiple are valid.
+				if (ExistingSetting->GetClass() == Class)
+				{
+					bAllowDuplicates = !SettingDefaultObject->IsSolo();
+					break;
+				}
+			}
+			if (!bAllowDuplicates)
+			{
+				continue;
+			}
+
+			FText SettingDisplayName;
+			if (SettingDefaultObject)
+			{
+				SettingDisplayName = SettingDefaultObject->GetDisplayText();
+			}
+
+			TSubclassOf<UMoviePipelineSetting> SubclassOf = Class;
+			MenuBuilder.AddMenuEntry(
+				SettingDisplayName,
+				Class->GetToolTipText(true),
+				FSlateIconFinder::FindIconForClass(Class),
+				FUIAction(
+					FExecuteAction::CreateSP(this, &SMoviePipelineConfigEditor::AddSettingFromClass, SubclassOf),
+					FCanExecuteAction::CreateSP(this, &SMoviePipelineConfigEditor::CanAddSettingFromClass, SubclassOf)
+				)
+			);
+		}
+
+		MenuBuilder.EndSection();
+	}
+
+	return MenuBuilder.MakeWidget();
+}
+
+void SMoviePipelineConfigEditor::UpdateDetails()
+{
+	FPropertyEditorModule& PropertyEditorModule = FModuleManager::Get().LoadModuleChecked<FPropertyEditorModule>("PropertyEditor");
+	FDetailsViewArgs DetailsViewArgs(false, false, false, FDetailsViewArgs::HideNameArea, true);
+	DetailsViewArgs.bShowScrollBar = false;
+	 
+	TArray<UMoviePipelineSetting*> SelectedSources;
+	SettingsWidget->GetSelectedSettings(SelectedSources);
+
+	
+	// Create 1 details panel per source class type
+	TSortedMap<const UClass*, TArray<UObject*>> ClassToSources;
+	for (UMoviePipelineSetting* Source : SelectedSources)
+	{
+		ClassToSources.FindOrAdd(Source->GetClass()).Add(Source);
+	}
+	
+	TArray<FObjectKey> PreviousClasses;
+	ClassToDetailsView.GenerateKeyArray(PreviousClasses);
+
+	TSharedPtr<INumericTypeInterface<double>> NumericTypeInterface = MakeShareable(new TDefaultNumericTypeInterface<double>);
+
+	for (auto& Pair : ClassToSources)
+	{
+		PreviousClasses.Remove(Pair.Key);
+	
+		TSharedPtr<IDetailsView> ExistingDetails = ClassToDetailsView.FindRef(Pair.Key);
+		if (ExistingDetails.IsValid())
+		{
+			ExistingDetails->SetObjects(Pair.Value);
+		}
+		else
+		{
+			TSharedRef<IDetailsView> Details = PropertyEditorModule.CreateDetailView(DetailsViewArgs);
+			Details->RegisterInstancedCustomPropertyTypeLayout("FrameNumber", FOnGetPropertyTypeCustomizationInstance::CreateStatic(&FFrameNumberDetailsCustomization::MakeInstance, NumericTypeInterface));
+
+			Details->SetObjects(Pair.Value);
+			DetailsBox->AddSlot()
+				[
+					Details
+				];
+	
+			ClassToDetailsView.Add(Pair.Key, Details);
+		}
+	}
+	
+	for (FObjectKey StaleClass : PreviousClasses)
+	{
+		TSharedPtr<IDetailsView> Details = ClassToDetailsView.FindRef(StaleClass);
+		DetailsBox->RemoveSlot(Details.ToSharedRef());
+		ClassToDetailsView.Remove(StaleClass);
+	}
+}
+
+void SMoviePipelineConfigEditor::AddSettingFromClass(TSubclassOf<UMoviePipelineSetting> SettingClass)
+{
+	UMoviePipelineConfigBase* Pipeline = PipelineConfigAttribute.Get();
+
+	if (*SettingClass && Pipeline)
+	{
+		FScopedTransaction Transaction(FText::Format(LOCTEXT("AddNewSetting", "Add New {0} Setting"), SettingClass->GetDisplayNameText()));
+		UMoviePipelineSetting* NewSetting = Pipeline->FindOrAddSettingByClass(SettingClass);
+
+		SettingsWidget->SetSelectedSettings({ NewSetting });
+	}
+}
+
+bool SMoviePipelineConfigEditor::CanAddSettingFromClass(TSubclassOf<UMoviePipelineSetting> SettingClass)
+{
+	return true;
+}
+
+EVisibility SMoviePipelineConfigEditor::IsSettingFooterVisible() const
+{
+	return (GetSettingsFooterText().IsEmpty()) ? EVisibility::Collapsed : EVisibility::Visible;
+}
+
+FText SMoviePipelineConfigEditor::GetSettingsFooterText() const
+{
+	TArray<UMoviePipelineSetting*> SelectedSettings;
+	SettingsWidget->GetSelectedSettings(SelectedSettings);
+
+	FTextBuilder TextBuilder;
+
+	if (CachedOwningJob.IsValid())
+	{
+		for (UMoviePipelineSetting* Setting : SelectedSettings)
+		{
+			TextBuilder.AppendLine(Setting->GetFooterText(CachedOwningJob.Get()));
+		}
+	}
+
+	return TextBuilder.ToText();
+}
+
+#undef LOCTEXT_NAMESPACE
\ No newline at end of file
diff --git a/Engine/Plugins/MovieScene/MovieRenderPipeline/Source/MovieRenderPipelineEditor/Private/Widgets/SMoviePipelineQueueEditor.cpp b/Engine/Plugins/MovieScene/MovieRenderPipeline/Source/MovieRenderPipelineEditor/Private/Widgets/SMoviePipelineQueueEditor.cpp
new file mode 100644
--- /dev/null
+++ b/Engine/Plugins/MovieScene/MovieRenderPipeline/Source/MovieRenderPipelineEditor/Private/Widgets/SMoviePipelineQueueEditor.cpp
@@ -0,0 +1,945 @@
+// Copyright 1998-2019 Epic Games, Inc. All Rights Reserved.
+
+// Movie Pipeline
+#include "Widgets/SMoviePipelineQueueEditor.h"
+#include "Widgets/MoviePipelineWidgetConstants.h"
+#include "MovieRenderPipelineDataTypes.h"
+#include "MoviePipelineQueue.h"
+#include "MoviePipelineOutputSetting.h"
+#include "MoviePipelineQueueSubsystem.h"
+#include "MoviePipelineMasterConfig.h"
+#include "MovieRenderPipelineStyle.h"
+#include "MovieRenderPipelineSettings.h"
+
+// Slate Includes
+#include "Widgets/Input/SComboButton.h"
+#include "Widgets/Input/SCheckBox.h"
+#include "Widgets/Input/SButton.h"
+#include "Widgets/Notifications/SProgressBar.h"
+#include "Widgets/Views/STreeView.h"
+#include "Widgets/Layout/SWidgetSwitcher.h"
+#include "Widgets/Views/SExpanderArrow.h"
+#include "Widgets/Input/SHyperlink.h"
+#include "Framework/MultiBox/MultiBoxBuilder.h"
+#include "EditorStyleSet.h"
+#include "EditorFontGlyphs.h"
+#include "SDropTarget.h"
+
+// Editor
+#include "Editor.h"
+#include "IContentBrowserSingleton.h"
+#include "ContentBrowserModule.h"
+#include "ScopedTransaction.h"
+#include "PropertyCustomizationHelpers.h"
+#include "Framework/Commands/UICommandList.h"
+#include "Framework/Commands/GenericCommands.h"
+#include "DragAndDrop/AssetDragDropOp.h"
+#include "AssetData.h"
+
+// Misc
+#include "LevelSequence.h"
+#include "Engine/EngineTypes.h"
+#include "Framework/Application/SlateApplication.h"
+#include "SMoviePipelineConfigPanel.h"
+#include "Widgets/SWindow.h"
+#include "HAL/FileManager.h"
+
+
+#define LOCTEXT_NAMESPACE "SMoviePipelineQueueEditor"
+
+struct FMoviePipelineQueueJobTreeItem;
+struct FMoviePipelineMapTreeItem;
+struct FMoviePipelineShotItem;
+class SQueueJobListRow;
+
+struct IMoviePipelineQueueTreeItem : TSharedFromThis<IMoviePipelineQueueTreeItem>
+{
+	virtual ~IMoviePipelineQueueTreeItem() {}
+
+	virtual TSharedPtr<FMoviePipelineQueueJobTreeItem> AsJob() { return nullptr; }
+	virtual void Delete(UMoviePipelineQueue* InOwningQueue) {}
+
+	virtual TSharedRef<ITableRow> ConstructWidget(TWeakPtr<SMoviePipelineQueueEditor> InQueueWidget, const TSharedRef<STableViewBase>& OwnerTable) = 0;
+};
+
+class SQueueJobListRow : public SMultiColumnTableRow<TSharedPtr<IMoviePipelineQueueTreeItem>>
+{
+public:
+	SLATE_BEGIN_ARGS(SQueueJobListRow) {}
+		SLATE_ARGUMENT(TSharedPtr<FMoviePipelineQueueJobTreeItem>, Item)
+		SLATE_EVENT(FOnMoviePipelineEditConfig, OnEditConfigRequested)
+	SLATE_END_ARGS()
+
+	static const FName NAME_Sequence;
+	static const FName NAME_Settings;
+	static const FName NAME_Output;
+	static const FName NAME_Status;
+
+	TSharedPtr<FMoviePipelineQueueJobTreeItem> Item;
+
+	void Construct(const FArguments& InArgs, const TSharedRef<STableViewBase>& OwnerTable);
+	virtual TSharedRef<SWidget> GenerateWidgetForColumn(const FName& ColumnName) override;
+
+private:
+	FOnMoviePipelineEditConfig OnEditConfigRequested;
+};
+
+struct FMoviePipelineQueueJobTreeItem : IMoviePipelineQueueTreeItem
+{
+	/** The job that this tree item represents */
+	TWeakObjectPtr<UMoviePipelineExecutorJob> WeakJob;
+
+	/** Sorted list of this category's children */
+	TArray<TSharedPtr<IMoviePipelineQueueTreeItem>> Children;
+
+	FOnMoviePipelineEditConfig Callback;
+
+	explicit FMoviePipelineQueueJobTreeItem(UMoviePipelineExecutorJob* InJob, FOnMoviePipelineEditConfig InCallback)
+		: WeakJob(InJob)
+		, Callback(InCallback)
+	{}
+
+	virtual TSharedRef<ITableRow> ConstructWidget(TWeakPtr<SMoviePipelineQueueEditor> InQueueWidget, const TSharedRef<STableViewBase>& OwnerTable) override
+	{
+		return SNew(SQueueJobListRow, OwnerTable)
+			.Item(SharedThis(this));
+	}
+
+	virtual TSharedPtr<FMoviePipelineQueueJobTreeItem> AsJob() override
+	{
+		return SharedThis(this);
+	}
+
+	virtual void Delete(UMoviePipelineQueue* InOwningQueue) override
+	{
+		InOwningQueue->DeleteJob(WeakJob.Get());
+	}
+
+public:
+	FString GetSequencePath() const
+	{
+		UMoviePipelineExecutorJob* Job = WeakJob.Get();
+		if (Job)
+		{
+			return Job->Sequence.ToString();
+		}
+
+		return FString();
+	}
+
+	void SetSequencePath(const FAssetData& AssetData)
+	{
+		UMoviePipelineExecutorJob* Job =WeakJob.Get();
+		if (Job)
+		{
+			Job->Sequence = AssetData.ToSoftObjectPath();
+		}
+	}
+
+	FString GetMasterConfigPath() const
+	{
+		UMoviePipelineExecutorJob* Job = WeakJob.Get();
+		if (Job)
+		{
+			UMoviePipelineConfigBase* Config = Job->GetConfiguration();
+			if (Config)
+			{
+				return Config->GetPathName();
+			}
+		}
+
+		return FString();
+	}
+
+	void SetMasterConfigPath(const FAssetData& AssetData)
+	{
+		UMoviePipelineExecutorJob* Job = WeakJob.Get();
+		if (Job)
+		{
+			Job->SetPresetOrigin(CastChecked<UMoviePipelineMasterConfig>(AssetData.GetAsset()));
+		}
+	}
+
+	EVisibility IsMasterConfigModified() const
+	{
+		UMoviePipelineExecutorJob* Job = WeakJob.Get();
+		if (Job)
+		{
+			
+			return (Job->GetPresetOrigin() == nullptr) ? EVisibility::Visible : EVisibility::Collapsed;
+		}
+		
+		return EVisibility::Collapsed;
+	}
+
+	FReply OnEditMasterConfigForJob()
+	{
+		Callback.ExecuteIfBound(WeakJob, nullptr);
+
+		return FReply::Handled();
+	}
+
+	FText GetOutputLabel() const
+	{
+		UMoviePipelineExecutorJob* Job = WeakJob.Get();
+		if (Job && Job->GetConfiguration())
+		{
+			UMoviePipelineOutputSetting* OutputSetting = Job->GetConfiguration()->FindSetting<UMoviePipelineOutputSetting>();
+			check(OutputSetting);
+
+			return FText::FromString(OutputSetting->OutputDirectory.Path);
+		}
+
+		return LOCTEXT("MissingConfigOutput_Label", "[No Config Set]");
+	}
+
+	void BrowseToOutputFolder()
+	{
+		UMoviePipelineExecutorJob* Job = WeakJob.Get();
+		if (Job && Job->GetConfiguration())
+		{
+			UMoviePipelineOutputSetting* OutputSetting = Job->GetConfiguration()->FindSetting<UMoviePipelineOutputSetting>();
+			check(OutputSetting);
+
+			// @ToDo: We should resolve the exact path (as much as we can) through the config.
+			// For now, we'll just split off any format strings and go to the base folder.
+			FString OutputFolderPath = FPaths::ConvertRelativePathToFull(OutputSetting->OutputDirectory.Path);
+
+			FString TrimmedPath;
+			if (OutputFolderPath.Split(TEXT("{"), &TrimmedPath, nullptr))
+			{
+				FPaths::NormalizeDirectoryName(TrimmedPath);
+				OutputFolderPath = TrimmedPath;
+			}
+
+			// Attempt to make the directory. The user can see the output folder before they render so the folder
+			// may not have been created yet and the ExploreFolder call will fail.
+			IFileManager::Get().MakeDirectory(*OutputFolderPath, true);
+
+			FPlatformProcess::ExploreFolder(*OutputFolderPath);
+		}
+	}
+
+	int32 GetStatusIndex() const
+	{
+		UMoviePipelineExecutorJob* Job = WeakJob.Get();
+		if (Job)
+		{
+			// JobStatus 0 is Uninitialized, so we take one off.
+			return FMath::Clamp(int32(Job->JobStatus) - 1, 0, 2);
+		}
+
+		return 0;
+	}
+
+	TOptional<float> GetProgressPercent() const
+	{
+		UMoviePipelineExecutorJob* Job = WeakJob.Get();
+		return Job ? Job->GetProgressPercentage() : TOptional<float>();
+	}
+};
+
+void SQueueJobListRow::Construct(const FArguments& InArgs, const TSharedRef<STableViewBase>& OwnerTable)
+{
+	Item = InArgs._Item;
+	OnEditConfigRequested = InArgs._OnEditConfigRequested;
+
+	FSuperRowType::FArguments SuperArgs = FSuperRowType::FArguments();
+	FSuperRowType::Construct(SuperArgs, OwnerTable);
+}
+
+TSharedRef<SWidget> SQueueJobListRow::GenerateWidgetForColumn(const FName& ColumnName)
+{
+	if (ColumnName == NAME_Sequence)
+	{
+		return SNew(SBox)
+			.Padding(2.0f)
+			[
+				SNew(SHorizontalBox)
+				+SHorizontalBox::Slot()
+				.AutoWidth()
+				.Padding(0, 0, 6, 0)
+				[
+					SNew(SExpanderArrow, SharedThis(this))
+				]
+
+				+SHorizontalBox::Slot()
+				.FillWidth(1.f)
+				[
+					SNew(SObjectPropertyEntryBox)
+					.ObjectPath(Item.Get(), &FMoviePipelineQueueJobTreeItem::GetSequencePath)
+					.AllowedClass(ULevelSequence::StaticClass())
+					.OnObjectChanged(Item.Get(), &FMoviePipelineQueueJobTreeItem::SetSequencePath)
+					.AllowClear(false)
+					.DisplayUseSelected(false)
+					.DisplayBrowse(true)
+					.DisplayThumbnail(true)
+					.DisplayCompactSize(false)
+				]
+			];
+	}
+	else if (ColumnName == NAME_Settings)
+	{
+		return SNew(SHorizontalBox)
+
+		// Preset Label
+		+ SHorizontalBox::Slot()
+		.AutoWidth()
+		.VAlign(VAlign_Center)
+		.Padding(2, 0)
+		[
+			SNew(SObjectPropertyEntryBox)
+			.ObjectPath(Item.Get(), &FMoviePipelineQueueJobTreeItem::GetMasterConfigPath)
+			.AllowedClass(UMoviePipelineMasterConfig::StaticClass())
+			.OnObjectChanged(Item.Get(), &FMoviePipelineQueueJobTreeItem::SetMasterConfigPath)
+			.AllowClear(false)
+			.DisplayUseSelected(false)
+			.DisplayBrowse(true)
+			.DisplayThumbnail(true)
+			.DisplayCompactSize(false)
+		]
+
+		// Modified button
+		+ SHorizontalBox::Slot()
+		.AutoWidth()
+		.VAlign(VAlign_Center)
+		[
+			SNew(STextBlock)
+			.Text(LOCTEXT("ModifiedConfigLabel", "(Modified)"))
+			.Visibility(Item.Get(), &FMoviePipelineQueueJobTreeItem::IsMasterConfigModified)
+		]
+
+		// Edit Button
+		+ SHorizontalBox::Slot()
+		.AutoWidth()
+		.VAlign(VAlign_Center)
+		[
+			SNew(SButton)
+			.OnClicked(Item.Get(), &FMoviePipelineQueueJobTreeItem::OnEditMasterConfigForJob)
+			.Content()
+			[
+				SNew(STextBlock)
+				.Text(LOCTEXT("EditMasterConfigButton_Label", "Edit"))
+			]
+		];
+	}
+	else if (ColumnName == NAME_Output)
+	{
+		return SNew(SBox)
+			.VAlign(VAlign_Center)
+			.HAlign(HAlign_Left)
+			[
+				SNew(SHyperlink)
+					.Text(Item.Get(), &FMoviePipelineQueueJobTreeItem::GetOutputLabel)
+					.OnNavigate(Item.Get(), &FMoviePipelineQueueJobTreeItem::BrowseToOutputFolder)
+			];
+
+	}
+	else if(ColumnName == NAME_Status)
+	{
+		return SNew(SWidgetSwitcher)
+			.WidgetIndex(Item.Get(), &FMoviePipelineQueueJobTreeItem::GetStatusIndex)
+
+			// Ready Label
+			+ SWidgetSwitcher::Slot()
+			.VAlign(VAlign_Center)
+			.HAlign(HAlign_Center)
+			[
+				SNew(STextBlock)
+				// .Style(FTakeRecorderStyle::Get(), "TakeRecorder.Source.Switch")
+				.Text(LOCTEXT("PendingJobStatus_Label", "Ready"))
+			]
+
+			// Progress Bar
+			+ SWidgetSwitcher::Slot()
+			.VAlign(VAlign_Center)
+			.HAlign(HAlign_Center)
+			[
+				SNew(SProgressBar)
+				.Percent(Item.Get(), &FMoviePipelineQueueJobTreeItem::GetProgressPercent)
+			]
+
+			// Completed
+			+ SWidgetSwitcher::Slot()
+			.VAlign(VAlign_Center)
+			.HAlign(HAlign_Center)
+			[
+				SNew(STextBlock)
+				.Text(LOCTEXT("PendingJobStatus_Label", "Completed!"))
+			];
+	}
+
+	return SNullWidget::NullWidget;
+}
+
+const FName SQueueJobListRow::NAME_Sequence = FName(TEXT("Sequence"));
+const FName SQueueJobListRow::NAME_Settings = FName(TEXT("Settings"));
+const FName SQueueJobListRow::NAME_Output = FName(TEXT("Output"));
+const FName SQueueJobListRow::NAME_Status = FName(TEXT("Status"));
+
+struct FMoviePipelineMapTreeItem : IMoviePipelineQueueTreeItem
+{
+	/** The job that this tree item represents */
+	TWeakObjectPtr<UMoviePipelineExecutorJob> WeakJob;
+
+	explicit FMoviePipelineMapTreeItem(UMoviePipelineExecutorJob* InJob)
+		: WeakJob(InJob)
+	{}
+
+	virtual TSharedRef<ITableRow> ConstructWidget(TWeakPtr<SMoviePipelineQueueEditor> InQueueWidget, const TSharedRef<STableViewBase>& OwnerTable) override
+	{
+		return SNew(STableRow<TSharedPtr<IMoviePipelineQueueTreeItem>>, OwnerTable)
+			.Content()
+			[
+				SNew(SHorizontalBox)
+
+				+ SHorizontalBox::Slot()
+				.AutoWidth()
+				.VAlign(VAlign_Center)
+				[
+					SNew(STextBlock)
+					.Text(LOCTEXT("MapRow_Label", "Target Map:"))
+				]
+
+				+ SHorizontalBox::Slot()
+				.AutoWidth()
+				.Padding(4, 0, 0, 0)
+				.VAlign(VAlign_Center)
+				[
+					SNew(SObjectPropertyEntryBox)
+					.ObjectPath(this, &FMoviePipelineMapTreeItem::GetMapPath)
+					.AllowedClass(UWorld::StaticClass())
+					.OnObjectChanged(this, &FMoviePipelineMapTreeItem::SetMapPath)
+					.AllowClear(false)
+					.DisplayUseSelected(false)
+					.DisplayBrowse(true)
+					.DisplayThumbnail(true)
+					.DisplayCompactSize(false)
+				]
+			];
+	}
+
+public:
+	FString GetMapPath() const
+	{
+		UMoviePipelineExecutorJob* Job = WeakJob.Get();
+		if (Job)
+		{
+			UWorld* World = Cast<UWorld>(Job->Map.TryLoad());
+			if (World)
+			{
+				return World->GetPathName();
+			}
+		}
+
+		return FString();
+	}
+
+	void SetMapPath(const FAssetData& AssetData)
+	{
+		UMoviePipelineExecutorJob* Job = WeakJob.Get();
+		if (Job)
+		{
+			Job->Map = CastChecked<UWorld>(AssetData.GetAsset());
+		}
+	}
+};
+
+struct FMoviePipelineShotItem : IMoviePipelineQueueTreeItem
+{
+	/** The job that this tree item represents */
+	TWeakObjectPtr<UMoviePipelineExecutorJob> WeakJob;
+
+	explicit FMoviePipelineShotItem(UMoviePipelineExecutorJob* InJob)
+		: WeakJob(InJob)
+	{}
+
+	virtual TSharedRef<ITableRow> ConstructWidget(TWeakPtr<SMoviePipelineQueueEditor> InQueueWidget, const TSharedRef<STableViewBase>& OwnerTable) override
+	{
+		return SNew(STableRow<TSharedPtr<IMoviePipelineQueueTreeItem>>, OwnerTable)
+			.Content()
+			[
+				SNullWidget::NullWidget
+			];
+	}
+
+	/*virtual TSharedRef<SWidget> ConstructWidget(TWeakPtr<SMoviePipelineQueueEditor> InQueueWidget) override
+	{
+		return SNew(SHorizontalBox)
+
+			// Toggle Checkbox for deciding to render or not.
+			+ SHorizontalBox::Slot()
+			.AutoWidth()
+			.Padding(20, 4)
+			.VAlign(VAlign_Center)
+			[
+				SNew(SCheckBox)
+				// .Style(FTakeRecorderStyle::Get(), "TakeRecorder.Source.Switch")
+				.IsFocusable(false)
+				.IsChecked(this, &FMoviePipelineShotItem::GetCheckState)
+				.OnCheckStateChanged(this, &FMoviePipelineShotItem::SetCheckState, InQueueWidget)
+			]
+			
+			// Shot Name Label
+			+SHorizontalBox::Slot()
+			.AutoWidth()
+			.Padding(20, 4)
+			.VAlign(VAlign_Center)
+			[
+				SNew(STextBlock)
+				.Text(this, &FMoviePipelineShotItem::GetShotLabel)
+			]
+
+			// Preset Label
+			+ SHorizontalBox::Slot()
+			.AutoWidth()
+			.VAlign(VAlign_Center)
+			.Padding(8, 0)
+			[
+				SNew(STextBlock)
+				.Text(this, &FMoviePipelineShotItem::GetPresetLabel)
+			]
+
+			// Spacer
+			+ SHorizontalBox::Slot()
+			.FillWidth(1)
+			[
+				SNew(SBox)
+			]
+
+			// Status 
+			+ SHorizontalBox::Slot()
+			.AutoWidth()
+			.VAlign(VAlign_Center)
+			.Padding(8, 0, 24, 0)
+			[
+				SNew(SWidgetSwitcher)
+				.WidgetIndex(this, &FMoviePipelineShotItem::GetStatusIndex)
+
+				// Ready Label
+				+ SWidgetSwitcher::Slot()
+				.VAlign(VAlign_Center)
+				.HAlign(HAlign_Center)
+				[
+					SNew(STextBlock)
+					// .Style(FTakeRecorderStyle::Get(), "TakeRecorder.Source.Switch")
+					.Text(LOCTEXT("PendingJobStatus_Label", "Ready"))
+				]
+
+				// Progress Bar
+				+ SWidgetSwitcher::Slot()
+				.VAlign(VAlign_Center)
+				.HAlign(HAlign_Center)
+				[
+					SNew(SProgressBar)
+					// .Style(FTakeRecorderStyle::Get(), "TakeRecorder.Source.Switch")
+					.Percent(this, &FMoviePipelineShotItem::GetProgressPercent)
+				]
+
+				// Completed
+				+ SWidgetSwitcher::Slot()
+				.VAlign(VAlign_Center)
+				.HAlign(HAlign_Center)
+				[
+					SNew(STextBlock)
+					.Text(LOCTEXT("PendingJobStatus_Label", "Completed!"))
+				]
+			]
+			;
+	}*/
+
+	ECheckBoxState GetCheckState() const
+	{
+		return ECheckBoxState::Checked;
+	}
+
+	void SetCheckState(ECheckBoxState InNewState, TWeakPtr<SMoviePipelineQueueEditor> InQueueWidget)
+	{
+
+	}
+
+	FText GetShotLabel() const
+	{
+		return LOCTEXT("ExampleShotName", "ShotName_001");
+	}
+
+	FText GetPresetLabel() const
+	{
+		return LOCTEXT("ExampleShotName2", "Use Master");
+	}
+
+	int32 GetStatusIndex() const
+	{
+		return 1;
+	}
+
+	TOptional<float> GetProgressPercent() const
+	{
+		return .1f;
+	}
+};
+
+PRAGMA_DISABLE_OPTIMIZATION
+void SMoviePipelineQueueEditor::Construct(const FArguments& InArgs)
+{
+	CachedQueueSerialNumber = uint32(-1);
+	OnEditConfigRequested = InArgs._OnEditConfigRequested;
+
+	TreeView = SNew(STreeView<TSharedPtr<IMoviePipelineQueueTreeItem>>)
+		.TreeItemsSource(&RootNodes)
+		// .OnSelectionChanged(InArgs._OnSelectionChanged)
+		.OnGenerateRow(this, &SMoviePipelineQueueEditor::OnGenerateRow)
+		.OnGetChildren(this, &SMoviePipelineQueueEditor::OnGetChildren)
+		.HeaderRow
+		(
+			SNew(SHeaderRow)
+
+			+ SHeaderRow::Column(SQueueJobListRow::NAME_Sequence)
+			.FillWidth(0.3f)
+			.DefaultLabel(LOCTEXT("QueueHeaderSequence_Text", "Sequence"))
+
+			+ SHeaderRow::Column(SQueueJobListRow::NAME_Settings)
+			.FillWidth(0.3)
+			.DefaultLabel(LOCTEXT("QueueHeaderSettings_Text", "Settings"))
+
+			+ SHeaderRow::Column(SQueueJobListRow::NAME_Output)
+			.FillWidth(0.4)
+			.DefaultLabel(LOCTEXT("QueueHeaderOutput_Text", "Output"))
+
+			+ SHeaderRow::Column(SQueueJobListRow::NAME_Status)
+			.FixedWidth(80)
+			.DefaultLabel(LOCTEXT("QueueHeaderStatus_Text", "Status"))
+		);
+
+	CommandList = MakeShared<FUICommandList>();
+	CommandList->MapAction(
+		FGenericCommands::Get().Delete,
+		FExecuteAction::CreateSP(this, &SMoviePipelineQueueEditor::OnDeleteSelected),
+		FCanExecuteAction::CreateSP(this, &SMoviePipelineQueueEditor::CanDeleteSelected)
+	);
+
+	ChildSlot
+	[
+		SNew(SDropTarget)
+		.OnDrop(this, &SMoviePipelineQueueEditor::OnDragDropTarget)
+		.OnAllowDrop(this, &SMoviePipelineQueueEditor::CanDragDropTarget)
+		.OnIsRecognized(this, &SMoviePipelineQueueEditor::CanDragDropTarget)
+		[
+			TreeView.ToSharedRef()
+		]
+	];
+}
+
+TSharedRef<SWidget> SMoviePipelineQueueEditor::MakeAddSequenceJobButton()
+{
+	return SNew(SComboButton)
+		.ContentPadding(MoviePipeline::ButtonPadding)
+		.ButtonStyle(FMovieRenderPipelineStyle::Get(), "FlatButton.Success")
+		.OnGetMenuContent(this, &SMoviePipelineQueueEditor::OnGenerateNewJobFromAssetMenu)
+		.ForegroundColor(FSlateColor::UseForeground())
+		.HasDownArrow(false)
+		.ButtonContent()
+		[
+			SNew(SHorizontalBox)
+
+			// Plus Icon
+			+ SHorizontalBox::Slot()
+			.VAlign(VAlign_Center)
+			.AutoWidth()
+			[
+				SNew(STextBlock)
+				.TextStyle(FEditorStyle::Get(), "NormalText.Important")
+				.Font(FEditorStyle::Get().GetFontStyle("FontAwesome.10"))
+				.Text(FEditorFontGlyphs::Plus)
+			]
+
+			// "Render" Text
+			+ SHorizontalBox::Slot()
+			.VAlign(VAlign_Center)
+			.AutoWidth()
+			.Padding(4, 0, 0, 0)
+			[
+				SNew(STextBlock)
+				.TextStyle(FEditorStyle::Get(), "NormalText.Important")
+				.Text(LOCTEXT("AddNewJob_Text", "Render"))
+			]
+
+			// Non-Default Down Caret arrow.
+			+ SHorizontalBox::Slot()
+			.VAlign(VAlign_Center)
+			.AutoWidth()
+			.Padding(4, 0, 0, 0)
+			[
+				SNew(STextBlock)
+				.TextStyle(FEditorStyle::Get(), "NormalText.Important")
+				.Font(FEditorStyle::Get().GetFontStyle("FontAwesome.10"))
+				.Text(FEditorFontGlyphs::Caret_Down)
+			]
+		];
+}
+
+TSharedRef<SWidget> SMoviePipelineQueueEditor::RemoveSelectedJobButton()
+{
+	return SNew(SButton)
+		.ContentPadding(MoviePipeline::ButtonPadding)
+		.IsEnabled(this, &SMoviePipelineQueueEditor::CanDeleteSelected)
+		.OnClicked(this, &SMoviePipelineQueueEditor::DeleteSelected)
+		.VAlign(VAlign_Center)
+		[
+			SNew(STextBlock)
+			.TextStyle(FEditorStyle::Get(), "NormalText.Important")
+			.Font(FEditorStyle::Get().GetFontStyle("FontAwesome.10"))
+			.Text(FEditorFontGlyphs::Minus)
+		];
+}
+
+TSharedRef<SWidget> SMoviePipelineQueueEditor::OnGenerateNewJobFromAssetMenu()
+{
+	FMenuBuilder MenuBuilder(true, nullptr);
+	IContentBrowserSingleton& ContentBrowser = FModuleManager::LoadModuleChecked<FContentBrowserModule>("ContentBrowser").Get();
+
+	FAssetPickerConfig AssetPickerConfig;
+	{
+		AssetPickerConfig.SelectionMode = ESelectionMode::Single;
+		AssetPickerConfig.InitialAssetViewType = EAssetViewType::List;
+		AssetPickerConfig.bFocusSearchBoxWhenOpened = true;
+		AssetPickerConfig.bAllowNullSelection = false;
+		AssetPickerConfig.bShowBottomToolbar = true;
+		AssetPickerConfig.bAutohideSearchBar = false;
+		AssetPickerConfig.bAllowDragging = false;
+		AssetPickerConfig.bCanShowClasses = false;
+		AssetPickerConfig.bShowPathInColumnView = true;
+		AssetPickerConfig.bShowTypeInColumnView = false;
+		AssetPickerConfig.bSortByPathInColumnView = false;
+		AssetPickerConfig.ThumbnailScale = 0.4f;
+		AssetPickerConfig.SaveSettingsName = TEXT("MoviePipelineQueueJobAsset");
+
+		AssetPickerConfig.AssetShowWarningText = LOCTEXT("NoSequences_Warning", "No Level Sequences Found");
+		AssetPickerConfig.Filter.ClassNames.Add(ULevelSequence::StaticClass()->GetFName());
+		AssetPickerConfig.OnAssetSelected = FOnAssetSelected::CreateSP(this, &SMoviePipelineQueueEditor::OnCreateJobFromAsset);
+	}
+
+	MenuBuilder.BeginSection(NAME_None, LOCTEXT("NewJob_MenuSection", "New Render Job"));
+	{
+		TSharedRef<SWidget> PresetPicker = SNew(SBox)
+			.WidthOverride(300.f)
+			.HeightOverride(300.f)
+			[
+				ContentBrowser.CreateAssetPicker(AssetPickerConfig)
+			];
+
+		MenuBuilder.AddWidget(PresetPicker, FText(), true, false);
+	}
+	MenuBuilder.EndSection();
+
+	return MenuBuilder.MakeWidget();
+}
+PRAGMA_ENABLE_OPTIMIZATION
+
+void SMoviePipelineQueueEditor::OnCreateJobFromAsset(const FAssetData& InAsset)
+{
+	// Close the dropdown menu that showed them the assets to pick from.
+	FSlateApplication::Get().DismissAllMenus();
+
+	// Only try to initialize level sequences, in the event they had more than a level sequence selected when drag/dropping.
+	ULevelSequence* LevelSequence = Cast<ULevelSequence>(InAsset.GetAsset());
+	if (LevelSequence)
+	{
+		FScopedTransaction Transaction(FText::Format(LOCTEXT("CreateJob_Transaction", "Add {0}|plural(one=Job, other=Jobs)"), 1));
+		
+		UMoviePipelineQueue* ActiveQueue = GEditor->GetEditorSubsystem<UMoviePipelineQueueSubsystem>()->GetQueue();
+		check(ActiveQueue);
+		ActiveQueue->Modify();
+
+		UMoviePipelineExecutorJob* NewJob = ActiveQueue->AllocateNewJob();
+		NewJob->Modify();
+
+		{
+			// We'll assume they went to render from the current world - they can always override it later.
+			FSoftObjectPath CurrentWorld = FSoftObjectPath(GEditor->GetEditorWorldContext().World());
+			FSoftObjectPath Sequence = InAsset.ToSoftObjectPath();
+
+			NewJob->Sequence = Sequence;
+			NewJob->Map = CurrentWorld;
+		}
+
+		{
+			// The job configuration is already set up with an empty configuration, but we'll try and use their last used preset 
+			// (or a engine supplied default) for better user experience. 
+			const UMovieRenderPipelineProjectSettings* ProjectSettings = GetDefault<UMovieRenderPipelineProjectSettings>();
+			if (ProjectSettings->LastPresetOrigin.IsValid())
+			{
+				NewJob->SetPresetOrigin(ProjectSettings->LastPresetOrigin.Get());
+			}
+		}
+	}
+}
+
+void SMoviePipelineQueueEditor::Tick(const FGeometry& AllottedGeometry, const double InCurrentTime, const float InDeltaTime)
+{
+	UMoviePipelineQueue* ActiveQueue = GEditor->GetEditorSubsystem<UMoviePipelineQueueSubsystem>()->GetQueue();
+	check(ActiveQueue);
+
+	if (ActiveQueue)
+	{
+		if (CachedQueueSerialNumber != ActiveQueue->GetQueueSerialNumber())
+		{
+			ReconstructTree();
+		}
+	}
+	// The sources are no longer valid, so we expect our cached serial number to be -1. If not, we haven't reset the tree yet.
+	else if (CachedQueueSerialNumber != uint32(-1))
+	{
+		ReconstructTree();
+	}
+}
+
+void SMoviePipelineQueueEditor::ReconstructTree()
+{
+	UMoviePipelineQueue* ActiveQueue = GEditor->GetEditorSubsystem<UMoviePipelineQueueSubsystem>()->GetQueue();
+	check(ActiveQueue);
+	if (!ActiveQueue)
+	{
+		CachedQueueSerialNumber = uint32(-1);
+		RootNodes.Reset();
+		return;
+	}
+
+	CachedQueueSerialNumber = ActiveQueue->GetQueueSerialNumber();
+
+	// TSortedMap<FString, TSharedPtr<FMoviePipelineQueueJobTreeItem>> RootJobs;
+	// for (TSharedPtr<IMoviePipelineQueueTreeItem> RootItem : RootNodes)
+	// {
+	// 	TSharedPtr<FMoviePipelineQueueJobTreeItem> RootCategory = RootItem->AsJob();
+	// 	if (RootCategory.IsValid())
+	// 	{
+	// 		RootCategory->Children.Reset();
+	// 		RootJobs.Add(RootCategory->Category.ToString(), RootCategory);
+	// 	}
+	// }
+
+	RootNodes.Reset();
+
+	// We attempt to re-use tree items in order to maintain selection states on them
+	// TMap<FObjectKey, TSharedPtr<FTakeRecorderSourceTreeItem>> OldSourceToTreeItem;
+	// Swap(SourceToTreeItem, OldSourceToTreeItem);
+
+	for (UMoviePipelineExecutorJob* Job : ActiveQueue->GetJobs())
+	{
+		if (!Job)
+		{
+			continue;
+		}
+
+		TSharedPtr< FMoviePipelineQueueJobTreeItem> JobTreeItem = MakeShared< FMoviePipelineQueueJobTreeItem>(Job, OnEditConfigRequested);
+		TSharedPtr<FMoviePipelineMapTreeItem> MapTreeItem = MakeShared< FMoviePipelineMapTreeItem>(Job);
+		JobTreeItem->Children.Add(MapTreeItem);
+		RootNodes.Add(JobTreeItem);
+	}
+
+	TreeView->RequestTreeRefresh();
+}
+
+
+FReply SMoviePipelineQueueEditor::OnKeyDown(const FGeometry& MyGeometry, const FKeyEvent& InKeyEvent)
+{
+	if (CommandList->ProcessCommandBindings(InKeyEvent))
+	{
+		return FReply::Handled();
+	}
+
+	return FReply::Unhandled();
+}
+
+TSharedRef<ITableRow> SMoviePipelineQueueEditor::OnGenerateRow(TSharedPtr<IMoviePipelineQueueTreeItem> Item, const TSharedRef<STableViewBase>& Tree)
+{
+	// Let the item construct itself.
+	return Item->ConstructWidget(SharedThis(this), Tree);
+}
+
+void SMoviePipelineQueueEditor::OnGetChildren(TSharedPtr<IMoviePipelineQueueTreeItem> Item, TArray<TSharedPtr<IMoviePipelineQueueTreeItem>>& OutChildItems)
+{
+	TSharedPtr<FMoviePipelineQueueJobTreeItem> Job = Item->AsJob();
+	if (Job.IsValid())
+	{
+		OutChildItems.Append(Job->Children);
+	}
+}
+
+FReply SMoviePipelineQueueEditor::OnDragDropTarget(TSharedPtr<FDragDropOperation> InOperation)
+{
+	if (InOperation)
+	{
+		if (InOperation->IsOfType<FAssetDragDropOp>())
+		{
+			TSharedPtr<FAssetDragDropOp> AssetDragDrop = StaticCastSharedPtr<FAssetDragDropOp>(InOperation);
+			FScopedTransaction Transaction(FText::Format(LOCTEXT("CreateJob_Transaction", "Add {0}|plural(one=Job, other=Jobs)"), AssetDragDrop->GetAssets().Num()));
+
+			for (const FAssetData& Asset : AssetDragDrop->GetAssets())
+			{
+				OnCreateJobFromAsset(Asset);
+			}
+
+			return FReply::Handled();
+		}
+	}
+
+	return FReply::Unhandled();
+}
+
+bool SMoviePipelineQueueEditor::CanDragDropTarget(TSharedPtr<FDragDropOperation> InOperation)
+{
+	bool bIsValid = false;
+	if (InOperation)
+	{
+		if (InOperation->IsOfType<FAssetDragDropOp>())
+		{
+			TSharedPtr<FAssetDragDropOp> AssetDragDrop = StaticCastSharedPtr<FAssetDragDropOp>(InOperation);
+			for (const FAssetData& Asset : AssetDragDrop->GetAssets())
+			{
+				ULevelSequence* LevelSequence = Cast<ULevelSequence>(Asset.GetAsset());
+				if (LevelSequence)
+				{
+					// If at least one of them is a Level Sequence then we'll accept the drop.
+					bIsValid = true;
+					break;
+				}
+			}
+		}
+	}
+
+	return bIsValid;
+}
+
+FReply SMoviePipelineQueueEditor::DeleteSelected()
+{
+	UMoviePipelineQueue* ActiveQueue = GEditor->GetEditorSubsystem<UMoviePipelineQueueSubsystem>()->GetQueue();
+	check(ActiveQueue);
+
+	if (ActiveQueue)
+	{
+		TArray<TSharedPtr<IMoviePipelineQueueTreeItem>> Items = TreeView->GetSelectedItems();
+
+		FScopedTransaction Transaction(FText::Format(LOCTEXT("DeleteSelection", "Delete Selected {0}|plural(one=Job, other=Jobs)"), Items.Num()));
+		ActiveQueue->Modify();
+
+		for (TSharedPtr<IMoviePipelineQueueTreeItem> Item : Items)
+		{
+			Item->Delete(ActiveQueue);
+		}
+	}
+
+	return FReply::Handled();
+}
+
+void SMoviePipelineQueueEditor::OnDeleteSelected()
+{
+	DeleteSelected();
+}
+
+bool SMoviePipelineQueueEditor::CanDeleteSelected() const
+{
+	return true;
+}
+
+#undef LOCTEXT_NAMESPACE
\ No newline at end of file
