Fixed compilation error

diff --git a/Engine/Plugins/MovieScene/MovieRenderPipeline/Source/MovieRenderPipelineRenderPasses/Private/MoviePipelineDeferredPasses.cpp b/Engine/Plugins/MovieScene/MovieRenderPipeline/Source/MovieRenderPipelineRenderPasses/Private/MoviePipelineDeferredPasses.cpp
--- a/Engine/Plugins/MovieScene/MovieRenderPipeline/Source/MovieRenderPipelineRenderPasses/Private/MoviePipelineDeferredPasses.cpp
+++ b/Engine/Plugins/MovieScene/MovieRenderPipeline/Source/MovieRenderPipelineRenderPasses/Private/MoviePipelineDeferredPasses.cpp
@@ -475,7 +475,7 @@ void UMoviePipelineDeferredPassBase::OnSetupView(FSceneViewFamily& InViewFamily,
 			: FinalPostProcessSettings(InFinalPostProcessSettings), RenderPasses(InRenderPasses)
 		{}
 
-		void ProcessValue(const FString& InName, UMaterial* Material, const FText& InText)
+		void ProcessValue(const FString& InName, UMaterialInterface* Material, const FText& InText)
 		{
 			if (!RenderPasses.Num() || RenderPasses.Contains(InName) || RenderPasses.Contains(InText.ToString()))
 			{
diff --git a/Engine/Source/Runtime/Engine/Public/BufferVisualizationData.h b/Engine/Source/Runtime/Engine/Public/BufferVisualizationData.h
--- a/Engine/Source/Runtime/Engine/Public/BufferVisualizationData.h
+++ b/Engine/Source/Runtime/Engine/Public/BufferVisualizationData.h
@@ -42,7 +42,7 @@ public:
 		for (TMaterialMap::TConstIterator It = MaterialMap.CreateConstIterator(); It; ++It)
 		{
 			const Record& Rec = It.Value();
-			Iterator.ProcessValue(Rec.Name, Rec.Material->GetMaterial(), Rec.DisplayName);
+			Iterator.ProcessValue(Rec.Name, Rec.Material, Rec.DisplayName);
 		}
 	}
 
