Fix-up deprecated calls in commented code
#rnx

diff --git a/Engine/Plugins/MovieScene/MovieRenderPipeline/Source/MovieRenderPipelineCore/Private/MoviePipelineAudioOutput.cpp b/Engine/Plugins/MovieScene/MovieRenderPipeline/Source/MovieRenderPipelineCore/Private/MoviePipelineAudioOutput.cpp
--- a/Engine/Plugins/MovieScene/MovieRenderPipeline/Source/MovieRenderPipelineCore/Private/MoviePipelineAudioOutput.cpp
+++ b/Engine/Plugins/MovieScene/MovieRenderPipeline/Source/MovieRenderPipelineCore/Private/MoviePipelineAudioOutput.cpp
@@ -23,7 +23,7 @@ void UMoviePipelineAudioOutput::SetupForPipelineImpl(UMoviePipeline* InPipeline)
 	PrevUnfocusedAudioMultiplier = FApp::GetUnfocusedVolumeMultiplier();
 	FApp::SetUnfocusedVolumeMultiplier(1.f);
 
-	if (FAudioDeviceManager* AudioDeviceManager = GEngine->GetAudioDeviceManager())
+	if (FAudioDeviceManager* AudioDeviceManager = FAudioDeviceManager::Get())
 	{
 		PrevAudioDevicePlatform = AudioDeviceManager->GetCurrentAudioDeviceModuleName();
 		bool bInitSuccessful = AudioDeviceManager->SwitchAudioDevicePlatform(TEXT("NonRealtimeAudioRenderer"));
@@ -49,7 +49,7 @@ void UMoviePipelineAudioOutput::SetupForPipelineImpl(UMoviePipeline* InPipeline)
 			UE_LOG(LogMovieRenderPipeline, Log, TEXT("Successfully initialized Non-Real Time audio backend."));
 		}
 	}
-	
+
 	// Grab a rough estimate of how long the expected output is to avoid unnecessarily resizing arrays.
 	// This may not be perfect (due to handle frames, slow-mo, etc.) but it's better than guessing.
 	ULevelSequence* TargetSequence = InPipeline->GetTargetSequence();
@@ -67,12 +67,12 @@ void UMoviePipelineAudioOutput::OnPipelineFinishedImpl()
 {
 	/*// Restore our cached CVar value.
 	IConsoleVariable* AudioRenderEveryTickCvar = IConsoleManager::Get().FindConsoleVariable(TEXT("au.nrt.RenderEveryTick"));
-	
+
 	// This will be null if initialization failed.
 	if (AudioRenderEveryTickCvar)
 	{
 		AudioRenderEveryTickCvar->Set(PrevRenderEveryTickValue, ECVF_SetByConstructor);
-		
+
 		// Convert it to absolute as the Audio Recorder will save relative paths to a different
 		// directory than we consider to be our relative root.
 		FString FormattedFileName = TEXT("Audio.wav");
@@ -82,7 +82,7 @@ void UMoviePipelineAudioOutput::OnPipelineFinishedImpl()
 
 	// Attempt to restore the previous Audio Platform.
 	check(GEngine);
-	if (FAudioDeviceManager* AudioDeviceManager = GEngine->GetAudioDeviceManager())
+	if (FAudioDeviceManager* AudioDeviceManager = FAudioDeviceManager::Get())
 	{
 		AudioDeviceManager->SwitchAudioDevicePlatform(PrevAudioDevicePlatform);
 	}
@@ -93,8 +93,8 @@ void UMoviePipelineAudioOutput::OnPipelineFinishedImpl()
 void UMoviePipelineAudioOutput::OnPostTickImpl()
 {
 	/*check(GEngine);
-	
-	if (FAudioDevice* AudioDevice = GEngine->GetActiveAudioDevice())
+
+	if (FAudioDevice* AudioDevice = FAudioDeviceManager::GetActiveDevice())
 	{
 		// Handle any game logic that changed Audio State.
 		AudioDevice->Update(true);
