Movie Pipeline: Fix non unity compile

diff --git a/Engine/Plugins/MovieScene/MovieRenderPipeline/Source/MovieRenderPipelineCore/Private/MoviePipeline.cpp b/Engine/Plugins/MovieScene/MovieRenderPipeline/Source/MovieRenderPipelineCore/Private/MoviePipeline.cpp
--- a/Engine/Plugins/MovieScene/MovieRenderPipeline/Source/MovieRenderPipelineCore/Private/MoviePipeline.cpp
+++ b/Engine/Plugins/MovieScene/MovieRenderPipeline/Source/MovieRenderPipelineCore/Private/MoviePipeline.cpp
@@ -14,6 +14,7 @@
 #include "AudioDeviceManager.h"
 #include "MovieSceneTimeHelpers.h"
 #include "Engine/World.h"
+#include "Modules/ModuleManager.h"
 #include "GameFramework/PlayerController.h"
 #include "MovieRenderDebugWidget.h"
 #include "MoviePipelineShotConfig.h"
