MoviePipeline: Shader/Mesh Distance Field compiler flushing now forces it to process results which ensures they actually get applied to the world.

#rnx

diff --git a/Engine/Plugins/MovieScene/MovieRenderPipeline/Source/MovieRenderPipelineCore/Private/MoviePipeline.cpp b/Engine/Plugins/MovieScene/MovieRenderPipeline/Source/MovieRenderPipelineCore/Private/MoviePipeline.cpp
--- a/Engine/Plugins/MovieScene/MovieRenderPipeline/Source/MovieRenderPipelineCore/Private/MoviePipeline.cpp
+++ b/Engine/Plugins/MovieScene/MovieRenderPipeline/Source/MovieRenderPipelineCore/Private/MoviePipeline.cpp
@@ -898,10 +898,17 @@ void UMoviePipeline::FlushAsyncEngineSystems()
 	if (GShaderCompilingManager)
 	{
 		bool bDidWork = false;
-		while (GShaderCompilingManager->GetNumRemainingJobs() > 0)
+		int32 NumShadersToCompile = GShaderCompilingManager->GetNumRemainingJobs();
+		if (NumShadersToCompile > 0)
 		{
-			UE_LOG(LogMovieRenderPipeline, Log, TEXT("[%d] Waiting for %d shaders to finish compiling..."), GFrameCounter, GShaderCompilingManager->GetNumRemainingJobs());
+			UE_LOG(LogMovieRenderPipeline, Log, TEXT("[%d] Starting build for %d shaders."), GFrameCounter, NumShadersToCompile);
+		}
 
+		while (GShaderCompilingManager->GetNumRemainingJobs() > 0 || GShaderCompilingManager->HasShaderJobs())
+		{
+			UE_LOG(LogMovieRenderPipeline, Log, TEXT("[%d] Waiting for %d shaders [Has Shader Jobs: %d] to finish compiling..."), GFrameCounter, GShaderCompilingManager->GetNumRemainingJobs(), GShaderCompilingManager->HasShaderJobs());
+			GShaderCompilingManager->ProcessAsyncResults(false, true);
+			
 			// Sleep for 1 second and then check again. This way we get an indication of progress as this works.
 			FPlatformProcess::Sleep(1.f);
 			bDidWork = true;
@@ -909,7 +916,7 @@ void UMoviePipeline::FlushAsyncEngineSystems()
 
 		if (bDidWork)
 		{
-			UE_LOG(LogMovieRenderPipeline, Log, TEXT("[%d] Done waiting for shaders to finish."), GFrameCounter);
+			UE_LOG(LogMovieRenderPipeline, Log, TEXT("[%d] Done building %d shaders."), GFrameCounter, NumShadersToCompile);
 		}
 	}
 
@@ -917,9 +924,16 @@ void UMoviePipeline::FlushAsyncEngineSystems()
 	if (GDistanceFieldAsyncQueue)
 	{
 		bool bDidWork = false;
+		int32 NumDistanceFieldsToBuild = GDistanceFieldAsyncQueue->GetNumOutstandingTasks();
+		if (NumDistanceFieldsToBuild > 0)
+		{
+			UE_LOG(LogMovieRenderPipeline, Log, TEXT("[%d] Starting build for %d mesh distance fields."), GFrameCounter, NumDistanceFieldsToBuild);
+		}
+
 		while (GDistanceFieldAsyncQueue->GetNumOutstandingTasks() > 0)
 		{
 			UE_LOG(LogMovieRenderPipeline, Log, TEXT("[%d] Waiting for %d Mesh Distance Fields to finish building..."), GFrameCounter, GDistanceFieldAsyncQueue->GetNumOutstandingTasks());
+			GDistanceFieldAsyncQueue->ProcessAsyncTasks();
 
 			// Sleep for 1 second and then check again. This way we get an indication of progress as this works.
 			FPlatformProcess::Sleep(1.f);
@@ -928,7 +942,7 @@ void UMoviePipeline::FlushAsyncEngineSystems()
 
 		if (bDidWork)
 		{
-			UE_LOG(LogMovieRenderPipeline, Log, TEXT("[%d] Done waiting for Mesh Distance Fields to build."), GFrameCounter);
+			UE_LOG(LogMovieRenderPipeline, Log, TEXT("[%d] Done building %d Mesh Distance Fields."), GFrameCounter, NumDistanceFieldsToBuild);
 		}
 
 	}
