MoviePipeline: Removed unneeded MoviePipelineBackbufferPass, cleaned up the Movie Pipeline Editor Module includes.

#rnx

diff --git a/Engine/Plugins/MovieScene/MovieRenderPipeline/Source/MovieRenderPipelineEditor/Private/MovieRenderPipelineEditorModule.cpp b/Engine/Plugins/MovieScene/MovieRenderPipeline/Source/MovieRenderPipelineEditor/Private/MovieRenderPipelineEditorModule.cpp
--- a/Engine/Plugins/MovieScene/MovieRenderPipeline/Source/MovieRenderPipelineEditor/Private/MovieRenderPipelineEditorModule.cpp
+++ b/Engine/Plugins/MovieScene/MovieRenderPipeline/Source/MovieRenderPipelineEditor/Private/MovieRenderPipelineEditorModule.cpp
@@ -1,31 +1,22 @@
 // Copyright 1998-2019 Epic Games, Inc. All Rights Reserved.
 
 #include "MovieRenderPipelineEditorModule.h"
-#include "Widgets/Docking/SDockTab.h"
-#include "Widgets/Text/STextBlock.h"
-#include "EditorStyleSet.h"
+#include "MovieRenderPipelineCoreModule.h"
+#include "MoviePipelineExecutor.h"
 #include "WorkspaceMenuStructureModule.h"
 #include "WorkspaceMenuStructure.h"
 #include "Widgets/SMoviePipelineTabContent.h"
-#include "Editor.h"
-#include "MoviePipeline.h"
-#include "MovieRenderPipelineDataTypes.h"
-#include "MoviePipelineBackbufferPass.h"
-#include "MoviePipelineImageSequenceContainer.h"
-#include "MoviePipelineShotConfig.h"
-#include "MovieRenderPipelineCoreModule.h"
-#include "MoviePipelineAudioOutput.h"
-#include "MoviePipelineBackbufferPass.h"
-#include "MoviePipelineGameOverrideSetting.h"
-#include "GameFramework/GameMode.h"
-#include "Engine/World.h"
-#include "GameFramework/WorldSettings.h"
-#include "MoviePipelineAccumulationSetting.h"
 #include "ISettingsModule.h"
 #include "MovieRenderPipelineSettings.h"
 #include "HAL/IConsoleManager.h"
-#include "MovieRenderPipelineSettings.h"
 #include "MoviePipelineExecutor.h"
+#include "Widgets/Docking/SDockTab.h"
+#include "Framework/Docking/TabManager.h"
+#include "EditorStyleSet.h"
+#include "Misc/PackageName.h"
+#include "UObject/UObjectGlobals.h"
+#include "UObject/Package.h"
+#include "Templates/SubclassOf.h"
 
 #define LOCTEXT_NAMESPACE "FMovieRenderPipelineEditorModule"
 
@@ -132,10 +123,6 @@ void FMovieRenderPipelineEditorModule::StartupModule()
 
 void FMovieRenderPipelineEditorModule::ShutdownModule()
 {
-	FEditorDelegates::PostPIEStarted.RemoveAll(this);
-	FEditorDelegates::PrePIEEnded.RemoveAll(this);
-	UGameViewportClient::OnViewportCreated().RemoveAll(this);
-
 	UnregisterSettings();
 }
 
diff --git a/Engine/Plugins/MovieScene/MovieRenderPipeline/Source/MovieRenderPipelineEditor/Private/MovieRenderPipelineEditorTest.cpp b/Engine/Plugins/MovieScene/MovieRenderPipeline/Source/MovieRenderPipelineEditor/Private/MovieRenderPipelineEditorTest.cpp
--- a/Engine/Plugins/MovieScene/MovieRenderPipeline/Source/MovieRenderPipelineEditor/Private/MovieRenderPipelineEditorTest.cpp
+++ b/Engine/Plugins/MovieScene/MovieRenderPipeline/Source/MovieRenderPipelineEditor/Private/MovieRenderPipelineEditorTest.cpp
@@ -1,9 +1,7 @@
 // Copyright 1998-2019 Epic Games, Inc. All Rights Reserved.
 #include "MovieRenderPipelineEditorModule.h"
-#include "MoviePipelineBackbufferPass.h"
 #include "MoviePipelineShotConfig.h"
 #include "MoviePipelineGameOverrideSetting.h"
-#include "GameFramework/GameMode.h"
 #include "MoviePipelineAccumulationSetting.h"
 #include "MoviePipelineImageSequenceContainer.h"
 #include "MoviePipelineAudioOutput.h"
@@ -13,6 +11,7 @@
 #include "MoviePipelineHighResSetting.h"
 #include "MoviePipelineDeferredPasses.h"
 #include "MoviePipelineCameraSetting.h"
+#include "GameFramework/GameModeBase.h"
 
 UMoviePipelineShotConfig* GenerateTestShotConfig(UObject* InOwner, int32 InSampleCount,
 	int32 InShutterAngle, EMoviePipelineShutterTiming InFrameTiming, int32 InTileCount,
@@ -21,9 +20,6 @@ UMoviePipelineShotConfig* GenerateTestShotConfig(UObject* InOwner, int32 InSampl
 	UMoviePipelineShotConfig* OutConfig = NewObject<UMoviePipelineShotConfig>(InOwner);
 	OutConfig->FindOrAddSetting<UMoviePipelineDeferredPassBase>();
 
-	UMoviePipelineGameOverrideSetting* GamemodeOverride = OutConfig->FindOrAddSetting<UMoviePipelineGameOverrideSetting>();
-	GamemodeOverride->GameModeOverride = AGameMode::StaticClass();
-
 	UMoviePipelineAccumulationSetting* Accumulation = OutConfig->FindOrAddSetting< UMoviePipelineAccumulationSetting>();
 	Accumulation->SpatialSampleCount = InSpatialSampleCount;
 	Accumulation->AccumulationGamma = AccumulationGamma;
@@ -77,9 +73,9 @@ TArray<UMoviePipelineMasterConfig*> FMovieRenderPipelineEditorModule::GenerateTe
 
 				int32 SizeX = 1920;
 				int32 SizeY = 1080;
-				int32 TileX = 4;
-				int32 TileY = 4;
-				int32 TestNumSamples = 8;
+				int32 TileX = 2;
+				int32 TileY = 2;
+				int32 TestNumSamples = 2;
 				float PadRatioX = 0.5f;
 				float PadRatioY = 0.5f;
 				float AccumulationGamma = 1.0f;
@@ -102,6 +98,9 @@ TArray<UMoviePipelineMasterConfig*> FMovieRenderPipelineEditorModule::GenerateTe
 				// Add some Outputs
 				OutPipeline->FindOrAddSetting<UMoviePipelineImageSequenceContainerBase>();
 
+				UMoviePipelineGameOverrideSetting* GamemodeOverride = OutPipeline->FindOrAddSetting<UMoviePipelineGameOverrideSetting>();
+				GamemodeOverride->GameModeOverride = AGameModeBase::StaticClass();
+
 				// Add it to our list to be processed.
 				Configs.Add(OutPipeline);
 			}
diff --git a/Engine/Plugins/MovieScene/MovieRenderPipeline/Source/MovieRenderPipelineRenderPasses/Private/MoviePipelineBackbufferPass.cpp b/Engine/Plugins/MovieScene/MovieRenderPipeline/Source/MovieRenderPipelineRenderPasses/Private/MoviePipelineBackbufferPass.cpp
deleted file mode 100644
diff --git a/Engine/Plugins/MovieScene/MovieRenderPipeline/Source/MovieRenderPipelineRenderPasses/Public/MoviePipelineBackbufferPass.h b/Engine/Plugins/MovieScene/MovieRenderPipeline/Source/MovieRenderPipelineRenderPasses/Public/MoviePipelineBackbufferPass.h
deleted file mode 100644
