Removing overlapped tile flags.

diff --git a/Engine/Plugins/MovieScene/MovieRenderPipeline/Source/MovieRenderPipelineCore/Private/MoviePipelineRendering.cpp b/Engine/Plugins/MovieScene/MovieRenderPipeline/Source/MovieRenderPipelineCore/Private/MoviePipelineRendering.cpp
--- a/Engine/Plugins/MovieScene/MovieRenderPipeline/Source/MovieRenderPipelineCore/Private/MoviePipelineRendering.cpp
+++ b/Engine/Plugins/MovieScene/MovieRenderPipeline/Source/MovieRenderPipelineCore/Private/MoviePipelineRendering.cpp
@@ -46,7 +46,8 @@ void UMoviePipeline::SetupRenderingPipelineForShot(FMoviePipelineShotInfo& Shot)
 	FIntPoint BackbufferTileCount = FIntPoint(1, 1);
 	
 	UMoviePipelineHighResSetting* HighResSettings = Shot.ShotConfig->FindSetting<UMoviePipelineHighResSetting>();
-	if (HighResSettings)
+
+	check(HighResSettings != nullptr)
 	{
 		// Figure out how big each sub-region (tile) is.
 		BackbufferResolution = FIntPoint(
@@ -54,10 +55,7 @@ void UMoviePipeline::SetupRenderingPipelineForShot(FMoviePipelineShotInfo& Shot)
 			FMath::CeilToInt(BackbufferResolution.Y / HighResSettings->TileCount));
 
 		// Then increase each sub-region by the overlap amount.
-		if (HighResSettings->bIsUsingOverlappedTiles)
-		{
-			BackbufferResolution = HighResSettings->CalculatePaddedBackbufferSize(BackbufferResolution);
-		}
+		BackbufferResolution = HighResSettings->CalculatePaddedBackbufferSize(BackbufferResolution);
 
 		// Note how many tiles we wish to render with.
 		BackbufferTileCount = FIntPoint(HighResSettings->TileCount, HighResSettings->TileCount);
@@ -168,13 +166,10 @@ void UMoviePipeline::RenderFrame()
 	FIntPoint TileCount = FIntPoint(1, 1);
 	FIntPoint OutputResolution = OutputSettings->OutputResolution;
 
-	bool bIsUsingOverlappedTiles = false;
-
 	UMoviePipelineHighResSetting* HighResSettings = CurrentShot.ShotConfig->FindSetting<UMoviePipelineHighResSetting>();
-	if (HighResSettings)
+	check(HighResSettings != nullptr)
 	{
 		TileCount = FIntPoint(HighResSettings->TileCount, HighResSettings->TileCount);
-		bIsUsingOverlappedTiles = HighResSettings->bIsUsingOverlappedTiles;
 	}
 
 	int32 NumSpatialSamples = AccumulationSettings->SpatialSampleCount;
@@ -265,7 +260,6 @@ void UMoviePipeline::RenderFrame()
 			double TileShiftX; 
 			double TileShiftY;
 
-			if (bIsUsingOverlappedTiles)
 			{
 				TileShiftX = 0.5f;
 				TileShiftY = 0.5f;
@@ -273,14 +267,6 @@ void UMoviePipeline::RenderFrame()
 				TilePixelSizeFractionX = 1.0f;
 				TilePixelSizeFractionY = 1.0f;
 			}
-			else
-			{
-				// Legacy interleaving
-				TilePixelSizeFractionX = 1.0 / TileCount.X;
-				TilePixelSizeFractionY = 1.0 / TileCount.Y;
-				TileShiftX = (TilePixelSizeFractionX * TileX) + (TilePixelSizeFractionX / 2.0);
-				TileShiftY = (TilePixelSizeFractionY * TileY) + (TilePixelSizeFractionY / 2.0);
-			}
 
 			// If they're trying to seed the histories, then we'll just override the NumSpatialSamples since they don't need all samples as
 			// this isn't output.
@@ -317,7 +303,6 @@ void UMoviePipeline::RenderFrame()
 				FIntPoint BackbufferResolution = OutputSettings->OutputResolution;
 				FIntPoint TileResolution = BackbufferResolution;
 
-				if (bIsUsingOverlappedTiles)
 				{
 					BackbufferResolution = FIntPoint(
 						FMath::CeilToInt(BackbufferResolution.X / HighResSettings->TileCount),
@@ -329,15 +314,6 @@ void UMoviePipeline::RenderFrame()
 					// Apply size padding.
 					BackbufferResolution = HighResSettings->CalculatePaddedBackbufferSize(BackbufferResolution);
 				}
-				else
-				{
-					// Legacy interleaving
-					BackbufferResolution.X = FMath::CeilToInt(OutputResolution.X / TileCount.X);
-					BackbufferResolution.Y = FMath::CeilToInt(OutputResolution.Y / TileCount.Y);
-
-					// No padding is used when interleaving
-					TileResolution = BackbufferResolution;
-				}
 
 				// We take all of the information needed to render a single sample and package it into a struct.
 				FMoviePipelineRenderPassMetrics SampleState;
@@ -351,7 +327,6 @@ void UMoviePipeline::RenderFrame()
 				SampleState.TemporalSampleIndex = CachedOutputState.TemporalSampleIndex;
 				SampleState.TemporalSampleCount = CameraSettings->TemporalSampleCount;
 				SampleState.AccumulationGamma = AccumulationSettings->AccumulationGamma;
-				SampleState.bIsUsingOverlappedTiles = bIsUsingOverlappedTiles;
 				SampleState.BackbufferSize = BackbufferResolution;
 				SampleState.TileSize = TileResolution;
 				SampleState.FrameInfo = FrameInfo;
@@ -364,7 +339,6 @@ void UMoviePipeline::RenderFrame()
 					SampleState.bWriteSampleToDisk = false;
 				}
 
-				if (bIsUsingOverlappedTiles)
 				{
 					SampleState.OverlappedPad = FIntPoint(FMath::CeilToInt(TileResolution.X * HighResSettings->OverlapPercentage), 
 														   FMath::CeilToInt(TileResolution.Y * HighResSettings->OverlapPercentage));
@@ -372,12 +346,6 @@ void UMoviePipeline::RenderFrame()
 															  TileY * TileResolution.Y - SampleState.OverlappedPad.Y);
 					SampleState.OverlappedSubpixelShift = FVector2D(1.0f - OffsetX,OffsetY);
 				}
-				else
-				{
-					SampleState.OverlappedOffset = FIntPoint(0,0);
-					SampleState.OverlappedPad = FIntPoint(0, 0);
-					SampleState.OverlappedSubpixelShift = FVector2D(0.5f,0.5f);
-				}
 
 				SampleState.WeightFunctionX.InitHelper(SampleState.OverlappedPad.X, SampleState.TileSize.X, SampleState.OverlappedPad.X);
 				SampleState.WeightFunctionY.InitHelper(SampleState.OverlappedPad.Y, SampleState.TileSize.Y, SampleState.OverlappedPad.Y);
diff --git a/Engine/Plugins/MovieScene/MovieRenderPipeline/Source/MovieRenderPipelineCore/Public/MoviePipelineHighResSetting.h b/Engine/Plugins/MovieScene/MovieRenderPipeline/Source/MovieRenderPipelineCore/Public/MoviePipelineHighResSetting.h
--- a/Engine/Plugins/MovieScene/MovieRenderPipeline/Source/MovieRenderPipelineCore/Public/MoviePipelineHighResSetting.h
+++ b/Engine/Plugins/MovieScene/MovieRenderPipeline/Source/MovieRenderPipelineCore/Public/MoviePipelineHighResSetting.h
@@ -15,7 +15,6 @@ public:
 	UMoviePipelineHighResSetting()
 		: TileCount(1)
 		, TextureSharpnessBias(0.f)
-		, bIsUsingOverlappedTiles(true)
 		, OverlapPercentage(.5f)
 		, bWriteAllSamples(true)
 	{
@@ -52,11 +51,6 @@ public:
 	UPROPERTY(EditAnywhere, BlueprintReadWrite, meta = (UIMin = -1, UIMax = 0), Category = "Movie Pipeline")
 	float TextureSharpnessBias;
 	
-	/**
-	* Do we overlap tiles or interlace? Only overlapping is supported at the moment.
-	*/
-	bool bIsUsingOverlappedTiles;
-	
 	/**
 	* How much should each tile overlap each other (0-0.5). Decreasing the overlap will result in smaller individual
 	* tiles (which means faster renders) but increases the likelyhood of edge-of-screen artifacts showing up which
diff --git a/Engine/Plugins/MovieScene/MovieRenderPipeline/Source/MovieRenderPipelineCore/Public/MovieRenderPipelineDataTypes.h b/Engine/Plugins/MovieScene/MovieRenderPipeline/Source/MovieRenderPipelineCore/Public/MovieRenderPipelineDataTypes.h
--- a/Engine/Plugins/MovieScene/MovieRenderPipeline/Source/MovieRenderPipelineCore/Public/MovieRenderPipelineDataTypes.h
+++ b/Engine/Plugins/MovieScene/MovieRenderPipeline/Source/MovieRenderPipelineCore/Public/MovieRenderPipelineDataTypes.h
@@ -736,9 +736,6 @@ public:
 	/** Of the TemporalSampleCount, which index is this? */
 	int32 TemporalSampleIndex;
 
-	/** Should we use the Overlapped Tiles method of doing high res screen shots? */
-	bool bIsUsingOverlappedTiles;
-
 	/** How many pixels are we overlapping with adjacent tiles (on each side) */
 	FIntPoint OverlappedPad;
 
diff --git a/Engine/Plugins/MovieScene/MovieRenderPipeline/Source/MovieRenderPipelineEditor/Private/MovieRenderPipelineEditorTest.cpp b/Engine/Plugins/MovieScene/MovieRenderPipeline/Source/MovieRenderPipelineEditor/Private/MovieRenderPipelineEditorTest.cpp
--- a/Engine/Plugins/MovieScene/MovieRenderPipeline/Source/MovieRenderPipelineEditor/Private/MovieRenderPipelineEditorTest.cpp
+++ b/Engine/Plugins/MovieScene/MovieRenderPipeline/Source/MovieRenderPipelineEditor/Private/MovieRenderPipelineEditorTest.cpp
@@ -15,7 +15,7 @@
 
 UMoviePipelineShotConfig* GenerateTestShotConfig(UObject* InOwner, int32 InSampleCount,
 	int32 InShutterAngle, EMoviePipelineShutterTiming InFrameTiming, int32 InTileCount,
-	int32 InSpatialSampleCount, bool bIsUsingOverlappedTiles, float PadRatioX, float PadRatioY, float AccumulationGamma)
+	int32 InSpatialSampleCount, float PadRatioX, float PadRatioY, float AccumulationGamma)
 {
 	UMoviePipelineShotConfig* OutConfig = NewObject<UMoviePipelineShotConfig>(InOwner);
 	OutConfig->FindOrAddSetting<UMoviePipelineDeferredPassBase>();
@@ -33,11 +33,9 @@ UMoviePipelineShotConfig* GenerateTestShotConfig(UObject* InOwner, int32 InSampl
 		CameraSetting->ExposureCompensation = 11;
 	}
 
-	if (InTileCount > 1)
 	{
 		UMoviePipelineHighResSetting* HighRes = OutConfig->FindOrAddSetting<UMoviePipelineHighResSetting>();
 		HighRes->TileCount = InTileCount;
-		HighRes->bIsUsingOverlappedTiles = bIsUsingOverlappedTiles;
 		HighRes->OverlapPercentage = PadRatioX;
 		HighRes->bWriteAllSamples = false;
 	}
@@ -85,13 +83,12 @@ TArray<UMoviePipelineMasterConfig*> FMovieRenderPipelineEditorModule::GenerateTe
 
 				const int32 NumSpatialSamples = TestNumSamples;
 				const int32 NumTiles = TileX;
-				static bool bIsUsingOverlappedTiles = true;
 
 				UMoviePipelineShotConfig* DefaultConfig = GenerateTestShotConfig(OutPipeline,
 					TemporalSampleCounts[TemporalSampleCountIndex],
 					ShutterAngles[ShutterAngleIndex],
 					ShutterTimings[ShutterTimingIndex],
-					NumTiles, NumSpatialSamples, bIsUsingOverlappedTiles, PadRatioX, PadRatioY, AccumulationGamma);
+					NumTiles, NumSpatialSamples, PadRatioX, PadRatioY, AccumulationGamma);
 
 				OutPipeline->DefaultShotConfig = DefaultConfig;
 
diff --git a/Engine/Plugins/MovieScene/MovieRenderPipeline/Source/MovieRenderPipelineRenderPasses/Private/MoviePipelineDeferredPasses.cpp b/Engine/Plugins/MovieScene/MovieRenderPipeline/Source/MovieRenderPipelineRenderPasses/Private/MoviePipelineDeferredPasses.cpp
--- a/Engine/Plugins/MovieScene/MovieRenderPipeline/Source/MovieRenderPipelineRenderPasses/Private/MoviePipelineDeferredPasses.cpp
+++ b/Engine/Plugins/MovieScene/MovieRenderPipeline/Source/MovieRenderPipelineRenderPasses/Private/MoviePipelineDeferredPasses.cpp
@@ -283,7 +283,7 @@ namespace MoviePipeline
 				);
 			}
 
-			if (InSampleState.bIsUsingOverlappedTiles)
+			// overlapped tile adjustment
 			{
 				float PadRatioX = 1.0f;
 				float PadRatioY = 1.0f;
