Fixed DOF scale in overlapped tile rendering.

diff --git a/Engine/Plugins/MovieScene/MovieRenderPipeline/Source/MovieRenderPipelineRenderPasses/Private/MoviePipelineDeferredPasses.cpp b/Engine/Plugins/MovieScene/MovieRenderPipeline/Source/MovieRenderPipelineRenderPasses/Private/MoviePipelineDeferredPasses.cpp
--- a/Engine/Plugins/MovieScene/MovieRenderPipeline/Source/MovieRenderPipelineRenderPasses/Private/MoviePipelineDeferredPasses.cpp
+++ b/Engine/Plugins/MovieScene/MovieRenderPipeline/Source/MovieRenderPipelineRenderPasses/Private/MoviePipelineDeferredPasses.cpp
@@ -235,6 +235,8 @@ namespace MoviePipeline
 			ViewFOV = GetPipeline()->GetWorld()->GetFirstPlayerController()->PlayerCameraManager->GetFOVAngle();
 		}
 
+		float DofSensorScale = 1.0f;
+
 		// Calculate a Projection Matrix
 		{
 			float XAxisMultiplier;
@@ -300,6 +302,8 @@ namespace MoviePipeline
 				BaseProjMatrix.M[0][0] /= ScaleX;
 				BaseProjMatrix.M[1][1] /= ScaleY;
 				
+				DofSensorScale = ScaleX;
+
 				// this offset would be correct with no pad
 				float OffsetX = -((float(InSampleState.TileIndexes.X) + 0.5f - float(InitSettings.TileCount.X) / 2.0f)* 2.0f);
 				float OffsetY = ((float(InSampleState.TileIndexes.Y) + 0.5f - float(InitSettings.TileCount.Y) / 2.0f)* 2.0f);
@@ -340,6 +344,8 @@ namespace MoviePipeline
 
 		}
 
+		View->FinalPostProcessSettings.DepthOfFieldSensorWidth *= DofSensorScale;
+
 		View->EndFinalPostprocessSettings(ViewInitOptions);
 
 		return View;
