CIS: Fixed dereferencing null pointer static analysis warnings. Fixed missing include guard in a header.

#rnx

diff --git a/Engine/Plugins/MovieScene/MovieRenderPipeline/Source/MovieRenderPipelineCore/Private/MoviePipelineBlueprintLibrary.cpp b/Engine/Plugins/MovieScene/MovieRenderPipeline/Source/MovieRenderPipelineCore/Private/MoviePipelineBlueprintLibrary.cpp
--- a/Engine/Plugins/MovieScene/MovieRenderPipeline/Source/MovieRenderPipelineCore/Private/MoviePipelineBlueprintLibrary.cpp
+++ b/Engine/Plugins/MovieScene/MovieRenderPipeline/Source/MovieRenderPipelineCore/Private/MoviePipelineBlueprintLibrary.cpp
@@ -77,7 +77,7 @@ UMovieSceneSequence* UMoviePipelineBlueprintLibrary::DuplicateSequence(UObject*
 	UMovieSceneSubTrack* SubTrack = DuplicatedSequence->GetMovieScene()->FindMasterTrack<UMovieSceneSubTrack>();
 	if (SubTrack)
 	{
-		AllSubSequenceSections.Append(ShotTrack->GetAllSections());
+		AllSubSequenceSections.Append(SubTrack->GetAllSections());
 	}
 
 	for (UMovieSceneSection* Section : AllSubSequenceSections)
diff --git a/Engine/Source/Editor/UnrealEd/Private/PlayLevel.cpp b/Engine/Source/Editor/UnrealEd/Private/PlayLevel.cpp
--- a/Engine/Source/Editor/UnrealEd/Private/PlayLevel.cpp
+++ b/Engine/Source/Editor/UnrealEd/Private/PlayLevel.cpp
@@ -579,7 +579,7 @@ void UEditorEngine::EndPlayMap()
 			if (WindowIndex == 0)
 			{
 				// Only update it if "Always center window to screen" is disabled
-				if (PlaySettingsConfig && !PlaySettingsConfig->CenterNewWindow)
+				if (!PlaySettingsConfig->CenterNewWindow)
 				{
 					PlaySettingsConfig->NewWindowPosition = PlaySettingsConfig->MultipleInstancePositions[WindowIndex];
 				}
diff --git a/Engine/Source/Editor/UnrealEd/Public/PlayInEditorDataTypes.h b/Engine/Source/Editor/UnrealEd/Public/PlayInEditorDataTypes.h
--- a/Engine/Source/Editor/UnrealEd/Public/PlayInEditorDataTypes.h
+++ b/Engine/Source/Editor/UnrealEd/Public/PlayInEditorDataTypes.h
@@ -77,7 +77,10 @@ struct FPieLoginStruct
 	int32 PIEInstanceIndex;
 
 	PRAGMA_DISABLE_DEPRECATION_WARNINGS
-	FPieLoginStruct& operator = (FPieLoginStruct const&) = default;
+	FPieLoginStruct& operator = (const FPieLoginStruct &) = default;
+	FPieLoginStruct& operator = (FPieLoginStruct &&) = default;
+	FPieLoginStruct(const FPieLoginStruct&) = default; // Copy
+	FPieLoginStruct(FPieLoginStruct&&) = default; // Move
 	PRAGMA_ENABLE_DEPRECATION_WARNINGS
 
 	FPieLoginStruct()
diff --git a/Engine/Source/Runtime/Core/Public/Math/sse_mathfun.h b/Engine/Source/Runtime/Core/Public/Math/sse_mathfun.h
--- a/Engine/Source/Runtime/Core/Public/Math/sse_mathfun.h
+++ b/Engine/Source/Runtime/Core/Public/Math/sse_mathfun.h
@@ -40,6 +40,7 @@
   2. It looks like some linked object files have conflicting names (log_ps, exp_ps, etc). So adding the 
      prefix SseMath_ to all the functions.
 */
+#pragma once
 
 #include <xmmintrin.h>
 
