Adding a parameter for the lens effects principal point to not be in the center of the screen. Needed for overlapped rendering.

diff --git a/Engine/Plugins/MovieScene/MovieRenderPipeline/Source/MovieRenderPipelineRenderPasses/Private/MoviePipelineDeferredPasses.cpp b/Engine/Plugins/MovieScene/MovieRenderPipeline/Source/MovieRenderPipelineRenderPasses/Private/MoviePipelineDeferredPasses.cpp
--- a/Engine/Plugins/MovieScene/MovieRenderPipeline/Source/MovieRenderPipelineRenderPasses/Private/MoviePipelineDeferredPasses.cpp
+++ b/Engine/Plugins/MovieScene/MovieRenderPipeline/Source/MovieRenderPipelineRenderPasses/Private/MoviePipelineDeferredPasses.cpp
@@ -344,8 +344,39 @@ namespace MoviePipeline
 
 		}
 
+		// Scaling sensor size inversely with the the projection matrix [0][0] should physically
+		// cause the circle of confusion to be unchanged.
 		View->FinalPostProcessSettings.DepthOfFieldSensorWidth *= DofSensorScale;
 
+		{
+			// We need our final view parameters to be in the space of [-1,1], including all the tiles.
+			// Starting with a single tile, the middle of the tile in offset screen space is:
+			FVector2D TilePrincipalPointOffset;
+
+			TilePrincipalPointOffset.X = (float(InSampleState.TileIndexes.X) + 0.5f - 0.5f * float(InSampleState.TileCounts.X)) * 2.0f;// / (.5f * float(InSampleState.TileCounts.Y));
+			TilePrincipalPointOffset.Y = (float(InSampleState.TileIndexes.Y) + 0.5f - 0.5f * float(InSampleState.TileCounts.Y)) * 2.0f;// / (.5f * float(InSampleState.TileCounts.X));
+
+			// For the tile size ratio, we have to multiply by (1.0 + overlap) and then divide by tile num
+			FVector2D OverlapScale;
+			OverlapScale.X = (1.0f + float(2 * InSampleState.OverlappedPad.X) / float(InSampleState.TileSize.X) );/// float(InSampleState.TileCounts.X);
+			OverlapScale.Y = (1.0f + float(2 * InSampleState.OverlappedPad.Y) / float(InSampleState.TileSize.Y) );/// float(InSampleState.TileCounts.Y);
+
+			TilePrincipalPointOffset.X /= OverlapScale.X;
+			TilePrincipalPointOffset.Y /= OverlapScale.Y;
+
+			FVector2D TilePrincipalPointScale;
+			TilePrincipalPointScale.X = OverlapScale.X / float(InSampleState.TileCounts.X);
+			TilePrincipalPointScale.Y = OverlapScale.Y / float(InSampleState.TileCounts.Y);
+
+			TilePrincipalPointOffset.X *= TilePrincipalPointScale.X;
+			TilePrincipalPointOffset.Y *= TilePrincipalPointScale.Y;
+
+			View->LensPrincipalPointOffsetScale = FVector4(TilePrincipalPointOffset.X, -TilePrincipalPointOffset.Y, TilePrincipalPointScale.X, TilePrincipalPointScale.Y);
+		}
+
+		//static FVector4 LensOffsetScale = FVector4(0, 0, 1, 1);
+		//View->LensPrincipalPointOffsetScale = LensOffsetScale;
+
 		View->EndFinalPostprocessSettings(ViewInitOptions);
 
 		return View;
diff --git a/Engine/Shaders/Private/PostProcessTonemap.usf b/Engine/Shaders/Private/PostProcessTonemap.usf
--- a/Engine/Shaders/Private/PostProcessTonemap.usf
+++ b/Engine/Shaders/Private/PostProcessTonemap.usf
@@ -68,12 +68,19 @@ float4 BloomDirtMaskTint;
 Texture2D BloomDirtMaskTexture;
 SamplerState BloomDirtMaskSampler;
 
+float4 LensPrincipalPointOffsetScale;
+
 half GrainFromUV(float2 GrainUV) 
 {
 	half Grain = frac(sin(GrainUV.x + GrainUV.y * 543.31) *  493013.0);
 	return Grain;
 }
 
+float2 ConvertScreenViewportSpaceToLensViewportSpace(float2 UV)
+{
+	return LensPrincipalPointOffsetScale.xy + UV * LensPrincipalPointOffsetScale.zw;
+}
+
 // LUT for color grading
 #if USE_VOLUME_LUT == 1
 	Texture3D		ColorGradingLUT;
@@ -175,7 +182,8 @@ void TonemapCommonVS(
 	float2 ColorViewportPos = UVToScreenPos(TexCoord, Color_Extent);
 
 	// Scale vignette to always be a circle with consistent corner intensity.
-	OutExposureScaleVignette.yz = VignetteSpace(ColorViewportPos, AspectRatio);
+	float2 LensViewportPos = ConvertScreenViewportSpaceToLensViewportSpace(ColorViewportPos);
+	OutExposureScaleVignette.yz = VignetteSpace(LensViewportPos, AspectRatio);
 
 	// Grain
 	OutGrainUV.xy = TexCoord + Color_ExtentInverse * float2(-0.5,0.5);
diff --git a/Engine/Source/Runtime/Engine/Private/SceneView.cpp b/Engine/Source/Runtime/Engine/Private/SceneView.cpp
--- a/Engine/Source/Runtime/Engine/Private/SceneView.cpp
+++ b/Engine/Source/Runtime/Engine/Private/SceneView.cpp
@@ -657,6 +657,7 @@ FSceneView::FSceneView(const FSceneViewInitOptions& InitOptions)
 	, bIsMobileMultiViewDirectEnabled(false)
 	, bShouldBindInstancedViewUB(false)
 	, GlobalClippingPlane(FPlane(0, 0, 0, 0))
+	, LensPrincipalPointOffsetScale(0.0f, 0.0f, 1.0f, 1.0f)
 #if WITH_EDITOR
 	, OverrideLODViewOrigin(InitOptions.OverrideLODViewOrigin)
 	, bAllowTranslucentPrimitivesInHitProxy( true )
diff --git a/Engine/Source/Runtime/Engine/Public/SceneView.h b/Engine/Source/Runtime/Engine/Public/SceneView.h
--- a/Engine/Source/Runtime/Engine/Public/SceneView.h
+++ b/Engine/Source/Runtime/Engine/Public/SceneView.h
@@ -1035,6 +1035,13 @@ public:
 	/** The frame index to override, useful for keeping determinism when rendering sequences. **/
 	TOptional<uint32> OverrideFrameIndexValue;
 
+	/** In some cases, the principal point of the lens is not at the center of the screen, especially for overlapped tile
+	 *  rendering. So given a UV in [-1,1] viewport space, convert it to the [-1,1] viewport space of the lens using
+	 *  LensUV = LensPrincipalPointOffsetScale.xy ScreenUV * LensPrincipalPointOffsetScale.zw;
+	 *  This value is FVector4(0,0,1,1) unless overridden.
+	 */
+	FVector4 LensPrincipalPointOffsetScale;
+
 #if WITH_EDITOR
 	/** The set of (the first 64) groups' visibility info for this view */
 	uint64 EditorViewBitflag;
diff --git a/Engine/Source/Runtime/Renderer/Private/PostProcess/PostProcessTonemap.cpp b/Engine/Source/Runtime/Renderer/Private/PostProcess/PostProcessTonemap.cpp
--- a/Engine/Source/Runtime/Renderer/Private/PostProcess/PostProcessTonemap.cpp
+++ b/Engine/Source/Runtime/Renderer/Private/PostProcess/PostProcessTonemap.cpp
@@ -533,6 +533,7 @@ BEGIN_SHADER_PARAMETER_STRUCT(FTonemapParameters, )
 	SHADER_PARAMETER(FVector4, BloomDirtMaskTint)
 	SHADER_PARAMETER(FVector4, ChromaticAberrationParams)
 	SHADER_PARAMETER(FVector4, TonemapperParams)
+	SHADER_PARAMETER(FVector4, LensPrincipalPointOffsetScale)
 	SHADER_PARAMETER(float, SwitchVerticalAxis)
 	SHADER_PARAMETER(float, DefaultEyeExposure)
 	SHADER_PARAMETER(float, EditorNITLevel)
@@ -777,6 +778,7 @@ FScreenPassTexture AddTonemapPass(FRDGBuilder& GraphBuilder, const FViewInfo& Vi
 	CommonParameters.DefaultEyeExposure = DefaultEyeExposure;
 	CommonParameters.EditorNITLevel = EditorNITLevel;
 	CommonParameters.bOutputInHDR = ViewFamily.bIsHDR;
+	CommonParameters.LensPrincipalPointOffsetScale = View.LensPrincipalPointOffsetScale;
 
 	// Generate permutation vector for the desktop tonemapper.
 	TonemapperPermutation::FDesktopDomain DesktopPermutationVector;
