Minor style  fixes to make overlapped rendering principal point calculation more clear.

diff --git a/Engine/Plugins/MovieScene/MovieRenderPipeline/Source/MovieRenderPipelineRenderPasses/Private/MoviePipelineDeferredPasses.cpp b/Engine/Plugins/MovieScene/MovieRenderPipeline/Source/MovieRenderPipelineRenderPasses/Private/MoviePipelineDeferredPasses.cpp
--- a/Engine/Plugins/MovieScene/MovieRenderPipeline/Source/MovieRenderPipelineRenderPasses/Private/MoviePipelineDeferredPasses.cpp
+++ b/Engine/Plugins/MovieScene/MovieRenderPipeline/Source/MovieRenderPipelineRenderPasses/Private/MoviePipelineDeferredPasses.cpp
@@ -353,8 +353,8 @@ namespace MoviePipeline
 			// Starting with a single tile, the middle of the tile in offset screen space is:
 			FVector2D TilePrincipalPointOffset;
 
-			TilePrincipalPointOffset.X = (float(InSampleState.TileIndexes.X) + 0.5f - 0.5f * float(InSampleState.TileCounts.X)) * 2.0f;// / (.5f * float(InSampleState.TileCounts.Y));
-			TilePrincipalPointOffset.Y = (float(InSampleState.TileIndexes.Y) + 0.5f - 0.5f * float(InSampleState.TileCounts.Y)) * 2.0f;// / (.5f * float(InSampleState.TileCounts.X));
+			TilePrincipalPointOffset.X = (float(InSampleState.TileIndexes.X) + 0.5f - (0.5f * float(InSampleState.TileCounts.X))) * 2.0f;
+			TilePrincipalPointOffset.Y = (float(InSampleState.TileIndexes.Y) + 0.5f - (0.5f * float(InSampleState.TileCounts.Y))) * 2.0f;
 
 			// For the tile size ratio, we have to multiply by (1.0 + overlap) and then divide by tile num
 			FVector2D OverlapScale;
