Minor comment fix in overlapped rendering.

diff --git a/Engine/Plugins/MovieScene/MovieRenderPipeline/Source/MovieRenderPipelineRenderPasses/Private/MoviePipelineDeferredPasses.cpp b/Engine/Plugins/MovieScene/MovieRenderPipeline/Source/MovieRenderPipelineRenderPasses/Private/MoviePipelineDeferredPasses.cpp
--- a/Engine/Plugins/MovieScene/MovieRenderPipeline/Source/MovieRenderPipelineRenderPasses/Private/MoviePipelineDeferredPasses.cpp
+++ b/Engine/Plugins/MovieScene/MovieRenderPipeline/Source/MovieRenderPipelineRenderPasses/Private/MoviePipelineDeferredPasses.cpp
@@ -358,8 +358,8 @@ namespace MoviePipeline
 
 			// For the tile size ratio, we have to multiply by (1.0 + overlap) and then divide by tile num
 			FVector2D OverlapScale;
-			OverlapScale.X = (1.0f + float(2 * InSampleState.OverlappedPad.X) / float(InSampleState.TileSize.X) );/// float(InSampleState.TileCounts.X);
-			OverlapScale.Y = (1.0f + float(2 * InSampleState.OverlappedPad.Y) / float(InSampleState.TileSize.Y) );/// float(InSampleState.TileCounts.Y);
+			OverlapScale.X = (1.0f + float(2 * InSampleState.OverlappedPad.X) / float(InSampleState.TileSize.X) );
+			OverlapScale.Y = (1.0f + float(2 * InSampleState.OverlappedPad.Y) / float(InSampleState.TileSize.Y) );
 
 			TilePrincipalPointOffset.X /= OverlapScale.X;
 			TilePrincipalPointOffset.Y /= OverlapScale.Y;
