Removing unused comment.

diff --git a/Engine/Plugins/MovieScene/MovieRenderPipeline/Source/MovieRenderPipelineRenderPasses/Private/MoviePipelineDeferredPasses.cpp b/Engine/Plugins/MovieScene/MovieRenderPipeline/Source/MovieRenderPipelineRenderPasses/Private/MoviePipelineDeferredPasses.cpp
--- a/Engine/Plugins/MovieScene/MovieRenderPipeline/Source/MovieRenderPipelineRenderPasses/Private/MoviePipelineDeferredPasses.cpp
+++ b/Engine/Plugins/MovieScene/MovieRenderPipeline/Source/MovieRenderPipelineRenderPasses/Private/MoviePipelineDeferredPasses.cpp
@@ -374,9 +374,6 @@ namespace MoviePipeline
 			View->LensPrincipalPointOffsetScale = FVector4(TilePrincipalPointOffset.X, -TilePrincipalPointOffset.Y, TilePrincipalPointScale.X, TilePrincipalPointScale.Y);
 		}
 
-		//static FVector4 LensOffsetScale = FVector4(0, 0, 1, 1);
-		//View->LensPrincipalPointOffsetScale = LensOffsetScale;
-
 		View->EndFinalPostprocessSettings(ViewInitOptions);
 
 		return View;
