MoviePipeline: Added support for writing out multiple image types at the same time, such as EXR and PNG. For incompatible bit depths they are simply clamped during conversion instead of rescaled.

#rnx

diff --git a/Engine/Plugins/MovieScene/MovieRenderPipeline/Source/MovieRenderPipelineCore/Private/MoviePipelineRendering.cpp b/Engine/Plugins/MovieScene/MovieRenderPipeline/Source/MovieRenderPipelineCore/Private/MoviePipelineRendering.cpp
--- a/Engine/Plugins/MovieScene/MovieRenderPipeline/Source/MovieRenderPipelineCore/Private/MoviePipelineRendering.cpp
+++ b/Engine/Plugins/MovieScene/MovieRenderPipeline/Source/MovieRenderPipelineCore/Private/MoviePipelineRendering.cpp
@@ -124,6 +124,11 @@ void UMoviePipeline::TeardownRenderingPipelineForShot(FMoviePipelineShotInfo& Sh
 	{
 		Input->Teardown();
 	}
+
+	for (TSharedPtr<MoviePipeline::FMoviePipelineEnginePass> EnginePass : ActiveRenderPasses)
+	{
+		EnginePass->Teardown();
+	}
 }
 
 void UMoviePipeline::RenderFrame()
@@ -405,7 +410,7 @@ void UMoviePipeline::OnSampleRendered(TUniquePtr<FImagePixelData>&& OutputSample
 	}
 
 	// JPEG output
-	TileImageTask->Format = EImageFormat::JPEG;
+	TileImageTask->Format = EImageFormat::EXR;
 	TileImageTask->CompressionQuality = 100;
 
 	FString OutputName = FString::Printf(TEXT("/%s_SS_%d_TS_%d_TileX_%d_TileY_%d.%d.jpeg"),
diff --git a/Engine/Plugins/MovieScene/MovieRenderPipeline/Source/MovieRenderPipelineEditor/Private/MovieRenderPipelineEditorTest.cpp b/Engine/Plugins/MovieScene/MovieRenderPipeline/Source/MovieRenderPipelineEditor/Private/MovieRenderPipelineEditorTest.cpp
--- a/Engine/Plugins/MovieScene/MovieRenderPipeline/Source/MovieRenderPipelineEditor/Private/MovieRenderPipelineEditorTest.cpp
+++ b/Engine/Plugins/MovieScene/MovieRenderPipeline/Source/MovieRenderPipelineEditor/Private/MovieRenderPipelineEditorTest.cpp
@@ -3,7 +3,7 @@
 #include "MoviePipelineShotConfig.h"
 #include "MoviePipelineGameOverrideSetting.h"
 #include "MoviePipelineAccumulationSetting.h"
-#include "MoviePipelineImageSequenceContainer.h"
+#include "MoviePipelineImageSequenceOutput.h"
 #include "MoviePipelineAudioOutput.h"
 #include "Misc/Paths.h"
 #include "MoviePipelineMasterConfig.h"
@@ -29,8 +29,8 @@ UMoviePipelineShotConfig* GenerateTestShotConfig(UObject* InOwner, int32 InSampl
 		CameraSetting->TemporalSampleCount = InSampleCount;
 		CameraSetting->CameraShutterAngle = InShutterAngle;
 		CameraSetting->ShutterTiming = InFrameTiming;
-		CameraSetting->bManualExposure = true;
-		CameraSetting->ExposureCompensation = 11;
+		CameraSetting->bManualExposure = false;
+		CameraSetting->ExposureCompensation = 10;
 	}
 
 	{
@@ -71,8 +71,8 @@ TArray<UMoviePipelineMasterConfig*> FMovieRenderPipelineEditorModule::GenerateTe
 
 				int32 SizeX = 1920;
 				int32 SizeY = 1080;
-				int32 TileX = 2;
-				int32 TileY = 2;
+				int32 TileX = 1;
+				int32 TileY = 1;
 				int32 TestNumSamples = 2;
 				float PadRatioX = 0.5f;
 				float PadRatioY = 0.5f;
@@ -93,7 +93,8 @@ TArray<UMoviePipelineMasterConfig*> FMovieRenderPipelineEditorModule::GenerateTe
 				OutPipeline->DefaultShotConfig = DefaultConfig;
 
 				// Add some Outputs
-				OutPipeline->FindOrAddSetting<UMoviePipelineImageSequenceContainerBase>();
+				OutPipeline->FindOrAddSetting<UMoviePipelineImageSequenceOutput_EXR>();
+				OutPipeline->FindOrAddSetting<UMoviePipelineImageSequenceOutput_PNG>();
 
 				UMoviePipelineGameOverrideSetting* GamemodeOverride = OutPipeline->FindOrAddSetting<UMoviePipelineGameOverrideSetting>();
 				GamemodeOverride->GameModeOverride = AGameModeBase::StaticClass();
diff --git a/Engine/Plugins/MovieScene/MovieRenderPipeline/Source/MovieRenderPipelineRenderPasses/Private/MoviePipelineDeferredPasses.cpp b/Engine/Plugins/MovieScene/MovieRenderPipeline/Source/MovieRenderPipelineRenderPasses/Private/MoviePipelineDeferredPasses.cpp
--- a/Engine/Plugins/MovieScene/MovieRenderPipeline/Source/MovieRenderPipelineRenderPasses/Private/MoviePipelineDeferredPasses.cpp
+++ b/Engine/Plugins/MovieScene/MovieRenderPipeline/Source/MovieRenderPipelineRenderPasses/Private/MoviePipelineDeferredPasses.cpp
@@ -41,10 +41,17 @@ namespace MoviePipeline
 		TileRenderTarget->AddToRoot();
 
 		// Initialize to the tile size (not final size) and use a 16 bit back buffer to avoid precision issues when accumulating later
-		TileRenderTarget->InitCustomFormat(InInitSettings.BackbufferResolution.X, InInitSettings.BackbufferResolution.Y, EPixelFormat::PF_A16B16G16R16, false);
+		TileRenderTarget->InitCustomFormat(InInitSettings.BackbufferResolution.X, InInitSettings.BackbufferResolution.Y, EPixelFormat::PF_FloatRGBA, false);
 
 		// Allocate 
 		ViewState.Allocate();
+
+		// Override us to use linear color output.
+		IConsoleVariable* OutputDeviceCVar = IConsoleManager::Get().FindConsoleVariable(TEXT("r.HDR.Display.OutputDevice"));
+		check(OutputDeviceCVar);
+
+		PreviousOutputDeviceIdx = OutputDeviceCVar->GetInt();
+		OutputDeviceCVar->Set(8 /*Linear No Tone Curve*/, EConsoleVariableFlags::ECVF_SetByConsole);
 	}
 
 	void FDeferredRenderEnginePass::Teardown()
@@ -57,6 +64,12 @@ namespace MoviePipeline
 		}
 
 		ViewState.Destroy();
+
+		// Restore the previous output device so that we don't leak changes into the editor.
+		IConsoleVariable* OutputDeviceCVar = IConsoleManager::Get().FindConsoleVariable(TEXT("r.HDR.Display.OutputDevice"));
+		check(OutputDeviceCVar);
+
+		OutputDeviceCVar->Set(PreviousOutputDeviceIdx, EConsoleVariableFlags::ECVF_SetByConsole);
 	}
 
 	void FDeferredRenderEnginePass::RenderSample_GameThread(const FMoviePipelineRenderPassMetrics& InSampleState)
@@ -84,7 +97,7 @@ namespace MoviePipeline
 			.SetWorldTimes(TimeSeconds, DeltaTimeSeconds, RealTimeSeconds)
 			.SetRealtimeUpdate(true));
 
-		ViewFamily.SceneCaptureSource = ESceneCaptureSource::SCS_FinalColorLDR;
+		ViewFamily.SceneCaptureSource = ESceneCaptureSource::SCS_FinalColorHDR;
 		ViewFamily.SetScreenPercentageInterface(new FLegacyScreenPercentageDriver(ViewFamily, 1.0f, false));
 
 		// View is added as a child of the ViewFamily
@@ -106,7 +119,7 @@ namespace MoviePipeline
 		UMoviePipelineCameraSetting* CameraSettings = GetPipeline()->GetPipelineCurrentShotConfig()->FindSetting<UMoviePipelineCameraSetting>();
 		UMoviePipelineHighResSetting* HighResSettings = GetPipeline()->GetPipelineCurrentShotConfig()->FindSetting<UMoviePipelineHighResSetting>();
 
-		if(CameraSettings)
+		if (CameraSettings)
 		{
 			if (CameraSettings->bManualExposure)
 			{
@@ -185,9 +198,9 @@ namespace MoviePipeline
 			{
 				UE_LOG(LogMovieRenderPipeline, Log, TEXT("Frame: %d TemporalSample: %d/%d Rendered Tile: %d/%d Sample: %d/%d"),
 					InSampleState.OutputState.OutputFrameNumber, InSampleState.TemporalSampleIndex + 1, InSampleState.TemporalSampleCount,
-					InSampleState.GetTileIndex() + 1 , InSampleState.GetTileCount(), InSampleState.SpatialSampleIndex + 1, InSampleState.SpatialSampleCount);
+					InSampleState.GetTileIndex() + 1, InSampleState.GetTileCount(), InSampleState.SpatialSampleIndex + 1, InSampleState.SpatialSampleCount);
 			}
-			
+
 			ReadBackbufferAndBroadcast_RenderThread(RHICmdList, BackbufferRenderTarget, InSampleState, OnBackbufferReadDelegate);
 		});
 	}
@@ -199,11 +212,11 @@ namespace MoviePipeline
 		FIntRect SourceRect = FIntRect(0, 0, InFromRenderTarget->GetSizeXY().X, InFromRenderTarget->GetSizeXY().Y);
 
 		// Read the data back to the CPU
-		TArray<FColor> RawPixels;
+		TArray<FLinearColor> RawPixels;
 		RawPixels.SetNum(SourceRect.Width() * SourceRect.Height());
 
-		FReadSurfaceDataFlags ReadDataFlags;
-		ReadDataFlags.SetLinearToGamma(true);
+		FReadSurfaceDataFlags ReadDataFlags(ERangeCompressionMode::RCM_MinMax);
+		ReadDataFlags.SetLinearToGamma(false);
 
 		RHICmdList.ReadSurfaceData(InFromRenderTarget->GetRenderTargetTexture(), SourceRect, RawPixels, ReadDataFlags);
 
@@ -289,25 +302,24 @@ namespace MoviePipeline
 			{
 				float PadRatioX = 1.0f;
 				float PadRatioY = 1.0f;
-				
+
 				if (InSampleState.OverlappedPad.X > 0 && InSampleState.OverlappedPad.Y > 0)
 				{
 					PadRatioX = float(InSampleState.OverlappedPad.X * 2 + InSampleState.TileSize.X) / float(InSampleState.TileSize.X);
 					PadRatioY = float(InSampleState.OverlappedPad.Y * 2 + InSampleState.TileSize.Y) / float(InSampleState.TileSize.Y);
 				}
-				
+
 				float ScaleX = PadRatioX / float(InitSettings.TileCount.X);
 				float ScaleY = PadRatioY / float(InitSettings.TileCount.Y);
-				
+
 				BaseProjMatrix.M[0][0] /= ScaleX;
 				BaseProjMatrix.M[1][1] /= ScaleY;
-				
 				DofSensorScale = ScaleX;
 
 				// this offset would be correct with no pad
 				float OffsetX = -((float(InSampleState.TileIndexes.X) + 0.5f - float(InitSettings.TileCount.X) / 2.0f)* 2.0f);
 				float OffsetY = ((float(InSampleState.TileIndexes.Y) + 0.5f - float(InitSettings.TileCount.Y) / 2.0f)* 2.0f);
-				
+
 				BaseProjMatrix.M[2][0] += OffsetX / PadRatioX;
 				BaseProjMatrix.M[2][1] += OffsetY / PadRatioX;
 			}
@@ -323,7 +335,7 @@ namespace MoviePipeline
 		View->ViewLocation = InSampleState.FrameInfo.CurrViewLocation;
 		View->ViewRotation = InSampleState.FrameInfo.CurrViewRotation;
 		View->PreviousViewTransform = FTransform(InSampleState.FrameInfo.PrevViewRotation, InSampleState.FrameInfo.PrevViewLocation);
-		
+
 		View->StartFinalPostprocessSettings(View->ViewLocation);
 		View->OverrideFrameIndexValue = InSampleState.SpatialSampleIndex;
 
@@ -428,7 +440,7 @@ void UMoviePipelineDeferredPassBase::GatherOutputPassesImpl(TArray<FMoviePipelin
 	}
 }
 
-void UMoviePipelineDeferredPassBase::OnBackbufferSampleReady(TArray<FColor> InPixelData, FMoviePipelineRenderPassMetrics InSampleState)
+void UMoviePipelineDeferredPassBase::OnBackbufferSampleReady(TArray<FLinearColor> InPixelData, FMoviePipelineRenderPassMetrics InSampleState)
 {
 	MoviePipeline::FSampleRenderThreadParams AccumulationParams;
 	{
@@ -442,7 +454,7 @@ void UMoviePipelineDeferredPassBase::OnBackbufferSampleReady(TArray<FColor> InPi
 	FrameData->PassIdentifier = FMoviePipelinePassIdentifier(TEXT("Backbuffer"));
 	FrameData->SampleState = InSampleState;
 
-	TUniquePtr<TImagePixelData<FColor>> PixelData = MakeUnique<TImagePixelData<FColor>>(InSampleState.BackbufferSize, TArray64<FColor>(MoveTemp(InPixelData)), FrameData);
+	TUniquePtr<TImagePixelData<FLinearColor>> PixelData = MakeUnique<TImagePixelData<FLinearColor>>(InSampleState.BackbufferSize, TArray64<FLinearColor>(MoveTemp(InPixelData)), FrameData);
 
 	AccumulateSample_RenderThread(MoveTemp(PixelData), FrameData, AccumulationParams);
 }
@@ -468,7 +480,7 @@ void UMoviePipelineDeferredPassBase::OnSetupView(FSceneViewFamily& InViewFamily,
 	} Iterator(InView.FinalPostProcessSettings, DesiredOutputPasses);
 	GetBufferVisualizationData().IterateOverAvailableMaterials(Iterator);
 
-	for(int32 PassIndex = 0; PassIndex < InView.FinalPostProcessSettings.BufferVisualizationOverviewMaterials.Num(); PassIndex++)
+	for (int32 PassIndex = 0; PassIndex < InView.FinalPostProcessSettings.BufferVisualizationOverviewMaterials.Num(); PassIndex++)
 	{
 		const UMaterialInterface* Pass = InView.FinalPostProcessSettings.BufferVisualizationOverviewMaterials[PassIndex];
 		UE_LOG(LogTemp, Log, TEXT("Pass: %s"), *Pass->GetName());
@@ -497,7 +509,7 @@ void UMoviePipelineDeferredPassBase::OnSetupView(FSceneViewFamily& InViewFamily,
 			{
 				return;
 			}
-			
+
 			TSharedRef<FImagePixelDataPayload, ESPMode::ThreadSafe> FrameData = MakeShared<FImagePixelDataPayload, ESPMode::ThreadSafe>();
 			FrameData->OutputState = AccumulationParams.SampleState.OutputState;
 			FrameData->PassIdentifier = FMoviePipelinePassIdentifier(PassName.ToString());
@@ -572,7 +584,7 @@ namespace MoviePipeline
 
 			// Now that a tile is fully built and accumulated we can notify the output builder that the
 			// data is ready so it can pass that onto the output containers (if needed).
-			/*if (InPixelData->GetType() == EImagePixelType::Float32)
+			if (InPixelData->GetType() == EImagePixelType::Float32)
 			{
 				// 32 bit FLinearColor
 				TUniquePtr<TImagePixelData<FLinearColor> > FinalPixelData = MakeUnique<TImagePixelData<FLinearColor>>(FIntPoint(FullSizeX, FullSizeY));
@@ -586,11 +598,11 @@ namespace MoviePipeline
 				// 32 bit FLinearColor
 				TUniquePtr<TImagePixelData<FFloat16Color> > FinalPixelData = MakeUnique<TImagePixelData<FFloat16Color>>(FIntPoint(FullSizeX, FullSizeY));
 				InParams.ImageAccumulator->FetchFinalPixelDataHalfFloat(FinalPixelData->Pixels);
-			
+
 				// Send the data to the Output Builder
 				InParams.OutputMerger->OnCompleteRenderPassDataAvailable_AnyThread(MoveTemp(FinalPixelData), InFrameData);
 			}
-			else*/ if(true)
+			else if (InPixelData->GetType() == EImagePixelType::Color)
 			{
 				// 8bit FColors
 				TUniquePtr<TImagePixelData<FColor>> FinalPixelData = MakeUnique<TImagePixelData<FColor>>(FIntPoint(FullSizeX, FullSizeY));
diff --git a/Engine/Plugins/MovieScene/MovieRenderPipeline/Source/MovieRenderPipelineRenderPasses/Private/MoviePipelineImageSequenceContainer.cpp b/Engine/Plugins/MovieScene/MovieRenderPipeline/Source/MovieRenderPipelineRenderPasses/Private/MoviePipelineImageSequenceContainer.cpp
deleted file mode 100644
diff --git a/Engine/Plugins/MovieScene/MovieRenderPipeline/Source/MovieRenderPipelineRenderPasses/Private/MoviePipelineImageSequenceOutput.cpp b/Engine/Plugins/MovieScene/MovieRenderPipeline/Source/MovieRenderPipelineRenderPasses/Private/MoviePipelineImageSequenceOutput.cpp
new file mode 100644
--- /dev/null
+++ b/Engine/Plugins/MovieScene/MovieRenderPipeline/Source/MovieRenderPipelineRenderPasses/Private/MoviePipelineImageSequenceOutput.cpp
@@ -0,0 +1,206 @@
+// Copyright 1998-2019 Epic Games, Inc. All Rights Reserved.
+
+#include "MoviePipelineImageSequenceOutput.h"
+#include "ImageWriteTask.h"
+#include "ImagePixelData.h"
+#include "Modules/ModuleManager.h"
+#include "ImageWriteQueue.h"
+#include "MoviePipeline.h"
+#include "ImageWriteStream.h"
+#include "MoviePipelineMasterConfig.h"
+#include "MovieRenderTileImage.h"
+#include "MovieRenderOverlappedImage.h"
+#include "MovieRenderPipelineCoreModule.h"
+#include "Misc/FrameRate.h"
+#include "MoviePipelineOutputSetting.h"
+
+
+UMoviePipelineImageSequenceOutputBase::UMoviePipelineImageSequenceOutputBase()
+{
+	// ToDo: Move this to a specific function to get out of CDO land.
+	ImageWriteQueue = &FModuleManager::Get().LoadModuleChecked<IImageWriteQueueModule>("ImageWriteQueue").GetWriteQueue();
+}
+
+void UMoviePipelineImageSequenceOutputBase::BeginFinalizeImpl()
+{
+	FinalizeFence = ImageWriteQueue->CreateFence();
+}
+
+bool UMoviePipelineImageSequenceOutputBase::HasFinishedProcessingImpl()
+{ 
+	// Wait until the finalization fence is reached meaning we've written everything to disk.
+	return Super::HasFinishedProcessingImpl() && (!FinalizeFence.IsValid() || FinalizeFence.WaitFor(0));
+}
+
+void UMoviePipelineImageSequenceOutputBase::SetupForPipelineImpl(UMoviePipeline* InPipeline)
+{
+	UMoviePipelineOutputSetting* OutputSettings = GetPipeline()->GetPipelineMasterConfig()->FindSetting<UMoviePipelineOutputSetting>();
+	check(OutputSettings);
+
+	OutputDirectory = OutputSettings->OutputDirectory.Path;
+}
+
+void UMoviePipelineImageSequenceOutputBase::OnRecieveImageDataImpl(FMoviePipelineMergerOutputFrame* InMergedOutputFrame)
+{
+	check(InMergedOutputFrame);
+
+	const TCHAR* Extension = TEXT("");
+	switch (OutputFormat)
+	{
+	case EImageFormat::PNG: Extension = TEXT("png"); break;
+	case EImageFormat::JPEG: Extension = TEXT("jpeg"); break;
+	case EImageFormat::BMP: Extension = TEXT("bmp"); break;
+	case EImageFormat::EXR: Extension = TEXT("exr"); break;
+	}
+
+
+	for (TPair<FMoviePipelinePassIdentifier, TUniquePtr<FImagePixelData>>& RenderPassData : InMergedOutputFrame->ImageOutputData)
+	{
+		TUniquePtr<FImagePixelData> QuantizedPixelData = nullptr;
+		FIntPoint RawSize = RenderPassData.Value->GetSize();
+		int32 RawNumChannels = RenderPassData.Value->GetNumChannels();
+
+		
+		switch (OutputFormat)
+		{
+		case EImageFormat::PNG:
+		case EImageFormat::JPEG:
+		case EImageFormat::BMP:
+		{
+			// All three of these formats only support 8 bit data, so we need to take the incoming buffer type,
+			// copy it into a new 8-bit array and optionally apply a little noise to the data to help hide gradient banding.
+			switch (RenderPassData.Value->GetBitDepth())
+			{
+				case 8:
+				{
+					// No work actually needs to be done, hooray! We'll copy the data though so that when it gets moved
+					// into the ImageWriteTask the original data can be passed to another output container.
+					QuantizedPixelData = RenderPassData.Value->CopyImageData();
+					break;
+				}
+				case 16:
+				{
+					TArray<FColor> ClampedPixels;
+					ClampedPixels.SetNum(RenderPassData.Value->GetSize().X * RenderPassData.Value->GetSize().Y);
+
+					int64 SizeInBytes = 0;
+					const void* SrcRawDataPtr = nullptr;
+					RenderPassData.Value->GetRawData(SrcRawDataPtr, SizeInBytes);
+
+					const uint16* RawDataPtr = static_cast<const uint16*>(SrcRawDataPtr);
+
+					// Copy pixels to new array
+					for (int32 Y = 0; Y < RawSize.Y; Y++)
+					{
+						for (int32 X = 0; X < RawSize.X; X++)
+						{
+							FColor* DestColor = &ClampedPixels.GetData()[Y*RawSize.X + X];
+							for (int32 ChanIter = 0; ChanIter < 4; ChanIter++)
+							{
+								FFloat16 RawValue;
+								RawValue.Encoded = RawDataPtr[(Y*RawSize.X + X)*RawNumChannels + ChanIter];
+
+								// c cast does the conversion from fp16 bits to float
+								uint8 Value = FMath::Clamp(FMath::FloorToInt(float(RawValue)/255), 0, 1);
+								switch (ChanIter)
+								{
+								case 0: DestColor->R = Value; break;
+								case 1: DestColor->G = Value; break;
+								case 2: DestColor->B = Value; break;
+								case 3: DestColor->A = Value; break;
+								}
+							}
+						}
+					}
+
+					QuantizedPixelData = MakeUnique<TImagePixelData<FColor>>(RenderPassData.Value->GetSize(), TArray64<FColor>(MoveTemp(ClampedPixels)));
+				}
+			case 32:
+			{
+				TArray<FColor> ClampedPixels;
+				ClampedPixels.SetNum(RenderPassData.Value->GetSize().X * RenderPassData.Value->GetSize().Y);
+
+				int64 SizeInBytes = 0;
+				const void* SrcRawDataPtr = nullptr;
+				RenderPassData.Value->GetRawData(SrcRawDataPtr, SizeInBytes);
+
+				const float* RawDataPtr = static_cast<const float*>(SrcRawDataPtr);
+
+				// Copy pixels to new array
+				for (int32 Y = 0; Y < RawSize.Y; Y++)
+				{
+					for (int32 X = 0; X < RawSize.X; X++)
+					{
+						FColor* DestColor = &ClampedPixels.GetData()[Y*RawSize.X + X];
+						for (int32 ChanIter = 0; ChanIter < 4; ChanIter++)
+						{
+							float RawValue = RawDataPtr[(Y*RawSize.X + X)*RawNumChannels + ChanIter];
+							RawValue = FMath::Pow(RawValue, 1/2.2f);
+
+							// c cast does the conversion from fp16 bits to float
+							uint8 Value = FMath::Clamp(FMath::RoundToInt(float(RawValue) * 255), 0, 255);
+							switch (ChanIter)
+							{
+							case 0: DestColor->R = Value; break;
+							case 1: DestColor->G = Value; break;
+							case 2: DestColor->B = Value; break;
+							case 3: DestColor->A = Value; break;
+							}
+						}
+					}
+				}
+
+				QuantizedPixelData = MakeUnique<TImagePixelData<FColor>>(RenderPassData.Value->GetSize(), TArray64<FColor>(MoveTemp(ClampedPixels)));
+				break;
+			}
+
+			default:
+				check(false);
+			}
+
+			break;
+		}
+		case EImageFormat::EXR:
+			QuantizedPixelData = RenderPassData.Value->CopyImageData();
+			break;
+		default:
+			check(false);
+		}
+
+		TUniquePtr<FImageWriteTask> TileImageTask = MakeUnique<FImageWriteTask>();
+
+		// Fill alpha for now 
+		// switch (RenderPassData.Value->GetType())
+		// {
+		// case EImagePixelType::Color:
+		// {
+		// 	TileImageTask->PixelPreProcessors.Add(TAsyncAlphaWrite<FColor>(255));
+		// 	break;
+		// }
+		// case EImagePixelType::Float16:
+		// {
+		// 
+		// 	break;
+		// }
+		// case EImagePixelType::Float32:
+		// {
+		// 	TileImageTask->PixelPreProcessors.Add(TAsyncAlphaWrite<FLinearColor>(255));
+		// 	break;
+		// }
+		// default:
+		// 	check(false);
+		// }
+
+		TileImageTask->Format = OutputFormat;
+		TileImageTask->CompressionQuality = 100;
+
+		FString OutputName = FString::Printf(TEXT("/%s.%d.%s"),
+			*RenderPassData.Key.Name, InMergedOutputFrame->FrameOutputState.OutputFrameNumber, Extension);
+
+		FString OutputPath = OutputDirectory + OutputName;
+		TileImageTask->Filename = OutputPath;
+
+		TileImageTask->PixelData = MoveTemp(QuantizedPixelData);
+		ImageWriteQueue->Enqueue(MoveTemp(TileImageTask));
+	}
+}
diff --git a/Engine/Plugins/MovieScene/MovieRenderPipeline/Source/MovieRenderPipelineRenderPasses/Public/MoviePipelineDeferredPasses.h b/Engine/Plugins/MovieScene/MovieRenderPipeline/Source/MovieRenderPipelineRenderPasses/Public/MoviePipelineDeferredPasses.h
--- a/Engine/Plugins/MovieScene/MovieRenderPipeline/Source/MovieRenderPipelineRenderPasses/Public/MoviePipelineDeferredPasses.h
+++ b/Engine/Plugins/MovieScene/MovieRenderPipeline/Source/MovieRenderPipelineRenderPasses/Public/MoviePipelineDeferredPasses.h
@@ -25,7 +25,7 @@ protected:
 	virtual void SetupImpl(TArray<TSharedPtr<MoviePipeline::FMoviePipelineEnginePass>>& InEnginePasses) override;
 	// ~UMovieRenderPassAPI
 
-	void OnBackbufferSampleReady(TArray<FColor> InPixelData, FMoviePipelineRenderPassMetrics InSampleState);
+	void OnBackbufferSampleReady(TArray<FLinearColor> InPixelData, FMoviePipelineRenderPassMetrics InSampleState);
 	void OnSetupView(FSceneViewFamily& InViewFamily, FSceneView& InView, const FMoviePipelineRenderPassMetrics& InSampleState);
 
 private:
@@ -36,7 +36,7 @@ private:
 	TArray<TSharedPtr<FImageOverlappedAccumulator, ESPMode::ThreadSafe>> ImageTileAccumulator;
 };
 
-DECLARE_MULTICAST_DELEGATE_TwoParams(FMoviePipelineSampleReady, TArray<FColor>, FMoviePipelineRenderPassMetrics);
+DECLARE_MULTICAST_DELEGATE_TwoParams(FMoviePipelineSampleReady, TArray<FLinearColor>, FMoviePipelineRenderPassMetrics);
 DECLARE_MULTICAST_DELEGATE_ThreeParams(FMoviePipelineSetupView, FSceneViewFamily&, FSceneView&, const FMoviePipelineRenderPassMetrics&);
 
 namespace MoviePipeline
@@ -59,6 +59,8 @@ namespace MoviePipeline
 		FSceneView* CalcSceneView(FSceneViewFamily* ViewFamily, const FMoviePipelineRenderPassMetrics& InSampleState);
 
 	protected:
+		/** Which output device were we using previously according to the cvar? */
+		int32 PreviousOutputDeviceIdx;
 
 		FSceneViewStateReference ViewState;
 
diff --git a/Engine/Plugins/MovieScene/MovieRenderPipeline/Source/MovieRenderPipelineRenderPasses/Public/MoviePipelineImageSequenceContainer.h b/Engine/Plugins/MovieScene/MovieRenderPipeline/Source/MovieRenderPipelineRenderPasses/Public/MoviePipelineImageSequenceContainer.h
deleted file mode 100644
diff --git a/Engine/Plugins/MovieScene/MovieRenderPipeline/Source/MovieRenderPipelineRenderPasses/Public/MoviePipelineImageSequenceOutput.h b/Engine/Plugins/MovieScene/MovieRenderPipeline/Source/MovieRenderPipelineRenderPasses/Public/MoviePipelineImageSequenceOutput.h
new file mode 100644
--- /dev/null
+++ b/Engine/Plugins/MovieScene/MovieRenderPipeline/Source/MovieRenderPipelineRenderPasses/Public/MoviePipelineImageSequenceOutput.h
@@ -0,0 +1,82 @@
+// Copyright 1998-2019 Epic Games, Inc. All Rights Reserved.
+#pragma once
+
+#include "MoviePipelineOutputBase.h"
+#include "Async/Future.h"
+#include "MoviePipelineImageSequenceOutput.generated.h"
+
+// Forward Declare
+class IImageWriteQueue;
+
+UCLASS(Blueprintable)
+class MOVIERENDERPIPELINERENDERPASSES_API UMoviePipelineImageSequenceOutputBase : public UMoviePipelineOutputBase
+{
+	GENERATED_BODY()
+public:
+	UMoviePipelineImageSequenceOutputBase();
+
+	virtual void OnRecieveImageDataImpl(FMoviePipelineMergerOutputFrame* InMergedOutputFrame) override;
+
+protected:
+	// UMovieRenderPipelineOutputContainer interface
+	virtual void SetupForPipelineImpl(UMoviePipeline* InPipeline);
+	virtual void BeginFinalizeImpl() override;
+	virtual bool HasFinishedProcessingImpl() override;
+	// ~UMovieRenderPipelineOutputContainer interface
+
+protected:
+	/** The format of the image to write out */
+	EImageFormat OutputFormat;
+
+private:
+	/** A pointer to the image write queue used for asynchronously writing images */
+	IImageWriteQueue* ImageWriteQueue;
+
+	/** A fence to keep track of when the Image Write queue has fully flushed. */
+	TFuture<void> FinalizeFence;
+	FString OutputDirectory;
+};
+
+UCLASS(meta = (DisplayName = "Image Sequence (.bmp [8bpp])"))
+class MOVIERENDERPIPELINERENDERPASSES_API UMoviePipelineImageSequenceOutput_BMP : public UMoviePipelineImageSequenceOutputBase
+{
+	GENERATED_BODY()
+public:
+	UMoviePipelineImageSequenceOutput_BMP()
+	{
+		OutputFormat = EImageFormat::BMP;
+	}
+};
+
+UCLASS(meta = (DisplayName = "Image Sequence (.png [8bpp])"))
+class MOVIERENDERPIPELINERENDERPASSES_API UMoviePipelineImageSequenceOutput_PNG : public UMoviePipelineImageSequenceOutputBase
+{
+	GENERATED_BODY()
+public:
+	UMoviePipelineImageSequenceOutput_PNG()
+	{
+		OutputFormat = EImageFormat::PNG;
+	}
+};
+
+UCLASS(meta = (DisplayName = "Image Sequence (.jpg [8bpp])"))
+class MOVIERENDERPIPELINERENDERPASSES_API UMoviePipelineImageSequenceOutput_JPG : public UMoviePipelineImageSequenceOutputBase
+{
+	GENERATED_BODY()
+public:
+	UMoviePipelineImageSequenceOutput_JPG()
+	{
+		OutputFormat = EImageFormat::JPEG;
+	}
+};
+
+UCLASS(meta = (DisplayName = "Image Sequence (.exr [32bpp])"))
+class MOVIERENDERPIPELINERENDERPASSES_API UMoviePipelineImageSequenceOutput_EXR : public UMoviePipelineImageSequenceOutputBase
+{
+	GENERATED_BODY()
+public:
+	UMoviePipelineImageSequenceOutput_EXR()
+	{
+		OutputFormat = EImageFormat::EXR;
+	}
+};
\ No newline at end of file
