MoviePipeline: Now attempts to flush virtual texture tile building and grass data.

#rnx

diff --git a/Engine/Plugins/MovieScene/MovieRenderPipeline/Source/MovieRenderPipelineCore/MovieRenderPipelineCore.Build.cs b/Engine/Plugins/MovieScene/MovieRenderPipeline/Source/MovieRenderPipelineCore/MovieRenderPipelineCore.Build.cs
--- a/Engine/Plugins/MovieScene/MovieRenderPipeline/Source/MovieRenderPipelineCore/MovieRenderPipelineCore.Build.cs
+++ b/Engine/Plugins/MovieScene/MovieRenderPipeline/Source/MovieRenderPipelineCore/MovieRenderPipelineCore.Build.cs
@@ -13,7 +13,8 @@ public class MovieRenderPipelineCore : ModuleRules
 				"InputCore",
 				"RenderCore",
 				"RHI",
-				"UMG"
+				"UMG",
+				"Landscape", // To flush grass
 			}
 		);
 
@@ -24,8 +25,7 @@ public class MovieRenderPipelineCore : ModuleRules
                 "MovieSceneTracks",
 				"LevelSequence",
 				"Engine",
-                "ImageWriteQueue", // Temp 
-
+                "ImageWriteQueue", // For debug tile writing
             }
         );
 	}
diff --git a/Engine/Plugins/MovieScene/MovieRenderPipeline/Source/MovieRenderPipelineCore/Private/MoviePipeline.cpp b/Engine/Plugins/MovieScene/MovieRenderPipeline/Source/MovieRenderPipelineCore/Private/MoviePipeline.cpp
--- a/Engine/Plugins/MovieScene/MovieRenderPipeline/Source/MovieRenderPipelineCore/Private/MoviePipeline.cpp
+++ b/Engine/Plugins/MovieScene/MovieRenderPipeline/Source/MovieRenderPipelineCore/Private/MoviePipeline.cpp
@@ -881,73 +881,6 @@ void UMoviePipeline::LoadDebugWidget()
 	}
 }
 
-
-void UMoviePipeline::FlushAsyncEngineSystems()
-{
-	// Flush Level Streaming. This solves the problem where levels that are not controlled
-	// by the Sequencer Level Visibility track are marked for Async Load by a gameplay system.
-	// This will register any new actors/components that were spawned during this frame. This needs 
-	// to be done before the shader compiler is flushed so that we compile shaders for any newly
-	// spawned component materials.
-	if (GetWorld())
-	{
-		GetWorld()->FlushLevelStreaming(EFlushLevelStreamingType::Full);
-	}
-
-	// Now we can flush the shader compiler. ToDo: This should probably happen right before SendAllEndOfFrameUpdates() is normally called
-	if (GShaderCompilingManager)
-	{
-		bool bDidWork = false;
-		int32 NumShadersToCompile = GShaderCompilingManager->GetNumRemainingJobs();
-		if (NumShadersToCompile > 0)
-		{
-			UE_LOG(LogMovieRenderPipeline, Log, TEXT("[%d] Starting build for %d shaders."), GFrameCounter, NumShadersToCompile);
-		}
-
-		while (GShaderCompilingManager->GetNumRemainingJobs() > 0 || GShaderCompilingManager->HasShaderJobs())
-		{
-			UE_LOG(LogMovieRenderPipeline, Log, TEXT("[%d] Waiting for %d shaders [Has Shader Jobs: %d] to finish compiling..."), GFrameCounter, GShaderCompilingManager->GetNumRemainingJobs(), GShaderCompilingManager->HasShaderJobs());
-			GShaderCompilingManager->ProcessAsyncResults(false, true);
-			
-			// Sleep for 1 second and then check again. This way we get an indication of progress as this works.
-			FPlatformProcess::Sleep(1.f);
-			bDidWork = true;
-		}
-
-		if (bDidWork)
-		{
-			UE_LOG(LogMovieRenderPipeline, Log, TEXT("[%d] Done building %d shaders."), GFrameCounter, NumShadersToCompile);
-		}
-	}
-
-	// Flush the Mesh Distance Field builder as well.
-	if (GDistanceFieldAsyncQueue)
-	{
-		bool bDidWork = false;
-		int32 NumDistanceFieldsToBuild = GDistanceFieldAsyncQueue->GetNumOutstandingTasks();
-		if (NumDistanceFieldsToBuild > 0)
-		{
-			UE_LOG(LogMovieRenderPipeline, Log, TEXT("[%d] Starting build for %d mesh distance fields."), GFrameCounter, NumDistanceFieldsToBuild);
-		}
-
-		while (GDistanceFieldAsyncQueue->GetNumOutstandingTasks() > 0)
-		{
-			UE_LOG(LogMovieRenderPipeline, Log, TEXT("[%d] Waiting for %d Mesh Distance Fields to finish building..."), GFrameCounter, GDistanceFieldAsyncQueue->GetNumOutstandingTasks());
-			GDistanceFieldAsyncQueue->ProcessAsyncTasks();
-
-			// Sleep for 1 second and then check again. This way we get an indication of progress as this works.
-			FPlatformProcess::Sleep(1.f);
-			bDidWork = true;
-		}
-
-		if (bDidWork)
-		{
-			UE_LOG(LogMovieRenderPipeline, Log, TEXT("[%d] Done building %d Mesh Distance Fields."), GFrameCounter, NumDistanceFieldsToBuild);
-		}
-
-	}
-}
-
 FFrameTime FMoviePipelineTimeController::OnRequestCurrentTime(const FQualifiedFrameTime& InCurrentTime, float InPlayRate)
 {
 	FFrameTime RequestTime = FFrameRate::TransformTime(TimeCache.Time, TimeCache.Rate, InCurrentTime.Rate);
diff --git a/Engine/Plugins/MovieScene/MovieRenderPipeline/Source/MovieRenderPipelineCore/Private/MoviePipelineRendering.cpp b/Engine/Plugins/MovieScene/MovieRenderPipeline/Source/MovieRenderPipelineCore/Private/MoviePipelineRendering.cpp
--- a/Engine/Plugins/MovieScene/MovieRenderPipeline/Source/MovieRenderPipelineCore/Private/MoviePipelineRendering.cpp
+++ b/Engine/Plugins/MovieScene/MovieRenderPipeline/Source/MovieRenderPipelineCore/Private/MoviePipelineRendering.cpp
@@ -21,6 +21,14 @@
 #include "Modules/ModuleManager.h"
 #include "MoviePipelineCameraSetting.h"
 
+// For flushing async systems
+#include "RendererInterface.h"
+#include "LandscapeProxy.h"
+#include "EngineModule.h"
+#include "DistanceFieldAtlas.h"
+#include "ShaderCompiler.h"
+#include "EngineUtils.h"
+
 #define LOCTEXT_NAMESPACE "MoviePipeline"
 
 void UMoviePipeline::SetupRenderingPipelineForShot(FMoviePipelineShotInfo& Shot)
@@ -388,7 +396,7 @@ void UMoviePipeline::OnSampleRendered(TUniquePtr<FImagePixelData>&& OutputSample
 	TUniquePtr<FImageWriteTask> TileImageTask = MakeUnique<FImageWriteTask>();
 
 	// Fill alpha for now 
-	switch (OutputSample->GetType())
+	/*switch (OutputSample->GetType())
 	{
 		case EImagePixelType::Color:
 		{
@@ -407,7 +415,7 @@ void UMoviePipeline::OnSampleRendered(TUniquePtr<FImagePixelData>&& OutputSample
 		}
 		default:
 			check(false);
-	}
+	}*/
 
 	// JPEG output
 	TileImageTask->Format = EImageFormat::EXR;
@@ -426,4 +434,95 @@ void UMoviePipeline::OnSampleRendered(TUniquePtr<FImagePixelData>&& OutputSample
 	ImageWriteQueue->Enqueue(MoveTemp(TileImageTask));
 }
 
+
+
+void UMoviePipeline::FlushAsyncEngineSystems()
+{
+	// Flush Level Streaming. This solves the problem where levels that are not controlled
+	// by the Sequencer Level Visibility track are marked for Async Load by a gameplay system.
+	// This will register any new actors/components that were spawned during this frame. This needs 
+	// to be done before the shader compiler is flushed so that we compile shaders for any newly
+	// spawned component materials.
+	if (GetWorld())
+	{
+		GetWorld()->FlushLevelStreaming(EFlushLevelStreamingType::Full);
+	}
+
+	// Now we can flush the shader compiler. ToDo: This should probably happen right before SendAllEndOfFrameUpdates() is normally called
+	if (GShaderCompilingManager)
+	{
+		bool bDidWork = false;
+		int32 NumShadersToCompile = GShaderCompilingManager->GetNumRemainingJobs();
+		if (NumShadersToCompile > 0)
+		{
+			UE_LOG(LogMovieRenderPipeline, Log, TEXT("[%d] Starting build for %d shaders."), GFrameCounter, NumShadersToCompile);
+		}
+
+		while (GShaderCompilingManager->GetNumRemainingJobs() > 0 || GShaderCompilingManager->HasShaderJobs())
+		{
+			UE_LOG(LogMovieRenderPipeline, Log, TEXT("[%d] Waiting for %d shaders [Has Shader Jobs: %d] to finish compiling..."), GFrameCounter, GShaderCompilingManager->GetNumRemainingJobs(), GShaderCompilingManager->HasShaderJobs());
+			GShaderCompilingManager->ProcessAsyncResults(false, true);
+
+			// Sleep for 1 second and then check again. This way we get an indication of progress as this works.
+			FPlatformProcess::Sleep(1.f);
+			bDidWork = true;
+		}
+
+		// The above while loop should have chewed through it all outstanding tasks but in the event it doesn't we'll ensure we're really finished.
+		GShaderCompilingManager->FinishAllCompilation();
+
+		if (bDidWork)
+		{
+			UE_LOG(LogMovieRenderPipeline, Log, TEXT("[%d] Done building %d shaders."), GFrameCounter, NumShadersToCompile);
+		}
+	}
+
+	// Flush the Mesh Distance Field builder as well.
+	if (GDistanceFieldAsyncQueue)
+	{
+		bool bDidWork = false;
+		int32 NumDistanceFieldsToBuild = GDistanceFieldAsyncQueue->GetNumOutstandingTasks();
+		if (NumDistanceFieldsToBuild > 0)
+		{
+			UE_LOG(LogMovieRenderPipeline, Log, TEXT("[%d] Starting build for %d mesh distance fields."), GFrameCounter, NumDistanceFieldsToBuild);
+		}
+
+		while (GDistanceFieldAsyncQueue->GetNumOutstandingTasks() > 0)
+		{
+			UE_LOG(LogMovieRenderPipeline, Log, TEXT("[%d] Waiting for %d Mesh Distance Fields to finish building..."), GFrameCounter, GDistanceFieldAsyncQueue->GetNumOutstandingTasks());
+			GDistanceFieldAsyncQueue->ProcessAsyncTasks();
+
+			// Sleep for 1 second and then check again. This way we get an indication of progress as this works.
+			FPlatformProcess::Sleep(1.f);
+			bDidWork = true;
+		}
+
+		if (bDidWork)
+		{
+			UE_LOG(LogMovieRenderPipeline, Log, TEXT("[%d] Done building %d Mesh Distance Fields."), GFrameCounter, NumDistanceFieldsToBuild);
+		}
+	}
+
+	// Flush grass
+	{
+		for (TActorIterator<ALandscapeProxy> It(GetWorld()); It; ++It)
+		{
+			ALandscapeProxy* LandscapeProxy = (*It);
+			if (LandscapeProxy)
+			{
+				TArray<FVector> CameraList;
+				LandscapeProxy->UpdateGrass(CameraList, true);
+			}
+		}
+	}
+
+	// Flush virtual texture tile calculations
+	ERHIFeatureLevel::Type FeatureLevel = GetWorld()->FeatureLevel;
+	ENQUEUE_RENDER_COMMAND(VirtualTextureSystemFlushCommand)(
+		[FeatureLevel](FRHICommandListImmediate& RHICmdList)
+	{
+		GetRendererModule().LoadPendingVirtualTextureTiles(RHICmdList, FeatureLevel);
+	});
+}
+
 #undef LOCTEXT_NAMESPACE // "MoviePipeline"
