MoviePipeline: Started implementing a UI for building and configuring movie pipelines through the editor ui.

#rnx

diff --git a/Engine/Plugins/MovieScene/MovieRenderPipeline/Source/MovieRenderPipelineCore/Private/MoviePipeline.cpp b/Engine/Plugins/MovieScene/MovieRenderPipeline/Source/MovieRenderPipelineCore/Private/MoviePipeline.cpp
--- a/Engine/Plugins/MovieScene/MovieRenderPipeline/Source/MovieRenderPipelineCore/Private/MoviePipeline.cpp
+++ b/Engine/Plugins/MovieScene/MovieRenderPipeline/Source/MovieRenderPipelineCore/Private/MoviePipeline.cpp
@@ -74,7 +74,7 @@ void ValidateSequence(ULevelSequence* InSequence)
 
 }
 
-void UMoviePipeline::Initialize(const FMoviePipelineExecutorJob& InJob)
+void UMoviePipeline::Initialize(const FMoviePipelineExecutorJobPrev& InJob)
 {
 	// This function is called after the PIE world has finished initializing, but before
 	// the PIE world is ticked for the first time. We'll end up waiting for the next tick
diff --git a/Engine/Plugins/MovieScene/MovieRenderPipeline/Source/MovieRenderPipelineCore/Private/MoviePipelineInProcessExecutor.cpp b/Engine/Plugins/MovieScene/MovieRenderPipeline/Source/MovieRenderPipelineCore/Private/MoviePipelineInProcessExecutor.cpp
--- a/Engine/Plugins/MovieScene/MovieRenderPipeline/Source/MovieRenderPipelineCore/Private/MoviePipelineInProcessExecutor.cpp
+++ b/Engine/Plugins/MovieScene/MovieRenderPipeline/Source/MovieRenderPipelineCore/Private/MoviePipelineInProcessExecutor.cpp
@@ -4,7 +4,7 @@
 #include "Engine/World.h"
 
 #define LOCTEXT_NAMESPACE "MoviePipelineInProcessExecutor"
-void UMoviePipelineInProcessExecutor::Start(const FMoviePipelineExecutorJob& InJob)
+void UMoviePipelineInProcessExecutor::Start(const FMoviePipelineExecutorJobPrev& InJob)
 {
 	FCoreUObjectDelegates::PostLoadMapWithWorld.AddUObject(this, &UMoviePipelineInProcessExecutor::OnMapLoadFinished);
 }
@@ -20,7 +20,7 @@ void UMoviePipelineInProcessExecutor::OnMapLoadFinished(UWorld* NewWorld)
 	// Stop listening for map load until we're done and know we want to start the next config.
 	FCoreUObjectDelegates::PostLoadMapWithWorld.RemoveAll(this);
 	
-	FMoviePipelineExecutorJob& CurrentJob = ExecutorJobs[CurrentPipelineIndex];
+	FMoviePipelineExecutorJobPrev& CurrentJob = ExecutorJobs[CurrentPipelineIndex];
 
 	ActiveMoviePipeline = NewObject<UMoviePipeline>(NewWorld, TargetPipelineClass);
 	ActiveMoviePipeline->Initialize(CurrentJob);
diff --git a/Engine/Plugins/MovieScene/MovieRenderPipeline/Source/MovieRenderPipelineCore/Private/MoviePipelineLinearExecutor.cpp b/Engine/Plugins/MovieScene/MovieRenderPipeline/Source/MovieRenderPipelineCore/Private/MoviePipelineLinearExecutor.cpp
--- a/Engine/Plugins/MovieScene/MovieRenderPipeline/Source/MovieRenderPipelineCore/Private/MoviePipelineLinearExecutor.cpp
+++ b/Engine/Plugins/MovieScene/MovieRenderPipeline/Source/MovieRenderPipelineCore/Private/MoviePipelineLinearExecutor.cpp
@@ -4,7 +4,7 @@
 
 #define LOCTEXT_NAMESPACE "MoviePipelineLinearExecutorBase"
 
-void UMoviePipelineLinearExecutorBase::ExecuteImpl(const TArray<FMoviePipelineExecutorJob>& InPipelines)
+void UMoviePipelineLinearExecutorBase::ExecuteImpl(const TArray<FMoviePipelineExecutorJobPrev>& InPipelines)
 {
 	if (InPipelines.Num() == 0)
 	{
diff --git a/Engine/Plugins/MovieScene/MovieRenderPipeline/Source/MovieRenderPipelineCore/Private/MovieRenderPipelineCommandLine.cpp b/Engine/Plugins/MovieScene/MovieRenderPipeline/Source/MovieRenderPipelineCore/Private/MovieRenderPipelineCommandLine.cpp
--- a/Engine/Plugins/MovieScene/MovieRenderPipeline/Source/MovieRenderPipelineCore/Private/MovieRenderPipelineCommandLine.cpp
+++ b/Engine/Plugins/MovieScene/MovieRenderPipeline/Source/MovieRenderPipelineCore/Private/MovieRenderPipelineCommandLine.cpp
@@ -63,8 +63,8 @@ void FMovieRenderPipelineCoreModule::InitializeCommandLineMovieRender()
 	ExecutorBase->OnExecutorErrored().AddRaw(this, &FMovieRenderPipelineCoreModule::OnCommandLineMovieRenderErrored);
 	ExecutorBase->SetMoviePipelineClass(PipelineClass);
 
-	TArray<FMoviePipelineExecutorJob> Jobs;
-	Jobs.Add(FMoviePipelineExecutorJob(SequenceAssetValue, MasterConfig));
+	TArray<FMoviePipelineExecutorJobPrev> Jobs;
+	Jobs.Add(FMoviePipelineExecutorJobPrev(SequenceAssetValue, MasterConfig));
 
 	ExecutorBase->Execute(Jobs);
 }
diff --git a/Engine/Plugins/MovieScene/MovieRenderPipeline/Source/MovieRenderPipelineCore/Public/MoviePipeline.h b/Engine/Plugins/MovieScene/MovieRenderPipeline/Source/MovieRenderPipelineCore/Public/MoviePipeline.h
--- a/Engine/Plugins/MovieScene/MovieRenderPipeline/Source/MovieRenderPipelineCore/Public/MoviePipeline.h
+++ b/Engine/Plugins/MovieScene/MovieRenderPipeline/Source/MovieRenderPipelineCore/Public/MoviePipeline.h
@@ -40,7 +40,7 @@ public:
 	
 	* @param InJob	- This contains settings and sequence to render this Movie Pipeline with.
 	*/
-	void Initialize(const FMoviePipelineExecutorJob& InJob);
+	void Initialize(const FMoviePipelineExecutorJobPrev& InJob);
 
 	/** 
 	* Call to shut down the pipeline. This flushes any outstanding file writes and unregisters all delegates.
@@ -269,7 +269,7 @@ public:
 private:
 	/** Keep track of which job we're working on. This holds our Configuration + which shots we're supposed to render from it. */
 	UPROPERTY(Transient)
-	FMoviePipelineExecutorJob CurrentJob;
+	FMoviePipelineExecutorJobPrev CurrentJob;
 };
 
 UCLASS()
diff --git a/Engine/Plugins/MovieScene/MovieRenderPipeline/Source/MovieRenderPipelineCore/Public/MoviePipelineAccumulationSetting.h b/Engine/Plugins/MovieScene/MovieRenderPipeline/Source/MovieRenderPipelineCore/Public/MoviePipelineAccumulationSetting.h
--- a/Engine/Plugins/MovieScene/MovieRenderPipeline/Source/MovieRenderPipelineCore/Public/MoviePipelineAccumulationSetting.h
+++ b/Engine/Plugins/MovieScene/MovieRenderPipeline/Source/MovieRenderPipelineCore/Public/MoviePipelineAccumulationSetting.h
@@ -19,9 +19,13 @@ public:
 	{
 	}
 
+public:
+#if WITH_EDITOR
+	virtual FText GetDisplayText() const override { return NSLOCTEXT("MovieRenderPipeline", "AccumulationSettingDisplayName", "Subsample Accumulation"); }
+#endif
 protected:
 	virtual bool IsValidOnShots() const override { return true; }
-	virtual bool IsValidOnMaster() const override { return false; }
+	virtual bool IsValidOnMaster() const override { return true; }
 	
 public:
 
diff --git a/Engine/Plugins/MovieScene/MovieRenderPipeline/Source/MovieRenderPipelineCore/Public/MoviePipelineAudioOutput.h b/Engine/Plugins/MovieScene/MovieRenderPipeline/Source/MovieRenderPipelineCore/Public/MoviePipelineAudioOutput.h
--- a/Engine/Plugins/MovieScene/MovieRenderPipeline/Source/MovieRenderPipelineCore/Public/MoviePipelineAudioOutput.h
+++ b/Engine/Plugins/MovieScene/MovieRenderPipeline/Source/MovieRenderPipelineCore/Public/MoviePipelineAudioOutput.h
@@ -4,7 +4,7 @@
 #include "MoviePipelineOutputBase.h"
 #include "MoviePipelineAudioOutput.generated.h"
 
-UCLASS(BlueprintType, meta=(MovieRenderPipelineDisplayName="Audio"))
+UCLASS(BlueprintType)
 class MOVIERENDERPIPELINECORE_API UMoviePipelineAudioOutput : public UMoviePipelineOutputBase
 {
 	GENERATED_BODY()
diff --git a/Engine/Plugins/MovieScene/MovieRenderPipeline/Source/MovieRenderPipelineCore/Public/MoviePipelineCameraSetting.h b/Engine/Plugins/MovieScene/MovieRenderPipeline/Source/MovieRenderPipelineCore/Public/MoviePipelineCameraSetting.h
--- a/Engine/Plugins/MovieScene/MovieRenderPipeline/Source/MovieRenderPipelineCore/Public/MoviePipelineCameraSetting.h
+++ b/Engine/Plugins/MovieScene/MovieRenderPipeline/Source/MovieRenderPipelineCore/Public/MoviePipelineCameraSetting.h
@@ -18,9 +18,14 @@ public:
 		, bManualExposure(true)
 		, ExposureCompensation(8.0)
 	{}
+
+public:
+#if WITH_EDITOR
+	virtual FText GetDisplayText() const override { return NSLOCTEXT("MovieRenderPipeline", "CameraSettingDisplayName", "Camera"); }
+#endif
 protected:
 	virtual bool IsValidOnShots() const override { return true; }
-	virtual bool IsValidOnMaster() const override { return false; }
+	virtual bool IsValidOnMaster() const override { return true; }
 
 public:
 
diff --git a/Engine/Plugins/MovieScene/MovieRenderPipeline/Source/MovieRenderPipelineCore/Public/MoviePipelineConfigBase.h b/Engine/Plugins/MovieScene/MovieRenderPipeline/Source/MovieRenderPipelineCore/Public/MoviePipelineConfigBase.h
--- a/Engine/Plugins/MovieScene/MovieRenderPipeline/Source/MovieRenderPipelineCore/Public/MoviePipelineConfigBase.h
+++ b/Engine/Plugins/MovieScene/MovieRenderPipeline/Source/MovieRenderPipelineCore/Public/MoviePipelineConfigBase.h
@@ -23,7 +23,7 @@ public:
 	* Removes the specific instance from our Setting list.
 	*/
 	UFUNCTION(BlueprintCallable, Category = "Movie Render Pipeline")
-		void RemoveSetting(UMoviePipelineSetting* InSetting);
+	void RemoveSetting(UMoviePipelineSetting* InSetting);
 
 	int32 GetSettingsSerialNumber() const { return SettingsSerialNumber; }
 
@@ -112,8 +112,8 @@ public:
 		return Found;
 	}
 
-protected:
-	virtual bool CanSettingBeAdded(UMoviePipelineSetting* InSetting) PURE_VIRTUAL( UMoviePipelineConfigBase::CanSettingBeAdded, return false; );
+public:
+	virtual bool CanSettingBeAdded(const UMoviePipelineSetting* InSetting) const PURE_VIRTUAL( UMoviePipelineConfigBase::CanSettingBeAdded, return false; );
 
 protected:
 	/** Array of settings classes that affect various parts of the output pipeline. */
diff --git a/Engine/Plugins/MovieScene/MovieRenderPipeline/Source/MovieRenderPipelineCore/Public/MoviePipelineExecutor.h b/Engine/Plugins/MovieScene/MovieRenderPipeline/Source/MovieRenderPipelineCore/Public/MoviePipelineExecutor.h
--- a/Engine/Plugins/MovieScene/MovieRenderPipeline/Source/MovieRenderPipelineCore/Public/MoviePipelineExecutor.h
+++ b/Engine/Plugins/MovieScene/MovieRenderPipeline/Source/MovieRenderPipelineCore/Public/MoviePipelineExecutor.h
@@ -41,7 +41,7 @@ public:
 	}
 
 	UFUNCTION(BlueprintCallable, Category = "Movie Render Pipeline")
-	virtual void Execute(const TArray<FMoviePipelineExecutorJob>& InPipelines)
+	virtual void Execute(const TArray<FMoviePipelineExecutorJobPrev>& InPipelines)
 	{
 		ExecuteImpl(InPipelines);
 	}
@@ -89,7 +89,7 @@ protected:
 		OnExecutorFinishedDelegate.Broadcast(this, bFatal);
 	}
 
-	virtual void ExecuteImpl(const TArray<FMoviePipelineExecutorJob>& InPipelines) PURE_VIRTUAL(UMoviePipelineExecutorBase::ExecuteImpl, );
+	virtual void ExecuteImpl(const TArray<FMoviePipelineExecutorJobPrev>& InPipelines) PURE_VIRTUAL(UMoviePipelineExecutorBase::ExecuteImpl, );
 
 private:
 	/** 
diff --git a/Engine/Plugins/MovieScene/MovieRenderPipeline/Source/MovieRenderPipelineCore/Public/MoviePipelineHighResSetting.h b/Engine/Plugins/MovieScene/MovieRenderPipeline/Source/MovieRenderPipelineCore/Public/MoviePipelineHighResSetting.h
--- a/Engine/Plugins/MovieScene/MovieRenderPipeline/Source/MovieRenderPipelineCore/Public/MoviePipelineHighResSetting.h
+++ b/Engine/Plugins/MovieScene/MovieRenderPipeline/Source/MovieRenderPipelineCore/Public/MoviePipelineHighResSetting.h
@@ -27,9 +27,13 @@ public:
 
 		return InTileSize + FIntPoint(2 * OverlappedPadX, 2 * OverlappedPadY);
 	}
-protected:
+
+public:
+#if WITH_EDITOR
+	virtual FText GetDisplayText() const override { return NSLOCTEXT("MovieRenderPipeline", "HighResSettingDisplayName", "High Resolution"); }
+#endif
 	virtual bool IsValidOnShots() const override { return true; }
-	virtual bool IsValidOnMaster() const override { return false; }
+	virtual bool IsValidOnMaster() const override { return true; }
 
 public:
 	/**
diff --git a/Engine/Plugins/MovieScene/MovieRenderPipeline/Source/MovieRenderPipelineCore/Public/MoviePipelineInProcessExecutor.h b/Engine/Plugins/MovieScene/MovieRenderPipeline/Source/MovieRenderPipelineCore/Public/MoviePipelineInProcessExecutor.h
--- a/Engine/Plugins/MovieScene/MovieRenderPipeline/Source/MovieRenderPipelineCore/Public/MoviePipelineInProcessExecutor.h
+++ b/Engine/Plugins/MovieScene/MovieRenderPipeline/Source/MovieRenderPipelineCore/Public/MoviePipelineInProcessExecutor.h
@@ -23,7 +23,7 @@ public:
 	}
 
 protected:
-	virtual void Start(const FMoviePipelineExecutorJob& InJob) override;
+	virtual void Start(const FMoviePipelineExecutorJobPrev& InJob) override;
 
 private:
 	void OnMapLoadFinished(UWorld* NewWorld);
diff --git a/Engine/Plugins/MovieScene/MovieRenderPipeline/Source/MovieRenderPipelineCore/Public/MoviePipelineLinearExecutor.h b/Engine/Plugins/MovieScene/MovieRenderPipeline/Source/MovieRenderPipelineCore/Public/MoviePipelineLinearExecutor.h
--- a/Engine/Plugins/MovieScene/MovieRenderPipeline/Source/MovieRenderPipelineCore/Public/MoviePipelineLinearExecutor.h
+++ b/Engine/Plugins/MovieScene/MovieRenderPipeline/Source/MovieRenderPipelineCore/Public/MoviePipelineLinearExecutor.h
@@ -27,10 +27,10 @@ public:
 	}
 
 protected:
-	virtual void ExecuteImpl(const TArray<FMoviePipelineExecutorJob>& InPipelines) override;
+	virtual void ExecuteImpl(const TArray<FMoviePipelineExecutorJobPrev>& InPipelines) override;
 
 	virtual void StartPipelineByIndex(int32 InPipelineIndex);
-	virtual void Start(const FMoviePipelineExecutorJob& InJob) {}
+	virtual void Start(const FMoviePipelineExecutorJobPrev& InJob) {}
 
 public:
 	virtual void OnIndividualPipelineFinished(UMoviePipeline* /* FinishedPipeline */);
@@ -40,7 +40,7 @@ protected:
 	
 	/** List of Pipeline Configs we've been asked to execute. */
 	UPROPERTY(Transient)
-	TArray<FMoviePipelineExecutorJob> ExecutorJobs;
+	TArray<FMoviePipelineExecutorJobPrev> ExecutorJobs;
 
 	/** A Movie Pipeline that has been spawned and is running (if any) */
 	UPROPERTY(Transient)
diff --git a/Engine/Plugins/MovieScene/MovieRenderPipeline/Source/MovieRenderPipelineCore/Public/MoviePipelineMasterConfig.h b/Engine/Plugins/MovieScene/MovieRenderPipeline/Source/MovieRenderPipelineCore/Public/MoviePipelineMasterConfig.h
--- a/Engine/Plugins/MovieScene/MovieRenderPipeline/Source/MovieRenderPipelineCore/Public/MoviePipelineMasterConfig.h
+++ b/Engine/Plugins/MovieScene/MovieRenderPipeline/Source/MovieRenderPipelineCore/Public/MoviePipelineMasterConfig.h
@@ -39,7 +39,7 @@ public:
 	UMoviePipelineShotConfig* GetConfigForShot(const FString& ShotName) const;
 
 protected:
-	virtual bool CanSettingBeAdded(UMoviePipelineSetting* InSetting) override
+	virtual bool CanSettingBeAdded(const UMoviePipelineSetting* InSetting) const override
 	{
 		check(InSetting);
 		return InSetting->IsValidOnMaster();
diff --git a/Engine/Plugins/MovieScene/MovieRenderPipeline/Source/MovieRenderPipelineCore/Public/MoviePipelineOutputSetting.h b/Engine/Plugins/MovieScene/MovieRenderPipeline/Source/MovieRenderPipelineCore/Public/MoviePipelineOutputSetting.h
--- a/Engine/Plugins/MovieScene/MovieRenderPipeline/Source/MovieRenderPipelineCore/Public/MoviePipelineOutputSetting.h
+++ b/Engine/Plugins/MovieScene/MovieRenderPipeline/Source/MovieRenderPipelineCore/Public/MoviePipelineOutputSetting.h
@@ -15,7 +15,12 @@ class MOVIERENDERPIPELINECORE_API UMoviePipelineOutputSetting : public UMoviePip
 	GENERATED_BODY()
 public:
 	UMoviePipelineOutputSetting();
-	
+public:
+#if WITH_EDITOR
+	virtual FText GetDisplayText() const override { return NSLOCTEXT("MovieRenderPipeline", "OutputSettingDisplayName", "Output"); }
+#endif
+	virtual bool IsValidOnShots() const override { return false; }
+	virtual bool IsValidOnMaster() const override { return true; }
 public:
 	/** What directory should all of our output files be relative to. */
 	UPROPERTY(EditAnywhere, BlueprintReadWrite, Category = "Movie Render Pipeline")
diff --git a/Engine/Plugins/MovieScene/MovieRenderPipeline/Source/MovieRenderPipelineCore/Public/MoviePipelineQueue.h b/Engine/Plugins/MovieScene/MovieRenderPipeline/Source/MovieRenderPipelineCore/Public/MoviePipelineQueue.h
new file mode 100644
--- /dev/null
+++ b/Engine/Plugins/MovieScene/MovieRenderPipeline/Source/MovieRenderPipelineCore/Public/MoviePipelineQueue.h
@@ -0,0 +1,145 @@
+// Copyright 1998-2019 Epic Games, Inc. All Rights Reserved.
+#pragma once
+
+#include "UObject/Object.h"
+#include "MovieRenderPipelineDataTypes.h"
+#include "LevelSequence.h"
+
+#include "MoviePipelineQueue.generated.h"
+
+class UMoviePipelineMasterConfig;
+
+UENUM(BlueprintType)
+enum class EMoviePipelineExecutorJobStatus : uint8
+{
+	Uninitialized = 0,
+	ReadyToStart = 1,
+	InProgress = 2,
+	Finished = 3
+};
+
+/**
+* A particular job within the Queue
+*/
+UCLASS(BlueprintType)
+class MOVIERENDERPIPELINECORE_API UMoviePipelineExecutorJob : public UObject
+{
+	GENERATED_BODY()
+public:
+	UMoviePipelineExecutorJob()
+	{
+		Configuration = CreateDefaultSubobject<UMoviePipelineMasterConfig>("DefaultMoviePipelineConfig");
+	}
+
+	UFUNCTION(BlueprintPure, Category = "Movie Render Pipeline")
+	bool HasFinished() const
+	{
+		return JobStatus == EMoviePipelineExecutorJobStatus::Finished;
+	}
+	
+	float GetProgressPercentage() const
+	{
+		return 0.2f;
+	}
+	
+public:
+	ULevelSequence* TryLoadSequence()
+	{
+		if(LoadedSequence)
+		{
+			return LoadedSequence;
+		}
+		
+		LoadedSequence = Cast<ULevelSequence>(Sequence.TryLoad());
+		return LoadedSequence;
+	}
+	
+
+public:
+	/** Which sequence should this job render? */
+	UPROPERTY(BlueprintReadWrite, Category = "Movie Render Pipeline")
+	FSoftObjectPath Sequence;
+	
+	/** Which map should this job render on */
+	UPROPERTY(BlueprintReadWrite, Category = "Movie Render Pipeline")
+	FSoftObjectPath Map;
+
+	/** What master configuration is used to render this sequence? This specifies output resolution/path, etc. */
+	UPROPERTY(BlueprintReadWrite, Category = "Movie Render Pipeline")
+	UMoviePipelineMasterConfig* Configuration;
+	
+	/** What state is this particular job instance currently in? */
+	UPROPERTY(BlueprintReadOnly, Category = "Movie Render Pipeline")
+	EMoviePipelineExecutorJobStatus JobStatus;
+	
+private:
+	/** Cache our loaded sequence after the first time someone tries to retrieve information from this job that requires the Sequence. */
+	UPROPERTY(Transient)
+	ULevelSequence* LoadedSequence;
+};
+
+/**
+* A queue is a list of jobs that have been executed, are executing and are waiting to be executed. These can be saved
+* to specific assets to allow 
+*/
+UCLASS(BlueprintType)
+class MOVIERENDERPIPELINECORE_API UMoviePipelineQueue : public UObject
+{
+	GENERATED_BODY()
+public:
+	
+	UMoviePipelineQueue()
+		: QueueSerialNumber(0)
+	{
+		// Ensure instances are always transactional
+		if (!HasAnyFlags(RF_ClassDefaultObject))
+		{
+			SetFlags(RF_Transactional);
+		}
+	}
+	
+	UFUNCTION(BlueprintCallable, Category = "Movie Render Pipeline")
+	UMoviePipelineExecutorJob* AllocateNewJob()
+	{
+		UMoviePipelineExecutorJob* NewJob = NewObject<UMoviePipelineExecutorJob>(this);
+		NewJob->SetFlags(RF_Transactional);
+		
+		Jobs.Add(NewJob);
+		QueueSerialNumber++;
+
+		return NewJob;
+	}
+
+	UFUNCTION(BlueprintCallable, Category = "Movie Render Pipeline")
+	void DeleteJob(UMoviePipelineExecutorJob* InJob)
+	{
+		if (!InJob)
+		{
+			return;
+		}
+
+		Jobs.Remove(InJob);
+		QueueSerialNumber++;
+	}
+	
+	UFUNCTION(BlueprintPure, Category = "Movie Render Pipeline")
+	TArray<UMoviePipelineExecutorJob*> GetJobs()
+	{
+		return Jobs;
+	}
+	
+	/**
+	 * Retrieve the serial number that is incremented when a job is added or removed from this list.
+	 * @note: This field is not serialized, and not copied along with UObject duplication.
+	 */
+	uint32 GetQueueSerialNumber() const
+	{
+		return QueueSerialNumber;
+	}
+private:
+	UPROPERTY()
+	TArray<UMoviePipelineExecutorJob*> Jobs;
+	
+private:
+	int32 QueueSerialNumber;
+};
\ No newline at end of file
diff --git a/Engine/Plugins/MovieScene/MovieRenderPipeline/Source/MovieRenderPipelineCore/Public/MoviePipelineRenderPass.h b/Engine/Plugins/MovieScene/MovieRenderPipeline/Source/MovieRenderPipelineCore/Public/MoviePipelineRenderPass.h
--- a/Engine/Plugins/MovieScene/MovieRenderPipeline/Source/MovieRenderPipelineCore/Public/MoviePipelineRenderPass.h
+++ b/Engine/Plugins/MovieScene/MovieRenderPipeline/Source/MovieRenderPipelineCore/Public/MoviePipelineRenderPass.h
@@ -36,7 +36,7 @@ public:
 
 protected:
 	virtual bool IsValidOnShots() const override { return true; }
-	virtual bool IsValidOnMaster() const override { return false; }
+	virtual bool IsValidOnMaster() const override { return true; }
 protected:
 	virtual void GetRequiredEnginePassesImpl(TSet<FMoviePipelinePassIdentifier>& RequiredEnginePasses) {}
 
diff --git a/Engine/Plugins/MovieScene/MovieRenderPipeline/Source/MovieRenderPipelineCore/Public/MoviePipelineSetting.h b/Engine/Plugins/MovieScene/MovieRenderPipeline/Source/MovieRenderPipelineCore/Public/MoviePipelineSetting.h
--- a/Engine/Plugins/MovieScene/MovieRenderPipeline/Source/MovieRenderPipelineCore/Public/MoviePipelineSetting.h
+++ b/Engine/Plugins/MovieScene/MovieRenderPipeline/Source/MovieRenderPipelineCore/Public/MoviePipelineSetting.h
@@ -40,6 +40,7 @@ protected:
 	
 public:
 #if WITH_EDITOR
+	/** Warning: This gets called on the CDO of the object */
 	virtual FText GetDisplayText() const { return this->GetClass()->GetDisplayNameText(); }
 #endif
 	/** Can this configuration setting be added to shots? If not, it will throw an error when trying to add it to a shot config. */
diff --git a/Engine/Plugins/MovieScene/MovieRenderPipeline/Source/MovieRenderPipelineCore/Public/MoviePipelineShotConfig.h b/Engine/Plugins/MovieScene/MovieRenderPipeline/Source/MovieRenderPipelineCore/Public/MoviePipelineShotConfig.h
--- a/Engine/Plugins/MovieScene/MovieRenderPipeline/Source/MovieRenderPipelineCore/Public/MoviePipelineShotConfig.h
+++ b/Engine/Plugins/MovieScene/MovieRenderPipeline/Source/MovieRenderPipelineCore/Public/MoviePipelineShotConfig.h
@@ -2,16 +2,12 @@
 #pragma once
 
 #include "MoviePipelineConfigBase.h"
-#include "MoviePipelineRenderPass.h"
 
 #include "MoviePipelineShotConfig.generated.h"
 
 
 // Forward Declares
-class ULevelSequence;
-class UMoviePipelineSetting;
 class UMoviePipelineRenderPass;
-class UMoviePipelineOutputBase;
 
 UCLASS(Blueprintable)
 class MOVIERENDERPIPELINECORE_API UMoviePipelineShotConfig : public UMoviePipelineConfigBase
@@ -26,7 +22,7 @@ public:
 	TArray<UMoviePipelineRenderPass*> GetRenderPasses() const;
 
 protected:
-	virtual bool CanSettingBeAdded(UMoviePipelineSetting* InSetting) override
+	virtual bool CanSettingBeAdded(const UMoviePipelineSetting* InSetting) const override
 	{
 		check(InSetting);
 		return InSetting->IsValidOnShots();
diff --git a/Engine/Plugins/MovieScene/MovieRenderPipeline/Source/MovieRenderPipelineCore/Public/MovieRenderPipelineDataTypes.h b/Engine/Plugins/MovieScene/MovieRenderPipeline/Source/MovieRenderPipelineCore/Public/MovieRenderPipelineDataTypes.h
--- a/Engine/Plugins/MovieScene/MovieRenderPipeline/Source/MovieRenderPipelineCore/Public/MovieRenderPipelineDataTypes.h
+++ b/Engine/Plugins/MovieScene/MovieRenderPipeline/Source/MovieRenderPipelineCore/Public/MovieRenderPipelineDataTypes.h
@@ -865,17 +865,17 @@ struct FImagePixelDataPayload : IImagePixelDataPayload, public TSharedFromThis<F
 * Describes an individual job executed by a Movie Pipeline Executor.
 */
 USTRUCT(BlueprintType)
-struct FMoviePipelineExecutorJob
+struct FMoviePipelineExecutorJobPrev
 {
 	GENERATED_BODY()
 public:
-	FMoviePipelineExecutorJob()
+	FMoviePipelineExecutorJobPrev()
 		: Sequence()
 		, Configuration(nullptr)
 	{
 	}
 
-	FMoviePipelineExecutorJob(const FSoftObjectPath& InSequenceAsset, UMoviePipelineMasterConfig* InConfig)
+	FMoviePipelineExecutorJobPrev(const FSoftObjectPath& InSequenceAsset, UMoviePipelineMasterConfig* InConfig)
 		: Sequence(InSequenceAsset)
 		, Configuration(InConfig)
 	{
diff --git a/Engine/Plugins/MovieScene/MovieRenderPipeline/Source/MovieRenderPipelineEditor/MovieRenderPipelineEditor.Build.cs b/Engine/Plugins/MovieScene/MovieRenderPipeline/Source/MovieRenderPipelineEditor/MovieRenderPipelineEditor.Build.cs
--- a/Engine/Plugins/MovieScene/MovieRenderPipeline/Source/MovieRenderPipelineEditor/MovieRenderPipelineEditor.Build.cs
+++ b/Engine/Plugins/MovieScene/MovieRenderPipeline/Source/MovieRenderPipelineEditor/MovieRenderPipelineEditor.Build.cs
@@ -29,7 +29,11 @@ public class MovieRenderPipelineEditor : ModuleRules
 				"MovieRenderPipelineCore",
                 "MovieRenderPipelineRenderPasses",
                 "MovieRenderPipelineSettings",
-				"Settings"
+				"Settings",
+				"ContentBrowser",
+				"PropertyEditor",
+				"EditorWidgets",
+				"EditorSubsystem",
             }
         );
 
diff --git a/Engine/Plugins/MovieScene/MovieRenderPipeline/Source/MovieRenderPipelineEditor/Private/MoviePipelinePIEExecutor.cpp b/Engine/Plugins/MovieScene/MovieRenderPipeline/Source/MovieRenderPipelineEditor/Private/MoviePipelinePIEExecutor.cpp
--- a/Engine/Plugins/MovieScene/MovieRenderPipeline/Source/MovieRenderPipelineEditor/Private/MoviePipelinePIEExecutor.cpp
+++ b/Engine/Plugins/MovieScene/MovieRenderPipeline/Source/MovieRenderPipelineEditor/Private/MoviePipelinePIEExecutor.cpp
@@ -25,7 +25,7 @@ FText UMoviePipelinePIEExecutor::GetWindowTitle(const int32 InConfigIndex, const
 	return WindowTitle;
 }
 
-void UMoviePipelinePIEExecutor::Start(const FMoviePipelineExecutorJob& InJob)
+void UMoviePipelinePIEExecutor::Start(const FMoviePipelineExecutorJobPrev& InJob)
 {
 	Super::Start(InJob);
 
diff --git a/Engine/Plugins/MovieScene/MovieRenderPipeline/Source/MovieRenderPipelineEditor/Private/MovieRenderPipelineEditorModule.cpp b/Engine/Plugins/MovieScene/MovieRenderPipeline/Source/MovieRenderPipelineEditor/Private/MovieRenderPipelineEditorModule.cpp
--- a/Engine/Plugins/MovieScene/MovieRenderPipeline/Source/MovieRenderPipelineEditor/Private/MovieRenderPipelineEditorModule.cpp
+++ b/Engine/Plugins/MovieScene/MovieRenderPipeline/Source/MovieRenderPipelineEditor/Private/MovieRenderPipelineEditorModule.cpp
@@ -5,7 +5,8 @@
 #include "MoviePipelineExecutor.h"
 #include "WorkspaceMenuStructureModule.h"
 #include "WorkspaceMenuStructure.h"
-#include "Widgets/SMoviePipelineTabContent.h"
+#include "Widgets/SMoviePipelineConfigTabContent.h"
+#include "Widgets/SMoviePipelineQueueTabContent.h"
 #include "ISettingsModule.h"
 #include "MovieRenderPipelineSettings.h"
 #include "HAL/IConsoleManager.h"
@@ -20,8 +21,10 @@
 
 #define LOCTEXT_NAMESPACE "FMovieRenderPipelineEditorModule"
 
-FName IMovieRenderPipelineEditorModule::MovieRenderPipelineTabName = "MovieRenderPipeline";
-FText IMovieRenderPipelineEditorModule::MovieRenderPipelineTabLabel = LOCTEXT("MovieRenderPipelineTab_Label", "Movie Render Pipeline");
+FName IMovieRenderPipelineEditorModule::MoviePipelineQueueTabName = "MoviePipelineQueue";
+FText IMovieRenderPipelineEditorModule::MoviePipelineQueueTabLabel = LOCTEXT("MovieRenderQueueTab_Label", "Movie Render Queue");
+FName IMovieRenderPipelineEditorModule::MoviePipelineConfigEditorTabName = "MovieRenderPipeline";
+FText IMovieRenderPipelineEditorModule::MoviePipelineConfigEditorTabLabel = LOCTEXT("MovieRenderPipelineTab_Label", "Movie Render Pipeline");
 										 
 namespace
 {
@@ -30,7 +33,18 @@ namespace
 		TSharedRef<SDockTab> NewTab = SNew(SDockTab)
 			.TabRole(ETabRole::NomadTab)
 			[
-				SNew(SMoviePipelineTabContent)
+				SNew(SMoviePipelineConfigTabContent)
+			];
+		
+		return NewTab;
+	}
+
+	static TSharedRef<SDockTab> SpawnMoviePipelineQueueTab(const FSpawnTabArgs& InSpawnTabArgs)
+	{
+		TSharedRef<SDockTab> NewTab = SNew(SDockTab)
+			.TabRole(ETabRole::MajorTab)
+			[
+				SNew(SMoviePipelineQueueTabContent)
 			];
 
 		return NewTab;
@@ -38,12 +52,20 @@ namespace
 
 	static void RegisterTabImpl()
 	{
-		FTabSpawnerEntry& MRPTabSpawner = FGlobalTabmanager::Get()->RegisterNomadTabSpawner(IMovieRenderPipelineEditorModule::MovieRenderPipelineTabName, FOnSpawnTab::CreateStatic(SpawnMovieRenderPipelineTab));
+		FTabSpawnerEntry& MRPConfigTabSpawner = FGlobalTabmanager::Get()->RegisterNomadTabSpawner(IMovieRenderPipelineEditorModule::MoviePipelineConfigEditorTabName, FOnSpawnTab::CreateStatic(SpawnMovieRenderPipelineTab));
+
+		MRPConfigTabSpawner
+			.SetGroup(WorkspaceMenu::GetMenuStructure().GetLevelEditorCinematicsCategory())
+			.SetDisplayName(IMovieRenderPipelineEditorModule::MoviePipelineConfigEditorTabLabel)
+			.SetTooltipText(LOCTEXT("MovieRenderPipelineConfigTab_Tooltip", "Open the Movie Render Config UI for creating and editing presets."))
+			.SetIcon(FSlateIcon(FEditorStyle::GetStyleSetName(), "ContentBrowser.TabIcon"));
+
+		FTabSpawnerEntry& MRPQueueTabSpawner = FGlobalTabmanager::Get()->RegisterNomadTabSpawner(IMovieRenderPipelineEditorModule::MoviePipelineQueueTabName, FOnSpawnTab::CreateStatic(SpawnMoviePipelineQueueTab));
 
-		MRPTabSpawner
+		MRPQueueTabSpawner
 			.SetGroup(WorkspaceMenu::GetMenuStructure().GetLevelEditorCinematicsCategory())
-			.SetDisplayName(IMovieRenderPipelineEditorModule::MovieRenderPipelineTabLabel)
-			.SetTooltipText(LOCTEXT("MovieRenderPipelineTab_Tooltip", "Open the Movie Render Pipeline UI"))
+			.SetDisplayName(IMovieRenderPipelineEditorModule::MoviePipelineQueueTabLabel)
+			.SetTooltipText(LOCTEXT("MovieRenderPipelineQueueTab_Tooltip", "Open the Movie Render Queue to render Sequences to disk at a higher quality than realtime allows."))
 			.SetIcon(FSlateIcon(FEditorStyle::GetStyleSetName(), "ContentBrowser.TabIcon"));
 	}
 }
@@ -68,10 +90,10 @@ void FMovieRenderPipelineEditorModule::PerformTestPipelineRender(const TArray<FS
 
 	TArray<UMoviePipelineMasterConfig*> Pipelines = GenerateTestPipelineConfigs();
 
-	TArray<FMoviePipelineExecutorJob> Jobs;
+	TArray<FMoviePipelineExecutorJobPrev> Jobs;
 	for(UMoviePipelineMasterConfig* Pipeline : Pipelines)
 	{
-		Jobs.Add(FMoviePipelineExecutorJob(SequencePath, Pipeline));
+		Jobs.Add(FMoviePipelineExecutorJobPrev(SequencePath, Pipeline));
 	}
 
 	Executor->Execute(Jobs);
diff --git a/Engine/Plugins/MovieScene/MovieRenderPipeline/Source/MovieRenderPipelineEditor/Private/MovieRenderPipelineEditorTest.cpp b/Engine/Plugins/MovieScene/MovieRenderPipeline/Source/MovieRenderPipelineEditor/Private/MovieRenderPipelineEditorTest.cpp
--- a/Engine/Plugins/MovieScene/MovieRenderPipeline/Source/MovieRenderPipelineEditor/Private/MovieRenderPipelineEditorTest.cpp
+++ b/Engine/Plugins/MovieScene/MovieRenderPipeline/Source/MovieRenderPipelineEditor/Private/MovieRenderPipelineEditorTest.cpp
@@ -37,7 +37,8 @@ UMoviePipelineShotConfig* GenerateTestShotConfig(UObject* InOwner, int32 InSampl
 		UMoviePipelineHighResSetting* HighRes = OutConfig->FindOrAddSetting<UMoviePipelineHighResSetting>();
 		HighRes->TileCount = InTileCount;
 		HighRes->OverlapPercentage = PadRatioX;
-		HighRes->bWriteAllSamples = false;
+		// HighRes->bWriteAllSamples = true;
+		// HighRes->OverlapPercentage = 0.f;
 	}
 
 	return OutConfig;
@@ -71,9 +72,9 @@ TArray<UMoviePipelineMasterConfig*> FMovieRenderPipelineEditorModule::GenerateTe
 
 				int32 SizeX = 1920;
 				int32 SizeY = 1080;
-				int32 TileX = 1;
-				int32 TileY = 1;
-				int32 TestNumSamples = 2;
+				int32 TileX = 4;
+				int32 TileY = 4;
+				int32 TestNumSamples = 8;
 				float PadRatioX = 0.5f;
 				float PadRatioY = 0.5f;
 				float AccumulationGamma = 1.0f;
diff --git a/Engine/Plugins/MovieScene/MovieRenderPipeline/Source/MovieRenderPipelineEditor/Private/Widgets/SMoviePipelineConfigPanel.cpp b/Engine/Plugins/MovieScene/MovieRenderPipeline/Source/MovieRenderPipelineEditor/Private/Widgets/SMoviePipelineConfigPanel.cpp
new file mode 100644
--- /dev/null
+++ b/Engine/Plugins/MovieScene/MovieRenderPipeline/Source/MovieRenderPipelineEditor/Private/Widgets/SMoviePipelineConfigPanel.cpp
@@ -0,0 +1,261 @@
+// Copyright 1998-2019 Epic Games, Inc. All Rights Reserved.
+
+#include "Widgets/SMoviePipelineConfigPanel.h"
+#include "Widgets/SMoviePipelineEditor.h"
+#include "Widgets/MoviePipelineWidgetConstants.h"
+#include "Misc/MessageDialog.h"
+#include "MoviePipelineShotConfig.h"
+
+// Slate Includes
+#include "Widgets/Text/STextBlock.h"
+#include "Widgets/SBoxPanel.h"
+#include "Widgets/Layout/SBorder.h"
+#include "Widgets/Layout/SBox.h"
+#include "Widgets/Layout/SSpacer.h"
+#include "Widgets/Input/SButton.h"
+#include "Widgets/Input/SComboButton.h"
+#include "Widgets/Images/SImage.h"
+#include "Framework/MultiBox/MultiBoxBuilder.h"
+#include "Styling/SlateIconFinder.h"
+#include "EditorStyleSet.h"
+#include "EditorFontGlyphs.h"
+
+
+// ContentBrowser Includes
+#include "IContentBrowserSingleton.h"
+#include "ContentBrowserModule.h"
+
+// UnrealEd Includes
+#include "ScopedTransaction.h"
+
+#define LOCTEXT_NAMESPACE "SMoviePipelineConfigPanel"
+
+SMoviePipelineConfigPanel::~SMoviePipelineConfigPanel()
+{
+}
+
+PRAGMA_DISABLE_OPTIMIZATION
+void SMoviePipelineConfigPanel::Construct(const FArguments& InArgs, TSubclassOf<UMoviePipelineConfigBase> InConfigType)
+{
+	// Allocate a transient preset automatically so they can start editing without having to create an asset.
+	TransientPreset = AllocateTransientPreset(InConfigType, InArgs._BasePreset);
+
+	// Copy the base preset into the transient preset if it was provided.
+	if (InArgs._BasePreset)
+	{
+		// TransientPreset->CopyFrom(InArgs._BasePreset);
+	}
+
+	// Create the child widgets that need to know about our pipeline
+	MoviePipelineEditorWidget = SNew(SMoviePipelineEditor)
+	.MoviePipeline(this, &SMoviePipelineConfigPanel::GetMoviePipeline);
+
+	ChildSlot
+	[
+		SNew(SVerticalBox)
+
+		// Create the toolbar for adding new stuff, choosing/saving a preset and resetting to the preset defaults.
+		+ SVerticalBox::Slot()
+		.Padding(FMargin(0.f, 1.0f))
+		.AutoHeight()
+		[
+			SNew(SBorder)
+			.BorderImage(FEditorStyle::GetBrush("DetailsView.CategoryTop"))
+			.BorderBackgroundColor(FLinearColor(.6, .6, .6, 1.0f))
+			[
+				SNew(SHorizontalBox)
+
+				// Adding things to the Pipeline
+				+ SHorizontalBox::Slot()
+				.Padding(MoviePipeline::ButtonOffset)
+				.VAlign(VAlign_Fill)
+				.AutoWidth()
+				[
+					MoviePipelineEditorWidget->MakeAddSettingButton()
+				]
+
+				// Presets Management Button
+				+ SHorizontalBox::Slot()
+				.Padding(MoviePipeline::ButtonOffset)
+				.VAlign(VAlign_Fill)
+				.AutoWidth()
+				[
+					SNew(SComboButton)
+					.ContentPadding(MoviePipeline::ButtonPadding)
+					// .ComboButtonStyle(FTakeRecorderStyle::Get(), "ComboButton")
+					.OnGetMenuContent(this, &SMoviePipelineConfigPanel::OnGeneratePresetsMenu)
+					.ForegroundColor(FSlateColor::UseForeground())
+					.ButtonContent()
+					[
+						SNew(SHorizontalBox)
+
+						+ SHorizontalBox::Slot()
+						.AutoWidth()
+						[
+							SNew(SImage)
+							.Image(FSlateIconFinder::FindIconBrushForClass(UMoviePipelineConfigBase::StaticClass()))
+						]
+
+						+ SHorizontalBox::Slot()
+						[
+							SNew(STextBlock)
+							.Text(LOCTEXT("PresetsToolbarButton", "Presets"))
+						]
+					]
+				]
+			]
+		]
+
+		+ SVerticalBox::Slot()
+		.FillHeight(1.0f)
+		[
+			SNew(SHorizontalBox)
+			+ SHorizontalBox::Slot()
+			[
+				MoviePipelineEditorWidget.ToSharedRef()
+			]
+		]
+
+		// Footer Bar
+		+ SVerticalBox::Slot()
+		.AutoHeight()
+		[
+			SNew(SBorder)
+			.BorderImage(FEditorStyle::GetBrush("DetailsView.CategoryTop"))
+			.BorderBackgroundColor(FLinearColor(.6, .6, .6, 1.0f))
+			// .IsEnabled_Lambda([this]() { return !CockpitWidget->Reviewing(); })
+			[
+				SNew(SHorizontalBox)
+
+				// Spacer to make the next widget right aligned
+				+ SHorizontalBox::Slot()
+				[
+					SNew(STextBlock)
+					.Text(LOCTEXT("DummyValidationError", "This is where a validation error might show up?"))
+				]
+
+				// Cancel Changes
+				+ SHorizontalBox::Slot()
+				.Padding(MoviePipeline::ButtonOffset)
+				.VAlign(VAlign_Fill)
+				.HAlign(HAlign_Right)
+				.AutoWidth()
+				[
+					SNew(SButton)
+					.ContentPadding(MoviePipeline::ButtonPadding)
+					.Text(LOCTEXT("CancelChangesButton_Text", "Cancel"))
+					.ToolTipText(LOCTEXT("CancelChangesButton_Tooltip", "Discards changes made and does not modify anything."))
+					// .ForegroundColor(FSlateColor::UseForeground())
+					// .ButtonStyle(FEditorStyle::Get(), "HoverHintOnly")
+					// .OnClicked(this, &SMoviePipelinePanel::OnRevertChanges)
+
+				]
+
+				// Accept Changes
+				+ SHorizontalBox::Slot()
+				.Padding(MoviePipeline::ButtonOffset)
+				.VAlign(VAlign_Fill)
+				.HAlign(HAlign_Right)
+				.AutoWidth()
+				[
+					SNew(SButton)
+					.ContentPadding(MoviePipeline::ButtonPadding)
+					.Text(LOCTEXT("ConfirmChangesButton_Text", "Accept"))
+					.ToolTipText(LOCTEXT("ConfirmChangesButton_Tooltip", "Accepts the changes made and applies them to the particular job instance."))
+					// .ForegroundColor(FSlateColor::UseForeground())
+					// .ButtonStyle(FEditorStyle::Get(), "HoverHintOnly")
+					// .OnClicked(this, &SMoviePipelinePanel::OnRevertChanges)
+
+				]
+			]
+		]
+	];
+}
+
+PRAGMA_ENABLE_OPTIMIZATION
+
+UMoviePipelineConfigBase* SMoviePipelineConfigPanel::AllocateTransientPreset(TSubclassOf<UMoviePipelineConfigBase> InConfigType, UMoviePipelineConfigBase* InTemplatePreset)
+{
+	static const TCHAR* PackageName = TEXT("/Temp/MoviePipeline/PendingConfig");
+
+	// Return a cached transient if it exists
+	UMoviePipelineConfigBase* ExistingPreset = FindObject<UMoviePipelineConfigBase>(nullptr, TEXT("/Temp/MoviePipeline/PendingConfig.PendingConfig"));
+	if (ExistingPreset)
+	{
+		return ExistingPreset;
+	}
+
+	static FName DesiredName = "PendingMoviePipelineConfig";
+	
+	UPackage* NewPackage = CreatePackage(nullptr, PackageName);
+	NewPackage->SetFlags(RF_Transient);
+	NewPackage->AddToRoot();
+
+	UMoviePipelineConfigBase* NewPreset = NewObject<UMoviePipelineConfigBase>(NewPackage, InConfigType, DesiredName, RF_Transient | RF_Transactional | RF_Standalone, InTemplatePreset);
+
+	return NewPreset;
+}
+
+TSharedRef<SWidget> SMoviePipelineConfigPanel::OnGeneratePresetsMenu()
+{
+	FMenuBuilder MenuBuilder(true, nullptr);
+
+	IContentBrowserSingleton& ContentBrowser = FModuleManager::LoadModuleChecked<FContentBrowserModule>("ContentBrowser").Get();
+
+	MenuBuilder.AddMenuEntry(
+		LOCTEXT("SaveAsPreset_Text", "Save As Preset"),
+		LOCTEXT("SaveAsPreset_Tip", "Save the current setup as a new preset that can be imported at a later date"),
+		FSlateIcon(FEditorStyle::Get().GetStyleSetName(), "AssetEditor.SaveAsset.Greyscale"),
+		FUIAction(
+			// FExecuteAction::CreateSP(this, &STakeRecorderPanel::OnSaveAsPreset)
+		)
+	);
+
+	FAssetPickerConfig AssetPickerConfig;
+	{
+		AssetPickerConfig.SelectionMode = ESelectionMode::Single;
+		AssetPickerConfig.InitialAssetViewType = EAssetViewType::Column;
+		AssetPickerConfig.bFocusSearchBoxWhenOpened = true;
+		AssetPickerConfig.bAllowNullSelection = false;
+		AssetPickerConfig.bShowBottomToolbar = true;
+		AssetPickerConfig.bAutohideSearchBar = false;
+		AssetPickerConfig.bAllowDragging = false;
+		AssetPickerConfig.bCanShowClasses = false;
+		AssetPickerConfig.bShowPathInColumnView = true;
+		AssetPickerConfig.bShowTypeInColumnView = false;
+		AssetPickerConfig.bSortByPathInColumnView = false;
+
+		AssetPickerConfig.AssetShowWarningText = LOCTEXT("NoPresets_Warning", "No Presets Found");
+		AssetPickerConfig.Filter.ClassNames.Add(UMoviePipelineConfigBase::StaticClass()->GetFName());
+		// AssetPickerConfig.OnAssetSelected = FOnAssetSelected::CreateSP(this, &SMoviePipelinePanel::OnImportPreset);
+	}
+
+	MenuBuilder.BeginSection(NAME_None, LOCTEXT("ImportPreset_MenuSection", "Import Preset"));
+	{
+		TSharedRef<SWidget> PresetPicker = SNew(SBox)
+			.MinDesiredWidth(400.f)
+			.MinDesiredHeight(400.f)
+			[
+				ContentBrowser.CreateAssetPicker(AssetPickerConfig)
+			];
+
+		MenuBuilder.AddWidget(PresetPicker, FText(), true, false);
+	}
+	MenuBuilder.EndSection();
+
+	return MenuBuilder.MakeWidget();
+}
+
+void SMoviePipelineConfigPanel::AddReferencedObjects(FReferenceCollector& Collector)
+{
+	Collector.AddReferencedObject(TransientPreset);
+	// Collector.AddReferencedObject(SuppliedLevelSequence);
+	// Collector.AddReferencedObject(RecordingLevelSequence);
+}
+
+UMoviePipelineConfigBase* SMoviePipelineConfigPanel::GetMoviePipeline() const
+{
+	return TransientPreset;
+}
+
+#undef LOCTEXT_NAMESPACE // SMoviePipelinePanel
\ No newline at end of file
diff --git a/Engine/Plugins/MovieScene/MovieRenderPipeline/Source/MovieRenderPipelineEditor/Private/Widgets/SMoviePipelineConfigPanel.h b/Engine/Plugins/MovieScene/MovieRenderPipeline/Source/MovieRenderPipelineEditor/Private/Widgets/SMoviePipelineConfigPanel.h
new file mode 100644
--- /dev/null
+++ b/Engine/Plugins/MovieScene/MovieRenderPipeline/Source/MovieRenderPipelineEditor/Private/Widgets/SMoviePipelineConfigPanel.h
@@ -0,0 +1,57 @@
+// Copyright 1998-2019 Epic Games, Inc. All Rights Reserved.
+
+#pragma once
+
+#include "Widgets/DeclarativeSyntaxSupport.h"
+#include "Widgets/SCompoundWidget.h"
+#include "UObject/GCObject.h"
+
+class UMoviePipelineConfigBase;
+class SMoviePipelineEditor;
+
+/**
+ * Outermost widget that is used for setting up a new movie render pipeline config. Operates on a transient UMovieRenderShotConfig that is internally owned and maintained 
+ */
+class SMoviePipelineConfigPanel : public SCompoundWidget, public FGCObject
+{
+public:
+
+	~SMoviePipelineConfigPanel();
+
+	SLATE_BEGIN_ARGS(SMoviePipelineConfigPanel)
+		: _BasePreset(nullptr)
+
+		{}
+
+		/*~ All following arguments are mutually-exclusive */
+		/*-------------------------------------------------*/
+		/** A preset asset to base the pipeline off. This will not get modified */
+		SLATE_ARGUMENT(UMoviePipelineConfigBase*, BasePreset)
+		/*-------------------------------------------------*/
+
+	SLATE_END_ARGS()
+
+	void Construct(const FArguments& InArgs, TSubclassOf<UMoviePipelineConfigBase> InConfigType);
+	UMoviePipelineConfigBase* GetMoviePipeline() const;
+
+private:
+	virtual void AddReferencedObjects(FReferenceCollector& Collector) override;
+
+private:
+	/** Generate the widget that is visible in the Choose Preset dropdown. */
+	TSharedRef<SWidget> OnGeneratePresetsMenu();
+
+	/** Called when the user wishes to revert their changes to the current preset. */
+	FReply OnRevertChanges();
+
+	/** Allocates a transient preset so that the user can use the pipeline without saving it to an asset first. */
+	UMoviePipelineConfigBase* AllocateTransientPreset(TSubclassOf<UMoviePipelineConfigBase> InConfigType, UMoviePipelineConfigBase* InTemplatePreset);
+
+
+private:
+	/** The transient preset that we use - kept alive by AddReferencedObjects */
+	UMoviePipelineConfigBase* TransientPreset;
+
+	/** The main movie pipeline editor widget */
+	TSharedPtr<SMoviePipelineEditor> MoviePipelineEditorWidget;
+};
\ No newline at end of file
diff --git a/Engine/Plugins/MovieScene/MovieRenderPipeline/Source/MovieRenderPipelineEditor/Private/Widgets/SMoviePipelineConfigTabContent.cpp b/Engine/Plugins/MovieScene/MovieRenderPipeline/Source/MovieRenderPipelineEditor/Private/Widgets/SMoviePipelineConfigTabContent.cpp
new file mode 100644
--- /dev/null
+++ b/Engine/Plugins/MovieScene/MovieRenderPipeline/Source/MovieRenderPipelineEditor/Private/Widgets/SMoviePipelineConfigTabContent.cpp
@@ -0,0 +1,52 @@
+// Copyright 1998-2019 Epic Games, Inc. All Rights Reserved.
+
+#include "Widgets/SMoviePipelineConfigTabContent.h"
+#include "Widgets/SMoviePipelineConfigPanel.h"
+
+
+// Analytics
+#include "EngineAnalytics.h"
+#include "Interfaces/IAnalyticsProvider.h"
+#include "MoviePipelineConfigBase.h"
+
+#define LOCTEXT_NAMESPACE "SMoviePipelineTabContent"
+
+
+void SMoviePipelineConfigTabContent::Construct(const FArguments& InArgs)
+{    
+	// Delay one tick before opening the default pipeline setup panel.
+	// this allows anything that just invoked the tab to customize it without the default UI being created
+	RegisterActiveTimer(0.f, FWidgetActiveTimerDelegate::CreateSP(this, &SMoviePipelineConfigTabContent::OnActiveTimer));
+     
+	if (FEngineAnalytics::IsAvailable())
+	{
+		FEngineAnalytics::GetProvider().RecordEvent(TEXT("MoviePipeline.PanelOpened"));
+	}
+}
+
+EActiveTimerReturnType SMoviePipelineConfigTabContent::OnActiveTimer(double InCurrentTime, float InDeltaTime)
+{
+	SetupForPipeline((UMoviePipelineConfigBase*)nullptr);
+	return EActiveTimerReturnType::Stop;
+}
+
+void SMoviePipelineConfigTabContent::SetupForPipeline(UMoviePipelineConfigBase* BasePreset)
+{
+	// Null out the tab content to ensure that all references have been cleaned up before constructing the new one
+	ChildSlot [ SNullWidget::NullWidget ];
+
+	ChildSlot
+	.HAlign(HAlign_Fill)
+	.VAlign(VAlign_Fill)
+	[
+		SAssignNew(WeakPanel, SMoviePipelineConfigPanel, UMoviePipelineConfigBase::StaticClass())
+		.BasePreset(BasePreset)
+	];
+
+	if (FEngineAnalytics::IsAvailable())
+	{
+		FEngineAnalytics::GetProvider().RecordEvent(TEXT("MoviePipeline.SetupForPipelineFromPreset"));
+	}
+}
+
+#undef LOCTEXT_NAMESPACE // SMoviePipelineTabContent
diff --git a/Engine/Plugins/MovieScene/MovieRenderPipeline/Source/MovieRenderPipelineEditor/Private/Widgets/SMoviePipelineConfigTabContent.h b/Engine/Plugins/MovieScene/MovieRenderPipeline/Source/MovieRenderPipelineEditor/Private/Widgets/SMoviePipelineConfigTabContent.h
new file mode 100644
--- /dev/null
+++ b/Engine/Plugins/MovieScene/MovieRenderPipeline/Source/MovieRenderPipelineEditor/Private/Widgets/SMoviePipelineConfigTabContent.h
@@ -0,0 +1,28 @@
+// Copyright 1998-2019 Epic Games, Inc. All Rights Reserved.
+
+#pragma once
+
+#include "Widgets/DeclarativeSyntaxSupport.h"
+#include "Widgets/SCompoundWidget.h"
+
+class UMoviePipelineConfigBase;
+class SMoviePipelineConfigPanel;
+
+class SMoviePipelineConfigTabContent : public SCompoundWidget
+{
+public:
+
+	SLATE_BEGIN_ARGS(SMoviePipelineConfigTabContent){}
+	SLATE_END_ARGS()
+
+	void Construct(const FArguments& InArgs);
+
+	/** Set this tab up for a specific Pipeline preset. Passing nullptr will initialize a default pipeline. */
+	void SetupForPipeline(UMoviePipelineConfigBase* BasePreset);
+
+private:
+	EActiveTimerReturnType OnActiveTimer(double InCurrentTime, float InDeltaTime);
+
+private:
+	TWeakPtr<SMoviePipelineConfigPanel> WeakPanel;
+};
diff --git a/Engine/Plugins/MovieScene/MovieRenderPipeline/Source/MovieRenderPipelineEditor/Private/Widgets/SMoviePipelineEditor.cpp b/Engine/Plugins/MovieScene/MovieRenderPipeline/Source/MovieRenderPipelineEditor/Private/Widgets/SMoviePipelineEditor.cpp
--- a/Engine/Plugins/MovieScene/MovieRenderPipeline/Source/MovieRenderPipelineEditor/Private/Widgets/SMoviePipelineEditor.cpp
+++ b/Engine/Plugins/MovieScene/MovieRenderPipeline/Source/MovieRenderPipelineEditor/Private/Widgets/SMoviePipelineEditor.cpp
@@ -211,10 +211,41 @@ TSharedRef<SWidget> SMoviePipelineEditor::OnGenerateSettingsMenu()
 
 		for (UClass* Class : SourceClasses)
 		{
-			TSubclassOf<UMoviePipelineSetting> SubclassOf = Class;
+			// Get a display name for the setting from the CDO.
+			const UMoviePipelineSetting* SettingDefaultObject = GetDefault<UMoviePipelineSetting>(Class);
+			
+			// Depending on the type of config we're editing, some settings may not be eligible. If this is the case, we omit them from the list.
+			bool bCanSettingBeAdded = CachedMoviePipeline->CanSettingBeAdded(SettingDefaultObject);
+			if (!bCanSettingBeAdded)
+			{
+				continue;
+			}
+
+			// If the setting already exists and it only allows a single instance, we omit them from the list.
+			bool bAllowDuplicates = true;
+			for (UMoviePipelineSetting* ExistingSetting : CachedMoviePipeline->GetSettings())
+			{
+				// If we found a setting with the same class as ours, ask the CDO if multiple are valid.
+				if (ExistingSetting->GetClass() == Class)
+				{
+					bAllowDuplicates = !SettingDefaultObject->IsSolo();
+					break;
+				}
+			}
+			if (!bAllowDuplicates)
+			{
+				continue;
+			}
+
+			FText SettingDisplayName;
+			if (SettingDefaultObject)
+			{
+				SettingDisplayName = SettingDefaultObject->GetDisplayText();
+			}
 
+			TSubclassOf<UMoviePipelineSetting> SubclassOf = Class;
 			MenuBuilder.AddMenuEntry(
-				FText::FromString(Class->GetMetaData(TEXT("MovieRenderPipelineDisplayName"))),
+				SettingDisplayName,
 				Class->GetToolTipText(true),
 				FSlateIconFinder::FindIconForClass(Class),
 				FUIAction(
diff --git a/Engine/Plugins/MovieScene/MovieRenderPipeline/Source/MovieRenderPipelineEditor/Private/Widgets/SMoviePipelinePanel.cpp b/Engine/Plugins/MovieScene/MovieRenderPipeline/Source/MovieRenderPipelineEditor/Private/Widgets/SMoviePipelinePanel.cpp
deleted file mode 100644
diff --git a/Engine/Plugins/MovieScene/MovieRenderPipeline/Source/MovieRenderPipelineEditor/Private/Widgets/SMoviePipelinePanel.h b/Engine/Plugins/MovieScene/MovieRenderPipeline/Source/MovieRenderPipelineEditor/Private/Widgets/SMoviePipelinePanel.h
deleted file mode 100644
diff --git a/Engine/Plugins/MovieScene/MovieRenderPipeline/Source/MovieRenderPipelineEditor/Private/Widgets/SMoviePipelineQueueEditor.cpp b/Engine/Plugins/MovieScene/MovieRenderPipeline/Source/MovieRenderPipelineEditor/Private/Widgets/SMoviePipelineQueueEditor.cpp
new file mode 100644
--- /dev/null
+++ b/Engine/Plugins/MovieScene/MovieRenderPipeline/Source/MovieRenderPipelineEditor/Private/Widgets/SMoviePipelineQueueEditor.cpp
@@ -0,0 +1,817 @@
+// Copyright 1998-2019 Epic Games, Inc. All Rights Reserved.
+
+// Movie Pipeline
+#include "Widgets/SMoviePipelineQueueEditor.h"
+#include "Widgets/MoviePipelineWidgetConstants.h"
+#include "MovieRenderPipelineDataTypes.h"
+#include "MoviePipelineQueue.h"
+#include "MoviePipelineOutputSetting.h"
+#include "MoviePipelineQueueSubsystem.h"
+#include "MoviePipelineMasterConfig.h"
+
+// Slate Includes
+#include "Widgets/Input/SComboButton.h"
+#include "Framework/MultiBox/MultiBoxBuilder.h"
+#include "EditorStyleSet.h"
+#include "EditorFontGlyphs.h"
+#include "Widgets/Input/SCheckBox.h"
+#include "Widgets/Input/SButton.h"
+#include "Widgets/Notifications/SProgressBar.h"
+#include "Widgets/Views/STreeView.h"
+#include "Widgets/Layout/SWidgetSwitcher.h"
+#include "SDropTarget.h"
+
+// Editor
+#include "Editor.h"
+#include "IContentBrowserSingleton.h"
+#include "ContentBrowserModule.h"
+#include "ScopedTransaction.h"
+#include "PropertyCustomizationHelpers.h"
+#include "Framework/Commands/UICommandList.h"
+#include "Framework/Commands/GenericCommands.h"
+#include "DragAndDrop/AssetDragDropOp.h"
+#include "AssetData.h"
+
+// Misc
+#include "LevelSequence.h"
+#include "Engine/EngineTypes.h"
+#include "Framework/Application/SlateApplication.h"
+#include "SMoviePipelineConfigPanel.h"
+#include "Widgets/SWindow.h"
+
+
+#define LOCTEXT_NAMESPACE "SMoviePipelineQueueEditor"
+
+struct FMoviePipelineQueueJobTreeItem;
+struct FMoviePipelineMapTreeItem;
+struct FMoviePipelineShotItem;
+
+struct IMoviePipelineQueueTreeItem : TSharedFromThis<IMoviePipelineQueueTreeItem>
+{
+	virtual ~IMoviePipelineQueueTreeItem() {}
+	
+	virtual TSharedPtr<FMoviePipelineQueueJobTreeItem> AsJob() { return nullptr; }
+	virtual void Delete(UMoviePipelineQueue* InOwningQueue) {}
+
+	virtual TSharedRef<SWidget> ConstructWidget(TWeakPtr<SMoviePipelineQueueEditor> InQueueWidget) = 0;
+};
+
+struct FMoviePipelineQueueJobTreeItem : IMoviePipelineQueueTreeItem
+{
+	/** The job that this tree item represents */
+	TWeakObjectPtr<UMoviePipelineExecutorJob> WeakJob;
+
+	/** Sorted list of this category's children */
+	TArray<TSharedPtr<IMoviePipelineQueueTreeItem>> Children;
+
+	explicit FMoviePipelineQueueJobTreeItem(UMoviePipelineExecutorJob* InJob)
+		: WeakJob(InJob)
+	{}
+
+	virtual TSharedRef<SWidget> ConstructWidget(TWeakPtr<SMoviePipelineQueueEditor> InQueueWidget) override
+	{
+		return SNew(SOverlay)
+
+		// Main Body
+		+SOverlay::Slot()
+		[
+			SNew(SHorizontalBox)
+
+			// Sequence Label
+			+ SHorizontalBox::Slot()
+			.AutoWidth()
+			.VAlign(VAlign_Center)
+			.Padding(8, 0)
+			[
+				SNew(SObjectPropertyEntryBox)
+				.ObjectPath(this, &FMoviePipelineQueueJobTreeItem::GetSequencePath)
+				.AllowedClass(ULevelSequence::StaticClass())
+				.OnObjectChanged(this, &FMoviePipelineQueueJobTreeItem::SetSequencePath)
+				.AllowClear(false)
+				.DisplayUseSelected(false)
+				.DisplayBrowse(true)
+				.DisplayThumbnail(true)
+				.DisplayCompactSize(false)
+			]
+
+			// Preset Label
+			+ SHorizontalBox::Slot()
+			.AutoWidth()
+			.VAlign(VAlign_Center)
+			.Padding(8, 0)
+			[
+				SNew(SObjectPropertyEntryBox)
+				.ObjectPath(this, &FMoviePipelineQueueJobTreeItem::GetMasterConfigPath)
+				.AllowedClass(UMoviePipelineMasterConfig::StaticClass())
+				.OnObjectChanged(this, &FMoviePipelineQueueJobTreeItem::SetMasterConfigPath)
+				.AllowClear(false)
+				.DisplayUseSelected(false)
+				.DisplayBrowse(true)
+				.DisplayThumbnail(true)
+				.DisplayCompactSize(false)
+			]
+
+			// Edit Button
+			+ SHorizontalBox::Slot()
+			.AutoWidth()
+			.VAlign(VAlign_Center)
+			.Padding(8, 0)
+			[
+				SNew(SButton)
+				.OnClicked(InQueueWidget.Pin().Get(), &SMoviePipelineQueueEditor::OnEditMasterConfigForJob, WeakJob)
+				.Content()
+				[
+					SNew(STextBlock)
+					.Text(LOCTEXT("EditMasterConfigButton_Label", "Edit"))
+				]
+			]
+
+			// Output File Label
+			+ SHorizontalBox::Slot()
+			.AutoWidth()
+			.VAlign(VAlign_Center)
+			.Padding(8, 0)
+			[
+				SNew(STextBlock)
+				.Text(this, &FMoviePipelineQueueJobTreeItem::GetOutputLabel)
+				// .TextStyle(FTakeRecorderStyle::Get(), "TakeRecorder.Source.Label")
+				.ColorAndOpacity(this, &FMoviePipelineQueueJobTreeItem::GetOutputColorAndOpacity)
+			]
+
+			// Spacer
+			+ SHorizontalBox::Slot()
+			.FillWidth(1)
+			[
+				SNew(SBox)
+			]
+
+			// Status 
+			+ SHorizontalBox::Slot()
+			.AutoWidth()
+			.VAlign(VAlign_Center)
+			.Padding(8, 0, 24, 0)
+			[
+				SNew(SWidgetSwitcher)
+				.WidgetIndex(this, &FMoviePipelineQueueJobTreeItem::GetStatusIndex)
+
+				// Ready Label
+				+ SWidgetSwitcher::Slot()
+				.VAlign(VAlign_Center)
+				.HAlign(HAlign_Center)
+				[
+					SNew(STextBlock)
+					// .Style(FTakeRecorderStyle::Get(), "TakeRecorder.Source.Switch")
+					.Text(LOCTEXT("PendingJobStatus_Label", "Ready"))
+				]
+
+				// Progress Bar
+				+ SWidgetSwitcher::Slot()
+				.VAlign(VAlign_Center)
+				.HAlign(HAlign_Center)
+				[
+					SNew(SProgressBar)
+					// .Style(FTakeRecorderStyle::Get(), "TakeRecorder.Source.Switch")
+					.Percent(this, &FMoviePipelineQueueJobTreeItem::GetProgressPercent)
+				]
+
+				// Completed
+				+ SWidgetSwitcher::Slot()
+				.VAlign(VAlign_Center)
+				.HAlign(HAlign_Center)
+				[
+					SNew(STextBlock)
+					.Text(LOCTEXT("PendingJobStatus_Label", "Completed!"))
+				]
+			]
+		];
+	}
+
+	virtual TSharedPtr<FMoviePipelineQueueJobTreeItem> AsJob() override
+	{
+		return SharedThis(this);
+	}
+
+	virtual void Delete(UMoviePipelineQueue* InOwningQueue) override
+	{
+		InOwningQueue->DeleteJob(WeakJob.Get());
+	}
+
+
+
+private:
+	FString GetSequencePath() const
+	{
+		UMoviePipelineExecutorJob* Job = WeakJob.Get();
+		if (Job)
+		{
+			return Job->Sequence.ToString();
+		}
+
+		return FString();
+	}
+
+	void SetSequencePath(const FAssetData& AssetData)
+	{
+		UMoviePipelineExecutorJob* Job = WeakJob.Get();
+		if (Job)
+		{
+			Job->Sequence = AssetData.ToSoftObjectPath();
+		}
+	}
+
+	FString GetMasterConfigPath() const
+	{
+		UMoviePipelineExecutorJob* Job = WeakJob.Get();
+		if (Job)
+		{
+			return Job->Configuration->GetPathName();
+		}
+
+		return FString();
+	}
+
+	void SetMasterConfigPath(const FAssetData& AssetData)
+	{
+		UMoviePipelineExecutorJob* Job = WeakJob.Get();
+		if (Job)
+		{
+			Job->Configuration = CastChecked<UMoviePipelineMasterConfig>(AssetData.GetAsset());
+		}
+	}
+
+	FText GetOutputLabel() const
+	{
+		UMoviePipelineExecutorJob* Job = WeakJob.Get();
+		if (Job && Job->Configuration)
+		{
+			UMoviePipelineOutputSetting* OutputSetting = Job->Configuration->FindSetting<UMoviePipelineOutputSetting>();
+			check(OutputSetting);
+
+			return FText::FromString(OutputSetting->OutputDirectory.Path);
+		}
+
+		return LOCTEXT("MissingConfigOutput_Label", "[No Config Set]");
+	}
+
+	FSlateColor GetOutputColorAndOpacity() const
+	{
+		UMoviePipelineExecutorJob* Job = WeakJob.Get();
+		return (Job && Job->HasFinished()) ? FSlateColor::UseSubduedForeground() : FSlateColor::UseForeground();
+	}
+
+	int32 GetStatusIndex() const
+	{
+		UMoviePipelineExecutorJob* Job = WeakJob.Get();
+		if (Job)
+		{
+			// JobStatus 0 is Uninitialized, so we take one off.
+			return FMath::Clamp(int32(Job->JobStatus) - 1, 0, 2);
+		}
+
+		return 0;
+	}
+
+	TOptional<float> GetProgressPercent() const
+	{
+		UMoviePipelineExecutorJob* Job = WeakJob.Get();
+		return Job ? Job->GetProgressPercentage() : TOptional<float>();
+	}
+};
+
+struct FMoviePipelineMapTreeItem : IMoviePipelineQueueTreeItem
+{
+	/** The job that this tree item represents */
+	TWeakObjectPtr<UMoviePipelineExecutorJob> WeakJob;
+
+	explicit FMoviePipelineMapTreeItem(UMoviePipelineExecutorJob* InJob)
+		: WeakJob(InJob)
+	{}
+
+	virtual TSharedRef<SWidget> ConstructWidget(TWeakPtr<SMoviePipelineQueueEditor> InQueueWidget) override
+	{
+		return SNew(SHorizontalBox)
+
+			+ SHorizontalBox::Slot()
+			.AutoWidth()
+			.Padding(20, 4)
+			.VAlign(VAlign_Center)
+			[
+				SNew(STextBlock)
+				.Text(this, &FMoviePipelineMapTreeItem::GetMapLabel)
+			];
+	}
+
+private:
+	FText GetMapLabel() const
+	{
+		UMoviePipelineExecutorJob* Job = WeakJob.Get();
+		if (Job)
+		{
+			UWorld* World = Cast<UWorld>(Job->Map.TryLoad());
+			if (World)
+			{
+				return FText::FromString(World->GetName());
+			}
+		}
+
+		return LOCTEXT("MissingMap_Label", "[Missing Map]");
+	}
+};
+
+struct FMoviePipelineShotItem : IMoviePipelineQueueTreeItem
+{
+	/** The job that this tree item represents */
+	TWeakObjectPtr<UMoviePipelineExecutorJob> WeakJob;
+
+	explicit FMoviePipelineShotItem(UMoviePipelineExecutorJob* InJob)
+		: WeakJob(InJob)
+	{}
+
+	virtual TSharedRef<SWidget> ConstructWidget(TWeakPtr<SMoviePipelineQueueEditor> InQueueWidget) override
+	{
+		return SNew(SHorizontalBox)
+
+			// Toggle Checkbox for deciding to render or not.
+			+ SHorizontalBox::Slot()
+			.AutoWidth()
+			.Padding(20, 4)
+			.VAlign(VAlign_Center)
+			[
+				SNew(SCheckBox)
+				// .Style(FTakeRecorderStyle::Get(), "TakeRecorder.Source.Switch")
+				.IsFocusable(false)
+				.IsChecked(this, &FMoviePipelineShotItem::GetCheckState)
+				.OnCheckStateChanged(this, &FMoviePipelineShotItem::SetCheckState, InQueueWidget)
+			]
+			
+			// Shot Name Label
+			+SHorizontalBox::Slot()
+			.AutoWidth()
+			.Padding(20, 4)
+			.VAlign(VAlign_Center)
+			[
+				SNew(STextBlock)
+				.Text(this, &FMoviePipelineShotItem::GetShotLabel)
+			]
+
+			// Preset Label
+			+ SHorizontalBox::Slot()
+			.AutoWidth()
+			.VAlign(VAlign_Center)
+			.Padding(8, 0)
+			[
+				SNew(STextBlock)
+				.Text(this, &FMoviePipelineShotItem::GetPresetLabel)
+			]
+
+			// Spacer
+			+ SHorizontalBox::Slot()
+			.FillWidth(1)
+			[
+				SNew(SBox)
+			]
+
+			// Status 
+			+ SHorizontalBox::Slot()
+			.AutoWidth()
+			.VAlign(VAlign_Center)
+			.Padding(8, 0, 24, 0)
+			[
+				SNew(SWidgetSwitcher)
+				.WidgetIndex(this, &FMoviePipelineShotItem::GetStatusIndex)
+
+				// Ready Label
+				+ SWidgetSwitcher::Slot()
+				.VAlign(VAlign_Center)
+				.HAlign(HAlign_Center)
+				[
+					SNew(STextBlock)
+					// .Style(FTakeRecorderStyle::Get(), "TakeRecorder.Source.Switch")
+					.Text(LOCTEXT("PendingJobStatus_Label", "Ready"))
+				]
+
+				// Progress Bar
+				+ SWidgetSwitcher::Slot()
+				.VAlign(VAlign_Center)
+				.HAlign(HAlign_Center)
+				[
+					SNew(SProgressBar)
+					// .Style(FTakeRecorderStyle::Get(), "TakeRecorder.Source.Switch")
+					.Percent(this, &FMoviePipelineShotItem::GetProgressPercent)
+				]
+
+				// Completed
+				+ SWidgetSwitcher::Slot()
+				.VAlign(VAlign_Center)
+				.HAlign(HAlign_Center)
+				[
+					SNew(STextBlock)
+					.Text(LOCTEXT("PendingJobStatus_Label", "Completed!"))
+				]
+			]
+			;
+	}
+
+	ECheckBoxState GetCheckState() const
+	{
+		return ECheckBoxState::Checked;
+	}
+
+	void SetCheckState(ECheckBoxState InNewState, TWeakPtr<SMoviePipelineQueueEditor> InQueueWidget)
+	{
+
+	}
+
+	FText GetShotLabel() const
+	{
+		return LOCTEXT("ExampleShotName", "ShotName_001");
+	}
+
+	FText GetPresetLabel() const
+	{
+		return LOCTEXT("ExampleShotName2", "Use Master");
+	}
+
+	int32 GetStatusIndex() const
+	{
+		return 1;
+	}
+
+	TOptional<float> GetProgressPercent() const
+	{
+		return .1f;
+	}
+};
+
+PRAGMA_DISABLE_OPTIMIZATION
+void SMoviePipelineQueueEditor::Construct(const FArguments& InArgs)
+{
+	CachedQueueSerialNumber = uint32(-1);
+
+	TreeView = SNew(STreeView<TSharedPtr<IMoviePipelineQueueTreeItem>>)
+		.TreeItemsSource(&RootNodes)
+		// .OnSelectionChanged(InArgs._OnSelectionChanged)
+		.OnGenerateRow(this, &SMoviePipelineQueueEditor::OnGenerateRow)
+		.OnGetChildren(this, &SMoviePipelineQueueEditor::OnGetChildren);
+
+	CommandList = MakeShared<FUICommandList>();
+	CommandList->MapAction(
+		FGenericCommands::Get().Delete,
+		FExecuteAction::CreateSP(this, &SMoviePipelineQueueEditor::OnDeleteSelected),
+		FCanExecuteAction::CreateSP(this, &SMoviePipelineQueueEditor::CanDeleteSelected)
+	);
+
+	ChildSlot
+	[
+		SNew(SDropTarget)
+		.OnDrop(this, &SMoviePipelineQueueEditor::OnDragDropTarget)
+		.OnAllowDrop(this, &SMoviePipelineQueueEditor::CanDragDropTarget)
+		.OnIsRecognized(this, &SMoviePipelineQueueEditor::CanDragDropTarget)
+		[
+			TreeView.ToSharedRef()
+		]
+	];
+
+	// Test
+	for(int32 i = 0; i < 5; i++)
+	{
+		UMoviePipelineExecutorJob* NewJob = NewObject<UMoviePipelineExecutorJob>();
+		NewJob->AddToRoot();
+		NewJob->Sequence = FSoftObjectPath(TEXT("/Game/Sequencer/MovieRenderPipeline/PipelineTestSeq"));
+		NewJob->Map = FSoftObjectPath(TEXT("/Game/Sequencer/MovieRenderPipeline/MovieRenderPipelineTest"));
+
+		TSharedPtr< FMoviePipelineQueueJobTreeItem> ExampleJobTreeItem = MakeShared<FMoviePipelineQueueJobTreeItem>(NewJob);
+		RootNodes.Add(ExampleJobTreeItem);
+
+		TSharedPtr<FMoviePipelineMapTreeItem> MapLine = MakeShared<FMoviePipelineMapTreeItem>(NewJob);
+		ExampleJobTreeItem->Children.Add(MapLine);
+
+		for (int32 j = 0; j < 2; j++)
+		{
+			TSharedPtr<FMoviePipelineShotItem> Shot = MakeShared<FMoviePipelineShotItem>(NewJob);
+			ExampleJobTreeItem->Children.Add(Shot);
+		}
+	}
+}
+
+TSharedRef<SWidget> SMoviePipelineQueueEditor::MakeAddSequenceJobButton()
+{
+	return SNew(SComboButton)
+		.ContentPadding(MoviePipeline::ButtonPadding)
+		// .ButtonStyle(FTakeRecorderStyle::Get(), "FlatButton.Success")
+		.OnGetMenuContent(this, &SMoviePipelineQueueEditor::OnGenerateNewJobFromAssetMenu)
+		.ForegroundColor(FSlateColor::UseForeground())
+		.HasDownArrow(false)
+		.ButtonContent()
+		[
+			SNew(SHorizontalBox)
+
+			// Plus Icon
+			+ SHorizontalBox::Slot()
+			.VAlign(VAlign_Center)
+			.AutoWidth()
+			[
+				SNew(STextBlock)
+				.TextStyle(FEditorStyle::Get(), "NormalText.Important")
+				.Font(FEditorStyle::Get().GetFontStyle("FontAwesome.10"))
+				.Text(FEditorFontGlyphs::Plus)
+			]
+
+			// "Job" Text
+			+ SHorizontalBox::Slot()
+			.VAlign(VAlign_Center)
+			.AutoWidth()
+			.Padding(4, 0, 0, 0)
+			[
+				SNew(STextBlock)
+				.TextStyle(FEditorStyle::Get(), "NormalText.Important")
+				.Text(LOCTEXT("AddNewSetting_Text", "Job"))
+			]
+
+			// Non-Default Down Caret arrow.
+			+ SHorizontalBox::Slot()
+			.VAlign(VAlign_Center)
+			.AutoWidth()
+			.Padding(4, 0, 0, 0)
+			[
+				SNew(STextBlock)
+				.TextStyle(FEditorStyle::Get(), "NormalText.Important")
+				.Font(FEditorStyle::Get().GetFontStyle("FontAwesome.10"))
+				.Text(FEditorFontGlyphs::Caret_Down)
+			]
+		];
+}
+
+TSharedRef<SWidget> SMoviePipelineQueueEditor::RemoveSelectedJobButton()
+{
+	return SNew(SButton)
+		.IsEnabled(this, &SMoviePipelineQueueEditor::CanDeleteSelected)
+		.OnClicked(this, &SMoviePipelineQueueEditor::DeleteSelected)
+		[
+			SNew(STextBlock)
+			.TextStyle(FEditorStyle::Get(), "NormalText.Important")
+			.Font(FEditorStyle::Get().GetFontStyle("FontAwesome.10"))
+			.Text(FEditorFontGlyphs::Minus)
+		];
+}
+
+TSharedRef<SWidget> SMoviePipelineQueueEditor::OnGenerateNewJobFromAssetMenu()
+{
+	FMenuBuilder MenuBuilder(true, nullptr);
+	IContentBrowserSingleton& ContentBrowser = FModuleManager::LoadModuleChecked<FContentBrowserModule>("ContentBrowser").Get();
+
+	FAssetPickerConfig AssetPickerConfig;
+	{
+		AssetPickerConfig.SelectionMode = ESelectionMode::Single;
+		AssetPickerConfig.InitialAssetViewType = EAssetViewType::List;
+		AssetPickerConfig.bFocusSearchBoxWhenOpened = true;
+		AssetPickerConfig.bAllowNullSelection = false;
+		AssetPickerConfig.bShowBottomToolbar = true;
+		AssetPickerConfig.bAutohideSearchBar = false;
+		AssetPickerConfig.bAllowDragging = false;
+		AssetPickerConfig.bCanShowClasses = false;
+		AssetPickerConfig.bShowPathInColumnView = true;
+		AssetPickerConfig.bShowTypeInColumnView = false;
+		AssetPickerConfig.bSortByPathInColumnView = false;
+		AssetPickerConfig.ThumbnailScale = 0.2f;
+
+		AssetPickerConfig.AssetShowWarningText = LOCTEXT("NoSequences_Warning", "No Level Sequences Found");
+		AssetPickerConfig.Filter.ClassNames.Add(ULevelSequence::StaticClass()->GetFName());
+		AssetPickerConfig.OnAssetSelected = FOnAssetSelected::CreateSP(this, &SMoviePipelineQueueEditor::OnCreateJobFromAsset);
+	}
+
+	MenuBuilder.BeginSection(NAME_None, LOCTEXT("NewJob_MenuSection", "New Render Job"));
+	{
+		TSharedRef<SWidget> PresetPicker = SNew(SBox)
+			.WidthOverride(300.f)
+			.HeightOverride(300.f)
+			[
+				ContentBrowser.CreateAssetPicker(AssetPickerConfig)
+			];
+
+		MenuBuilder.AddWidget(PresetPicker, FText(), true, false);
+	}
+	MenuBuilder.EndSection();
+
+	return MenuBuilder.MakeWidget();
+}
+PRAGMA_ENABLE_OPTIMIZATION
+
+void SMoviePipelineQueueEditor::OnCreateJobFromAsset(const FAssetData& InAsset)
+{
+	// Close the dropdown menu that showed them the assets to pick from.
+	FSlateApplication::Get().DismissAllMenus();
+
+	// Only try to initialize level sequences, in the event they had more than a level sequence selected when drag/dropping.
+	ULevelSequence* LevelSequence = Cast<ULevelSequence>(InAsset.GetAsset());
+	if (LevelSequence)
+	{
+		FScopedTransaction Transaction(FText::Format(LOCTEXT("CreateJob_Transaction", "Add {0}|plural(one=Job, other=Jobs)"), 1));
+		
+		UMoviePipelineQueue* ActiveQueue = GEditor->GetEditorSubsystem<UMoviePipelineQueueSubsystem>()->GetQueue();
+		check(ActiveQueue);
+
+		
+		// We'll assume they went to render from the current world - they can always override it later.
+		FSoftObjectPath CurrentWorld = FSoftObjectPath(GEditor->GetEditorWorldContext().World());
+		FSoftObjectPath Sequence = InAsset.ToSoftObjectPath();
+
+		ActiveQueue->Modify();
+		UMoviePipelineExecutorJob* NewJob = ActiveQueue->AllocateNewJob();
+		NewJob->Modify();
+		NewJob->Sequence = Sequence;
+		NewJob->Map = CurrentWorld;
+	}
+}
+
+void SMoviePipelineQueueEditor::Tick(const FGeometry& AllottedGeometry, const double InCurrentTime, const float InDeltaTime)
+{
+	UMoviePipelineQueue* ActiveQueue = GEditor->GetEditorSubsystem<UMoviePipelineQueueSubsystem>()->GetQueue();
+	check(ActiveQueue);
+
+	if (ActiveQueue)
+	{
+		if (CachedQueueSerialNumber != ActiveQueue->GetQueueSerialNumber())
+		{
+			ReconstructTree();
+		}
+	}
+	// The sources are no longer valid, so we expect our cached serial number to be -1. If not, we haven't reset the tree yet.
+	else if (CachedQueueSerialNumber != uint32(-1))
+	{
+		ReconstructTree();
+	}
+}
+
+void SMoviePipelineQueueEditor::ReconstructTree()
+{
+	UMoviePipelineQueue* ActiveQueue = GEditor->GetEditorSubsystem<UMoviePipelineQueueSubsystem>()->GetQueue();
+	check(ActiveQueue);
+	if (!ActiveQueue)
+	{
+		CachedQueueSerialNumber = uint32(-1);
+		RootNodes.Reset();
+		return;
+	}
+
+	CachedQueueSerialNumber = ActiveQueue->GetQueueSerialNumber();
+
+	// TSortedMap<FString, TSharedPtr<FMoviePipelineQueueJobTreeItem>> RootJobs;
+	// for (TSharedPtr<IMoviePipelineQueueTreeItem> RootItem : RootNodes)
+	// {
+	// 	TSharedPtr<FMoviePipelineQueueJobTreeItem> RootCategory = RootItem->AsJob();
+	// 	if (RootCategory.IsValid())
+	// 	{
+	// 		RootCategory->Children.Reset();
+	// 		RootJobs.Add(RootCategory->Category.ToString(), RootCategory);
+	// 	}
+	// }
+
+	RootNodes.Reset();
+
+	// We attempt to re-use tree items in order to maintain selection states on them
+	// TMap<FObjectKey, TSharedPtr<FTakeRecorderSourceTreeItem>> OldSourceToTreeItem;
+	// Swap(SourceToTreeItem, OldSourceToTreeItem);
+
+	for (UMoviePipelineExecutorJob* Job : ActiveQueue->GetJobs())
+	{
+		if (!Job)
+		{
+			continue;
+		}
+
+		TSharedPtr< FMoviePipelineQueueJobTreeItem> JobTreeItem = MakeShared< FMoviePipelineQueueJobTreeItem>(Job);
+		TSharedPtr<FMoviePipelineMapTreeItem> MapTreeItem = MakeShared< FMoviePipelineMapTreeItem>(Job);
+		JobTreeItem->Children.Add(MapTreeItem);
+		RootNodes.Add(JobTreeItem);
+	}
+
+	TreeView->RequestTreeRefresh();
+}
+
+
+FReply SMoviePipelineQueueEditor::OnKeyDown(const FGeometry& MyGeometry, const FKeyEvent& InKeyEvent)
+{
+	if (CommandList->ProcessCommandBindings(InKeyEvent))
+	{
+		return FReply::Handled();
+	}
+
+	return FReply::Unhandled();
+}
+
+TSharedRef<ITableRow> SMoviePipelineQueueEditor::OnGenerateRow(TSharedPtr<IMoviePipelineQueueTreeItem> Item, const TSharedRef<STableViewBase>& Tree)
+{
+	return
+		SNew(STableRow<TSharedPtr<IMoviePipelineQueueTreeItem>>, Tree)
+		[
+			Item->ConstructWidget(SharedThis(this))
+		];
+}
+
+void SMoviePipelineQueueEditor::OnGetChildren(TSharedPtr<IMoviePipelineQueueTreeItem> Item, TArray<TSharedPtr<IMoviePipelineQueueTreeItem>>& OutChildItems)
+{
+	TSharedPtr<FMoviePipelineQueueJobTreeItem> Job = Item->AsJob();
+	if (Job.IsValid())
+	{
+		OutChildItems.Append(Job->Children);
+	}
+}
+
+FReply SMoviePipelineQueueEditor::OnDragDropTarget(TSharedPtr<FDragDropOperation> InOperation)
+{
+	if (InOperation)
+	{
+		if (InOperation->IsOfType<FAssetDragDropOp>())
+		{
+			TSharedPtr<FAssetDragDropOp> AssetDragDrop = StaticCastSharedPtr<FAssetDragDropOp>(InOperation);
+			FScopedTransaction Transaction(FText::Format(LOCTEXT("CreateJob_Transaction", "Add {0}|plural(one=Job, other=Jobs)"), AssetDragDrop->GetAssets().Num()));
+
+			for (const FAssetData& Asset : AssetDragDrop->GetAssets())
+			{
+				OnCreateJobFromAsset(Asset);
+			}
+
+			return FReply::Handled();
+		}
+	}
+
+	return FReply::Unhandled();
+}
+
+bool SMoviePipelineQueueEditor::CanDragDropTarget(TSharedPtr<FDragDropOperation> InOperation)
+{
+	bool bIsValid = false;
+	if (InOperation)
+	{
+		if (InOperation->IsOfType<FAssetDragDropOp>())
+		{
+			TSharedPtr<FAssetDragDropOp> AssetDragDrop = StaticCastSharedPtr<FAssetDragDropOp>(InOperation);
+			for (const FAssetData& Asset : AssetDragDrop->GetAssets())
+			{
+				ULevelSequence* LevelSequence = Cast<ULevelSequence>(Asset.GetAsset());
+				if (LevelSequence)
+				{
+					// If at least one of them is a Level Sequence then we'll accept the drop.
+					bIsValid = true;
+					break;
+				}
+			}
+		}
+	}
+
+	return bIsValid;
+}
+
+FReply SMoviePipelineQueueEditor::DeleteSelected()
+{
+	UMoviePipelineQueue* ActiveQueue = GEditor->GetEditorSubsystem<UMoviePipelineQueueSubsystem>()->GetQueue();
+	check(ActiveQueue);
+
+	if (ActiveQueue)
+	{
+		TArray<TSharedPtr<IMoviePipelineQueueTreeItem>> Items = TreeView->GetSelectedItems();
+
+		FScopedTransaction Transaction(FText::Format(LOCTEXT("DeleteSelection", "Delete Selected {0}|plural(one=Job, other=Jobs)"), Items.Num()));
+		ActiveQueue->Modify();
+
+		for (TSharedPtr<IMoviePipelineQueueTreeItem> Item : Items)
+		{
+			Item->Delete(ActiveQueue);
+		}
+	}
+
+	return FReply::Handled();
+}
+
+void SMoviePipelineQueueEditor::OnDeleteSelected()
+{
+	DeleteSelected();
+}
+
+bool SMoviePipelineQueueEditor::CanDeleteSelected() const
+{
+	return true;
+}
+
+FReply SMoviePipelineQueueEditor::OnEditMasterConfigForJob(TWeakObjectPtr<UMoviePipelineExecutorJob> InJob)
+{
+	TSharedRef<SWindow> EditorWindow =
+		SNew(SWindow)
+		.ClientSize(FVector2D(700, 600));
+	
+	TSharedRef<SMoviePipelineConfigPanel> ConfigEditorPanel = 
+		SNew(SMoviePipelineConfigPanel, UMoviePipelineMasterConfig::StaticClass())
+		.BasePreset(InJob->Configuration);
+
+	EditorWindow->SetContent(ConfigEditorPanel);
+
+	TSharedPtr<SWindow> ParentWindow = FSlateApplication::Get().FindWidgetWindow(AsShared());
+	if (ParentWindow.IsValid())
+	{
+		FSlateApplication::Get().AddWindowAsNativeChild(EditorWindow, ParentWindow.ToSharedRef());
+	}
+	
+	return FReply::Handled();
+}
+#undef LOCTEXT_NAMESPACE
\ No newline at end of file
diff --git a/Engine/Plugins/MovieScene/MovieRenderPipeline/Source/MovieRenderPipelineEditor/Private/Widgets/SMoviePipelineQueueEditor.h b/Engine/Plugins/MovieScene/MovieRenderPipeline/Source/MovieRenderPipelineEditor/Private/Widgets/SMoviePipelineQueueEditor.h
new file mode 100644
--- /dev/null
+++ b/Engine/Plugins/MovieScene/MovieRenderPipeline/Source/MovieRenderPipelineEditor/Private/Widgets/SMoviePipelineQueueEditor.h
@@ -0,0 +1,64 @@
+// Copyright 1998-2019 Epic Games, Inc. All Rights Reserved.
+
+#pragma once
+
+#include "Widgets/DeclarativeSyntaxSupport.h"
+#include "Widgets/SCompoundWidget.h"
+
+struct FAssetData;
+struct IMoviePipelineQueueTreeItem;
+template<typename> class STreeView;
+class FUICommandList;
+class ITableRow;
+class STableViewBase;
+class UMoviePipelineExecutorJob;
+
+
+/**
+ * Widget used to edit a Movie Pipeline Queue
+ */
+class SMoviePipelineQueueEditor : public SCompoundWidget
+{
+public:
+
+	SLATE_BEGIN_ARGS(SMoviePipelineQueueEditor)
+		{}
+
+	SLATE_END_ARGS()
+
+	void Construct(const FArguments& InArgs);
+
+	TSharedRef<SWidget> MakeAddSequenceJobButton();
+	TSharedRef<SWidget> RemoveSelectedJobButton();
+	TSharedRef<SWidget> OnGenerateNewJobFromAssetMenu();
+
+	FReply OnEditMasterConfigForJob(TWeakObjectPtr<UMoviePipelineExecutorJob> InJob);
+
+private:
+	// SWidget Interface
+	FReply OnKeyDown(const FGeometry& MyGeometry, const FKeyEvent& InKeyEvent);
+	void Tick(const FGeometry& AllottedGeometry, const double InCurrentTime, const float InDeltaTime);
+	// ~SWidget Interface
+private:
+	void OnCreateJobFromAsset(const FAssetData& InAsset);
+
+	TSharedRef<ITableRow> OnGenerateRow(TSharedPtr<IMoviePipelineQueueTreeItem> Item, const TSharedRef<STableViewBase>& Tree);
+
+	void OnGetChildren(TSharedPtr<IMoviePipelineQueueTreeItem> Item, TArray<TSharedPtr<IMoviePipelineQueueTreeItem>>& OutChildItems);
+
+	FReply OnDragDropTarget(TSharedPtr<FDragDropOperation> InOperation);
+
+	bool CanDragDropTarget(TSharedPtr<FDragDropOperation> InOperation);
+
+	void OnDeleteSelected();
+	bool CanDeleteSelected() const;
+	FReply DeleteSelected();
+
+	void ReconstructTree();
+
+private:
+	TArray<TSharedPtr<IMoviePipelineQueueTreeItem>> RootNodes;
+	TSharedPtr<STreeView<TSharedPtr<IMoviePipelineQueueTreeItem>>> TreeView;
+	TSharedPtr<FUICommandList> CommandList;
+	uint32 CachedQueueSerialNumber;
+};
\ No newline at end of file
diff --git a/Engine/Plugins/MovieScene/MovieRenderPipeline/Source/MovieRenderPipelineEditor/Private/Widgets/SMoviePipelineQueuePanel.cpp b/Engine/Plugins/MovieScene/MovieRenderPipeline/Source/MovieRenderPipelineEditor/Private/Widgets/SMoviePipelineQueuePanel.cpp
new file mode 100644
--- /dev/null
+++ b/Engine/Plugins/MovieScene/MovieRenderPipeline/Source/MovieRenderPipelineEditor/Private/Widgets/SMoviePipelineQueuePanel.cpp
@@ -0,0 +1,194 @@
+// Copyright 1998-2019 Epic Games, Inc. All Rights Reserved.
+
+// Movie Pipeline Includes
+#include "Widgets/SMoviePipelineQueuePanel.h"
+#include "Widgets/MoviePipelineWidgetConstants.h"
+#include "SMoviePipelineQueueEditor.h"
+#include "MovieRenderPipelineSettings.h"
+
+// Slate Includes
+#include "Widgets/SBoxPanel.h"
+#include "Widgets/Layout/SBorder.h"
+#include "Widgets/Text/STextBlock.h"
+#include "Widgets/Input/SButton.h"
+#include "EditorStyleSet.h"
+
+// Misc
+
+#define LOCTEXT_NAMESPACE "SMoviePipelineQueuePanel"
+
+PRAGMA_DISABLE_OPTIMIZATION
+void SMoviePipelineQueuePanel::Construct(const FArguments& InArgs)
+{
+	// Allocate a transient preset automatically so they can start editing without having to create an asset.
+	// TransientPreset = AllocateTransientPreset();
+
+	// Copy the base preset into the transient preset if it was provided.
+	if (InArgs._BasePreset)
+	{
+		// TransientPreset->CopyFrom(InArgs._BasePreset);
+	}
+
+	// Create the child widgets that need to know about our pipeline
+	PipelineQueueEditorWidget = SNew(SMoviePipelineQueueEditor);
+	// .MoviePipeline(this, &SMoviePipelineConfigPanel::GetMoviePipeline);
+
+	ChildSlot
+	[
+		SNew(SVerticalBox)
+
+		// Create the toolbar for adding new items to the queue
+		+ SVerticalBox::Slot()
+		.Padding(FMargin(0.f, 1.0f))
+		.AutoHeight()
+		[
+			SNew(SBorder)
+			.BorderImage(FEditorStyle::GetBrush("DetailsView.CategoryTop"))
+			.BorderBackgroundColor(FLinearColor(.6, .6, .6, 1.0f))
+			[
+				SNew(SHorizontalBox)
+
+				// Add a Level Sequence to the queue 
+				+ SHorizontalBox::Slot()
+				.Padding(MoviePipeline::ButtonOffset)
+				.VAlign(VAlign_Fill)
+				.AutoWidth()
+				[
+					PipelineQueueEditorWidget->MakeAddSequenceJobButton()
+				]
+
+				// Remove a job (potentially already processed) from the the queue 
+				+ SHorizontalBox::Slot()
+				.Padding(MoviePipeline::ButtonOffset)
+				.VAlign(VAlign_Fill)
+				.AutoWidth()
+				[
+					PipelineQueueEditorWidget->RemoveSelectedJobButton()
+				]	
+			]
+		]
+	
+		// Main Queue Body
+		+ SVerticalBox::Slot()
+		.FillHeight(1.0f)
+		[
+			SNew(SHorizontalBox)
+			+ SHorizontalBox::Slot()
+			[
+				PipelineQueueEditorWidget.ToSharedRef()
+			]
+		]
+
+		// Footer Bar
+		+ SVerticalBox::Slot()
+		.AutoHeight()
+		[
+			SNew(SBorder)
+			.BorderImage(FEditorStyle::GetBrush("DetailsView.CategoryTop"))
+			.BorderBackgroundColor(FLinearColor(.6, .6, .6, 1.0f))
+			.Padding(FMargin(0, 2, 0, 2))
+			[
+				SNew(SHorizontalBox)
+
+				+ SHorizontalBox::Slot()
+				.VAlign(VAlign_Fill)
+				.HAlign(HAlign_Left)
+				.FillWidth(1.f)
+				[
+					SNullWidget::NullWidget
+				]
+
+				// Render Local in Process
+				+ SHorizontalBox::Slot()
+				.Padding(MoviePipeline::ButtonOffset)
+				.VAlign(VAlign_Fill)
+				.HAlign(HAlign_Right)
+				.AutoWidth()
+				[
+					SNew(SButton)
+					.ContentPadding(MoviePipeline::ButtonPadding)
+					.Text(LOCTEXT("RenderQueueLocal_Text", "Render (Local)"))
+					.ToolTipText(LOCTEXT("RenderQueueLocal_Tooltip", "Renders the current queue in the current process using Play in Editor."))
+					.ForegroundColor(FSlateColor::UseForeground())
+					.IsEnabled(this, &SMoviePipelineQueuePanel::IsRenderLocalEnabled)
+					.OnClicked(this, &SMoviePipelineQueuePanel::OnRenderLocalRequested)
+				]
+
+				// Render Remotely (Separate Process or Farm)
+				+ SHorizontalBox::Slot()
+				.Padding(MoviePipeline::ButtonOffset)
+				.VAlign(VAlign_Fill)
+				.HAlign(HAlign_Right)
+				.AutoWidth()
+				[
+					SNew(SButton)
+					.ContentPadding(MoviePipeline::ButtonPadding)
+					.Text(LOCTEXT("RenderQueueRemote_Text", "Render (Remote)"))
+					.ToolTipText(LOCTEXT("RenderQueueRemote_Tooltip", "Renders the current queue in a separate process."))
+					.ForegroundColor(FSlateColor::UseForeground())
+					.IsEnabled(this, &SMoviePipelineQueuePanel::IsRenderLocalEnabled)
+					.OnClicked(this, &SMoviePipelineQueuePanel::OnRenderLocalRequested)
+				]
+			]
+		]
+	];
+}
+
+PRAGMA_ENABLE_OPTIMIZATION
+
+FReply SMoviePipelineQueuePanel::OnRenderLocalRequested()
+{
+	UE_LOG(LogTemp, Log, TEXT("Render locally!"));
+	return FReply::Handled();
+}
+
+bool SMoviePipelineQueuePanel::IsRenderLocalEnabled() const
+{
+	const UMovieRenderPipelineProjectSettings* ProjectSettings = GetDefault<UMovieRenderPipelineProjectSettings>();
+	return ProjectSettings->DefaultLocalExecutor != nullptr;
+}
+
+FReply SMoviePipelineQueuePanel::OnRenderRemoteRequested()
+{
+	UE_LOG(LogTemp, Log, TEXT("Render remotely!"));
+	return FReply::Handled();
+}
+
+bool SMoviePipelineQueuePanel::IsRenderRemoteEnabled() const
+{
+	const UMovieRenderPipelineProjectSettings* ProjectSettings = GetDefault<UMovieRenderPipelineProjectSettings>();
+	return ProjectSettings->DefaultRemoteExecutor != nullptr;
+}
+
+/*UMoviePipelineConfigBase* SMoviePipelineConfigPanel::AllocateTransientPreset()
+{
+	static const TCHAR* PackageName = TEXT("/Temp/MoviePipeline/PendingPipeline");
+
+	// Return a cached transient if it exists
+	UMoviePipelineConfigBase* ExistingPreset = FindObject<UMoviePipelineConfigBase>(nullptr, TEXT("/Temp/MoviePipeline/PendingPipeline.PendingPipeline"));
+	if (ExistingPreset)
+	{
+		return ExistingPreset;
+	}
+
+	static FName DesiredName = "PendingMoviePipeline";
+	
+	UPackage* NewPackage = CreatePackage(nullptr, PackageName);
+	NewPackage->SetFlags(RF_Transient);
+	NewPackage->AddToRoot();
+
+	UMoviePipelineConfigBase* NewPreset = NewObject<UMoviePipelineConfigBase>(NewPackage, DesiredName, RF_Transient | RF_Transactional | RF_Standalone);
+
+	return NewPreset;
+}*/
+
+void SMoviePipelineQueuePanel::AddReferencedObjects(FReferenceCollector& Collector)
+{
+	// Collector.AddReferencedObject(TransientPreset);
+	// Collector.AddReferencedObject(SuppliedLevelSequence);
+	// Collector.AddReferencedObject(RecordingLevelSequence);
+}
+
+
+
+#undef LOCTEXT_NAMESPACE // SMoviePipelineQueuePanel
\ No newline at end of file
diff --git a/Engine/Plugins/MovieScene/MovieRenderPipeline/Source/MovieRenderPipelineEditor/Private/Widgets/SMoviePipelineQueuePanel.h b/Engine/Plugins/MovieScene/MovieRenderPipeline/Source/MovieRenderPipelineEditor/Private/Widgets/SMoviePipelineQueuePanel.h
new file mode 100644
--- /dev/null
+++ b/Engine/Plugins/MovieScene/MovieRenderPipeline/Source/MovieRenderPipelineEditor/Private/Widgets/SMoviePipelineQueuePanel.h
@@ -0,0 +1,53 @@
+// Copyright 1998-2019 Epic Games, Inc. All Rights Reserved.
+
+#pragma once
+
+#include "Widgets/DeclarativeSyntaxSupport.h"
+#include "Widgets/SCompoundWidget.h"
+#include "UObject/GCObject.h"
+
+class UMoviePipelineConfigBase;
+class SMoviePipelineQueueEditor;
+
+/**
+ * Outermost widget that is used for setting up a new movie render pipeline queue. Operates on a transient object that is internally owned and maintained 
+ */
+class SMoviePipelineQueuePanel : public SCompoundWidget, public FGCObject
+{
+public:
+	SLATE_BEGIN_ARGS(SMoviePipelineQueuePanel)
+		: _BasePreset(nullptr)
+
+		{}
+
+		/*~ All following arguments are mutually-exclusive */
+		/*-------------------------------------------------*/
+		/** A preset asset to base the pipeline off */
+		SLATE_ARGUMENT(UMoviePipelineConfigBase*, BasePreset)
+
+	SLATE_END_ARGS()
+
+	void Construct(const FArguments& InArgs);
+
+private:
+	// FGCObject Interface
+	virtual void AddReferencedObjects(FReferenceCollector& Collector) override;
+	// ~FGCObject Interface
+
+	FReply OnRenderLocalRequested();
+	bool IsRenderLocalEnabled() const;
+	FReply OnRenderRemoteRequested();
+	bool IsRenderRemoteEnabled() const;
+
+private:
+	/** Allocates a transient preset so that the user can use the pipeline without saving it to an asset first. */
+	//UMoviePipelineConfigBase* AllocateTransientPreset();
+
+
+private:
+	/** The main movie pipeline queue editor widget */
+	TSharedPtr<SMoviePipelineQueueEditor> PipelineQueueEditorWidget;
+
+	/** The transient preset that we use - kept alive by AddReferencedObjects */
+	// UMoviePipelineConfigBase* TransientPreset;
+};
\ No newline at end of file
diff --git a/Engine/Plugins/MovieScene/MovieRenderPipeline/Source/MovieRenderPipelineEditor/Private/Widgets/SMoviePipelineQueueTabContent.cpp b/Engine/Plugins/MovieScene/MovieRenderPipeline/Source/MovieRenderPipelineEditor/Private/Widgets/SMoviePipelineQueueTabContent.cpp
new file mode 100644
--- /dev/null
+++ b/Engine/Plugins/MovieScene/MovieRenderPipeline/Source/MovieRenderPipelineEditor/Private/Widgets/SMoviePipelineQueueTabContent.cpp
@@ -0,0 +1,31 @@
+// Copyright 1998-2019 Epic Games, Inc. All Rights Reserved.
+
+#include "Widgets/SMoviePipelineQueueTabContent.h"
+#include "SMoviePipelineQueuePanel.h"
+
+// Analytics
+#include "EngineAnalytics.h"
+#include "Interfaces/IAnalyticsProvider.h"
+
+#define LOCTEXT_NAMESPACE "SMoviePipelineQueueTabContent"
+
+
+void SMoviePipelineQueueTabContent::Construct(const FArguments& InArgs)
+{    
+	ChildSlot
+		.HAlign(HAlign_Fill)
+		.VAlign(VAlign_Fill)
+		[
+			SAssignNew(WeakPanel, SMoviePipelineQueuePanel)
+		];
+
+
+	if (FEngineAnalytics::IsAvailable())
+	{
+		FEngineAnalytics::GetProvider().RecordEvent(TEXT("MoviePipeline.QueueTabOpened"));
+	}
+}
+
+
+
+#undef LOCTEXT_NAMESPACE // SMoviePipelineQueueTabContent
diff --git a/Engine/Plugins/MovieScene/MovieRenderPipeline/Source/MovieRenderPipelineEditor/Private/Widgets/SMoviePipelineQueueTabContent.h b/Engine/Plugins/MovieScene/MovieRenderPipeline/Source/MovieRenderPipelineEditor/Private/Widgets/SMoviePipelineQueueTabContent.h
new file mode 100644
--- /dev/null
+++ b/Engine/Plugins/MovieScene/MovieRenderPipeline/Source/MovieRenderPipelineEditor/Private/Widgets/SMoviePipelineQueueTabContent.h
@@ -0,0 +1,21 @@
+// Copyright 1998-2019 Epic Games, Inc. All Rights Reserved.
+
+#pragma once
+
+#include "Widgets/DeclarativeSyntaxSupport.h"
+#include "Widgets/SCompoundWidget.h"
+
+class SMoviePipelineQueuePanel;
+
+class SMoviePipelineQueueTabContent : public SCompoundWidget
+{
+public:
+
+	SLATE_BEGIN_ARGS(SMoviePipelineQueueTabContent){}
+	SLATE_END_ARGS()
+
+	void Construct(const FArguments& InArgs);
+
+private:
+	TWeakPtr<SMoviePipelineQueuePanel> WeakPanel;
+};
diff --git a/Engine/Plugins/MovieScene/MovieRenderPipeline/Source/MovieRenderPipelineEditor/Private/Widgets/SMoviePipelineSettings.cpp b/Engine/Plugins/MovieScene/MovieRenderPipeline/Source/MovieRenderPipelineEditor/Private/Widgets/SMoviePipelineSettings.cpp
--- a/Engine/Plugins/MovieScene/MovieRenderPipeline/Source/MovieRenderPipelineEditor/Private/Widgets/SMoviePipelineSettings.cpp
+++ b/Engine/Plugins/MovieScene/MovieRenderPipeline/Source/MovieRenderPipelineEditor/Private/Widgets/SMoviePipelineSettings.cpp
@@ -100,22 +100,6 @@ struct FMoviePipelineSettingTreeItem : IMoviePipelineSettingTreeItem
 				]
 			]
 
-			+ SHorizontalBox::Slot()
-			.Padding(8, 4, 8, 4)
-			.AutoWidth()
-			.VAlign(VAlign_Center)
-			.HAlign(HAlign_Center)
-			[
-				SNew(SBox)
-				.WidthOverride(32)
-				.HeightOverride(32)
-				[
-					SNew(SImage)
-					.Image(this, &FMoviePipelineSettingTreeItem::GetIcon)
-					.ColorAndOpacity(this, &FMoviePipelineSettingTreeItem::GetImageColorAndOpacity)
-				]
-			]
-
 			+ SHorizontalBox::Slot()
 			.AutoWidth()
 			.VAlign(VAlign_Center)
@@ -163,7 +147,7 @@ struct FMoviePipelineSettingTreeItem : IMoviePipelineSettingTreeItem
 			[
 				SNew(SColorBlock)
 				.Size(FVector2D(6.0, 38.0))
-				.Color(FLinearColor::White) 
+				.Color(FLinearColor::Yellow) 
 			]
 
 		];
@@ -172,12 +156,6 @@ struct FMoviePipelineSettingTreeItem : IMoviePipelineSettingTreeItem
 
 private:
 
-	const FSlateBrush* GetIcon() const
-	{
-		UMoviePipelineSetting* Setting = WeakSetting.Get();
-		return Setting ? Setting->GetDisplayIcon() : nullptr;
-	}
-
 	FText GetDescription() const
 	{
 		UMoviePipelineSetting* Setting = WeakSetting.Get();
diff --git a/Engine/Plugins/MovieScene/MovieRenderPipeline/Source/MovieRenderPipelineEditor/Private/Widgets/SMoviePipelineTabContent.cpp b/Engine/Plugins/MovieScene/MovieRenderPipeline/Source/MovieRenderPipelineEditor/Private/Widgets/SMoviePipelineTabContent.cpp
--- a/Engine/Plugins/MovieScene/MovieRenderPipeline/Source/MovieRenderPipelineEditor/Private/Widgets/SMoviePipelineTabContent.cpp
+++ b/Engine/Plugins/MovieScene/MovieRenderPipeline/Source/MovieRenderPipelineEditor/Private/Widgets/SMoviePipelineTabContent.cpp
@@ -1,8 +1,7 @@
 // Copyright 1998-2019 Epic Games, Inc. All Rights Reserved.
 
 #include "Widgets/SMoviePipelineTabContent.h"
-#include "Widgets/SMoviePipelinePanel.h"
-
+#include "SMoviePipelineConfigPanel.h"
 
 // Analytics
 #include "EngineAnalytics.h"
@@ -34,13 +33,13 @@ void SMoviePipelineTabContent::SetupForPipeline(UMoviePipelineConfigBase* BasePr
 	// Null out the tab content to ensure that all references have been cleaned up before constructing the new one
 	ChildSlot [ SNullWidget::NullWidget ];
 
-	ChildSlot
-	.HAlign(HAlign_Fill)
-	.VAlign(VAlign_Fill)
-	[
-		SAssignNew(WeakPanel, SMoviePipelinePanel)
-		.BasePreset(BasePreset)
-	];
+	// ChildSlot
+	// .HAlign(HAlign_Fill)
+	// .VAlign(VAlign_Fill)
+	// [
+	// 	SAssignNew(WeakPanel, SMoviePipelineConfigPanel, BasePreset->GetClass())
+	// 	.BasePreset(BasePreset)
+	// ];
 
 	if (FEngineAnalytics::IsAvailable())
 	{
diff --git a/Engine/Plugins/MovieScene/MovieRenderPipeline/Source/MovieRenderPipelineEditor/Private/Widgets/SMoviePipelineTabContent.h b/Engine/Plugins/MovieScene/MovieRenderPipeline/Source/MovieRenderPipelineEditor/Private/Widgets/SMoviePipelineTabContent.h
--- a/Engine/Plugins/MovieScene/MovieRenderPipeline/Source/MovieRenderPipelineEditor/Private/Widgets/SMoviePipelineTabContent.h
+++ b/Engine/Plugins/MovieScene/MovieRenderPipeline/Source/MovieRenderPipelineEditor/Private/Widgets/SMoviePipelineTabContent.h
@@ -6,7 +6,7 @@
 #include "Widgets/SCompoundWidget.h"
 
 class UMoviePipelineConfigBase;
-class SMoviePipelinePanel;
+class SMoviePipelineConfigPanel;
 
 class SMoviePipelineTabContent : public SCompoundWidget
 {
@@ -24,5 +24,5 @@ private:
 	EActiveTimerReturnType OnActiveTimer(double InCurrentTime, float InDeltaTime);
 
 private:
-	TWeakPtr<SMoviePipelinePanel> WeakPanel;
+	TWeakPtr<SMoviePipelineConfigPanel> WeakPanel;
 };
diff --git a/Engine/Plugins/MovieScene/MovieRenderPipeline/Source/MovieRenderPipelineEditor/Public/IMovieRenderPipelineEditorModule.h b/Engine/Plugins/MovieScene/MovieRenderPipeline/Source/MovieRenderPipelineEditor/Public/IMovieRenderPipelineEditorModule.h
--- a/Engine/Plugins/MovieScene/MovieRenderPipeline/Source/MovieRenderPipelineEditor/Public/IMovieRenderPipelineEditorModule.h
+++ b/Engine/Plugins/MovieScene/MovieRenderPipeline/Source/MovieRenderPipelineEditor/Public/IMovieRenderPipelineEditorModule.h
@@ -11,9 +11,15 @@
 class IMovieRenderPipelineEditorModule : public IModuleInterface
 {
 public:
-	/** The tab name for the movie render pipeline tab */
-	static FName MovieRenderPipelineTabName;
+	/** The tab name for the movie render queue editor tab */
+	static FName MoviePipelineQueueTabName;
 
-	/** The default label for the movie render pipeline tab. */
-	static FText MovieRenderPipelineTabLabel;
+	/** The default label for the movie render queue editor tab. */
+	static FText MoviePipelineQueueTabLabel;
+
+	/** The tab name for the movie render config editor tab */
+	static FName MoviePipelineConfigEditorTabName;
+
+	/** The default label for the movie render config editor tab. */
+	static FText MoviePipelineConfigEditorTabLabel;
 };
diff --git a/Engine/Plugins/MovieScene/MovieRenderPipeline/Source/MovieRenderPipelineEditor/Public/MoviePipelinePIEExecutor.h b/Engine/Plugins/MovieScene/MovieRenderPipeline/Source/MovieRenderPipelineEditor/Public/MoviePipelinePIEExecutor.h
--- a/Engine/Plugins/MovieScene/MovieRenderPipeline/Source/MovieRenderPipelineEditor/Public/MoviePipelinePIEExecutor.h
+++ b/Engine/Plugins/MovieScene/MovieRenderPipeline/Source/MovieRenderPipelineEditor/Public/MoviePipelinePIEExecutor.h
@@ -24,7 +24,7 @@ public:
 	}
 
 protected:
-	virtual void Start(const FMoviePipelineExecutorJob& InJob) override;
+	virtual void Start(const FMoviePipelineExecutorJobPrev& InJob) override;
 
 private:
 	/** Called when PIE finishes booting up and it is safe for us to spawn an object into that world. */
diff --git a/Engine/Plugins/MovieScene/MovieRenderPipeline/Source/MovieRenderPipelineEditor/Public/MoviePipelineQueueSubsystem.h b/Engine/Plugins/MovieScene/MovieRenderPipeline/Source/MovieRenderPipelineEditor/Public/MoviePipelineQueueSubsystem.h
new file mode 100644
--- /dev/null
+++ b/Engine/Plugins/MovieScene/MovieRenderPipeline/Source/MovieRenderPipelineEditor/Public/MoviePipelineQueueSubsystem.h
@@ -0,0 +1,29 @@
+// Copyright 1998-2019 Epic Games, Inc. All Rights Reserved.
+#pragma once
+
+#include "EditorSubsystem.h"
+#include "MoviePipelineQueue.h"
+#include "MoviePipelineQueueSubsystem.generated.h"
+
+UCLASS(BlueprintType)
+class MOVIERENDERPIPELINEEDITOR_API UMoviePipelineQueueSubsystem : public UEditorSubsystem
+{
+	GENERATED_BODY()
+	
+public:
+
+	UMoviePipelineQueueSubsystem()
+	{
+		CurrentQueue = CreateDefaultSubobject<UMoviePipelineQueue>("EditorMoviePipelineQueue");
+	}
+	
+	UFUNCTION(BlueprintPure, Category = "Movie Render Pipeline")
+	UMoviePipelineQueue* GetQueue() const
+	{
+		return CurrentQueue;
+	}
+
+private:
+	UPROPERTY(Transient, Instanced)
+	UMoviePipelineQueue* CurrentQueue;
+};
\ No newline at end of file
diff --git a/Engine/Plugins/MovieScene/MovieRenderPipeline/Source/MovieRenderPipelineEditor/Public/MovieRenderPipelineSettings.h b/Engine/Plugins/MovieScene/MovieRenderPipeline/Source/MovieRenderPipelineEditor/Public/MovieRenderPipelineSettings.h
--- a/Engine/Plugins/MovieScene/MovieRenderPipeline/Source/MovieRenderPipelineEditor/Public/MovieRenderPipelineSettings.h
+++ b/Engine/Plugins/MovieScene/MovieRenderPipeline/Source/MovieRenderPipelineEditor/Public/MovieRenderPipelineSettings.h
@@ -3,6 +3,8 @@
 
 #include "UObject/Object.h"
 #include "Templates/SubclassOf.h"
+#include "MoviePipelineExecutor.h"
+#include "MoviePipeline.h"
 #include "MovieRenderPipelineSettings.generated.h"
 
 class UMoviePipelineExecutorBase;
diff --git a/Engine/Plugins/MovieScene/MovieRenderPipeline/Source/MovieRenderPipelineRenderPasses/Public/MoviePipelineImageSequenceOutput.h b/Engine/Plugins/MovieScene/MovieRenderPipeline/Source/MovieRenderPipelineRenderPasses/Public/MoviePipelineImageSequenceOutput.h
--- a/Engine/Plugins/MovieScene/MovieRenderPipeline/Source/MovieRenderPipelineRenderPasses/Public/MoviePipelineImageSequenceOutput.h
+++ b/Engine/Plugins/MovieScene/MovieRenderPipeline/Source/MovieRenderPipelineRenderPasses/Public/MoviePipelineImageSequenceOutput.h
@@ -37,10 +37,14 @@ private:
 	FString OutputDirectory;
 };
 
-UCLASS(meta = (DisplayName = "Image Sequence (.bmp [8bpp])"))
+UCLASS()
 class MOVIERENDERPIPELINERENDERPASSES_API UMoviePipelineImageSequenceOutput_BMP : public UMoviePipelineImageSequenceOutputBase
 {
 	GENERATED_BODY()
+public:
+#if WITH_EDITOR
+	virtual FText GetDisplayText() const override { return NSLOCTEXT("MovieRenderPipeline", "ImgSequenceBMPSettingDisplayName", "Image Sequence (.bmp [8bpp])"); }
+#endif
 public:
 	UMoviePipelineImageSequenceOutput_BMP()
 	{
@@ -48,10 +52,14 @@ public:
 	}
 };
 
-UCLASS(meta = (DisplayName = "Image Sequence (.png [8bpp])"))
+UCLASS()
 class MOVIERENDERPIPELINERENDERPASSES_API UMoviePipelineImageSequenceOutput_PNG : public UMoviePipelineImageSequenceOutputBase
 {
 	GENERATED_BODY()
+public:
+#if WITH_EDITOR
+	virtual FText GetDisplayText() const override { return NSLOCTEXT("MovieRenderPipeline", "ImgSequencePNGSettingDisplayName", "Image Sequence (.png [8bpp])"); }
+#endif
 public:
 	UMoviePipelineImageSequenceOutput_PNG()
 	{
@@ -59,10 +67,14 @@ public:
 	}
 };
 
-UCLASS(meta = (DisplayName = "Image Sequence (.jpg [8bpp])"))
+UCLASS()
 class MOVIERENDERPIPELINERENDERPASSES_API UMoviePipelineImageSequenceOutput_JPG : public UMoviePipelineImageSequenceOutputBase
 {
 	GENERATED_BODY()
+public:
+#if WITH_EDITOR
+	virtual FText GetDisplayText() const override { return NSLOCTEXT("MovieRenderPipeline", "ImgSequenceJPGSettingDisplayName", "Image Sequence (.jpg [8bpp])"); }
+#endif
 public:
 	UMoviePipelineImageSequenceOutput_JPG()
 	{
@@ -70,10 +82,14 @@ public:
 	}
 };
 
-UCLASS(meta = (DisplayName = "Image Sequence (.exr [32bpp])"))
+UCLASS()
 class MOVIERENDERPIPELINERENDERPASSES_API UMoviePipelineImageSequenceOutput_EXR : public UMoviePipelineImageSequenceOutputBase
 {
 	GENERATED_BODY()
+public:
+#if WITH_EDITOR
+	virtual FText GetDisplayText() const override { return NSLOCTEXT("MovieRenderPipeline", "ImgSequenceEXRSettingDisplayName", "Image Sequence (.exr [32bpp])"); }
+#endif
 public:
 	UMoviePipelineImageSequenceOutput_EXR()
 	{
diff --git a/Engine/Plugins/MovieScene/MovieRenderPipeline/Source/MovieRenderPipelineSettings/Public/MoviePipelineBurnInSetting.h b/Engine/Plugins/MovieScene/MovieRenderPipeline/Source/MovieRenderPipelineSettings/Public/MoviePipelineBurnInSetting.h
--- a/Engine/Plugins/MovieScene/MovieRenderPipeline/Source/MovieRenderPipelineSettings/Public/MoviePipelineBurnInSetting.h
+++ b/Engine/Plugins/MovieScene/MovieRenderPipeline/Source/MovieRenderPipelineSettings/Public/MoviePipelineBurnInSetting.h
@@ -5,12 +5,16 @@
 #include "LevelSequenceActor.h"
 #include "MoviePipelineBurnInSetting.generated.h"
 
-UCLASS(Blueprintable, meta=(MovieRenderPipelineDisplayName="Burn In"))
+UCLASS(Blueprintable)
 class UMoviePipelineBurnInSetting : public UMoviePipelineSetting
 {
 	GENERATED_BODY()
 public:
+#if WITH_EDITOR
 	virtual FText GetDisplayText() const override { return NSLOCTEXT("MovieRenderPipeline", "BurnInSettingDisplayName", "Burn In"); }
+#endif
+	virtual bool IsValidOnShots() const override { return false; }
+	virtual bool IsValidOnMaster() const override { return true; }
 public:
 	UPROPERTY(EditAnywhere, BlueprintReadWrite, meta=(MetaClass="LevelSequenceBurnIn"), Category = "Movie Pipeline")
 	FSoftClassPath BurnInClass;
diff --git a/Engine/Plugins/MovieScene/MovieRenderPipeline/Source/MovieRenderPipelineSettings/Public/MoviePipelineGameOverrideSetting.h b/Engine/Plugins/MovieScene/MovieRenderPipeline/Source/MovieRenderPipelineSettings/Public/MoviePipelineGameOverrideSetting.h
--- a/Engine/Plugins/MovieScene/MovieRenderPipeline/Source/MovieRenderPipelineSettings/Public/MoviePipelineGameOverrideSetting.h
+++ b/Engine/Plugins/MovieScene/MovieRenderPipeline/Source/MovieRenderPipelineSettings/Public/MoviePipelineGameOverrideSetting.h
@@ -20,7 +20,10 @@ public:
 	{
 	}
 
-protected:
+public:
+#if WITH_EDITOR
+	virtual FText GetDisplayText() const override { return NSLOCTEXT("MovieRenderPipeline", "GameOverrideSettingDisplayName", "Game Overrides"); }
+#endif
 	virtual bool IsValidOnShots() const override { return false; }
 	virtual bool IsValidOnMaster() const override { return true; }
 
diff --git a/Engine/Plugins/MovieScene/MovieRenderPipeline/Source/MovieRenderPipelineSettings/Public/MoviePipelineRayTracingSetting.h b/Engine/Plugins/MovieScene/MovieRenderPipeline/Source/MovieRenderPipelineSettings/Public/MoviePipelineRayTracingSetting.h
--- a/Engine/Plugins/MovieScene/MovieRenderPipeline/Source/MovieRenderPipelineSettings/Public/MoviePipelineRayTracingSetting.h
+++ b/Engine/Plugins/MovieScene/MovieRenderPipeline/Source/MovieRenderPipelineSettings/Public/MoviePipelineRayTracingSetting.h
@@ -12,22 +12,98 @@ class UMoviePipelineRayTracingSetting : public UMoviePipelineSetting
 
 	UMoviePipelineRayTracingSetting()
 		: bEnableDenoisers(true)
+		, AmbientOcclusionSPP(-1)
+		, GlobalIlluminationSPP(-1)
+		, ReflectionsSPP(-1)
+		, ShadowSPP(-1)
+		, bUsePathTracer(false)
 	{}
+
+public:
+#if WITH_EDITOR
+	virtual FText GetDisplayText() const override { return NSLOCTEXT("MovieRenderPipeline", "RayTracingSettingDisplayName", "Raytracing"); }
+#endif
+	virtual bool IsValidOnShots() const override { return true; }
+	virtual bool IsValidOnMaster() const override { return true; }
 public:
+	virtual void OnSetupForShot()
+	{
+		FString DenoisersEnabled = bEnableDenoisers ? TEXT("1") : TEXT("0");
+
+		 PreviousDenoiserState.Reflections = IConsoleManager::Get().FindConsoleVariable(TEXT("r.Reflections.Denoiser"))->GetInt();
+		 PreviousDenoiserState.Shadows = IConsoleManager::Get().FindConsoleVariable(TEXT("r.Shadow.Denoiser"))->GetInt();
+		 PreviousDenoiserState.AmbientOcclusion = IConsoleManager::Get().FindConsoleVariable(TEXT("r.AmbientOcclusion.Denoiser"))->GetInt();
+		 PreviousDenoiserState.GlobalIllumination = IConsoleManager::Get().FindConsoleVariable(TEXT("r.Raytracing.GLobalIllumination.Denoiser"))->GetInt();
+		 
+		 PreviousSamplesPerPixelState.Reflections = IConsoleManager::Get().FindConsoleVariable(TEXT("r.Raytracing.Reflections.SamplesPerPixel"))->GetInt();
+		 PreviousSamplesPerPixelState.Shadows = IConsoleManager::Get().FindConsoleVariable(TEXT("r.Raytracing.Shadows.SamplesPerPixel"))->GetInt();
+		 PreviousSamplesPerPixelState.AmbientOcclusion = IConsoleManager::Get().FindConsoleVariable(TEXT("r.Raytracing.AmbientOcclusion.SamplesPerPixel"))->GetInt();
+		 PreviousSamplesPerPixelState.GlobalIllumination = IConsoleManager::Get().FindConsoleVariable(TEXT("r.Raytracing.GlobalIllumination.SamplesPerPixel"))->GetInt();
+	}
+
 	/**
 	* Enabling the ray tracing denoisers will reduce the number of samples required per frame,
 	* but can introduce artifacts (ghosting). If you disable this expect to need 64-128 samples
 	* per frame to overcome the lack of denoiser.
 	*/
-	UPROPERTY(EditAnywhere, BlueprintReadWrite, meta = (EditCondition="bUsePathTracer"), Category = "Movie Pipeline")
+	UPROPERTY(EditAnywhere, BlueprintReadWrite, meta = (EditCondition="!bUsePathTracer"), Category = "Movie Pipeline")
 	bool bEnableDenoisers;
 	
+	/**
+	* The number of samples per pixel to use during Ambient Occlusion raytracing. Too large of a value
+	* will cause a TDR (Timeout Detection and Recovery) which results in a GPU/Engine crash. More effective
+	* to use more temporal/spatial samples (better quality, less likely to trigger TDR). Set to -1 to
+	* fall back to the Post Processing/Camera settings.
+	*/
+	UPROPERTY(EditAnywhere, BlueprintReadWrite, meta = (UIMin="-1", UIMax="2"), Category = "Movie Pipeline")
 	int32 AmbientOcclusionSPP;
+
+	/**
+	* The number of samples per pixel to use during Global Illumination raytracing. Too large of a value
+	* will cause a TDR (Timeout Detection and Recovery) which results in a GPU/Engine crash. More effective
+	* to use more temporal/spatial samples (better quality, less likely to trigger TDR). Set to -1 to
+	* fall back to the Post Processing/Camera settings.
+	*/
+	UPROPERTY(EditAnywhere, BlueprintReadWrite, meta = (UIMin = "-1", UIMax = "2"), Category = "Movie Pipeline")
 	int32 GlobalIlluminationSPP;
+
+	/**
+	* The number of samples per pixel to use during Reflection raytracing. Too large of a value
+	* will cause a TDR (Timeout Detection and Recovery) which results in a GPU/Engine crash. More effective
+	* to use more temporal/spatial samples (better quality, less likely to trigger TDR). Set to -1 to
+	* fall back to the Post Processing/Camera settings.
+	*/
+	UPROPERTY(EditAnywhere, BlueprintReadWrite, meta = (UIMin = "-1", UIMax = "2"), Category = "Movie Pipeline")
 	int32 ReflectionsSPP;
+
+	/**
+	* The number of samples per pixel to use during Global Illumination raytracing. Too large of a value
+	* will cause a TDR (Timeout Detection and Recovery) which results in a GPU/Engine crash. More effective
+	* to use more temporal/spatial samples (better quality, less likely to trigger TDR). Set to -1 to
+	* fall back to the Post Processing/Camera settings.
+	*/
+	UPROPERTY(EditAnywhere, BlueprintReadWrite, meta = (UIMin = "-1", UIMax = "2"), Category = "Movie Pipeline")
 	int32 ShadowSPP;
 		
+	/**
+	* Should we use the Path Tracer instead of the real-time raytracer? The path tracer is a reference implementation
+	* and may not support all rendering features ray tracing does. It is a debugging tool to compare the output
+	* of realtime raytracing and raster engines compared to a more traditional approach.
+	*/
+	UPROPERTY(EditAnywhere, BlueprintReadWrite, Category = "Movie Pipeline")
 	bool bUsePathTracer;
+
+private:
+	struct FCVarStateCache
+	{
+		int32 Reflections;
+		int32 Shadows;
+		int32 AmbientOcclusion;
+		int32 GlobalIllumination;
+	};
+
+	FCVarStateCache PreviousDenoiserState;
+	FCVarStateCache PreviousSamplesPerPixelState;
 	/**
 	
 	UPROPERTY(Instanced, EditAnywhere, BlueprintReadWrite, Category="Movie Pipeline")
