MoviePipeline: Renamed some classes/files for additional clarity between editor ui windows.

#rnx

diff --git a/Engine/Plugins/MovieScene/MovieRenderPipeline/Source/MovieRenderPipelineCore/Private/MoviePipelineRendering.cpp b/Engine/Plugins/MovieScene/MovieRenderPipeline/Source/MovieRenderPipelineCore/Private/MoviePipelineRendering.cpp
--- a/Engine/Plugins/MovieScene/MovieRenderPipeline/Source/MovieRenderPipelineCore/Private/MoviePipelineRendering.cpp
+++ b/Engine/Plugins/MovieScene/MovieRenderPipeline/Source/MovieRenderPipelineCore/Private/MoviePipelineRendering.cpp
@@ -318,14 +318,17 @@ void UMoviePipeline::RenderFrame()
 
 				{
 					BackbufferResolution = FIntPoint(
-						FMath::CeilToInt(BackbufferResolution.X / HighResSettings->TileCount),
-						FMath::CeilToInt(BackbufferResolution.Y / HighResSettings->TileCount));
+						FMath::CeilToInt(BackbufferResolution.X / TileCount.X),
+						FMath::CeilToInt(BackbufferResolution.Y / TileCount.Y));
 
 					// Cache tile resolution before we add padding.
 					TileResolution = BackbufferResolution;
 
 					// Apply size padding.
-					BackbufferResolution = HighResSettings->CalculatePaddedBackbufferSize(BackbufferResolution);
+					if (HighResSettings)
+					{
+						BackbufferResolution = HighResSettings->CalculatePaddedBackbufferSize(BackbufferResolution);
+					}
 				}
 
 				// We take all of the information needed to render a single sample and package it into a struct.
diff --git a/Engine/Plugins/MovieScene/MovieRenderPipeline/Source/MovieRenderPipelineEditor/Private/Widgets/SMoviePipelineConfigEditor.cpp b/Engine/Plugins/MovieScene/MovieRenderPipeline/Source/MovieRenderPipelineEditor/Private/Widgets/SMoviePipelineConfigEditor.cpp
new file mode 100644
--- /dev/null
+++ b/Engine/Plugins/MovieScene/MovieRenderPipeline/Source/MovieRenderPipelineEditor/Private/Widgets/SMoviePipelineConfigEditor.cpp
@@ -0,0 +1,330 @@
+// Copyright 1998-2019 Epic Games, Inc. All Rights Reserved.
+
+#include "Widgets/SMoviePipelineConfigEditor.h"
+#include "Widgets/MoviePipelineWidgetConstants.h"
+#include "Widgets/SMoviePipelineConfigSettings.h"
+
+#include "MoviePipelineShotConfig.h"
+#include "MoviePipelineSetting.h"
+
+// Core includes
+#include "UObject/UObjectIterator.h"
+#include "Templates/SubclassOf.h"
+#include "ClassIconFinder.h"
+// 
+// AssetRegistry includes
+#include "AssetRegistryModule.h"
+#include "AssetData.h"
+ 
+// ContentBrowser includes
+#include "IContentBrowserSingleton.h"
+#include "ContentBrowserModule.h"
+ 
+// AssetTools includes
+#include "AssetToolsModule.h"
+ 
+// Slate includes
+#include "Widgets/SBoxPanel.h"
+#include "Widgets/Text/STextBlock.h"
+#include "Widgets/Input/SButton.h"
+#include "Widgets/Input/SComboButton.h"
+#include "Widgets/Layout/SSplitter.h"
+#include "Widgets/Layout/SScrollBox.h"
+#include "Widgets/Layout/SBorder.h"
+#include "Framework/MultiBox/MultiBoxBuilder.h"
+#include "Styling/SlateIconFinder.h"
+
+// EditorStyle includes
+#include "EditorStyleSet.h"
+#include "EditorFontGlyphs.h"
+#include "ScopedTransaction.h"
+#include "PropertyEditorModule.h"
+#include "IDetailsView.h"
+
+#define LOCTEXT_NAMESPACE "SMoviePipelineEditor"
+
+
+TArray<UClass*> FindMoviePipelineSettingClasses()
+{
+	FAssetRegistryModule& AssetRegistryModule = FModuleManager::LoadModuleChecked<FAssetRegistryModule>(TEXT("AssetRegistry"));
+
+	TArray<FAssetData> ClassList;
+
+	FARFilter Filter;
+	Filter.ClassNames.Add(UMoviePipelineSetting::StaticClass()->GetFName());
+
+	// Include any Blueprint based objects as well, this includes things like Blutilities, UMG, and GameplayAbility objects
+	Filter.bRecursiveClasses = true;
+	AssetRegistryModule.Get().GetAssets(Filter, ClassList);
+
+	TArray<UClass*> Classes;
+
+	for (const FAssetData& Data : ClassList)
+	{
+		UClass* Class = Data.GetClass();
+		if (Class)
+		{
+			Classes.Add(Class);
+		}
+	}
+
+	for (TObjectIterator<UClass> ClassIterator; ClassIterator; ++ClassIterator)
+	{
+		if (ClassIterator->IsChildOf(UMoviePipelineSetting::StaticClass()) && !ClassIterator->HasAnyClassFlags(CLASS_Abstract | CLASS_Deprecated | CLASS_NewerVersionExists))
+		{
+			Classes.Add(*ClassIterator);
+		}
+	}
+
+	return Classes;
+}
+
+PRAGMA_DISABLE_OPTIMIZATION
+void SMoviePipelineConfigEditor::Construct(const FArguments& InArgs)
+{
+	bRequestDetailsRefresh = true;
+	MoviePipelineAttribute = InArgs._MoviePipeline;
+    
+	DetailsBox = SNew(SScrollBox);
+	DetailsBox->SetScrollBarRightClickDragAllowed(true);
+     
+	SettingsWidget = SNew(SMoviePipelineConfigSettings)
+	.OnSelectionChanged(this, &SMoviePipelineConfigEditor::OnSettingsSelectionChanged);
+    
+	CheckForNewMoviePipeline();
+     
+	ChildSlot
+	[
+		SNew(SSplitter)
+		.Orientation(Orient_Horizontal)
+    
+		+ SSplitter::Slot()
+		.Value(.25f)
+		[
+			SNew(SBorder)
+			.Padding(4)
+			.BorderImage( FEditorStyle::GetBrush("ToolPanel.GroupBorder") )
+			[
+				SettingsWidget.ToSharedRef()
+			]
+		]
+    
+		+ SSplitter::Slot()
+		.Value(.75f)
+		[
+			DetailsBox.ToSharedRef()
+		]
+	];
+}
+
+TSharedRef<SWidget> SMoviePipelineConfigEditor::MakeAddSettingButton()
+{
+	return SNew(SComboButton)
+		.ContentPadding(MoviePipeline::ButtonPadding)
+		// .ButtonStyle(FTakeRecorderStyle::Get(), "FlatButton.Success")
+		.OnGetMenuContent(this, &SMoviePipelineConfigEditor::OnGenerateSettingsMenu)
+		.ForegroundColor(FSlateColor::UseForeground())
+		.HasDownArrow(false)
+		.ButtonContent()
+		[
+			SNew(SHorizontalBox)
+    
+			// Plus Icon
+			+ SHorizontalBox::Slot()
+			.VAlign(VAlign_Center)
+			.AutoWidth()
+			[
+				SNew(STextBlock)
+				.TextStyle(FEditorStyle::Get(), "NormalText.Important")
+				.Font(FEditorStyle::Get().GetFontStyle("FontAwesome.10"))
+				.Text(FEditorFontGlyphs::Plus)
+			]
+    
+			// "Setting" Text
+			+ SHorizontalBox::Slot()
+			.VAlign(VAlign_Center)
+			.AutoWidth()
+			.Padding(4, 0, 0, 0)
+			[
+				SNew(STextBlock)
+				.TextStyle(FEditorStyle::Get(), "NormalText.Important")
+				.Text(LOCTEXT("AddNewSetting_Text", "Setting"))
+			]
+    
+			// Non-Default Down Caret arrow.
+			+ SHorizontalBox::Slot()
+			.VAlign(VAlign_Center)
+			.AutoWidth()
+			.Padding(4, 0, 0, 0)
+			[
+				SNew(STextBlock)
+				.TextStyle(FEditorStyle::Get(), "NormalText.Important")
+				.Font(FEditorStyle::Get().GetFontStyle("FontAwesome.10"))
+				.Text(FEditorFontGlyphs::Caret_Down)
+			]
+		];
+}
+PRAGMA_ENABLE_OPTIMIZATION
+
+void SMoviePipelineConfigEditor::CheckForNewMoviePipeline()
+{
+	UMoviePipelineConfigBase* NewMoviePipeline = MoviePipelineAttribute.Get();
+	if (CachedMoviePipeline != NewMoviePipeline)
+	{
+		CachedMoviePipeline = NewMoviePipeline;
+
+		SettingsWidget->SetShotConfigObject(NewMoviePipeline);
+		bRequestDetailsRefresh = true;
+	}
+}
+
+void SMoviePipelineConfigEditor::Tick(const FGeometry& AllottedGeometry, const double InCurrentTime, const float InDeltaTime)
+{
+	CheckForNewMoviePipeline();
+	if (bRequestDetailsRefresh)
+	{
+		UpdateDetails();
+		bRequestDetailsRefresh = false;
+	}
+}
+
+TSharedRef<SWidget> SMoviePipelineConfigEditor::OnGenerateSettingsMenu()
+{
+	TSharedRef<FExtender> Extender = MakeShared<FExtender>();
+	{
+		// ULevelSequence*       LevelSequence = LevelSequenceAttribute.Get();
+		// UTakeRecorderSources* Sources = LevelSequence ? LevelSequence->FindOrAddMetaData<UTakeRecorderSources>() : nullptr;
+		// 
+		// if (Sources)
+		// {
+		// 	FTakeRecorderModule& TakeRecorderModule = FModuleManager::GetModuleChecked<FTakeRecorderModule>("TakeRecorder");
+		// 	TakeRecorderModule.PopulateSourcesMenu(Extender, Sources);
+		// }
+	}
+
+	FMenuBuilder MenuBuilder(true, nullptr, Extender);
+
+	MenuBuilder.BeginSection("Settings", LOCTEXT("SettingsMenuSection", "Available Settings"));
+	{
+		TArray<UClass*> SourceClasses = FindMoviePipelineSettingClasses();
+		Algo::SortBy(SourceClasses, &UClass::GetDisplayNameText, FText::FSortPredicate());
+
+		for (UClass* Class : SourceClasses)
+		{
+			// Get a display name for the setting from the CDO.
+			const UMoviePipelineSetting* SettingDefaultObject = GetDefault<UMoviePipelineSetting>(Class);
+			
+			// Depending on the type of config we're editing, some settings may not be eligible. If this is the case, we omit them from the list.
+			bool bCanSettingBeAdded = CachedMoviePipeline->CanSettingBeAdded(SettingDefaultObject);
+			if (!bCanSettingBeAdded)
+			{
+				continue;
+			}
+
+			// If the setting already exists and it only allows a single instance, we omit them from the list.
+			bool bAllowDuplicates = true;
+			for (UMoviePipelineSetting* ExistingSetting : CachedMoviePipeline->GetSettings())
+			{
+				// If we found a setting with the same class as ours, ask the CDO if multiple are valid.
+				if (ExistingSetting->GetClass() == Class)
+				{
+					bAllowDuplicates = !SettingDefaultObject->IsSolo();
+					break;
+				}
+			}
+			if (!bAllowDuplicates)
+			{
+				continue;
+			}
+
+			FText SettingDisplayName;
+			if (SettingDefaultObject)
+			{
+				SettingDisplayName = SettingDefaultObject->GetDisplayText();
+			}
+
+			TSubclassOf<UMoviePipelineSetting> SubclassOf = Class;
+			MenuBuilder.AddMenuEntry(
+				SettingDisplayName,
+				Class->GetToolTipText(true),
+				FSlateIconFinder::FindIconForClass(Class),
+				FUIAction(
+					FExecuteAction::CreateSP(this, &SMoviePipelineConfigEditor::AddSettingFromClass, SubclassOf),
+					FCanExecuteAction::CreateSP(this, &SMoviePipelineConfigEditor::CanAddSettingFromClass, SubclassOf)
+				)
+			);
+		}
+	}
+	MenuBuilder.EndSection();
+
+	return MenuBuilder.MakeWidget();
+}
+
+void SMoviePipelineConfigEditor::UpdateDetails()
+{
+	FPropertyEditorModule& PropertyEditorModule = FModuleManager::Get().LoadModuleChecked<FPropertyEditorModule>("PropertyEditor");
+	FDetailsViewArgs DetailsViewArgs(false, false, false, FDetailsViewArgs::HideNameArea, true);
+	DetailsViewArgs.bShowScrollBar = false;
+	 
+	TArray<UMoviePipelineSetting*> SelectedSources;
+	SettingsWidget->GetSelectedSettings(SelectedSources);
+
+	
+	// Create 1 details panel per source class type
+	TSortedMap<const UClass*, TArray<UObject*>> ClassToSources;
+	for (UMoviePipelineSetting* Source : SelectedSources)
+	{
+		ClassToSources.FindOrAdd(Source->GetClass()).Add(Source);
+	}
+	
+	TArray<FObjectKey> PreviousClasses;
+	ClassToDetailsView.GenerateKeyArray(PreviousClasses);
+	
+	for (auto& Pair : ClassToSources)
+	{
+		PreviousClasses.Remove(Pair.Key);
+	
+		TSharedPtr<IDetailsView> ExistingDetails = ClassToDetailsView.FindRef(Pair.Key);
+		if (ExistingDetails.IsValid())
+		{
+			ExistingDetails->SetObjects(Pair.Value);
+		}
+		else
+		{
+			TSharedRef<IDetailsView> Details = PropertyEditorModule.CreateDetailView(DetailsViewArgs);
+			Details->SetObjects(Pair.Value);
+			DetailsBox->AddSlot()
+				[
+					Details
+				];
+	
+			ClassToDetailsView.Add(Pair.Key, Details);
+		}
+	}
+	
+	for (FObjectKey StaleClass : PreviousClasses)
+	{
+		TSharedPtr<IDetailsView> Details = ClassToDetailsView.FindRef(StaleClass);
+		DetailsBox->RemoveSlot(Details.ToSharedRef());
+		ClassToDetailsView.Remove(StaleClass);
+	}
+}
+
+void SMoviePipelineConfigEditor::AddSettingFromClass(TSubclassOf<UMoviePipelineSetting> SettingClass)
+{
+	UMoviePipelineConfigBase* Pipeline = MoviePipelineAttribute.Get();
+
+	if (*SettingClass && Pipeline)
+	{
+		FScopedTransaction Transaction(FText::Format(LOCTEXT("AddNewSetting", "Add New {0} Setting"), SettingClass->GetDisplayNameText()));
+		Pipeline->FindOrAddSettingByClass(SettingClass);
+	}
+}
+
+bool SMoviePipelineConfigEditor::CanAddSettingFromClass(TSubclassOf<UMoviePipelineSetting> SettingClass)
+{
+	return true;
+}
+
+
+#undef LOCTEXT_NAMESPACE
\ No newline at end of file
diff --git a/Engine/Plugins/MovieScene/MovieRenderPipeline/Source/MovieRenderPipelineEditor/Private/Widgets/SMoviePipelineConfigEditor.h b/Engine/Plugins/MovieScene/MovieRenderPipeline/Source/MovieRenderPipelineEditor/Private/Widgets/SMoviePipelineConfigEditor.h
new file mode 100644
--- /dev/null
+++ b/Engine/Plugins/MovieScene/MovieRenderPipeline/Source/MovieRenderPipelineEditor/Private/Widgets/SMoviePipelineConfigEditor.h
@@ -0,0 +1,77 @@
+// Copyright 1998-2019 Epic Games, Inc. All Rights Reserved.
+
+#pragma once
+
+#include "Widgets/DeclarativeSyntaxSupport.h"
+#include "Widgets/SCompoundWidget.h"
+#include "UObject/GCObject.h"
+#include "UObject/ObjectKey.h"
+
+struct FAssetData;
+struct IMoviePipelineSettingTreeItem;
+
+class SScrollBox;
+class SMoviePipelineConfigSettings;
+class UMoviePipelineConfigBase;
+class IDetailsView;
+class UMoviePipelineSetting;
+
+template<class> class TSubclassOf;
+
+/**
+ * Widget used to edit a Movie Render Pipeline Shot Config.
+ */
+class SMoviePipelineConfigEditor : public SCompoundWidget
+{
+public:
+
+	SLATE_BEGIN_ARGS(SMoviePipelineConfigEditor)
+		: _MoviePipeline(nullptr)
+		{}
+
+		SLATE_ATTRIBUTE(UMoviePipelineConfigBase*, MoviePipeline)
+
+	SLATE_END_ARGS()
+
+	void Construct(const FArguments& InArgs);
+
+	/**
+	 * Construct a button that can add sources to this widget's preset
+	 */
+	TSharedRef<SWidget> MakeAddSettingButton();
+
+private:
+
+	virtual void Tick(const FGeometry& AllottedGeometry, const double InCurrentTime, const float InDeltaTime) override;
+
+private:
+
+	/**
+	 * Check to see whether the movie pipeline ptr has changed, and propagate that change if necessary
+	 */
+	void CheckForNewMoviePipeline();
+
+	/**
+	 * Update the details panel for the current selection
+	 */
+	void UpdateDetails();
+
+private:
+
+	TSharedRef<SWidget> OnGenerateSettingsMenu();
+     
+	void AddSettingFromClass(TSubclassOf<UMoviePipelineSetting> SettingClass);
+	bool CanAddSettingFromClass(TSubclassOf<UMoviePipelineSetting> SettingClass);
+	
+	void OnSettingsSelectionChanged(TSharedPtr<IMoviePipelineSettingTreeItem>, ESelectInfo::Type) { bRequestDetailsRefresh = true; }
+
+private:
+
+	bool bRequestDetailsRefresh;
+	TAttribute<UMoviePipelineConfigBase*> MoviePipelineAttribute;
+	TWeakObjectPtr<UMoviePipelineConfigBase> CachedMoviePipeline;
+    
+	TSharedPtr<SMoviePipelineConfigSettings> SettingsWidget;
+	TSharedPtr<SScrollBox> DetailsBox;
+	TMap<FObjectKey, TSharedPtr<IDetailsView>> ClassToDetailsView;
+};
\ No newline at end of file
diff --git a/Engine/Plugins/MovieScene/MovieRenderPipeline/Source/MovieRenderPipelineEditor/Private/Widgets/SMoviePipelineConfigPanel.cpp b/Engine/Plugins/MovieScene/MovieRenderPipeline/Source/MovieRenderPipelineEditor/Private/Widgets/SMoviePipelineConfigPanel.cpp
--- a/Engine/Plugins/MovieScene/MovieRenderPipeline/Source/MovieRenderPipelineEditor/Private/Widgets/SMoviePipelineConfigPanel.cpp
+++ b/Engine/Plugins/MovieScene/MovieRenderPipeline/Source/MovieRenderPipelineEditor/Private/Widgets/SMoviePipelineConfigPanel.cpp
@@ -1,7 +1,7 @@
 // Copyright 1998-2019 Epic Games, Inc. All Rights Reserved.
 
 #include "Widgets/SMoviePipelineConfigPanel.h"
-#include "Widgets/SMoviePipelineEditor.h"
+#include "Widgets/SMoviePipelineConfigEditor.h"
 #include "Widgets/MoviePipelineWidgetConstants.h"
 #include "Misc/MessageDialog.h"
 #include "MoviePipelineShotConfig.h"
@@ -47,7 +47,7 @@ void SMoviePipelineConfigPanel::Construct(const FArguments& InArgs, TSubclassOf<
 	}
 
 	// Create the child widgets that need to know about our pipeline
-	MoviePipelineEditorWidget = SNew(SMoviePipelineEditor)
+	MoviePipelineEditorWidget = SNew(SMoviePipelineConfigEditor)
 	.MoviePipeline(this, &SMoviePipelineConfigPanel::GetMoviePipeline);
 
 	ChildSlot
diff --git a/Engine/Plugins/MovieScene/MovieRenderPipeline/Source/MovieRenderPipelineEditor/Private/Widgets/SMoviePipelineConfigPanel.h b/Engine/Plugins/MovieScene/MovieRenderPipeline/Source/MovieRenderPipelineEditor/Private/Widgets/SMoviePipelineConfigPanel.h
--- a/Engine/Plugins/MovieScene/MovieRenderPipeline/Source/MovieRenderPipelineEditor/Private/Widgets/SMoviePipelineConfigPanel.h
+++ b/Engine/Plugins/MovieScene/MovieRenderPipeline/Source/MovieRenderPipelineEditor/Private/Widgets/SMoviePipelineConfigPanel.h
@@ -7,7 +7,7 @@
 #include "UObject/GCObject.h"
 
 class UMoviePipelineConfigBase;
-class SMoviePipelineEditor;
+class SMoviePipelineConfigEditor;
 
 /**
  * Outermost widget that is used for setting up a new movie render pipeline config. Operates on a transient UMovieRenderShotConfig that is internally owned and maintained 
@@ -53,5 +53,5 @@ private:
 	UMoviePipelineConfigBase* TransientPreset;
 
 	/** The main movie pipeline editor widget */
-	TSharedPtr<SMoviePipelineEditor> MoviePipelineEditorWidget;
+	TSharedPtr<SMoviePipelineConfigEditor> MoviePipelineEditorWidget;
 };
\ No newline at end of file
diff --git a/Engine/Plugins/MovieScene/MovieRenderPipeline/Source/MovieRenderPipelineEditor/Private/Widgets/SMoviePipelineConfigSettings.cpp b/Engine/Plugins/MovieScene/MovieRenderPipeline/Source/MovieRenderPipelineEditor/Private/Widgets/SMoviePipelineConfigSettings.cpp
new file mode 100644
--- /dev/null
+++ b/Engine/Plugins/MovieScene/MovieRenderPipeline/Source/MovieRenderPipelineEditor/Private/Widgets/SMoviePipelineConfigSettings.cpp
@@ -0,0 +1,497 @@
+// Copyright 1998-2019 Epic Games, Inc. All Rights Reserved.
+
+#include "Widgets/SMoviePipelineConfigSettings.h"
+#include "MoviePipelineShotConfig.h"
+#include "MoviePipelineSetting.h"
+
+// Core includes
+#include "Algo/Sort.h"
+#include "UObject/UObjectIterator.h"
+
+// AssetRegistry includes
+#include "AssetRegistryModule.h"
+
+// Slate includes
+#include "Widgets/SBoxPanel.h"
+#include "Widgets/Colors/SColorBlock.h"
+// #include "Widgets/Colors/SColorPicker.h"
+#include "Widgets/Input/SButton.h"
+#include "Widgets/Input/SCheckBox.h"
+#include "Widgets/Images/SThrobber.h"
+#include "Widgets/Layout/SWidgetSwitcher.h"
+#include "Widgets/Views/STreeView.h"
+#include "Widgets/Input/SComboButton.h"
+#include "Widgets/SOverlay.h"
+#include "Widgets/Images/SImage.h"
+#include "Framework/Commands/UICommandList.h"
+#include "Framework/Commands/GenericCommands.h"
+
+// UnrealEd includes
+#include "ScopedTransaction.h"
+
+#define LOCTEXT_NAMESPACE "SMoviePipelineConfigSettings"
+
+struct FMoviePipelineSettingCategory;
+struct FMoviePipelineSettingTreeItem;
+
+struct IMoviePipelineSettingTreeItem : TSharedFromThis<IMoviePipelineSettingTreeItem>
+{
+	virtual ~IMoviePipelineSettingTreeItem() {}
+
+	virtual void Delete(UMoviePipelineConfigBase* Owner) = 0;
+
+	virtual TSharedPtr<FMoviePipelineSettingCategory> AsCategory() { return nullptr; }
+	virtual TSharedPtr<FMoviePipelineSettingTreeItem> AsSetting()   { return nullptr; }
+
+	virtual TSharedRef<SWidget> ConstructWidget(TWeakPtr<SMoviePipelineConfigSettings> SettingsWidget) = 0;
+};
+
+struct FMoviePipelineSettingTreeItem : IMoviePipelineSettingTreeItem
+{
+	/** Weak pointer to the setting that this tree item represents */
+	TWeakObjectPtr<UMoviePipelineSetting> WeakSetting;
+
+	explicit FMoviePipelineSettingTreeItem(UMoviePipelineSetting* InSetting)
+		: WeakSetting(InSetting)
+	{}
+
+	FText GetLabel() const
+	{
+		UMoviePipelineSetting* Setting = WeakSetting.Get();
+		return Setting ? Setting->GetDisplayText() : FText();
+	}
+
+	virtual void Delete(UMoviePipelineConfigBase* Owner) override
+	{
+		UMoviePipelineSetting* Setting = WeakSetting.Get();
+		if (Setting)
+		{
+			Owner->RemoveSetting(Setting);
+		}
+	}
+
+	virtual TSharedPtr<FMoviePipelineSettingTreeItem> AsSetting()
+	{
+		return SharedThis(this);
+	}
+
+	virtual TSharedRef<SWidget> ConstructWidget(TWeakPtr<SMoviePipelineConfigSettings> SettingsWidget) override
+	{
+		return SNew(SOverlay)
+
+		+SOverlay::Slot()
+		[
+		
+			SNew(SHorizontalBox)
+
+			+ SHorizontalBox::Slot()
+			.AutoWidth()
+			.VAlign(VAlign_Center)
+			[
+				SNew(SBox)
+				.WidthOverride(32)
+				.HeightOverride(32)
+		  		[
+					SNew(SCheckBox)
+					//.Style(FTakeRecorderStyle::Get(), "TakeRecorder.Source.Switch")
+					.IsFocusable(false)
+					.IsChecked(this, &FMoviePipelineSettingTreeItem::GetCheckState)
+					.OnCheckStateChanged(this, &FMoviePipelineSettingTreeItem::SetCheckState, SettingsWidget)
+				]
+			]
+
+			+ SHorizontalBox::Slot()
+			.AutoWidth()
+			.VAlign(VAlign_Center)
+			.Padding(8, 0)
+			[
+				SNew(STextBlock)
+				.Text(this, &FMoviePipelineSettingTreeItem::GetLabel)
+				//.TextStyle(FTakeRecorderStyle::Get(), "TakeRecorder.Source.Label")
+				.ColorAndOpacity(this, &FMoviePipelineSettingTreeItem::GetColorAndOpacity)
+			]
+
+			+ SHorizontalBox::Slot()
+			.FillWidth(1)
+			[
+				SNew(SBox)
+			]
+
+			+ SHorizontalBox::Slot()
+			.AutoWidth()
+			.VAlign(VAlign_Center)
+			.Padding(8, 0, 24, 0)
+			[
+				SNew(STextBlock)
+				.Text(this, &FMoviePipelineSettingTreeItem::GetDescription)
+				//.TextStyle(FTakeRecorderStyle::Get(), "TakeRecorder.Source.Label")
+				.Font(FCoreStyle::GetDefaultFontStyle("Regular", 8))
+				.ColorAndOpacity(this, &FMoviePipelineSettingTreeItem::GetColorAndOpacity)
+			]
+		]
+
+		+SOverlay::Slot()
+		[
+			SNew(SHorizontalBox)
+
+			+ SHorizontalBox::Slot()
+			.FillWidth(1)
+			[
+				SNew(SBox)
+			]
+
+			+ SHorizontalBox::Slot()
+			.AutoWidth()
+			.VAlign(VAlign_Center)
+			.Padding(0, 0, 0, 0)
+			[
+				SNew(SColorBlock)
+				.Size(FVector2D(6.0, 38.0))
+				.Color(FLinearColor::Yellow) 
+			]
+
+		];
+
+	}
+
+private:
+
+	FText GetDescription() const
+	{
+		UMoviePipelineSetting* Setting = WeakSetting.Get();
+		return Setting ? Setting->GetDescriptionText() : FText();
+	}
+
+	ECheckBoxState GetCheckState() const
+	{
+		UMoviePipelineSetting* Setting = WeakSetting.Get();
+		return Setting && Setting->bEnabled ? ECheckBoxState::Checked : ECheckBoxState::Unchecked;
+	}
+
+	void SetCheckState(const ECheckBoxState NewState, TWeakPtr<SMoviePipelineConfigSettings> WeakSettingsWidget)
+	{
+		const bool bEnable = NewState == ECheckBoxState::Checked;
+
+		TSharedPtr<SMoviePipelineConfigSettings>		SettingsWidget = WeakSettingsWidget.Pin();
+		UMoviePipelineSetting*            ThisSetting    = WeakSetting.Get();
+
+		if (ThisSetting && SettingsWidget.IsValid())
+		{
+			TArray<UMoviePipelineSetting*> SelectedSettings;
+			SettingsWidget->GetSelectedSettings(SelectedSettings);
+
+			FText TransactionFormat = bEnable
+				? LOCTEXT("EnableSetting", "Enable {0}|plural(one=Setting, other=Settings)")
+				: LOCTEXT("DisableSetting", "Disable {0}|plural(one=Setting, other=Settings)");
+
+			if (!SelectedSettings.Contains(ThisSetting))
+			{
+				FScopedTransaction Transaction(FText::Format(TransactionFormat, 1));
+
+				ThisSetting->Modify();
+				ThisSetting->bEnabled = bEnable;
+			}
+			else 
+			{
+				FScopedTransaction Transaction(FText::Format(TransactionFormat, SelectedSettings.Num()));
+
+				for (UMoviePipelineSetting* SelectedSetting : SelectedSettings)
+				{
+					SelectedSetting->Modify();
+					SelectedSetting->bEnabled = bEnable;
+				}
+			}
+		}
+	}
+
+	FSlateColor GetColorAndOpacity() const
+	{
+		UMoviePipelineSetting* Setting = WeakSetting.Get();
+		return Setting && Setting->bEnabled ? FSlateColor::UseForeground() : FSlateColor::UseSubduedForeground();
+	}
+
+	FSlateColor GetImageColorAndOpacity() const
+	{
+		UMoviePipelineSetting* Setting = WeakSetting.Get();
+		return Setting && Setting->bEnabled ? FLinearColor::White : FLinearColor::White.CopyWithNewOpacity(0.3f);
+	}
+};
+
+struct FMoviePipelineSettingCategory : IMoviePipelineSettingTreeItem
+{
+	/** The title of this category */
+	FText Category;
+
+	/** Sorted list of this category's children */
+	TArray<TSharedPtr<FMoviePipelineSettingTreeItem>> Children;
+
+	explicit FMoviePipelineSettingCategory(const FString& InCategory)
+		: Category(FText::FromString(InCategory))
+	{}
+
+	virtual void Delete(UMoviePipelineConfigBase* Owner) override
+	{
+		for (TSharedPtr<FMoviePipelineSettingTreeItem> Child : Children)
+		{
+			Child->Delete(Owner);
+		}
+	}
+
+	virtual TSharedPtr<FMoviePipelineSettingCategory> AsCategory()
+	{
+		return SharedThis(this);
+	}
+
+	virtual TSharedRef<SWidget> ConstructWidget(TWeakPtr<SMoviePipelineConfigSettings> SettingsWidget) override
+	{
+		return SNew(SHorizontalBox)
+
+			+SHorizontalBox::Slot()
+			.AutoWidth()
+			.Padding(20, 4)
+			.VAlign(VAlign_Center)
+			[
+				SNew(STextBlock)
+				.Text(this, &FMoviePipelineSettingCategory::GetLabel)
+			];
+	}
+
+private:
+
+	FText GetLabel() const
+	{
+		return FText::Format(LOCTEXT("CategoryFormatString", "{0} ({1})"),
+			Category, Children.Num());
+	}
+};
+
+
+void SMoviePipelineConfigSettings::Construct(const FArguments& InArgs)
+{
+	CachedSettingsSerialNumber = uint32(-1);
+
+	TreeView = SNew(STreeView<TSharedPtr<IMoviePipelineSettingTreeItem>>)
+		.TreeItemsSource(&RootNodes)
+		.OnSelectionChanged(InArgs._OnSelectionChanged)
+		.OnGenerateRow(this, &SMoviePipelineConfigSettings::OnGenerateRow)
+		.OnGetChildren(this, &SMoviePipelineConfigSettings::OnGetChildren);
+
+	CommandList = MakeShared<FUICommandList>();
+	CommandList->MapAction(
+		FGenericCommands::Get().Delete,
+		FExecuteAction::CreateSP(this, &SMoviePipelineConfigSettings::OnDeleteSelected),
+		FCanExecuteAction()
+	);
+
+	ChildSlot
+	[
+		TreeView.ToSharedRef()
+	];
+}
+
+void SMoviePipelineConfigSettings::Tick(const FGeometry& AllottedGeometry, const double InCurrentTime, const float InDeltaTime)
+{
+	UMoviePipelineConfigBase* ShotConfig = WeakShotConfig.Get();
+
+	// If we have a settings ptr, we expect its serial number to match our cached one, if not, we rebuild the tree
+	if (ShotConfig)
+	{
+		if (CachedSettingsSerialNumber != ShotConfig->GetSettingsSerialNumber())
+		{
+			ReconstructTree();
+		}
+	}
+	// The settings are no longer valid, so we expect our cached serial number to be -1. If not, we haven't reset the tree yet.
+	else if (CachedSettingsSerialNumber != uint32(-1))
+	{
+		ReconstructTree();
+	}
+}
+
+FReply SMoviePipelineConfigSettings::OnKeyDown(const FGeometry& MyGeometry, const FKeyEvent& InKeyEvent)
+{
+	if (CommandList->ProcessCommandBindings(InKeyEvent))
+	{
+		return FReply::Handled();
+	}
+
+	return FReply::Unhandled();
+}
+
+void SMoviePipelineConfigSettings::GetSelectedSettings(TArray<UMoviePipelineSetting*>& OutSettings) const
+{
+	TArray<TSharedPtr<IMoviePipelineSettingTreeItem>> SelectedItems;
+
+	TreeView->GetSelectedItems(SelectedItems);
+	for (TSharedPtr<IMoviePipelineSettingTreeItem> Item : SelectedItems)
+	{
+		TSharedPtr<FMoviePipelineSettingTreeItem> SettingItem = Item->AsSetting();
+		UMoviePipelineSetting* SettingPtr = SettingItem.IsValid() ? SettingItem->WeakSetting.Get() : nullptr;
+		if (SettingPtr)
+		{
+			OutSettings.Add(SettingPtr);
+		}
+	}
+}
+
+void SMoviePipelineConfigSettings::SetShotConfigObject(UMoviePipelineConfigBase* InShotConfig)
+{
+	WeakShotConfig = InShotConfig;
+	ReconstructTree();
+}
+
+void SMoviePipelineConfigSettings::ReconstructTree()
+{
+	UMoviePipelineConfigBase* ShotConfig = WeakShotConfig.Get();
+	if (!ShotConfig)
+	{
+		CachedSettingsSerialNumber = uint32(-1);
+		RootNodes.Reset();
+		return;
+	}
+
+	CachedSettingsSerialNumber = ShotConfig->GetSettingsSerialNumber();
+
+	TSortedMap<FString, TSharedPtr<FMoviePipelineSettingCategory>> RootCategories;
+	for (TSharedPtr<IMoviePipelineSettingTreeItem> RootItem : RootNodes)
+	{
+		TSharedPtr<FMoviePipelineSettingCategory> RootCategory = RootItem->AsCategory();
+		if (RootCategory.IsValid())
+		{
+			RootCategory->Children.Reset();
+			RootCategories.Add(RootCategory->Category.ToString(), RootCategory);
+		}
+	}
+
+	RootNodes.Reset();
+
+	// We attempt to re-use tree items in order to maintain selection states on them
+	TMap<FObjectKey, TSharedPtr<FMoviePipelineSettingTreeItem>> OldSettingToTreeItem;
+	Swap(SettingToTreeItem, OldSettingToTreeItem);
+
+
+	static const FName CategoryName = "Category";
+
+	for (UMoviePipelineSetting* Setting : ShotConfig->GetSettings())
+	{
+		if (!Setting)
+		{
+			continue;
+		}
+
+		// The category in the UI is taken from the class itself
+		FString Category = TEXT("Default"); // Source->GetCategoryText().ToString();
+		if (Category.IsEmpty())
+		{
+			Category = Setting->GetClass()->GetMetaData(CategoryName);
+		}
+
+		// Attempt to find an existing category node, creating one if necessary
+		TSharedPtr<FMoviePipelineSettingCategory> CategoryNode = RootCategories.FindRef(Category);
+		if (!CategoryNode.IsValid())
+		{
+			CategoryNode = RootCategories.Add(Category, MakeShared<FMoviePipelineSettingCategory>(Category));
+
+			TreeView->SetItemExpansion(CategoryNode, true);
+		}
+
+		// Attempt to find an existing setting item node from the previous data, creating one if necessary
+		FObjectKey SettingKey(Setting);
+		TSharedPtr<FMoviePipelineSettingTreeItem> SettingItem = SettingToTreeItem.FindRef(SettingKey);
+		if (!SettingItem.IsValid())
+		{
+			SettingItem = OldSettingToTreeItem.FindRef(Setting);
+			if (SettingItem.IsValid())
+			{
+				SettingToTreeItem.Add(SettingKey, SettingItem);
+			}
+			else
+			{
+				SettingItem = MakeShared<FMoviePipelineSettingTreeItem>(Setting);
+				SettingToTreeItem.Add(SettingKey, SettingItem);
+
+				TreeView->SetItemExpansion(SettingItem, true);
+			}
+		}
+
+		check(SettingItem.IsValid());
+		CategoryNode->Children.Add(SettingItem);
+	}
+
+
+	RootNodes.Reset(RootCategories.Num());
+	for(TTuple<FString, TSharedPtr<FMoviePipelineSettingCategory>>& Pair : RootCategories)
+	{
+		if (Pair.Value->Children.Num() == 0)
+		{
+			continue;
+		}
+
+		// Sort children by name. Work with a tuple of index and string to avoid excessively calling GetLabel().ToString()
+		TArray<TTuple<int32, FString>> SortData;
+		SortData.Reserve(Pair.Value->Children.Num());
+		for (TSharedPtr<FMoviePipelineSettingTreeItem> Item : Pair.Value->Children)
+		{
+			SortData.Add(MakeTuple(SortData.Num(), Item->GetLabel().ToString()));
+		}
+
+		auto SortPredicate = [](TTuple<int32, FString>& A, TTuple<int32, FString>& B)
+		{
+			return A.Get<1>() < B.Get<1>();
+		};
+
+		Algo::Sort(SortData, SortPredicate);
+
+		// Create a new sorted list of the child entries
+		TArray<TSharedPtr<FMoviePipelineSettingTreeItem>> NewChildren;
+		NewChildren.Reserve(SortData.Num());
+		for (const TTuple<int32, FString>& Item : SortData)
+		{
+			NewChildren.Add(Pair.Value->Children[Item.Get<0>()]);
+		}
+
+		Swap(Pair.Value->Children, NewChildren);
+
+		// Add the category
+		RootNodes.Add(Pair.Value);
+	}
+
+	TreeView->RequestTreeRefresh();
+}
+
+TSharedRef<ITableRow> SMoviePipelineConfigSettings::OnGenerateRow(TSharedPtr<IMoviePipelineSettingTreeItem> Item, const TSharedRef<STableViewBase>& Tree)
+{
+	return
+		SNew(STableRow<TSharedPtr<IMoviePipelineSettingTreeItem>>, Tree)
+		[
+			Item->ConstructWidget(SharedThis(this))
+		];
+}
+
+void SMoviePipelineConfigSettings::OnGetChildren(TSharedPtr<IMoviePipelineSettingTreeItem> Item, TArray<TSharedPtr<IMoviePipelineSettingTreeItem>>& OutChildItems)
+{
+	TSharedPtr<FMoviePipelineSettingCategory> Category = Item->AsCategory();
+	if (Category.IsValid())
+	{
+		OutChildItems.Append(Category->Children);
+	}
+}
+
+void SMoviePipelineConfigSettings::OnDeleteSelected()
+{
+	UMoviePipelineConfigBase* ShotConfig = WeakShotConfig.Get();
+	if (ShotConfig)
+	{
+		TArray<TSharedPtr<IMoviePipelineSettingTreeItem>> Items = TreeView->GetSelectedItems();
+
+		FScopedTransaction Transaction(FText::Format(LOCTEXT("DeleteSelection", "Delete Selected {0}|plural(one=Setting, other=Settings)"), Items.Num()));
+		ShotConfig->Modify();
+
+		for (TSharedPtr<IMoviePipelineSettingTreeItem> Item : Items)
+		{
+			Item->Delete(ShotConfig);
+		}
+	}
+}
+
+#undef LOCTEXT_NAMESPACE
\ No newline at end of file
diff --git a/Engine/Plugins/MovieScene/MovieRenderPipeline/Source/MovieRenderPipelineEditor/Private/Widgets/SMoviePipelineConfigSettings.h b/Engine/Plugins/MovieScene/MovieRenderPipeline/Source/MovieRenderPipelineEditor/Private/Widgets/SMoviePipelineConfigSettings.h
new file mode 100644
--- /dev/null
+++ b/Engine/Plugins/MovieScene/MovieRenderPipeline/Source/MovieRenderPipelineEditor/Private/Widgets/SMoviePipelineConfigSettings.h
@@ -0,0 +1,67 @@
+// Copyright 1998-2019 Epic Games, Inc. All Rights Reserved.
+
+#pragma once
+
+#include "Containers/Array.h"
+#include "Templates/SharedPointer.h"
+#include "Widgets/DeclarativeSyntaxSupport.h"
+#include "Widgets/SCompoundWidget.h"
+#include "UObject/ObjectKey.h"
+#include "Styling/SlateTypes.h" 
+
+class ITableRow;
+class FUICommandList;
+class STableViewBase;
+class UMoviePipelineSetting;
+class UMoviePipelineConfigBase;
+struct IMoviePipelineSettingTreeItem;
+struct FMoviePipelineSettingTreeItem;
+template<typename> class STreeView;
+
+DECLARE_DELEGATE_TwoParams(
+	FOnSettingSelectionChanged,
+	TSharedPtr<IMoviePipelineSettingTreeItem>,
+	ESelectInfo::Type
+)
+
+/** Main widget for the Movie Render Shot Config panel */
+class SMoviePipelineConfigSettings : public SCompoundWidget
+{
+public:
+
+	SLATE_BEGIN_ARGS(SMoviePipelineConfigSettings){}
+		SLATE_EVENT(FOnSettingSelectionChanged, OnSelectionChanged)
+	SLATE_END_ARGS()
+
+	void Construct(const FArguments& InArgs);
+
+	void SetShotConfigObject(UMoviePipelineConfigBase* InShotConfig);
+
+	void GetSelectedSettings(TArray<UMoviePipelineSetting*>& OutSettings) const;
+
+private:
+
+	virtual FReply OnKeyDown(const FGeometry& MyGeometry, const FKeyEvent& InKeyEvent) override;
+	virtual void Tick(const FGeometry& AllottedGeometry, const double InCurrentTime, const float InDeltaTime) override;
+
+private:
+
+	void ReconstructTree();
+
+	TSharedRef<ITableRow> OnGenerateRow(TSharedPtr<IMoviePipelineSettingTreeItem> Item, const TSharedRef<STableViewBase>& Tree);
+
+	void OnGetChildren(TSharedPtr<IMoviePipelineSettingTreeItem> Item, TArray<TSharedPtr<IMoviePipelineSettingTreeItem>>& OutChildItems);
+
+	void OnDeleteSelected();
+
+private:
+
+	uint32 CachedSettingsSerialNumber;
+	TWeakObjectPtr<UMoviePipelineConfigBase> WeakShotConfig;
+
+	TArray<TSharedPtr<IMoviePipelineSettingTreeItem>> RootNodes;
+	TMap<FObjectKey, TSharedPtr<FMoviePipelineSettingTreeItem>> SettingToTreeItem;
+	TSharedPtr<STreeView<TSharedPtr<IMoviePipelineSettingTreeItem>>> TreeView;
+
+	TSharedPtr<FUICommandList> CommandList;
+};
\ No newline at end of file
diff --git a/Engine/Plugins/MovieScene/MovieRenderPipeline/Source/MovieRenderPipelineEditor/Private/Widgets/SMoviePipelineEditor.cpp b/Engine/Plugins/MovieScene/MovieRenderPipeline/Source/MovieRenderPipelineEditor/Private/Widgets/SMoviePipelineEditor.cpp
deleted file mode 100644
diff --git a/Engine/Plugins/MovieScene/MovieRenderPipeline/Source/MovieRenderPipelineEditor/Private/Widgets/SMoviePipelineEditor.h b/Engine/Plugins/MovieScene/MovieRenderPipeline/Source/MovieRenderPipelineEditor/Private/Widgets/SMoviePipelineEditor.h
deleted file mode 100644
diff --git a/Engine/Plugins/MovieScene/MovieRenderPipeline/Source/MovieRenderPipelineEditor/Private/Widgets/SMoviePipelineSettings.cpp b/Engine/Plugins/MovieScene/MovieRenderPipeline/Source/MovieRenderPipelineEditor/Private/Widgets/SMoviePipelineSettings.cpp
deleted file mode 100644
diff --git a/Engine/Plugins/MovieScene/MovieRenderPipeline/Source/MovieRenderPipelineEditor/Private/Widgets/SMoviePipelineSettings.h b/Engine/Plugins/MovieScene/MovieRenderPipeline/Source/MovieRenderPipelineEditor/Private/Widgets/SMoviePipelineSettings.h
deleted file mode 100644
diff --git a/Engine/Plugins/MovieScene/MovieRenderPipeline/Source/MovieRenderPipelineEditor/Private/Widgets/SMoviePipelineTabContent.cpp b/Engine/Plugins/MovieScene/MovieRenderPipeline/Source/MovieRenderPipelineEditor/Private/Widgets/SMoviePipelineTabContent.cpp
deleted file mode 100644
diff --git a/Engine/Plugins/MovieScene/MovieRenderPipeline/Source/MovieRenderPipelineEditor/Private/Widgets/SMoviePipelineTabContent.h b/Engine/Plugins/MovieScene/MovieRenderPipeline/Source/MovieRenderPipelineEditor/Private/Widgets/SMoviePipelineTabContent.h
deleted file mode 100644
