CIS: Fixed missing include.

#rnx

diff --git a/Engine/Plugins/MovieScene/MovieRenderPipeline/Source/MovieRenderPipelineCore/Public/MoviePipelineQueue.h b/Engine/Plugins/MovieScene/MovieRenderPipeline/Source/MovieRenderPipelineCore/Public/MoviePipelineQueue.h
--- a/Engine/Plugins/MovieScene/MovieRenderPipeline/Source/MovieRenderPipelineCore/Public/MoviePipelineQueue.h
+++ b/Engine/Plugins/MovieScene/MovieRenderPipeline/Source/MovieRenderPipelineCore/Public/MoviePipelineQueue.h
@@ -4,6 +4,7 @@
 #include "UObject/Object.h"
 #include "MovieRenderPipelineDataTypes.h"
 #include "LevelSequence.h"
+#include "MoviePipelineMasterConfig.h"
 
 #include "MoviePipelineQueue.generated.h"
 
