Fix compile GetActiveDevice - > FAudioDeviceManager::Get()->GetActiveAudioDevice() after merging from main

diff --git a/Engine/Plugins/MovieScene/MovieRenderPipeline/Source/MovieRenderPipelineCore/Private/MoviePipelineAudioRendering.cpp b/Engine/Plugins/MovieScene/MovieRenderPipeline/Source/MovieRenderPipelineCore/Private/MoviePipelineAudioRendering.cpp
--- a/Engine/Plugins/MovieScene/MovieRenderPipeline/Source/MovieRenderPipelineCore/Private/MoviePipelineAudioRendering.cpp
+++ b/Engine/Plugins/MovieScene/MovieRenderPipeline/Source/MovieRenderPipelineCore/Private/MoviePipelineAudioRendering.cpp
@@ -152,7 +152,7 @@ void UMoviePipeline::ProcessAudioTick()
 		AudioDeltaTime = GetPipelineMasterConfig()->GetEffectiveFrameRate(TargetSequence).AsInterval();
 	}
 
-	if (FAudioDevice* AudioDevice = FAudioDeviceManager::GetActiveDevice())
+	if (FAudioDevice* AudioDevice = FAudioDeviceManager::Get()->GetActiveAudioDevice())
 	{
 		// Handle any game logic that changed Audio State.
 		AudioDevice->Update(true);
