Fixed render pipeline for 64bit arrays.

diff --git a/Engine/Plugins/MovieScene/MovieRenderPipeline/Source/MovieRenderPipelineRenderPasses/Private/MoviePipelineImageSequenceOutput.cpp b/Engine/Plugins/MovieScene/MovieRenderPipeline/Source/MovieRenderPipelineRenderPasses/Private/MoviePipelineImageSequenceOutput.cpp
--- a/Engine/Plugins/MovieScene/MovieRenderPipeline/Source/MovieRenderPipelineRenderPasses/Private/MoviePipelineImageSequenceOutput.cpp
+++ b/Engine/Plugins/MovieScene/MovieRenderPipeline/Source/MovieRenderPipelineRenderPasses/Private/MoviePipelineImageSequenceOutput.cpp
@@ -246,12 +246,12 @@ static TUniquePtr<FImagePixelData> QuantizePixelDataTo8bpp(FImagePixelData* InPi
 		{
 			for (int32 X = 0; X < RawSize.X; X++)
 			{
-				FColor* DestColor = &ClampedPixels.GetData()[Y*RawSize.X + X];
+				FColor* DestColor = &ClampedPixels.GetData()[int64(Y)*int64(RawSize.X) + int64(X)];
 				FLinearColor SrcColor;
 				for (int32 ChanIter = 0; ChanIter < 4; ChanIter++)
 				{
 					FFloat16 Value;
-					Value.Encoded = RawDataPtr[(Y*RawSize.X + X)*RawNumChannels + ChanIter];
+					Value.Encoded = RawDataPtr[(int64(Y)*int64(RawSize.X) + int64(X))*int64(RawNumChannels) + int64(ChanIter)];
 
 					switch (ChanIter)
 					{
@@ -284,11 +284,11 @@ static TUniquePtr<FImagePixelData> QuantizePixelDataTo8bpp(FImagePixelData* InPi
 		{
 			for (int32 X = 0; X < RawSize.X; X++)
 			{
-				FColor* DestColor = &ClampedPixels.GetData()[Y*RawSize.X + X];
+				FColor* DestColor = &ClampedPixels.GetData()[int64(Y)*int64(RawSize.X) + int64(X)];
 				FLinearColor SrcColor;
 				for (int32 ChanIter = 0; ChanIter < 4; ChanIter++)
 				{
-					float Value = RawDataPtr[(Y*RawSize.X + X)*RawNumChannels + ChanIter];
+					float Value = RawDataPtr[(int64(Y)*int64(RawSize.X) + int64(X))*int64(RawNumChannels) + int64(ChanIter)];
 
 					switch (ChanIter)
 					{
