MoviePipeline: BurnIns now respect the Enabled flag.

diff --git a/Engine/Plugins/MovieScene/MovieRenderPipeline/Source/MovieRenderPipelineSettings/Private/MoviePipelineBurnInSetting.cpp b/Engine/Plugins/MovieScene/MovieRenderPipeline/Source/MovieRenderPipelineSettings/Private/MoviePipelineBurnInSetting.cpp
--- a/Engine/Plugins/MovieScene/MovieRenderPipeline/Source/MovieRenderPipelineSettings/Private/MoviePipelineBurnInSetting.cpp
+++ b/Engine/Plugins/MovieScene/MovieRenderPipeline/Source/MovieRenderPipelineSettings/Private/MoviePipelineBurnInSetting.cpp
@@ -13,6 +13,12 @@
 
 void UMoviePipelineBurnInSetting::GatherOutputPassesImpl(TArray<FMoviePipelinePassIdentifier>& ExpectedRenderPasses)
 {
+	// If this was transiently added, don't make a burn-in.
+	if (!GetIsUserCustomized() || !IsEnabled())
+	{
+		return;
+	}
+
 	if (BurnInClass.IsValid())
 	{
 		ExpectedRenderPasses.Add(FMoviePipelinePassIdentifier(TEXT("BurnInOverlay")));
@@ -69,6 +75,12 @@ void UMoviePipelineBurnInSetting::RenderSample_GameThreadImpl(const FMoviePipeli
 
 void UMoviePipelineBurnInSetting::SetupImpl(TArray<TSharedPtr<MoviePipeline::FMoviePipelineEnginePass>>& InEnginePasses, const MoviePipeline::FMoviePipelineRenderPassInitSettings& InPassInitSettings)
 {
+	// If this was transiently added, don't make a burn-in.
+	if (!GetIsUserCustomized() || !IsEnabled())
+	{
+		return;
+	}
+
 	if (BurnInClass.IsValid())
 	{
 		BurnInWidgetInstance = NewObject<UMoviePipelineBurnInWidget>(this, BurnInClass.TryLoadClass<UMoviePipelineBurnInWidget>());
@@ -95,7 +107,7 @@ void UMoviePipelineBurnInSetting::SetupImpl(TArray<TSharedPtr<MoviePipeline::FMo
 
 void UMoviePipelineBurnInSetting::TeardownImpl() 
 {
-	if (FSlateApplication::IsInitialized())
+	if (FSlateApplication::IsInitialized() && VirtualWindow.IsValid())
 	{
 		FSlateApplication::Get().UnregisterVirtualWindow(VirtualWindow.ToSharedRef());
 	}
