MoviePipeline: Refactored image quantization to be shared for more output formats.

diff --git a/Engine/Plugins/MovieScene/MovieRenderPipeline/Source/MovieRenderPipelineCore/Private/MoviePipelineImageQuantization.cpp b/Engine/Plugins/MovieScene/MovieRenderPipeline/Source/MovieRenderPipelineCore/Private/MoviePipelineImageQuantization.cpp
new file mode 100644
--- /dev/null
+++ b/Engine/Plugins/MovieScene/MovieRenderPipeline/Source/MovieRenderPipelineCore/Private/MoviePipelineImageQuantization.cpp
@@ -0,0 +1,136 @@
+// Copyright Epic Games, Inc. All Rights Reserved.
+
+#include "MoviePipelineImageQuantization.h"
+#include "ImagePixelData.h"
+
+namespace UE
+{
+namespace MoviePipeline
+{
+
+static TUniquePtr<FImagePixelData> QuantizePixelDataTo8bpp(const FImagePixelData* InPixelData)
+{
+	TUniquePtr<FImagePixelData> QuantizedPixelData = nullptr;
+
+	FIntPoint RawSize = InPixelData->GetSize();
+	int32 RawNumChannels = InPixelData->GetNumChannels();
+
+	// Look at our incoming bit depth
+	switch (InPixelData->GetBitDepth())
+	{
+	case 8:
+	{
+		// No work actually needs to be done, hooray! We return a copy of the data for consistency
+		QuantizedPixelData = InPixelData->CopyImageData();
+		break;
+	}
+	case 16:
+	{
+		TArray<FColor> ClampedPixels;
+		ClampedPixels.SetNum(RawSize.X * RawSize.Y);
+
+		int64 SizeInBytes = 0;
+		const void* SrcRawDataPtr = nullptr;
+		InPixelData->GetRawData(SrcRawDataPtr, SizeInBytes);
+
+		const uint16* RawDataPtr = static_cast<const uint16*>(SrcRawDataPtr);
+
+		// Copy pixels to new array
+		for (int32 Y = 0; Y < RawSize.Y; Y++)
+		{
+			for (int32 X = 0; X < RawSize.X; X++)
+			{
+				FColor* DestColor = &ClampedPixels.GetData()[Y*RawSize.X + X];
+				FLinearColor SrcColor;
+				for (int32 ChanIter = 0; ChanIter < 4; ChanIter++)
+				{
+					FFloat16 Value;
+					Value.Encoded = RawDataPtr[(Y*RawSize.X + X)*RawNumChannels + ChanIter];
+
+					switch (ChanIter)
+					{
+					case 0: SrcColor.R = Value; break;
+					case 1: SrcColor.G = Value; break;
+					case 2: SrcColor.B = Value; break;
+					case 3: SrcColor.A = Value; break;
+					}
+				}
+				// convert to FColor using sRGB conversion
+				*DestColor = SrcColor.ToFColor(true);
+			}
+		}
+
+		QuantizedPixelData = MakeUnique<TImagePixelData<FColor>>(RawSize, TArray64<FColor>(MoveTemp(ClampedPixels)));
+	}
+	case 32:
+	{
+		TArray<FColor> ClampedPixels;
+		ClampedPixels.SetNum(RawSize.X * RawSize.Y);
+
+		int64 SizeInBytes = 0;
+		const void* SrcRawDataPtr = nullptr;
+		InPixelData->GetRawData(SrcRawDataPtr, SizeInBytes);
+
+		const float* RawDataPtr = static_cast<const float*>(SrcRawDataPtr);
+
+		// Copy pixels to new array
+		for (int32 Y = 0; Y < RawSize.Y; Y++)
+		{
+			for (int32 X = 0; X < RawSize.X; X++)
+			{
+				FColor* DestColor = &ClampedPixels.GetData()[Y*RawSize.X + X];
+				FLinearColor SrcColor;
+				for (int32 ChanIter = 0; ChanIter < 4; ChanIter++)
+				{
+					float Value = RawDataPtr[(Y*RawSize.X + X)*RawNumChannels + ChanIter];
+
+					switch (ChanIter)
+					{
+					case 0: SrcColor.R = Value; break;
+					case 1: SrcColor.G = Value; break;
+					case 2: SrcColor.B = Value; break;
+					case 3: SrcColor.A = Value; break;
+					}
+				}
+				// convert to FColor using sRGB conversion
+				*DestColor = SrcColor.ToFColor(true);
+			}
+		}
+
+		QuantizedPixelData = MakeUnique<TImagePixelData<FColor>>(RawSize, TArray64<FColor>(MoveTemp(ClampedPixels)));
+		break;
+	}
+
+	default:
+		// Unsupported source bit-depth, consider adding it!
+		check(false);
+	}
+
+	return QuantizedPixelData;
+}
+
+TUniquePtr<FImagePixelData> QuantizeImagePixelDataToBitDepth(const FImagePixelData* InData, const int32 TargetBitDepth)
+{
+	TUniquePtr<FImagePixelData> QuantizedPixelData = nullptr;
+	switch (TargetBitDepth)
+	{
+	case 8:
+		// Convert to 8 bit FColor
+		QuantizedPixelData = UE::MoviePipeline::QuantizePixelDataTo8bpp(InData);
+		break;
+	case 10:
+	case 12:
+	case 16:
+		// Convert to 16 bit FFloat16Color
+	case 32:
+		// Convert to 32 bit FLinearColor
+
+	default:
+		// Unsupported bit-depth to convert from, please consider implementing!
+		check(false);
+	}
+
+	return QuantizedPixelData;
+}
+}
+}
diff --git a/Engine/Plugins/MovieScene/MovieRenderPipeline/Source/MovieRenderPipelineCore/Private/MoviePipelineOutputBase.cpp b/Engine/Plugins/MovieScene/MovieRenderPipeline/Source/MovieRenderPipelineCore/Private/MoviePipelineOutputBase.cpp
--- a/Engine/Plugins/MovieScene/MovieRenderPipeline/Source/MovieRenderPipelineCore/Private/MoviePipelineOutputBase.cpp
+++ b/Engine/Plugins/MovieScene/MovieRenderPipeline/Source/MovieRenderPipelineCore/Private/MoviePipelineOutputBase.cpp
@@ -18,4 +18,70 @@ void UMoviePipelineOutputBase::ValidateStateImpl()
 		}
 
 	}
+}
+
+namespace UE
+{
+namespace MoviePipeline
+{
+	void ValidateOutputFormatString(FString& InOutFilenameFormatString, const bool bTestRenderPass, const bool bTestFrameNumber)
+	{
+		// If there is more than one file being written for this frame, make sure they uniquely identify.
+		if (bTestRenderPass)
+		{
+			if (!InOutFilenameFormatString.Contains(TEXT("{render_pass}"), ESearchCase::IgnoreCase))
+			{
+				InOutFilenameFormatString += TEXT("{render_pass}");
+			}
+		}
+
+		if (bTestFrameNumber)
+		{
+			// Ensure there is a frame number in the output string somewhere to uniquely identify individual files in an image sequence.
+			FString FrameNumberIdentifiers[] = { TEXT("{frame_number}"), TEXT("{frame_number_shot}"), TEXT("{frame_number_rel}"), TEXT("{frame_number_shot_rel}") };
+			int32 FrameNumberIndex = INDEX_NONE;
+			for (const FString& Identifier : FrameNumberIdentifiers)
+			{
+				FrameNumberIndex = InOutFilenameFormatString.Find(Identifier, ESearchCase::IgnoreCase, ESearchDir::FromEnd);
+				if (FrameNumberIndex != INDEX_NONE)
+				{
+					break;
+				}
+			}
+
+			// We want to insert a {file_dup} before the frame number. This instructs the name resolver to put the (2) before
+			// the frame number, so that they're still properly recognized as image sequences by other software. It will resolve
+			// to "" if not needed.
+			if (FrameNumberIndex == INDEX_NONE)
+			{
+				// Automatically append {frame_number} so the files are uniquely identified.
+				InOutFilenameFormatString.Append(TEXT("{file_dup}.{frame_number}"));
+			}
+			else
+			{
+				// The user had already specified a frame number identifier, so we need to insert the
+				// file_dup tag before it.
+				InOutFilenameFormatString.InsertAt(FrameNumberIndex, TEXT("{file_dup}"));
+			}
+		}
+
+		if (!InOutFilenameFormatString.Contains(TEXT("{file_dup}"), ESearchCase::IgnoreCase))
+		{
+			InOutFilenameFormatString.Append(TEXT("{file_dup}"));
+		}
+	}
+
+	void RemoveFrameNumberFormatStrings(FString& InOutFilenameFormatString, const bool bIncludeShots)
+	{
+		// Strip {frame_number} related separators from their file name, otherwise it will create one output file per frame.
+		InOutFilenameFormatString.ReplaceInline(TEXT("{frame_number}"), TEXT(""));
+		InOutFilenameFormatString.ReplaceInline(TEXT("{frame_number_rel}"), TEXT(""));
+
+		if (bIncludeShots)
+		{
+			InOutFilenameFormatString.ReplaceInline(TEXT("{frame_number_shot}"), TEXT(""));
+			InOutFilenameFormatString.ReplaceInline(TEXT("{frame_number_shot_rel}"), TEXT(""));
+		}
+	}
+}
 }
\ No newline at end of file
diff --git a/Engine/Plugins/MovieScene/MovieRenderPipeline/Source/MovieRenderPipelineCore/Public/MoviePipelineImageQuantization.h b/Engine/Plugins/MovieScene/MovieRenderPipeline/Source/MovieRenderPipelineCore/Public/MoviePipelineImageQuantization.h
new file mode 100644
--- /dev/null
+++ b/Engine/Plugins/MovieScene/MovieRenderPipeline/Source/MovieRenderPipelineCore/Public/MoviePipelineImageQuantization.h
@@ -0,0 +1,15 @@
+// Copyright Epic Games, Inc. All Rights Reserved.
+#pragma once
+#include "Engine/EngineTypes.h"
+
+// Forward Declare
+struct FImagePixelData;
+
+namespace UE
+{
+namespace MoviePipeline
+{
+	MOVIERENDERPIPELINECORE_API TUniquePtr<FImagePixelData> QuantizeImagePixelDataToBitDepth(const FImagePixelData* InData, const int32 TargetBitDepth);
+}
+}
+	  
\ No newline at end of file
diff --git a/Engine/Plugins/MovieScene/MovieRenderPipeline/Source/MovieRenderPipelineCore/Public/MoviePipelineOutputBase.h b/Engine/Plugins/MovieScene/MovieRenderPipeline/Source/MovieRenderPipelineCore/Public/MoviePipelineOutputBase.h
--- a/Engine/Plugins/MovieScene/MovieRenderPipeline/Source/MovieRenderPipelineCore/Public/MoviePipelineOutputBase.h
+++ b/Engine/Plugins/MovieScene/MovieRenderPipeline/Source/MovieRenderPipelineCore/Public/MoviePipelineOutputBase.h
@@ -9,6 +9,15 @@
 
 class UMoviePipeline;
 
+namespace UE
+{
+namespace MoviePipeline
+{
+	MOVIERENDERPIPELINECORE_API void ValidateOutputFormatString(FString& InOutFilenameFormatString, const bool bTestRenderPass, const bool bTestFrameNumber);
+	MOVIERENDERPIPELINECORE_API void RemoveFrameNumberFormatStrings(FString& InOutFilenameFormatString, const bool bIncludeShots);
+}
+}
+
 /**
 * A base class for different output containers for the data generated by the Movie Pipeline,
 * i.e: image sequences, video containers, etc.
diff --git a/Engine/Plugins/MovieScene/MovieRenderPipeline/Source/MovieRenderPipelineRenderPasses/Private/MoviePipelineImageSequenceOutput.cpp b/Engine/Plugins/MovieScene/MovieRenderPipeline/Source/MovieRenderPipelineRenderPasses/Private/MoviePipelineImageSequenceOutput.cpp
--- a/Engine/Plugins/MovieScene/MovieRenderPipeline/Source/MovieRenderPipelineRenderPasses/Private/MoviePipelineImageSequenceOutput.cpp
+++ b/Engine/Plugins/MovieScene/MovieRenderPipeline/Source/MovieRenderPipelineRenderPasses/Private/MoviePipelineImageSequenceOutput.cpp
@@ -15,9 +15,10 @@
 #include "MoviePipelineOutputSetting.h"
 #include "Containers/UnrealString.h"
 #include "Misc/StringFormatArg.h"
+#include "MoviePipelineOutputBase.h"
+#include "MoviePipelineImageQuantization.h"
 
 // Forward Declare
-static void ValidateOutputFormatString(FString& InOutFilenameFormatString, const bool bIncludeRenderPass);
 static TUniquePtr<FImagePixelData> QuantizePixelDataTo8bpp(FImagePixelData* InPixelData);
 
 
@@ -84,6 +85,7 @@ void UMoviePipelineImageSequenceOutputBase::OnRecieveImageDataImpl(FMoviePipelin
 
 		TUniquePtr<FImagePixelData> QuantizedPixelData = nullptr;
 		
+
 		switch (OutputFormat)
 		{
 		case EImageFormat::PNG:
@@ -92,7 +94,7 @@ void UMoviePipelineImageSequenceOutputBase::OnRecieveImageDataImpl(FMoviePipelin
 		{
 			// All three of these formats only support 8 bit data, so we need to take the incoming buffer type,
 			// copy it into a new 8-bit array and optionally apply a little noise to the data to help hide gradient banding.
-			QuantizedPixelData = QuantizePixelDataTo8bpp(RenderPassData.Value.Get());
+			QuantizedPixelData = UE::MoviePipeline::QuantizeImagePixelDataToBitDepth(RenderPassData.Value.Get(), 8);
 			break;
 		}
 		case EImageFormat::EXR:
@@ -111,8 +113,9 @@ void UMoviePipelineImageSequenceOutputBase::OnRecieveImageDataImpl(FMoviePipelin
 			// If we're writing more than one render pass out, we need to ensure the file name has the format string in it so we don't
 			// overwrite the same file multiple times. Burn In overlays don't count because they get composited on top of an existing file.
 			const bool bIncludeRenderPass = InMergedOutputFrame->ImageOutputData.Num() - (BurnInImageData ? 1 : 0) > 1;
+			const bool bTestFrameNumber = true;
 
-			ValidateOutputFormatString(FileNameFormatString, bIncludeRenderPass);
+			UE::MoviePipeline::ValidateOutputFormatString(FileNameFormatString, bIncludeRenderPass, bTestFrameNumber);
 
 			// Create specific data that needs to override 
 			FStringFormatNamedArguments FormatOverrides;
@@ -174,142 +177,3 @@ void UMoviePipelineImageSequenceOutputBase::GetFilenameFormatArguments(FMoviePip
 	InOutFormatArgs.Arguments.Add(TEXT("render_pass"), TEXT("RenderPassName"));
 }
 
-static void ValidateOutputFormatString(FString& InOutFilenameFormatString, const bool bIncludeRenderPass)
-{
-	// If there is more than one file being written for this frame, make sure they uniquely identify.
-	if (bIncludeRenderPass)
-	{
-		if (!InOutFilenameFormatString.Contains(TEXT("{render_pass}"), ESearchCase::IgnoreCase))
-		{
-			InOutFilenameFormatString += TEXT("{render_pass}");
-		}
-	}
-
-	// Ensure there is a frame number in the output string somewhere to uniquely identify individual files in an image sequence.
-	FString FrameNumberIdentifiers[] = { TEXT("{frame_number}"), TEXT("{frame_number_shot}"), TEXT("{frame_number_rel}"), TEXT("{frame_number_shot_rel}") };
-	int32 FrameNumberIndex = INDEX_NONE;
-	for (const FString& Identifier : FrameNumberIdentifiers)
-	{
-		FrameNumberIndex = InOutFilenameFormatString.Find(Identifier, ESearchCase::IgnoreCase, ESearchDir::FromEnd);
-		if(FrameNumberIndex != INDEX_NONE)
-		{
-			break;
-		}
-	}
-
-	// We want to insert a {file_dup} before the frame number. This instructs the name resolver to put the (2) before
-	// the frame number, so that they're still properly recognized as image sequences by other software. It will resolve
-	// to "" if not needed.
-	if (FrameNumberIndex == INDEX_NONE)
-	{
-		// Automatically append {frame_number} so the files are uniquely identified.
-		InOutFilenameFormatString.Append(TEXT("{file_dup}.{frame_number}"));
-	}
-	else
-	{
-		// The user had already specified a frame number identifier, so we need to insert the
-		// file_dup tag before it.
-		InOutFilenameFormatString.InsertAt(FrameNumberIndex, TEXT("{file_dup}"));
-	}
-}
-
-static TUniquePtr<FImagePixelData> QuantizePixelDataTo8bpp(FImagePixelData* InPixelData)
-{
-	TUniquePtr<FImagePixelData> QuantizedPixelData = nullptr;
-
-	FIntPoint RawSize = InPixelData->GetSize();
-	int32 RawNumChannels = InPixelData->GetNumChannels();
-
-	// Look at our incoming bit depth
-	switch (InPixelData->GetBitDepth())
-	{
-	case 8:
-	{
-		// No work actually needs to be done, hooray! We'll copy the data though so that when it gets moved
-		// into the ImageWriteTask the original data can be passed to another output container.
-		QuantizedPixelData = InPixelData->CopyImageData();
-		break;
-	}
-	case 16:
-	{
-		TArray<FColor> ClampedPixels;
-		ClampedPixels.SetNum(RawSize.X * RawSize.Y);
-
-		int64 SizeInBytes = 0;
-		const void* SrcRawDataPtr = nullptr;
-		InPixelData->GetRawData(SrcRawDataPtr, SizeInBytes);
-
-		const uint16* RawDataPtr = static_cast<const uint16*>(SrcRawDataPtr);
-
-		// Copy pixels to new array
-		for (int32 Y = 0; Y < RawSize.Y; Y++)
-		{
-			for (int32 X = 0; X < RawSize.X; X++)
-			{
-				FColor* DestColor = &ClampedPixels.GetData()[int64(Y)*int64(RawSize.X) + int64(X)];
-				FLinearColor SrcColor;
-				for (int32 ChanIter = 0; ChanIter < 4; ChanIter++)
-				{
-					FFloat16 Value;
-					Value.Encoded = RawDataPtr[(int64(Y)*int64(RawSize.X) + int64(X))*int64(RawNumChannels) + int64(ChanIter)];
-
-					switch (ChanIter)
-					{
-					case 0: SrcColor.R = Value; break;
-					case 1: SrcColor.G = Value; break;
-					case 2: SrcColor.B = Value; break;
-					case 3: SrcColor.A = Value; break;
-					}
-				}
-				// convert to FColor using sRGB conversion
-				*DestColor = SrcColor.ToFColor(true);
-			}
-		}
-
-		QuantizedPixelData = MakeUnique<TImagePixelData<FColor>>(RawSize, TArray64<FColor>(MoveTemp(ClampedPixels)));
-	}
-	case 32:
-	{
-		TArray<FColor> ClampedPixels;
-		ClampedPixels.SetNum(RawSize.X * RawSize.Y);
-
-		int64 SizeInBytes = 0;
-		const void* SrcRawDataPtr = nullptr;
-		InPixelData->GetRawData(SrcRawDataPtr, SizeInBytes);
-
-		const float* RawDataPtr = static_cast<const float*>(SrcRawDataPtr);
-
-		// Copy pixels to new array
-		for (int32 Y = 0; Y < RawSize.Y; Y++)
-		{
-			for (int32 X = 0; X < RawSize.X; X++)
-			{
-				FColor* DestColor = &ClampedPixels.GetData()[int64(Y)*int64(RawSize.X) + int64(X)];
-				FLinearColor SrcColor;
-				for (int32 ChanIter = 0; ChanIter < 4; ChanIter++)
-				{
-					float Value = RawDataPtr[(int64(Y)*int64(RawSize.X) + int64(X))*int64(RawNumChannels) + int64(ChanIter)];
-
-					switch (ChanIter)
-					{
-					case 0: SrcColor.R = Value; break;
-					case 1: SrcColor.G = Value; break;
-					case 2: SrcColor.B = Value; break;
-					case 3: SrcColor.A = Value; break;
-					}
-				}
-				// convert to FColor using sRGB conversion
-				*DestColor = SrcColor.ToFColor(true);
-			}
-		}
-
-		QuantizedPixelData = MakeUnique<TImagePixelData<FColor>>(RawSize, TArray64<FColor>(MoveTemp(ClampedPixels)));
-		break;
-	}
-
-	default:
-		check(false);
-	}
-
-	return QuantizedPixelData;
-}
\ No newline at end of file
