MoviePipeline: Close Editor option in New Process rendering settings now closes the editor after succesfully launching a remote render.

#rnx

diff --git a/Engine/Plugins/MovieScene/MovieRenderPipeline/Source/MovieRenderPipelineEditor/Private/MoviePipelineNewProcessExecutor.cpp b/Engine/Plugins/MovieScene/MovieRenderPipeline/Source/MovieRenderPipelineEditor/Private/MoviePipelineNewProcessExecutor.cpp
--- a/Engine/Plugins/MovieScene/MovieRenderPipeline/Source/MovieRenderPipelineEditor/Private/MoviePipelineNewProcessExecutor.cpp
+++ b/Engine/Plugins/MovieScene/MovieRenderPipeline/Source/MovieRenderPipelineEditor/Private/MoviePipelineNewProcessExecutor.cpp
@@ -184,8 +184,15 @@ void UMoviePipelineNewProcessExecutor::ExecuteImpl(UMoviePipelineQueue* InPipeli
 	}
 	else
 	{
-		// Register a tick handler to listen every frame to see if the process shut down gracefully, we'll use return codes to tell success vs cancel.
-		FCoreDelegates::OnBeginFrame.AddUObject(this, &UMoviePipelineNewProcessExecutor::CheckForProcessFinished);
+		if (ExecutorSettings->bCloseEditor)
+		{
+			FPlatformMisc::RequestExit(false);
+		}
+		else
+		{
+			// Register a tick handler to listen every frame to see if the process shut down gracefully, we'll use return codes to tell success vs cancel.
+			FCoreDelegates::OnBeginFrame.AddUObject(this, &UMoviePipelineNewProcessExecutor::CheckForProcessFinished);
+		}
 	}
 }
 
