MoviePipeline: Fixed Game Overrides not being applied to PIE renders.

diff --git a/Engine/Plugins/MovieScene/MovieRenderPipeline/Source/MovieRenderPipelineEditor/Private/MoviePipelinePIEExecutor.cpp b/Engine/Plugins/MovieScene/MovieRenderPipeline/Source/MovieRenderPipelineEditor/Private/MoviePipelinePIEExecutor.cpp
--- a/Engine/Plugins/MovieScene/MovieRenderPipeline/Source/MovieRenderPipelineEditor/Private/MoviePipelinePIEExecutor.cpp
+++ b/Engine/Plugins/MovieScene/MovieRenderPipeline/Source/MovieRenderPipelineEditor/Private/MoviePipelinePIEExecutor.cpp
@@ -54,17 +54,26 @@ void UMoviePipelinePIEExecutor::Start(const UMoviePipelineExecutorJob* InJob)
 	FRequestPlaySessionParams Params;
 	Params.EditorPlaySettings = PlayInEditorSettings;
 	Params.CustomPIEWindow = CustomWindow;
+	Params.GlobalMapOverride = InJob->Map.GetAssetPathString();
 
-	UMoviePipelineGameOverrideSetting* GameOverrides = InJob->GetConfiguration()->FindSetting<UMoviePipelineGameOverrideSetting>();
-	if (GameOverrides)
-	{
-		Params.GameModeOverride = GameOverrides->GameModeOverride;
+	// Initialize the transient settings so that they will exist in time for the GameOverrides check.
+	InJob->GetConfiguration()->InitializeTransientSettings();
+
+	TArray<UMoviePipelineSetting*> AllSettings = InJob->GetConfiguration()->GetAllSettings();
+	UMoviePipelineSetting** GameOverridesPtr = AllSettings.FindByPredicate([](UMoviePipelineSetting* InSetting) { return InSetting->GetClass() == UMoviePipelineGameOverrideSetting::StaticClass(); });
+	if (GameOverridesPtr)
+	{	
+		UMoviePipelineSetting* Setting = *GameOverridesPtr;
+		if (Setting)
+		{
+			Params.GameModeOverride = CastChecked<UMoviePipelineGameOverrideSetting>(Setting)->GameModeOverride;
+		}
 	}
 
 	bPreviousUseFixedTimeStep = FApp::UseFixedTimeStep();
 	PreviousFixedTimeStepDelta = FApp::GetFixedDeltaTime();
 
-	// Force the engine into fixed timestep mode. It's going to get overriden on the first frame by the movie pipeline,
+	// Force the engine into fixed timestep mode. It's going to get overridden on the first frame by the movie pipeline,
 	// and everything controlled by Sequencer will use the correct timestep for renders but non-controlled things (such
 	// as pawns) use an uncontrolled DT on the first frame which lowers determinism.
 	ULevelSequence* LevelSequence = CastChecked<ULevelSequence>(InJob->Sequence.TryLoad());
