MoviePipeline: Minor optimizations to reduce render thread flushing and added profiling stat group for further analysis.

diff --git a/Engine/Plugins/MovieScene/MovieRenderPipeline/Source/MovieRenderPipelineCore/Private/MoviePipelineOutputSetting.cpp b/Engine/Plugins/MovieScene/MovieRenderPipeline/Source/MovieRenderPipelineCore/Private/MoviePipelineOutputSetting.cpp
--- a/Engine/Plugins/MovieScene/MovieRenderPipeline/Source/MovieRenderPipelineCore/Private/MoviePipelineOutputSetting.cpp
+++ b/Engine/Plugins/MovieScene/MovieRenderPipeline/Source/MovieRenderPipelineCore/Private/MoviePipelineOutputSetting.cpp
@@ -19,6 +19,7 @@ UMoviePipelineOutputSetting::UMoviePipelineOutputSetting()
 	, ZeroPadFrameNumbers(4)
 	, FrameNumberOffset(0)
 	, bDisableToneCurve(false)
+	, bHidePreview(true)
 {
 	FileNameFormat = TEXT("{sequence_name}.{frame_number}");
 	OutputDirectory.Path = FPaths::ProjectSavedDir() / TEXT("MovieRenders/");
diff --git a/Engine/Plugins/MovieScene/MovieRenderPipeline/Source/MovieRenderPipelineCore/Private/MoviePipelineRendering.cpp b/Engine/Plugins/MovieScene/MovieRenderPipeline/Source/MovieRenderPipelineCore/Private/MoviePipelineRendering.cpp
--- a/Engine/Plugins/MovieScene/MovieRenderPipeline/Source/MovieRenderPipelineCore/Private/MoviePipelineRendering.cpp
+++ b/Engine/Plugins/MovieScene/MovieRenderPipeline/Source/MovieRenderPipelineCore/Private/MoviePipelineRendering.cpp
@@ -20,6 +20,7 @@
 #include "MoviePipelineHighResSetting.h"
 #include "Modules/ModuleManager.h"
 #include "MoviePipelineCameraSetting.h"
+#include "Engine/GameViewportClient.h"
 
 // For flushing async systems
 #include "RendererInterface.h"
@@ -154,6 +155,11 @@ void UMoviePipeline::SetupRenderingPipelineForShot(FMoviePipelineShotInfo& Shot)
 		NumOutputPasses++;
 	}
 
+	if (OutputSettings->bHidePreview)
+	{
+		GetWorld()->GetGameViewport()->bDisableWorldRendering = true;
+	}
+
 	UE_LOG(LogMovieRenderPipeline, Log, TEXT("Finished setting up rendering for shot. Shot has %d Engine Passes and %d Output Passes."), ActiveRenderPasses.Num(), NumOutputPasses);
 }
 
@@ -431,7 +437,7 @@ void UMoviePipeline::RenderFrame()
 	}
 	
 	// UE_LOG(LogMovieRenderPipeline, Warning, TEXT("[%d] Pre-FlushRenderingCommands"), GFrameCounter);
-	FlushRenderingCommands();
+	// FlushRenderingCommands();
 	// UE_LOG(LogMovieRenderPipeline, Warning, TEXT("[%d] Post-FlushRenderingCommands"), GFrameCounter);
 }
 
diff --git a/Engine/Plugins/MovieScene/MovieRenderPipeline/Source/MovieRenderPipelineCore/Private/MovieRenderOverlappedImage.cpp b/Engine/Plugins/MovieScene/MovieRenderPipeline/Source/MovieRenderPipelineCore/Private/MovieRenderOverlappedImage.cpp
--- a/Engine/Plugins/MovieScene/MovieRenderPipeline/Source/MovieRenderPipelineCore/Private/MovieRenderOverlappedImage.cpp
+++ b/Engine/Plugins/MovieScene/MovieRenderPipeline/Source/MovieRenderPipelineCore/Private/MovieRenderOverlappedImage.cpp
@@ -9,6 +9,16 @@
 #include "Math/Float16.h"
 #include "HAL/PlatformTime.h"
 
+DECLARE_CYCLE_STAT(TEXT("MoviePipeline_SampleReadback"), STAT_SampleReadback, STATGROUP_MoviePipeline);
+DECLARE_CYCLE_STAT(TEXT("OverlapAccumulator_AccumulateSinglePlane"), STAT_AccumulateSinglePlane, STATGROUP_MoviePipeline);
+DECLARE_CYCLE_STAT(TEXT("OverlapAccumulator_InitMemory"), STAT_InitMemory, STATGROUP_MoviePipeline);
+DECLARE_CYCLE_STAT(TEXT("OverlapAccumulator_ZeroPlanes"), STAT_ZeroPlanes, STATGROUP_MoviePipeline);
+DECLARE_CYCLE_STAT(TEXT("OverlapAccumulator_AccumulatePixelData"), STAT_AccumulatePixelData, STATGROUP_MoviePipeline);
+DECLARE_CYCLE_STAT(TEXT("OverlapAccumulator_FetchFull"), STAT_FetchFullImageValue, STATGROUP_MoviePipeline);
+DECLARE_CYCLE_STAT(TEXT("OverlapAccumulator_FetchPixelData8"), STAT_FetchFinalPixelData8bit, STATGROUP_MoviePipeline);
+DECLARE_CYCLE_STAT(TEXT("OverlapAccumulator_FetchPixelData16"), STAT_FetchFinalPixelData16bit, STATGROUP_MoviePipeline);
+DECLARE_CYCLE_STAT(TEXT("OverlapAccumulator_FetchPixelData32"), STAT_FetchFinalPixelData32bit, STATGROUP_MoviePipeline);
+
 void FImageOverlappedPlane::Init(FIntPoint InSize)
 {
 	Size = InSize;
@@ -41,6 +51,8 @@ void FImageOverlappedPlane::AccumulateSinglePlane(const TArray64<float>& InRawDa
 								const TArray<float> & WeightDataX,
 								const TArray<float> & WeightDataY)
 {
+	SCOPE_CYCLE_COUNTER(STAT_AccumulateSinglePlane);
+
 	check(InRawData.Num() == InSize.X * InSize.Y);
 	check(WeightDataX.Num() == InSize.X);
 	check(WeightDataY.Num() == InSize.Y);
@@ -335,6 +347,7 @@ void FImageOverlappedPlane::AccumulateSinglePlane(const TArray64<float>& InRawDa
 
 void FImageOverlappedAccumulator::InitMemory(FIntPoint InPlaneSize, int32 InNumChannels)
 {
+	SCOPE_CYCLE_COUNTER(STAT_InitMemory);
 	PlaneSize.X = InPlaneSize.X;
 	PlaneSize.Y = InPlaneSize.Y;
 	NumChannels = InNumChannels;
@@ -351,6 +364,7 @@ void FImageOverlappedAccumulator::InitMemory(FIntPoint InPlaneSize, int32 InNumC
 
 void FImageOverlappedAccumulator::ZeroPlanes()
 {
+	SCOPE_CYCLE_COUNTER(STAT_ZeroPlanes);
 	check(ChannelPlanes.Num() == NumChannels);
 
 	for (int32 Channel = 0; Channel < NumChannels; Channel++)
@@ -471,6 +485,7 @@ void FImageOverlappedAccumulator::CheckTileSubRect(const TArray64<float>& OutWei
 
 void FImageOverlappedAccumulator::AccumulatePixelData(const FImagePixelData& InPixelData, FIntPoint InTileOffset, FVector2D InSubpixelOffset, const MoviePipeline::FTileWeight1D & WeightX, const MoviePipeline::FTileWeight1D & WeightY)
 {
+	SCOPE_CYCLE_COUNTER(STAT_AccumulatePixelData);
 	EImagePixelType Fmt = InPixelData.GetType();
 
 	int32 RawNumChan = InPixelData.GetNumChannels();
@@ -729,6 +744,8 @@ void FImageOverlappedAccumulator::AccumulatePixelData(const FImagePixelData& InP
 
 void FImageOverlappedAccumulator::FetchFullImageValue(float Rgba[4], int32 FullX, int32 FullY) const
 {
+	SCOPE_CYCLE_COUNTER(STAT_FetchFullImageValue);
+
 	Rgba[0] = 0.0f;
 	Rgba[1] = 0.0f;
 	Rgba[2] = 0.0f;
@@ -757,6 +774,7 @@ void FImageOverlappedAccumulator::FetchFullImageValue(float Rgba[4], int32 FullX
 
 void FImageOverlappedAccumulator::FetchFinalPixelDataByte(TArray64<FColor> & OutPixelData) const
 {
+	SCOPE_CYCLE_COUNTER(STAT_FetchFinalPixelData8bit);
 	int32 FullSizeX = PlaneSize.X;
 	int32 FullSizeY = PlaneSize.Y;
 	OutPixelData.SetNumUninitialized(FullSizeX * FullSizeY);
@@ -779,6 +797,7 @@ void FImageOverlappedAccumulator::FetchFinalPixelDataByte(TArray64<FColor> & Out
 
 void FImageOverlappedAccumulator::FetchFinalPixelDataHalfFloat(TArray64<FFloat16Color>& OutPixelData) const
 {
+	SCOPE_CYCLE_COUNTER(STAT_FetchFinalPixelData16bit);
 	int32 FullSizeX = PlaneSize.X;
 	int32 FullSizeY = PlaneSize.Y;
 	OutPixelData.SetNumUninitialized(FullSizeX * FullSizeY);
@@ -801,6 +820,7 @@ void FImageOverlappedAccumulator::FetchFinalPixelDataHalfFloat(TArray64<FFloat16
 
 void FImageOverlappedAccumulator::FetchFinalPixelDataLinearColor(TArray64<FLinearColor> & OutPixelData) const
 {
+	SCOPE_CYCLE_COUNTER(STAT_FetchFinalPixelData32bit);
 	int32 FullSizeX = PlaneSize.X;
 	int32 FullSizeY = PlaneSize.Y;
 	OutPixelData.SetNumUninitialized(FullSizeX * FullSizeY);
diff --git a/Engine/Plugins/MovieScene/MovieRenderPipeline/Source/MovieRenderPipelineCore/Public/MoviePipelineOutputSetting.h b/Engine/Plugins/MovieScene/MovieRenderPipeline/Source/MovieRenderPipelineCore/Public/MoviePipelineOutputSetting.h
--- a/Engine/Plugins/MovieScene/MovieRenderPipeline/Source/MovieRenderPipelineCore/Public/MoviePipelineOutputSetting.h
+++ b/Engine/Plugins/MovieScene/MovieRenderPipeline/Source/MovieRenderPipelineCore/Public/MoviePipelineOutputSetting.h
@@ -82,4 +82,8 @@ public:
 	/** If true the Filmic Tone Curve will not be applied. Disabling this will allow you to export linear data for EXRs. */
 	UPROPERTY(EditAnywhere, BlueprintReadWrite, AdvancedDisplay, Category = "File Output")
 	bool bDisableToneCurve;
+
+	/** Should the normal player viewport be visible when rendering? This increases render cost and is only an approximation of what is actually being rendered. */
+	UPROPERTY(EditAnywhere, BlueprintReadWrite, AdvancedDisplay, Category = "File Output")
+	bool bHidePreview;
 };
\ No newline at end of file
diff --git a/Engine/Plugins/MovieScene/MovieRenderPipeline/Source/MovieRenderPipelineCore/Public/MovieRenderPipelineCoreModule.h b/Engine/Plugins/MovieScene/MovieRenderPipeline/Source/MovieRenderPipelineCore/Public/MovieRenderPipelineCoreModule.h
--- a/Engine/Plugins/MovieScene/MovieRenderPipeline/Source/MovieRenderPipelineCore/Public/MovieRenderPipelineCoreModule.h
+++ b/Engine/Plugins/MovieScene/MovieRenderPipeline/Source/MovieRenderPipelineCore/Public/MovieRenderPipelineCoreModule.h
@@ -13,6 +13,9 @@ class UMoviePipeline;
 
 namespace MoviePipeline { struct FMoviePipelineEnginePass; }
 
+// Declare a stat-group for our performance stats to be counted under, readable in game by "stat MovieRenderPipeline".
+DECLARE_STATS_GROUP(TEXT("MovieRenderPipeline"), STATGROUP_MoviePipeline, STATCAT_Advanced);
+
 namespace MoviePipelineErrorCodes
 {
 	/** Everything completed as expected or we (unfortunately) couldn't detect the error. */
diff --git a/Engine/Plugins/MovieScene/MovieRenderPipeline/Source/MovieRenderPipelineRenderPasses/Private/MoviePipelineDeferredPasses.cpp b/Engine/Plugins/MovieScene/MovieRenderPipeline/Source/MovieRenderPipelineRenderPasses/Private/MoviePipelineDeferredPasses.cpp
--- a/Engine/Plugins/MovieScene/MovieRenderPipeline/Source/MovieRenderPipelineRenderPasses/Private/MoviePipelineDeferredPasses.cpp
+++ b/Engine/Plugins/MovieScene/MovieRenderPipeline/Source/MovieRenderPipelineRenderPasses/Private/MoviePipelineDeferredPasses.cpp
@@ -28,6 +28,11 @@
 #include "MoviePipelineQueue.h"
 #include "MoviePipelineAntiAliasingSetting.h"
 #include "MoviePipelineOutputSetting.h"
+#include "MovieRenderPipelineCoreModule.h"
+
+DECLARE_CYCLE_STAT(TEXT("MoviePipeline_SampleReadback"), STAT_SampleReadback, STATGROUP_MoviePipeline);
+DECLARE_CYCLE_STAT(TEXT("MoviePipeline_AccumulateSample_RT"), STAT_AccumulateSample_RenderThread, STATGROUP_MoviePipeline);
+
 
 namespace MoviePipeline
 {
@@ -189,6 +194,8 @@ namespace MoviePipeline
 
 	void ReadBackbufferAndBroadcast_RenderThread(FRHICommandListImmediate &RHICmdList, FRenderTarget* InFromRenderTarget, const FMoviePipelineRenderPassMetrics& InSampleState, const FMoviePipelineSampleReady& OnReadDelegate)
 	{
+		SCOPE_CYCLE_COUNTER(STAT_SampleReadback);
+
 		check(IsInRenderingThread());
 
 		FIntRect SourceRect = FIntRect(0, 0, InFromRenderTarget->GetSizeXY().X, InFromRenderTarget->GetSizeXY().Y);
@@ -208,7 +215,6 @@ namespace MoviePipeline
 	FSceneView* FDeferredRenderEnginePass::CalcSceneView(FSceneViewFamily* ViewFamily, const FMoviePipelineRenderPassMetrics& InSampleState)
 	{
 		APlayerController* LocalPlayerController = GetPipeline()->GetWorld()->GetFirstPlayerController();
-		// GetPipeline()->GetWorld()->GetGameViewport()->bDisableWorldRendering = true;
 
 		int32 TileSizeX = InitSettings.BackbufferResolution.X;
 		int32 TileSizeY = InitSettings.BackbufferResolution.Y;
@@ -521,6 +527,8 @@ namespace MoviePipeline
 {
 	void AccumulateSample_RenderThread(TUniquePtr<FImagePixelData>&& InPixelData, TSharedRef<FImagePixelDataPayload, ESPMode::ThreadSafe> InFrameData, MoviePipeline::FSampleRenderThreadParams& InParams)
 	{
+		SCOPE_CYCLE_COUNTER(STAT_AccumulateSample_RenderThread);
+
 		check(InPixelData->IsDataWellFormed());
 
 		// Writing tiles can be useful for debug reasons. These get passed onto the output every frame.
@@ -533,6 +541,19 @@ namespace MoviePipeline
 			InParams.OutputMerger->OnSingleSampleDataAvailable_AnyThread(MoveTemp(SampleData), InFrameData);
 		}
 
+		// Optimization! If we don't need the accumulator (no tiling, no supersampling) then we'll skip it and just send it straight to the output stage.
+		// This significantly improves performance in the baseline case.
+		const bool bOneTile = InFrameData->IsFirstTile() && InFrameData->IsLastTile();
+		const bool bOneTS = InFrameData->IsFirstTemporalSample() && InFrameData->IsLastTemporalSample();
+		const bool bOneSS = InFrameData->SampleState.SpatialSampleCount == 1;
+
+		if (bOneTile && bOneTS && bOneSS)
+		{
+			// Send the data directly to the Output Builder and skip the accumulator.
+			InParams.OutputMerger->OnCompleteRenderPassDataAvailable_AnyThread(MoveTemp(InPixelData), InFrameData);
+			return;
+		}
+
 		// For the first sample in a new output, we allocate memory
 		if (InFrameData->IsFirstTile() && InFrameData->IsFirstTemporalSample())
 		{
diff --git a/Engine/Plugins/MovieScene/MovieRenderPipeline/Source/MovieRenderPipelineRenderPasses/Private/MoviePipelineImageSequenceOutput.cpp b/Engine/Plugins/MovieScene/MovieRenderPipeline/Source/MovieRenderPipelineRenderPasses/Private/MoviePipelineImageSequenceOutput.cpp
--- a/Engine/Plugins/MovieScene/MovieRenderPipeline/Source/MovieRenderPipelineRenderPasses/Private/MoviePipelineImageSequenceOutput.cpp
+++ b/Engine/Plugins/MovieScene/MovieRenderPipeline/Source/MovieRenderPipelineRenderPasses/Private/MoviePipelineImageSequenceOutput.cpp
@@ -21,6 +21,7 @@
 // Forward Declare
 static TUniquePtr<FImagePixelData> QuantizePixelDataTo8bpp(FImagePixelData* InPixelData);
 
+DECLARE_CYCLE_STAT(TEXT("ImgSeqOutput_RecieveImageData"), STAT_ImgSeqRecieveImageData, STATGROUP_MoviePipeline);
 
 UMoviePipelineImageSequenceOutputBase::UMoviePipelineImageSequenceOutputBase()
 {
@@ -44,6 +45,8 @@ bool UMoviePipelineImageSequenceOutputBase::HasFinishedProcessingImpl()
 
 void UMoviePipelineImageSequenceOutputBase::OnRecieveImageDataImpl(FMoviePipelineMergerOutputFrame* InMergedOutputFrame)
 {
+	SCOPE_CYCLE_COUNTER(STAT_ImgSeqRecieveImageData);
+
 	check(InMergedOutputFrame);
 
 	// We do a little special handling for Burn In overlays, because we need to composite them on top of the main image. We may also want to write them
