MoviePipeline: Fixed Output Containers being finalized before all frames had been passed to them. This was caused by async queueing the finished work on the game thread which wouldn't get processed until the next tick, even if the render thread had been flushed and all frames had been sent to the output merger. The output merger now uses a queue and does FIFO and is pulled at explicit times from the main MoviePipeline instead of being pushed to the MoviePipeline.

#rnx

diff --git a/Engine/Plugins/MovieScene/MovieRenderPipeline/Source/MovieRenderPipelineCore/Private/MoviePipeline.cpp b/Engine/Plugins/MovieScene/MovieRenderPipeline/Source/MovieRenderPipelineCore/Private/MoviePipeline.cpp
--- a/Engine/Plugins/MovieScene/MovieRenderPipeline/Source/MovieRenderPipelineCore/Private/MoviePipeline.cpp
+++ b/Engine/Plugins/MovieScene/MovieRenderPipeline/Source/MovieRenderPipelineCore/Private/MoviePipeline.cpp
@@ -376,6 +376,12 @@ bool UMoviePipeline::ProcessEndOfCameraCut(FMoviePipelineShotInfo &CurrentShot,
 	// tick the new shot will be initialized and processed.
 	if (CurrentShotIndex >= ShotList.Num())
 	{
+		// Ensure all frames have been processed by the GPU and sent to the Output Merger
+		FlushRenderingCommands();
+		
+		// And then make sure all frames are sent to the Output Containers before we finalize.
+		ProcessOutstandingFinishedFrames();
+
 		UE_LOG(LogMovieRenderPipeline, Log, TEXT("[%d] Finished rendering last shot. Moving to Finalize to finish writing items to disk."), GFrameCounter);
 		PipelineState = EMovieRenderPipelineState::Finalize;
 
diff --git a/Engine/Plugins/MovieScene/MovieRenderPipeline/Source/MovieRenderPipelineCore/Private/MoviePipelineOutputBuilder.cpp b/Engine/Plugins/MovieScene/MovieRenderPipeline/Source/MovieRenderPipelineCore/Private/MoviePipelineOutputBuilder.cpp
--- a/Engine/Plugins/MovieScene/MovieRenderPipeline/Source/MovieRenderPipelineCore/Private/MoviePipelineOutputBuilder.cpp
+++ b/Engine/Plugins/MovieScene/MovieRenderPipeline/Source/MovieRenderPipelineCore/Private/MoviePipelineOutputBuilder.cpp
@@ -19,9 +19,9 @@ FMoviePipelineMergerOutputFrame& FMoviePipelineOutputMerger::QueueOutputFrame_Ga
 	FScopeLock ScopeLock(&ActiveDataMutex);
 
 	// Ensure this frame hasn't already been entered somehow.
-	check(!ActiveData.Find(CachedOutputState));
+	check(!PendingData.Find(CachedOutputState));
 
-	FMoviePipelineMergerOutputFrame& NewFrame = ActiveData.Add(CachedOutputState);
+	FMoviePipelineMergerOutputFrame& NewFrame = PendingData.Add(CachedOutputState);
 	NewFrame.FrameOutputState = CachedOutputState;
 
 	return NewFrame;
@@ -57,7 +57,7 @@ void FMoviePipelineOutputMerger::OnCompleteRenderPassDataAvailable_AnyThread(TUn
 	// Instead of just finding the result in the TMap with the equality operator, we find it by hand so that we can
 	// ignore certain parts of equality (such as Temporal Sample, as the last sample has a temporal index different
 	// than the first sample!)
-	for (TPair<FMoviePipelineFrameOutputState, FMoviePipelineMergerOutputFrame>& KVP : ActiveData)
+	for (TPair<FMoviePipelineFrameOutputState, FMoviePipelineMergerOutputFrame>& KVP : PendingData)
 	{
 		if (KVP.Key.OutputFrameNumber == InFrameData->OutputState.OutputFrameNumber)
 	 	{
@@ -87,18 +87,7 @@ void FMoviePipelineOutputMerger::OnCompleteRenderPassDataAvailable_AnyThread(TUn
 	{
 		// Transfer ownership from the map to here;
 		FMoviePipelineMergerOutputFrame FinalFrame;
-		ActiveData.RemoveAndCopyValue(InFrameData->OutputState, FinalFrame);
-
-		// Notify the Movie Pipeline that this frame has been completed and can be forwarded to the output containers.
-		TWeakObjectPtr<UMoviePipeline> LocalWeakPipeline = WeakMoviePipeline;
-
-		AsyncTask(ENamedThreads::GameThread, [LocalFinalFrame = MoveTemp(FinalFrame), InFrameData, LocalWeakPipeline]() mutable
-		{
-			if (ensureAlwaysMsgf(LocalWeakPipeline.IsValid(), TEXT("A memory lifespan issue has left an output builder alive without an owning Movie Pipeline.")))
-			{
-				LocalWeakPipeline->OnFrameCompletelyRendered(MoveTemp(LocalFinalFrame), InFrameData);
-			}
-		}
-		);
+		PendingData.RemoveAndCopyValue(InFrameData->OutputState, FinalFrame);
+		FinishedFrames.Enqueue(MoveTemp(FinalFrame));
 	}
 }
\ No newline at end of file
diff --git a/Engine/Plugins/MovieScene/MovieRenderPipeline/Source/MovieRenderPipelineCore/Private/MoviePipelineRendering.cpp b/Engine/Plugins/MovieScene/MovieRenderPipeline/Source/MovieRenderPipelineCore/Private/MoviePipelineRendering.cpp
--- a/Engine/Plugins/MovieScene/MovieRenderPipeline/Source/MovieRenderPipelineCore/Private/MoviePipelineRendering.cpp
+++ b/Engine/Plugins/MovieScene/MovieRenderPipeline/Source/MovieRenderPipelineCore/Private/MoviePipelineRendering.cpp
@@ -185,6 +185,9 @@ void UMoviePipeline::RenderFrame()
 	// the render thread uses it.
 	FlushAsyncEngineSystems();
 
+	// Send any output frames that have been completed since the last render.
+	ProcessOutstandingFinishedFrames();
+
 	FMoviePipelineShotInfo& CurrentShot = ShotList[CurrentShotIndex];
 	FMoviePipelineCameraCutInfo& CurrentCameraCut = CurrentShot.GetCurrentCameraCut();
 	APlayerController* LocalPlayerController = GetWorld()->GetFirstPlayerController();
@@ -435,19 +438,19 @@ void UMoviePipeline::RenderFrame()
 			}
 		}
 	}
-	
-	// UE_LOG(LogMovieRenderPipeline, Warning, TEXT("[%d] Pre-FlushRenderingCommands"), GFrameCounter);
-	// FlushRenderingCommands();
-	// UE_LOG(LogMovieRenderPipeline, Warning, TEXT("[%d] Post-FlushRenderingCommands"), GFrameCounter);
 }
 
-void UMoviePipeline::OnFrameCompletelyRendered(FMoviePipelineMergerOutputFrame&& OutputFrame, const TSharedRef<FImagePixelDataPayload, ESPMode::ThreadSafe> InFrameData)
+void UMoviePipeline::ProcessOutstandingFinishedFrames()
 {
-	UE_LOG(LogMovieRenderPipeline, Warning, TEXT("[%d] Data required for output available! Frame: %d"), GFrameCounter, OutputFrame.FrameOutputState.OutputFrameNumber);
-	
-	for (UMoviePipelineOutputBase* OutputContainer : GetPipelineMasterConfig()->GetOutputContainers())
+	while (!OutputBuilder->FinishedFrames.IsEmpty())
 	{
-		OutputContainer->OnRecieveImageData(&OutputFrame);
+		FMoviePipelineMergerOutputFrame OutputFrame;
+		OutputBuilder->FinishedFrames.Dequeue(OutputFrame);
+
+		for (UMoviePipelineOutputBase* OutputContainer : GetPipelineMasterConfig()->GetOutputContainers())
+		{
+			OutputContainer->OnRecieveImageData(&OutputFrame);
+		}
 	}
 }
 
diff --git a/Engine/Plugins/MovieScene/MovieRenderPipeline/Source/MovieRenderPipelineCore/Public/MoviePipeline.h b/Engine/Plugins/MovieScene/MovieRenderPipeline/Source/MovieRenderPipelineCore/Public/MoviePipeline.h
--- a/Engine/Plugins/MovieScene/MovieRenderPipeline/Source/MovieRenderPipelineCore/Public/MoviePipeline.h
+++ b/Engine/Plugins/MovieScene/MovieRenderPipeline/Source/MovieRenderPipelineCore/Public/MoviePipeline.h
@@ -90,7 +90,7 @@ public:
 
 	FDateTime GetInitializationTime() const { return InitializationTime; }
 public:
-	void OnFrameCompletelyRendered(FMoviePipelineMergerOutputFrame&& OutputFrame, const TSharedRef<FImagePixelDataPayload, ESPMode::ThreadSafe> InFrameData);
+	void ProcessOutstandingFinishedFrames();
 	void OnSampleRendered(TUniquePtr<FImagePixelData>&& OutputSample, const TSharedRef<FImagePixelDataPayload, ESPMode::ThreadSafe> InFrameData);
 	const MoviePipeline::FAudioState& GetAudioState() const { return AudioState; }
 public:
diff --git a/Engine/Plugins/MovieScene/MovieRenderPipeline/Source/MovieRenderPipelineCore/Public/MoviePipelineOutputBuilder.h b/Engine/Plugins/MovieScene/MovieRenderPipeline/Source/MovieRenderPipelineCore/Public/MoviePipelineOutputBuilder.h
--- a/Engine/Plugins/MovieScene/MovieRenderPipeline/Source/MovieRenderPipelineCore/Public/MoviePipelineOutputBuilder.h
+++ b/Engine/Plugins/MovieScene/MovieRenderPipeline/Source/MovieRenderPipelineCore/Public/MoviePipelineOutputBuilder.h
@@ -2,6 +2,7 @@
 
 #pragma once
 #include "MovieRenderPipelineDataTypes.h"
+#include "Containers/Queue.h"
 
 // Forward Declares
 struct FImagePixelData;
@@ -22,12 +23,13 @@ public:
 	void OnCompleteRenderPassDataAvailable_AnyThread(TUniquePtr<FImagePixelData>&& InData, const TSharedRef<FImagePixelDataPayload, ESPMode::ThreadSafe> InFrameData);
 	void OnSingleSampleDataAvailable_AnyThread(TUniquePtr<FImagePixelData>&& InData, const TSharedRef<FImagePixelDataPayload, ESPMode::ThreadSafe> InFrameData);
 
+	TQueue<FMoviePipelineMergerOutputFrame> FinishedFrames;
 private:
 	/** The Movie Pipeline that owns us. */
 	TWeakObjectPtr<UMoviePipeline> WeakMoviePipeline;
 
 	/** Data that is expected but not fully available yet. */
-	TMap<FMoviePipelineFrameOutputState, FMoviePipelineMergerOutputFrame> ActiveData;
+	TMap<FMoviePipelineFrameOutputState, FMoviePipelineMergerOutputFrame> PendingData;
 
 	/** Mutex that protects adding/updating/removing from ActiveData */
 	FCriticalSection ActiveDataMutex;
diff --git a/Engine/Plugins/MovieScene/MovieRenderPipeline/Source/MovieRenderPipelineCore/Public/MovieRenderPipelineCoreModule.h b/Engine/Plugins/MovieScene/MovieRenderPipeline/Source/MovieRenderPipelineCore/Public/MovieRenderPipelineCoreModule.h
--- a/Engine/Plugins/MovieScene/MovieRenderPipeline/Source/MovieRenderPipelineCore/Public/MovieRenderPipelineCoreModule.h
+++ b/Engine/Plugins/MovieScene/MovieRenderPipeline/Source/MovieRenderPipelineCore/Public/MovieRenderPipelineCoreModule.h
@@ -3,6 +3,7 @@
 #pragma once
 
 #include "CoreMinimal.h"
+#include "Stats/Stats2.h"
 #include "Modules/ModuleInterface.h"
 
 // Forward Declare
