MoviePipeline: Wave Output now converts the path to absolute before passing it to the wave writer which makes some assumptions if the path is relative.

diff --git a/Engine/Plugins/MovieScene/MovieRenderPipeline/Source/MovieRenderPipelineRenderPasses/Public/MoviePipelineWaveOutput.cpp b/Engine/Plugins/MovieScene/MovieRenderPipeline/Source/MovieRenderPipelineRenderPasses/Public/MoviePipelineWaveOutput.cpp
--- a/Engine/Plugins/MovieScene/MovieRenderPipeline/Source/MovieRenderPipelineRenderPasses/Public/MoviePipelineWaveOutput.cpp
+++ b/Engine/Plugins/MovieScene/MovieRenderPipeline/Source/MovieRenderPipelineRenderPasses/Public/MoviePipelineWaveOutput.cpp
@@ -28,6 +28,8 @@ void UMoviePipelineWaveOutput::BeginFinalizeImpl()
 		return;
 	}
 
+	ActiveWriters.Reset();
+
 	// There should be no active submixes by the time we finalize - they should all have been converted to recorded samples.
 	check(GetPipeline()->GetAudioState().ActiveSubmixes.Num() == 0);
 
@@ -130,13 +132,21 @@ void UMoviePipelineWaveOutput::BeginFinalizeImpl()
 			FileFolder = FinalFilePath.Left(LastFolderSeparator);
 			FileName = FinalFilePath.RightChop(LastFolderSeparator + 1);
 		}
+
+		const bool bIsRelativePath = FPaths::IsRelative(FileFolder);
+		if (bIsRelativePath)
+		{
+			FileFolder = FPaths::ConvertRelativePathToFull(FileFolder);
+		}
 		
 		OutstandingWrites++;
-		Audio::FSoundWavePCMWriter Writer;
-		Writer.BeginWriteToWavFile(SampleBuffer, FileName, FileFolder, [this]()
+		TUniquePtr<Audio::FSoundWavePCMWriter> Writer = MakeUnique<Audio::FSoundWavePCMWriter>();
+		Writer->BeginWriteToWavFile(SampleBuffer, FileName, FileFolder, [this]()
 		{
 			this->OutstandingWrites--;
 		});
+
+		ActiveWriters.Add(MoveTemp(Writer));
 	}
 }
 
diff --git a/Engine/Plugins/MovieScene/MovieRenderPipeline/Source/MovieRenderPipelineRenderPasses/Public/MoviePipelineWaveOutput.h b/Engine/Plugins/MovieScene/MovieRenderPipeline/Source/MovieRenderPipelineRenderPasses/Public/MoviePipelineWaveOutput.h
--- a/Engine/Plugins/MovieScene/MovieRenderPipeline/Source/MovieRenderPipelineRenderPasses/Public/MoviePipelineWaveOutput.h
+++ b/Engine/Plugins/MovieScene/MovieRenderPipeline/Source/MovieRenderPipelineRenderPasses/Public/MoviePipelineWaveOutput.h
@@ -30,5 +30,7 @@ public:
 	FString FileNameFormat;
 
 private:
+	/** Kept alive during finalization because the writer writes async to disk but doesn't expect to fall out of scope */
+	TArray<TUniquePtr<Audio::FSoundWavePCMWriter>> ActiveWriters;
 	TAtomic<int32> OutstandingWrites;
 };
\ No newline at end of file
