MoviePipeline: 8 bit quantization now uses a pre-calculated sRGB table and does table lookups instead of calculating the sRGB value each time for every pixel from scratch. This significantly improves performance of image quantization, from about ~130ms for a 1080p image to ~20ms.

diff --git a/Engine/Plugins/MovieScene/MovieRenderPipeline/Source/MovieRenderPipelineCore/Private/MoviePipelineImageQuantization.cpp b/Engine/Plugins/MovieScene/MovieRenderPipeline/Source/MovieRenderPipelineCore/Private/MoviePipelineImageQuantization.cpp
--- a/Engine/Plugins/MovieScene/MovieRenderPipeline/Source/MovieRenderPipelineCore/Private/MoviePipelineImageQuantization.cpp
+++ b/Engine/Plugins/MovieScene/MovieRenderPipeline/Source/MovieRenderPipelineCore/Private/MoviePipelineImageQuantization.cpp
@@ -2,11 +2,136 @@
 
 #include "MoviePipelineImageQuantization.h"
 #include "ImagePixelData.h"
+#include "Stats/Stats2.h"
+#include "MovieRenderPipelineCoreModule.h"
 
 namespace UE
 {
 namespace MoviePipeline
 {
+DECLARE_CYCLE_STAT(TEXT("STAT_MoviePipeline_ImageQuantization"), STAT_ImageQuantization, STATGROUP_MoviePipeline);
+
+static TArray<uint8> GenerateSRGBTable(uint32 InPrecision)
+{
+	QUICK_SCOPE_CYCLE_COUNTER(STAT_ImageQuant_TableGeneration);
+	TArray<uint8> OutsRGBTable;
+	OutsRGBTable.SetNumUninitialized(InPrecision);
+	for (int32 TableIndex = 0; TableIndex < OutsRGBTable.Num(); TableIndex++)
+	{
+		float ValueAsLinear = (float)TableIndex / (OutsRGBTable.Num() - 1);
+
+		// sRGB is linear under 0.0031308 and Pow(1/2.4) above that.
+		if (ValueAsLinear <= 0.0031308f)
+		{
+			ValueAsLinear = ValueAsLinear * 12.92f;
+		}
+		else
+		{
+			ValueAsLinear = FMath::Pow(ValueAsLinear, 1.0f / 2.4f) * 1.055f - 0.055f;
+		}
+
+		// Flooring avoids an extra branch for Round. Using [] on GetData() avoids bounds checking cost.
+		OutsRGBTable.GetData()[TableIndex] = (uint8)FMath::Clamp(FMath::FloorToInt(ValueAsLinear * 256.f), 0, 255);
+	}
+	
+	return OutsRGBTable;
+}
+
+static TArray<uint8> GenerateSRGBTableFloat16()
+{
+	QUICK_SCOPE_CYCLE_COUNTER(STAT_ImageQuant_TableGeneration);
+	TArray<uint8> OutsRGBTable;
+	OutsRGBTable.SetNumUninitialized(65536);
+	uint8* OutsRGBTableData = OutsRGBTable.GetData();
+
+	FFloat16 OnePointZero = 1.0f;
+	for (int32 TableIndex = OnePointZero.Encoded; TableIndex < 32768; TableIndex++)
+	{
+		// Fill table for values >= 1.0 to be 255
+		OutsRGBTableData[TableIndex] = 255;
+	}
+	for (int32 TableIndex = 32768; TableIndex < 65536; TableIndex++)
+	{
+		// Fill table for all negative values to be 0
+		OutsRGBTableData[TableIndex] = 0;
+	}
+	for (int32 TableIndex = 0; TableIndex < OnePointZero.Encoded; TableIndex++)
+	{
+		FFloat16 Value;
+		Value.Encoded = TableIndex;
+		float ValueAsLinear = (float)Value;
+		// sRGB is linear under 0.0031308 and Pow(1/2.4) above that.
+		if (ValueAsLinear <= 0.0031308f)
+		{
+			ValueAsLinear = ValueAsLinear * 12.92f;
+		}
+		else
+		{
+			ValueAsLinear = FMath::Pow(ValueAsLinear, 1.0f / 2.4f) * 1.055f - 0.055f;
+		}
+		// Flooring avoids an extra branch for Round. Using [] on GetData() avoids bounds checking cost.
+		OutsRGBTableData[TableIndex] = (uint8)FMath::FloorToInt((ValueAsLinear * 255.f) + 0.5f);
+	}
+
+	return OutsRGBTable;
+}
+	
+static TArray<FColor> ConvertLinearTosRGB8bppViaLookupTable(FFloat16Color* InColor, const int32 InCount)
+{
+	TArray<uint8> sRGBTable = GenerateSRGBTableFloat16();
+
+	// Convert all of our pixels.
+	TArray<FColor> OutsRGBData;
+	OutsRGBData.SetNumUninitialized(InCount);
+
+	uint8* sRGBTableData = static_cast<uint8*>(sRGBTable.GetData());
+	FColor* OutData = static_cast<FColor*>(OutsRGBData.GetData());
+
+	QUICK_SCOPE_CYCLE_COUNTER(STAT_ImageQuant_ApplysRGB);
+	for (int32 PixelIndex = 0; PixelIndex < InCount; PixelIndex++)
+	{
+		// Avoid the bounds checking of TArray[]
+		FColor* OutColor = &OutData[PixelIndex];
+		OutColor->R = sRGBTableData[InColor[PixelIndex].R.Encoded];
+		OutColor->G = sRGBTableData[InColor[PixelIndex].G.Encoded];
+		OutColor->B = sRGBTableData[InColor[PixelIndex].B.Encoded];
+
+		// Alpha doesn't get sRGB conversion, it just gets linearly converted to 8 bit. Flooring avoids an extra branch for Round.
+		OutColor->A = (uint8)FMath::Clamp(FMath::FloorToInt((InColor[PixelIndex].A * 255.f) + 0.5f), 0, 255);
+	}
+
+	return OutsRGBData;
+}
+
+static TArray<FColor> ConvertLinearTosRGB8bppViaLookupTable(FLinearColor* InColor, const int32 InCount)
+{
+	TArray<uint8> sRGBTable = GenerateSRGBTable(4096); 
+
+	// Convert all of our pixels.
+	TArray<FColor> OutsRGBData;
+	OutsRGBData.SetNumUninitialized(InCount);
+
+	int32 TableUpperBound = sRGBTable.Num() - 1;
+
+	for (int32 PixelIndex = 0; PixelIndex < InCount; PixelIndex++)
+	{
+		// Avoid the bounds checking of TArray[]
+		FColor* OutColor = &OutsRGBData.GetData()[PixelIndex];
+		int32 TableIndexR = int32(InColor[PixelIndex].R * TableUpperBound);
+		int32 TableIndexG = int32(InColor[PixelIndex].G * TableUpperBound);
+		int32 TableIndexB = int32(InColor[PixelIndex].B * TableUpperBound);
+		
+		// We clamp the table index because values greater than 1.0 in the color would generate out of bound indexes.
+		OutColor->R = sRGBTable.GetData()[FMath::Clamp(TableIndexR, 0, TableUpperBound)];
+		OutColor->G = sRGBTable.GetData()[FMath::Clamp(TableIndexG, 0, TableUpperBound)];
+		OutColor->B = sRGBTable.GetData()[FMath::Clamp(TableIndexB, 0, TableUpperBound)];
+		
+		// Alpha doesn't get sRGB conversion, it just gets linearly converted to 8 bit. Flooring avoids an extra branch for Round.
+		OutColor->A = (uint8)FMath::Clamp(FMath::FloorToInt((InColor[PixelIndex].A * 255.f) + 0.5f), 0, 255);
+	}
+
+	return OutsRGBData;
+}
 
 static TUniquePtr<FImagePixelData> QuantizePixelDataTo8bpp(const FImagePixelData* InPixelData)
 {
@@ -26,78 +151,26 @@ static TUniquePtr<FImagePixelData> QuantizePixelDataTo8bpp(const FImagePixelData
 	}
 	case 16:
 	{
-		TArray<FColor> ClampedPixels;
-		ClampedPixels.SetNum(RawSize.X * RawSize.Y);
-
 		int64 SizeInBytes = 0;
 		const void* SrcRawDataPtr = nullptr;
 		InPixelData->GetRawData(SrcRawDataPtr, SizeInBytes);
 
-		const uint16* RawDataPtr = static_cast<const uint16*>(SrcRawDataPtr);
-
-		// Copy pixels to new array
-		for (int32 Y = 0; Y < RawSize.Y; Y++)
 		{
-			for (int32 X = 0; X < RawSize.X; X++)
-			{
-				FColor* DestColor = &ClampedPixels.GetData()[Y*RawSize.X + X];
-				FLinearColor SrcColor;
-				for (int32 ChanIter = 0; ChanIter < 4; ChanIter++)
-				{
-					FFloat16 Value;
-					Value.Encoded = RawDataPtr[(Y*RawSize.X + X)*RawNumChannels + ChanIter];
-
-					switch (ChanIter)
-					{
-					case 0: SrcColor.R = Value; break;
-					case 1: SrcColor.G = Value; break;
-					case 2: SrcColor.B = Value; break;
-					case 3: SrcColor.A = Value; break;
-					}
-				}
-				// convert to FColor using sRGB conversion
-				*DestColor = SrcColor.ToFColor(true);
-			}
+			TArray<FColor> sRGBEncoded = ConvertLinearTosRGB8bppViaLookupTable((FFloat16Color*)SrcRawDataPtr, RawSize.X * RawSize.Y);
+			QuantizedPixelData = MakeUnique<TImagePixelData<FColor>>(RawSize, TArray64<FColor>(MoveTemp(sRGBEncoded)));
 		}
-
-		QuantizedPixelData = MakeUnique<TImagePixelData<FColor>>(RawSize, TArray64<FColor>(MoveTemp(ClampedPixels)));
+		break;
 	}
 	case 32:
 	{
-		TArray<FColor> ClampedPixels;
-		ClampedPixels.SetNum(RawSize.X * RawSize.Y);
-
 		int64 SizeInBytes = 0;
 		const void* SrcRawDataPtr = nullptr;
 		InPixelData->GetRawData(SrcRawDataPtr, SizeInBytes);
-
-		const float* RawDataPtr = static_cast<const float*>(SrcRawDataPtr);
-
-		// Copy pixels to new array
-		for (int32 Y = 0; Y < RawSize.Y; Y++)
+		
 		{
-			for (int32 X = 0; X < RawSize.X; X++)
-			{
-				FColor* DestColor = &ClampedPixels.GetData()[Y*RawSize.X + X];
-				FLinearColor SrcColor;
-				for (int32 ChanIter = 0; ChanIter < 4; ChanIter++)
-				{
-					float Value = RawDataPtr[(Y*RawSize.X + X)*RawNumChannels + ChanIter];
-
-					switch (ChanIter)
-					{
-					case 0: SrcColor.R = Value; break;
-					case 1: SrcColor.G = Value; break;
-					case 2: SrcColor.B = Value; break;
-					case 3: SrcColor.A = Value; break;
-					}
-				}
-				// convert to FColor using sRGB conversion
-				*DestColor = SrcColor.ToFColor(true);
-			}
+			TArray<FColor> sRGBEncoded = ConvertLinearTosRGB8bppViaLookupTable((FLinearColor*)SrcRawDataPtr, RawSize.X * RawSize.Y);
+			QuantizedPixelData = MakeUnique<TImagePixelData<FColor>>(RawSize, TArray64<FColor>(MoveTemp(sRGBEncoded)));
 		}
-
-		QuantizedPixelData = MakeUnique<TImagePixelData<FColor>>(RawSize, TArray64<FColor>(MoveTemp(ClampedPixels)));
 		break;
 	}
 
@@ -111,6 +184,7 @@ static TUniquePtr<FImagePixelData> QuantizePixelDataTo8bpp(const FImagePixelData
 
 TUniquePtr<FImagePixelData> QuantizeImagePixelDataToBitDepth(const FImagePixelData* InData, const int32 TargetBitDepth)
 {
+	SCOPE_CYCLE_COUNTER(STAT_ImageQuantization);
 	TUniquePtr<FImagePixelData> QuantizedPixelData = nullptr;
 	switch (TargetBitDepth)
 	{
