MoviePipeline: Fixed Cinematic Quality Settings overrides applying the cinematic quality levels to the wrong value which then weren't applied until the next run.

diff --git a/Engine/Plugins/MovieScene/MovieRenderPipeline/Source/MovieRenderPipelineSettings/Private/MoviePipelineGameOverrideSetting.cpp b/Engine/Plugins/MovieScene/MovieRenderPipeline/Source/MovieRenderPipelineSettings/Private/MoviePipelineGameOverrideSetting.cpp
--- a/Engine/Plugins/MovieScene/MovieRenderPipeline/Source/MovieRenderPipelineSettings/Private/MoviePipelineGameOverrideSetting.cpp
+++ b/Engine/Plugins/MovieScene/MovieRenderPipeline/Source/MovieRenderPipelineSettings/Private/MoviePipelineGameOverrideSetting.cpp
@@ -63,7 +63,7 @@ void UMoviePipelineGameOverrideSetting::ApplyCVarSettings(const bool bRestoreOld
 
 			// Create a copy and override to the maximum level for each Scalability category
 			Scalability::FQualityLevels QualityLevels = PreviousQualityLevels;
-			PreviousQualityLevels.SetFromSingleQualityLevelRelativeToMax(0);
+			QualityLevels.SetFromSingleQualityLevelRelativeToMax(0);
 
 			// Apply
 			Scalability::SetQualityLevels(QualityLevels);
