MoviePipeline: Fixed audio rendering rendering more audio than desired when using Temporal sub-sampling.

diff --git a/Engine/Plugins/MovieScene/MovieRenderPipeline/Source/MovieRenderPipelineCore/Private/MoviePipelineAudioRendering.cpp b/Engine/Plugins/MovieScene/MovieRenderPipeline/Source/MovieRenderPipelineCore/Private/MoviePipelineAudioRendering.cpp
--- a/Engine/Plugins/MovieScene/MovieRenderPipeline/Source/MovieRenderPipelineCore/Private/MoviePipelineAudioRendering.cpp
+++ b/Engine/Plugins/MovieScene/MovieRenderPipeline/Source/MovieRenderPipelineCore/Private/MoviePipelineAudioRendering.cpp
@@ -25,7 +25,7 @@ static FAudioDevice* GetAudioDeviceFromWorldContext(const UObject* WorldContextO
 		return nullptr;
 	}
 
-	return ThisWorld->GetAudioDevice().GetAudioDevice();
+	return ThisWorld->GetAudioDeviceRaw();
 }
 
 static Audio::FMixerDevice* GetAudioMixerDeviceFromWorldContext(const UObject* WorldContextObject)
@@ -152,14 +152,13 @@ void UMoviePipeline::ProcessAudioTick()
 		AudioDeltaTime = GetPipelineMasterConfig()->GetEffectiveFrameRate(TargetSequence).AsInterval();
 	}
 
-	if (FAudioDevice* AudioDevice = FAudioDeviceManager::Get()->GetActiveAudioDevice().GetAudioDevice())
 	{
 		// Handle any game logic that changed Audio State.
-		AudioDevice->Update(true);
+		MixerDevice->Update(true);
 	}
 
 	Audio::IAudioMixerPlatformInterface* AudioMixerPlatform = MixerDevice->GetAudioMixerPlatform();
-	if (AudioMixerPlatform && AudioMixerPlatform->IsNonRealtime())
+	if (bCanRenderAudio && AudioMixerPlatform && AudioMixerPlatform->IsNonRealtime())
 	{
 		// There is only one non-realtime audio platform at this time, so we can safely static_cast this due to the IsNonRealtime check.
 		Audio::FMixerPlatformNonRealtime* Platform = static_cast<Audio::FMixerPlatformNonRealtime*>(MixerDevice->GetAudioMixerPlatform());
@@ -171,4 +170,4 @@ void UMoviePipeline::ProcessAudioTick()
 			Platform->RenderAudio(OutputFrameRate.AsInterval());
 		}
 	}
-}
\ No newline at end of file
+}
