MoviePipeline: Fixed camera cut tracks not being respected and range intersection failing which caused shots outside of the playback bounds to be processed despite not needing to be.

diff --git a/Engine/Plugins/MovieScene/MovieRenderPipeline/Source/MovieRenderPipelineCore/Private/MoviePipeline.cpp b/Engine/Plugins/MovieScene/MovieRenderPipeline/Source/MovieRenderPipelineCore/Private/MoviePipeline.cpp
--- a/Engine/Plugins/MovieScene/MovieRenderPipeline/Source/MovieRenderPipelineCore/Private/MoviePipeline.cpp
+++ b/Engine/Plugins/MovieScene/MovieRenderPipeline/Source/MovieRenderPipelineCore/Private/MoviePipeline.cpp
@@ -1,5 +1,6 @@
 // Copyright Epic Games, Inc. All Rights Reserved.
 #include "MoviePipeline.h"
+#include "MovieScene.h"
 #include "MovieRenderPipelineCoreModule.h"
 #include "LevelSequence.h"
 #include "Tracks/MovieSceneSubTrack.h"
@@ -557,7 +558,7 @@ void UMoviePipeline::InitializeLevelSequenceActor(ULevelSequence* OriginalLevelS
 }
 
 
-FMoviePipelineShotInfo CreateShotFromMovieScene(const UMovieScene* InMovieScene, const TRange<FFrameNumber>& InIntersectionRange)
+FMoviePipelineShotInfo CreateShotFromMovieScene(const UMovieScene* InMovieScene, const TRange<FFrameNumber>& InIntersectionRange, UMovieSceneSubSection* InSubSection)
 {
 	check(InMovieScene);
 
@@ -572,7 +573,7 @@ FMoviePipelineShotInfo CreateShotFromMovieScene(const UMovieScene* InMovieScene,
 	TArray<FCameraCutRange> IntersectedRanges;
 
 	// We're going to search for Camera Cut tracks within this shot. If none are found, we'll use the whole range of the shot.
-	UMovieSceneCameraCutTrack* CameraCutTrack = InMovieScene->FindMasterTrack<UMovieSceneCameraCutTrack>();
+	const UMovieSceneCameraCutTrack* CameraCutTrack = Cast<UMovieSceneCameraCutTrack>(InMovieScene->GetCameraCutTrack());
 	if (CameraCutTrack)
 	{
 		for (UMovieSceneSection* Section : CameraCutTrack->GetAllSections())
@@ -581,7 +582,19 @@ FMoviePipelineShotInfo CreateShotFromMovieScene(const UMovieScene* InMovieScene,
 
 			// ToDo: Inner vs. Outer resolution differences.
 			// Intersect this cut with the outer range in the likely event that the section goes past the bounds.
-			TRange<FFrameNumber> IntersectingRange = TRange<FFrameNumber>::Intersection(Section->GetRange(), InIntersectionRange);
+			
+			TRange<FFrameNumber> LocalSectionRange = Section->GetRange(); // Section in local space
+			
+			LocalSectionRange = MovieScene::TranslateRange(LocalSectionRange, -InMovieScene->GetPlaybackRange().GetLowerBoundValue()); // Section relative to zero
+
+			TRange<FFrameNumber> SectionRangeInMaster = LocalSectionRange;
+			// Add the offset from the parent subsection to put it back into parent space.
+			if (InSubSection)
+			{
+				SectionRangeInMaster = MovieScene::TranslateRange(LocalSectionRange, InSubSection->GetRange().GetLowerBoundValue());
+			}
+
+			TRange<FFrameNumber> IntersectingRange = TRange<FFrameNumber>::Intersection(SectionRangeInMaster, InIntersectionRange);
 
 			FCameraCutRange& NewRange = IntersectedRanges.AddDefaulted_GetRef();
 			NewRange.Range = IntersectingRange;
@@ -641,6 +654,14 @@ TArray<FMoviePipelineShotInfo> UMoviePipeline::BuildShotListFromSequence(const U
 				continue;
 			}
 
+			// Skip this section if it falls entirely outside of our playback bounds.
+			TRange<FFrameNumber> MasterPlaybackBounds = GetCurrentJob()->GetConfiguration()->GetEffectivePlaybackRange(TargetSequence);
+			if (!ShotSection->GetRange().Overlaps(MasterPlaybackBounds))
+			{
+				UE_LOG(LogMovieRenderPipeline, Verbose, TEXT("Skipped adding Shot %s to Shot List due to not overlapping playback bounds."), *ShotSection->GetShotDisplayName());
+				continue;
+			}
+			
 			//if (CurrentJob->ShotRenderMask.Num() > 0)
 			//{
 			//	if (!CurrentJob->ShotRenderMask.Contains(ShotSection->GetShotDisplayName()))
@@ -653,8 +674,8 @@ TArray<FMoviePipelineShotInfo> UMoviePipeline::BuildShotListFromSequence(const U
 			// The Shot Section may extend past our Sequence's Playback Bounds. We intersect the two bounds to ensure that
 			// the Playback Start/Playback End of the overall sequence is respected.
 			TRange<FFrameNumber> CinematicShotSectionRange = TRange<FFrameNumber>::Intersection(ShotSection->GetRange(), InSequence->GetMovieScene()->GetPlaybackRange());
+			FMoviePipelineShotInfo NewShot = CreateShotFromMovieScene(ShotSection->GetSequence()->GetMovieScene(), CinematicShotSectionRange, ShotSection);
 
-			FMoviePipelineShotInfo NewShot = CreateShotFromMovieScene(ShotSection->GetSequence()->GetMovieScene(), CinematicShotSectionRange);
 
 			// Convert the offset time to ticks
 			FFrameTime OffsetTime = ShotSection->GetOffsetTime().GetValue();
@@ -674,7 +695,7 @@ TArray<FMoviePipelineShotInfo> UMoviePipeline::BuildShotListFromSequence(const U
 	else
 	{
 		// They don't have a cinematic shot track. We'll slice them up by camera cuts instead.
-		FMoviePipelineShotInfo NewShot = CreateShotFromMovieScene(TargetSequence->GetMovieScene(), TargetSequence->GetMovieScene()->GetPlaybackRange());
+		FMoviePipelineShotInfo NewShot = CreateShotFromMovieScene(TargetSequence->GetMovieScene(), TargetSequence->GetMovieScene()->GetPlaybackRange(), nullptr);
 		// NewShot.ShotConfig = GetPipelineMasterConfig()->DefaultShotConfig;
 		
 		NewShotList.Add(MoveTemp(NewShot));
@@ -729,7 +750,19 @@ TArray<FMoviePipelineShotInfo> UMoviePipeline::BuildShotListFromSequence(const U
 			FString CameraName;
 			if (CameraCut.CameraCutSection.IsValid())
 			{
-				CameraName = TargetSequence->GetMovieScene()->FindBinding(CameraCut.CameraCutSection->GetCameraBindingID().GetGuid())->GetName();
+				const FMovieSceneObjectBindingID& CameraObjectBindingId = CameraCut.CameraCutSection->GetCameraBindingID();
+				if (CameraObjectBindingId.IsValid())
+				{
+					UMovieScene* OwningMovieScene = Cast<UMovieScene>(CameraCut.CameraCutSection->GetOuter());
+					if (OwningMovieScene)
+					{
+						FMovieSceneBinding* Binding = OwningMovieScene->FindBinding(CameraObjectBindingId.GetGuid());
+						if (Binding)
+						{
+							CameraName = Binding->GetName();
+						}
+					}
+				}
 			}
 
 			CameraCut.ShotName = ShotName;
@@ -893,15 +926,14 @@ void UMoviePipeline::ExpandShot(FMoviePipelineShotInfo& InShot)
 		ShotMovieScene->SetPlaybackRange(NewPlaybackRange, false);
 
 		// Expand the outer owning section
+		ShotSection->SetIsLocked(false); // Unlock the section so we can modify the range
 		ShotSection->SetRange(MovieScene::DilateRange(ShotSection->GetRange(), -StartOffset, EndOffset));
-
 	}
 
 	// Ensure the overall Movie Scene Playback Range is large enough. This will clamp evaluation if we don't expand it. We hull the existing range
 	// with the new range.
 	TRange<FFrameNumber> EncompassingPlaybackRange = TRange<FFrameNumber>::Hull(TotalPlaybackRange, TargetSequence->MovieScene->GetPlaybackRange());
 	TargetSequence->GetMovieScene()->SetPlaybackRange(EncompassingPlaybackRange);
-
 }
 
 bool UMoviePipeline::IsDebugFrameStepIdling() const
diff --git a/Engine/Plugins/MovieScene/MovieRenderPipeline/Source/MovieRenderPipelineCore/Private/MoviePipelineMasterConfig.cpp b/Engine/Plugins/MovieScene/MovieRenderPipeline/Source/MovieRenderPipelineCore/Private/MoviePipelineMasterConfig.cpp
--- a/Engine/Plugins/MovieScene/MovieRenderPipeline/Source/MovieRenderPipelineCore/Private/MoviePipelineMasterConfig.cpp
+++ b/Engine/Plugins/MovieScene/MovieRenderPipeline/Source/MovieRenderPipelineCore/Private/MoviePipelineMasterConfig.cpp
@@ -1,4 +1,4 @@
-// Copyright Epic Games, Inc. All Rights Reserved.
+ // Copyright Epic Games, Inc. All Rights Reserved.
 
 #include "MoviePipelineMasterConfig.h"
 #include "MovieScene.h"
@@ -167,6 +167,28 @@ FFrameRate UMoviePipelineMasterConfig::GetEffectiveFrameRate(const ULevelSequenc
 }
 
 
+TRange<FFrameNumber> UMoviePipelineMasterConfig::GetEffectivePlaybackRange(const ULevelSequence* InSequence) const
+{
+	UMoviePipelineOutputSetting* OutputSettings = FindSetting<UMoviePipelineOutputSetting>();
+	check(OutputSettings);
+	check(InSequence);
+
+	// Check to see if they overrode the frame rate.
+	if (OutputSettings->bUseCustomPlaybackRange)
+	{
+		// Convert the custom playback range from frames to ticks.
+		FFrameNumber StartTick = FFrameRate::TransformTime(FFrameTime(FFrameNumber(OutputSettings->CustomStartFrame)), InSequence->GetMovieScene()->GetDisplayRate(), InSequence->GetMovieScene()->GetTickResolution()).FloorToFrame();
+		FFrameNumber EndTick = FFrameRate::TransformTime(FFrameTime(FFrameNumber(OutputSettings->CustomEndFrame)), InSequence->GetMovieScene()->GetDisplayRate(), InSequence->GetMovieScene()->GetTickResolution()).FloorToFrame();
+
+		// [Inclusive, Exclusive)
+		return TRange<FFrameNumber>(StartTick, EndTick);
+	}
+
+	// Pull it from the sequence if they didn't.
+	return InSequence->GetMovieScene()->GetPlaybackRange();
+}
+
+
 /*bool UMoviePipelineShotConfig::ValidateConfig(TArray<FText>& OutValidationErrors) const{
 	bool bValidPipeline = true;
 	OutValidationErrors.Reset();
diff --git a/Engine/Plugins/MovieScene/MovieRenderPipeline/Source/MovieRenderPipelineCore/Public/MoviePipelineMasterConfig.h b/Engine/Plugins/MovieScene/MovieRenderPipeline/Source/MovieRenderPipelineCore/Public/MoviePipelineMasterConfig.h
--- a/Engine/Plugins/MovieScene/MovieRenderPipeline/Source/MovieRenderPipelineCore/Public/MoviePipelineMasterConfig.h
+++ b/Engine/Plugins/MovieScene/MovieRenderPipeline/Source/MovieRenderPipelineCore/Public/MoviePipelineMasterConfig.h
@@ -63,6 +63,9 @@ public:
 	*/
 	UFUNCTION(BlueprintCallable, Category = "Movie Render Pipeline")
 	FFrameRate GetEffectiveFrameRate(const ULevelSequence* InSequence) const;
+
+	TRange<FFrameNumber> GetEffectivePlaybackRange(const ULevelSequence* InSequence) const;
+
 protected:
 	virtual bool CanSettingBeAdded(const UMoviePipelineSetting* InSetting) const override
 	{
diff --git a/Engine/Source/Runtime/MovieScene/Private/MovieScene.cpp b/Engine/Source/Runtime/MovieScene/Private/MovieScene.cpp
--- a/Engine/Source/Runtime/MovieScene/Private/MovieScene.cpp
+++ b/Engine/Source/Runtime/MovieScene/Private/MovieScene.cpp
@@ -904,7 +904,7 @@ UMovieSceneTrack* UMovieScene::AddCameraCutTrack( TSubclassOf<UMovieSceneTrack>
 }
 
 
-UMovieSceneTrack* UMovieScene::GetCameraCutTrack()
+UMovieSceneTrack* UMovieScene::GetCameraCutTrack() const
 {
 	return CameraCutTrack;
 }
diff --git a/Engine/Source/Runtime/MovieScene/Public/MovieScene.h b/Engine/Source/Runtime/MovieScene/Public/MovieScene.h
--- a/Engine/Source/Runtime/MovieScene/Public/MovieScene.h
+++ b/Engine/Source/Runtime/MovieScene/Public/MovieScene.h
@@ -584,7 +584,7 @@ public:
 	UMovieSceneTrack* AddCameraCutTrack( TSubclassOf<UMovieSceneTrack> TrackClass );
 	
 	/** @return The camera cut track if it exists. */
-	UMovieSceneTrack* GetCameraCutTrack();
+	UMovieSceneTrack* GetCameraCutTrack() const;
 
 	/** Removes the camera cut track if it exists. */
 	void RemoveCameraCutTrack();
