MoviePipeline: Fixed Temporal AA reprojection issues on the first frame of camera cuts caused by using the difference between frame 1 and frame 0 as the reprojection source (for motion vectors) but using image history from frame 0.

diff --git a/Engine/Plugins/MovieScene/MovieRenderPipeline/Source/MovieRenderPipelineCore/Private/MoviePipelineTiming.cpp b/Engine/Plugins/MovieScene/MovieRenderPipeline/Source/MovieRenderPipelineCore/Private/MoviePipelineTiming.cpp
--- a/Engine/Plugins/MovieScene/MovieRenderPipeline/Source/MovieRenderPipelineCore/Private/MoviePipelineTiming.cpp
+++ b/Engine/Plugins/MovieScene/MovieRenderPipeline/Source/MovieRenderPipelineCore/Private/MoviePipelineTiming.cpp
@@ -176,7 +176,7 @@ void UMoviePipeline::TickProducingFrames()
 	// This block is optional, as they may not want motion blur fixes.
 	if (CurrentCameraCut.State == EMovieRenderShotState::MotionBlur)
 	{
-		CachedOutputState.bSkipRendering = true;
+		CachedOutputState.bSkipRendering = false;
 		CachedOutputState.bDiscardRenderResult = true;
 
 		const MoviePipeline::FFrameConstantMetrics FrameMetrics = CalculateShotFrameMetrics(CurrentShot);
