MoviePipeline: FSceneView contents are now added to the reference collector to prevent material instances in the reusable pool from dropping out of scope under GC-heavy loads.

diff --git a/Engine/Plugins/MovieScene/MovieRenderPipeline/Source/MovieRenderPipelineRenderPasses/Private/MoviePipelineDeferredPasses.cpp b/Engine/Plugins/MovieScene/MovieRenderPipeline/Source/MovieRenderPipelineRenderPasses/Private/MoviePipelineDeferredPasses.cpp
--- a/Engine/Plugins/MovieScene/MovieRenderPipeline/Source/MovieRenderPipelineRenderPasses/Private/MoviePipelineDeferredPasses.cpp
+++ b/Engine/Plugins/MovieScene/MovieRenderPipeline/Source/MovieRenderPipelineRenderPasses/Private/MoviePipelineDeferredPasses.cpp
@@ -59,12 +59,26 @@ namespace MoviePipeline
 	{
 		FMoviePipelineEnginePass::Teardown();
 
+		FSceneViewStateInterface* Ref = ViewState.GetReference();
+		if (Ref)
+		{
+			Ref->ClearMIDPool();
+		}
+		ViewState.Destroy();
+
 		if (TileRenderTarget.IsValid())
 		{
 			TileRenderTarget->RemoveFromRoot();
 		}
-
-		ViewState.Destroy();
+	}
+	
+	void FDeferredRenderEnginePass::AddReferencedObjects(FReferenceCollector& Collector)
+	{
+		FSceneViewStateInterface* Ref = ViewState.GetReference();
+		if (Ref)
+		{
+			Ref->AddReferencedObjects(Collector);
+		}
 	}
 
 	void FDeferredRenderEnginePass::RenderSample_GameThread(const FMoviePipelineRenderPassMetrics& InSampleState)
diff --git a/Engine/Plugins/MovieScene/MovieRenderPipeline/Source/MovieRenderPipelineRenderPasses/Public/MoviePipelineDeferredPasses.h b/Engine/Plugins/MovieScene/MovieRenderPipeline/Source/MovieRenderPipelineRenderPasses/Public/MoviePipelineDeferredPasses.h
--- a/Engine/Plugins/MovieScene/MovieRenderPipeline/Source/MovieRenderPipelineRenderPasses/Public/MoviePipelineDeferredPasses.h
+++ b/Engine/Plugins/MovieScene/MovieRenderPipeline/Source/MovieRenderPipelineRenderPasses/Public/MoviePipelineDeferredPasses.h
@@ -5,6 +5,7 @@
 #include "MovieRenderPipelineDataTypes.h"
 #include "SceneTypes.h"
 #include "SceneView.h"
+#include "UObject/GCObject.h"
 #include "MoviePipelineDeferredPasses.generated.h"
 
 class UTextureRenderTarget2D;
@@ -46,7 +47,7 @@ DECLARE_MULTICAST_DELEGATE_ThreeParams(FMoviePipelineSetupView, FSceneViewFamily
 
 namespace MoviePipeline
 {
-	struct FDeferredRenderEnginePass : public FMoviePipelineEnginePass
+	struct FDeferredRenderEnginePass : public FMoviePipelineEnginePass, FGCObject
 	{
 		FDeferredRenderEnginePass()
 			: FMoviePipelineEnginePass(FMoviePipelinePassIdentifier(TEXT("MainDeferredPass")))
@@ -57,6 +58,10 @@ namespace MoviePipeline
 		virtual void RenderSample_GameThread(const FMoviePipelineRenderPassMetrics& InSampleState) override;
 		virtual void Teardown() override;
 
+		// FGCObject Interface
+		virtual void AddReferencedObjects(FReferenceCollector& Collector) override;
+		// ~FGCObject Interface
+
 		FMoviePipelineSampleReady BackbufferReadyDelegate;
 		FMoviePipelineSetupView SetupViewDelegate;
 
