MoviePipeline: Particle System LODs are now biased if you have bUseLODZero enabled to help in cases where the content may be culled for the lowest quality LOD.

diff --git a/Engine/Plugins/MovieScene/MovieRenderPipeline/Source/MovieRenderPipelineSettings/Private/MoviePipelineGameOverrideSetting.cpp b/Engine/Plugins/MovieScene/MovieRenderPipeline/Source/MovieRenderPipelineSettings/Private/MoviePipelineGameOverrideSetting.cpp
--- a/Engine/Plugins/MovieScene/MovieRenderPipeline/Source/MovieRenderPipelineSettings/Private/MoviePipelineGameOverrideSetting.cpp
+++ b/Engine/Plugins/MovieScene/MovieRenderPipeline/Source/MovieRenderPipelineSettings/Private/MoviePipelineGameOverrideSetting.cpp
@@ -104,7 +104,8 @@ void UMoviePipelineGameOverrideSetting::ApplyCVarSettings(const bool bRestoreOld
 	if (bUseLODZero)
 	{
 		STORE_AND_OVERRIDE_INT_SETTING(PreviousForceLOD, TEXT("r.ForceLOD"), 0, bRestoreOldValues);
-		STORE_AND_OVERRIDE_INT_SETTING(PreviousSkeletalMeshBias, TEXT("r.SkeletalMeshLODBias"), -4, bRestoreOldValues);
+		STORE_AND_OVERRIDE_INT_SETTING(PreviousSkeletalMeshBias, TEXT("r.SkeletalMeshLODBias"), -10, bRestoreOldValues);
+		STORE_AND_OVERRIDE_INT_SETTING(PreviousParticleLODBias, TEXT("r.ParticleLODBias"), -10, bRestoreOldValues);
 	}
 
 	if (bDisableHLODs)
diff --git a/Engine/Plugins/MovieScene/MovieRenderPipeline/Source/MovieRenderPipelineSettings/Public/MoviePipelineGameOverrideSetting.h b/Engine/Plugins/MovieScene/MovieRenderPipeline/Source/MovieRenderPipelineSettings/Public/MoviePipelineGameOverrideSetting.h
--- a/Engine/Plugins/MovieScene/MovieRenderPipeline/Source/MovieRenderPipelineSettings/Public/MoviePipelineGameOverrideSetting.h
+++ b/Engine/Plugins/MovieScene/MovieRenderPipeline/Source/MovieRenderPipelineSettings/Public/MoviePipelineGameOverrideSetting.h
@@ -63,7 +63,7 @@ public:
 	UPROPERTY(EditAnywhere, BlueprintReadWrite, Category = "Rendering")
 	EMoviePipelineTextureStreamingMethod TextureStreaming;
 
-	/** Should we try to use the highest quality LOD for meshes regardless of distance? */
+	/** Should we try to use the highest quality LOD for meshes and particle systems regardless of distance? */
 	UPROPERTY(EditAnywhere, BlueprintReadWrite, Category = "Rendering")
 	bool bUseLODZero;
 
@@ -99,6 +99,7 @@ private:
 	int32 PreviousTextureStreaming;
 	int32 PreviousForceLOD;
 	int32 PreviousSkeletalMeshBias;
+	int32 PreviousParticleLODBias;
 	int32 PreviousShadowDistanceScale;
 	int32 PreviousShadowQuality;
 	float PreviousShadowRadiusThreshold;
