MoviePipeline: Let MakeUniqueObjectName pick our name instead of us suggesting it.

diff --git a/Engine/Plugins/MovieScene/MovieRenderPipeline/Source/MovieRenderPipelineCore/Private/MoviePipelineBlueprintLibrary.cpp b/Engine/Plugins/MovieScene/MovieRenderPipeline/Source/MovieRenderPipelineCore/Private/MoviePipelineBlueprintLibrary.cpp
--- a/Engine/Plugins/MovieScene/MovieRenderPipeline/Source/MovieRenderPipelineCore/Private/MoviePipelineBlueprintLibrary.cpp
+++ b/Engine/Plugins/MovieScene/MovieRenderPipeline/Source/MovieRenderPipelineCore/Private/MoviePipelineBlueprintLibrary.cpp
@@ -142,7 +142,6 @@ void UMoviePipelineBlueprintLibrary::GetOverallSegmentCounts(const UMoviePipelin
 	}
 }
 
-
 UMovieSceneSequence* UMoviePipelineBlueprintLibrary::DuplicateSequence(UObject* Outer, UMovieSceneSequence* InSequence)
 {
 	if (!InSequence)
@@ -151,10 +150,10 @@ UMovieSceneSequence* UMoviePipelineBlueprintLibrary::DuplicateSequence(UObject*
 	}
 
 	FObjectDuplicationParameters DuplicationParams(InSequence, Outer);
-	DuplicationParams.DestName = MakeUniqueObjectName(Outer, UMovieSceneSequence::StaticClass(), InSequence->GetFName());
+	DuplicationParams.DestName = MakeUniqueObjectName(Outer, UMovieSceneSequence::StaticClass());
 
 	// Duplicate the given sequence.
-	ULevelSequence* DuplicatedSequence = (ULevelSequence*)StaticDuplicateObjectEx(DuplicationParams);
+	UMovieSceneSequence* DuplicatedSequence = CastChecked<UMovieSceneSequence>(StaticDuplicateObjectEx(DuplicationParams));
 
 	// Now go through looking for Shot and SubSequence tracks. These currently point to the same (shared) sequence as the InSequence.
 	TArray<UMovieSceneSection*> AllSubSequenceSections;
