MoviePipeline: Reduced gpu stall introduced by surface readback to copy results to the CPU.

diff --git a/Engine/Plugins/MovieScene/MovieRenderPipeline/Source/MovieRenderPipelineCore/Private/MoviePipelineSurfaceReader.cpp b/Engine/Plugins/MovieScene/MovieRenderPipeline/Source/MovieRenderPipelineCore/Private/MoviePipelineSurfaceReader.cpp
new file mode 100644
--- /dev/null
+++ b/Engine/Plugins/MovieScene/MovieRenderPipeline/Source/MovieRenderPipelineCore/Private/MoviePipelineSurfaceReader.cpp
@@ -0,0 +1,290 @@
+// Copyright Epic Games, Inc. All Rights Reserved.
+
+#include "MoviePipelineSurfaceReader.h"
+#include "RenderingThread.h"
+#include "RHICommandList.h"
+#include "RHIStaticStates.h"
+#include "CommonRenderResources.h"
+#include "RenderTargetPool.h"
+#include "ScreenRendering.h"
+#include "GlobalShader.h"
+#include "Shader.h"
+#include "StaticBoundShaderState.h"
+#include "Modules/ModuleManager.h"
+
+FMoviePipelineSurfaceReader::FMoviePipelineSurfaceReader(EPixelFormat InPixelFormat, FIntPoint InSurfaceSize)
+{
+	AvailableEvent = nullptr;
+	ReadbackTexture = nullptr;
+	PixelFormat = InPixelFormat;
+	bQueuedForCapture = false;
+	
+	Resize(InSurfaceSize.X, InSurfaceSize.Y);
+}
+
+FMoviePipelineSurfaceReader::~FMoviePipelineSurfaceReader()
+{
+	BlockUntilAvailable();
+
+	ReadbackTexture.SafeRelease();
+}
+
+void FMoviePipelineSurfaceReader::Initialize()
+{
+	// Initialize shouldn't be called if we're currently not available.
+	check(!AvailableEvent);
+	AvailableEvent = FPlatformProcess::GetSynchEventFromPool();
+}
+
+void FMoviePipelineSurfaceReader::Resize(uint32 Width, uint32 Height)
+{
+	ReadbackTexture.SafeRelease();
+
+	// TexCreate_CPUReadback is important to make this texture available for reading later.
+	FMoviePipelineSurfaceReader* This = this;
+	ENQUEUE_RENDER_COMMAND(CreateCaptureFrameTexture)(
+		[Width, Height, This](FRHICommandListImmediate& RHICmdList)
+		{
+			FRHIResourceCreateInfo CreateInfo;
+
+			This->ReadbackTexture = RHICreateTexture2D(
+				Width,
+				Height,
+				This->PixelFormat,
+				1,
+				1,
+				TexCreate_CPUReadback,
+				CreateInfo
+				);
+		});
+}
+
+void FMoviePipelineSurfaceReader::BlockUntilAvailable()
+{
+	if (AvailableEvent)
+	{
+		// Make this thread wait until another thread (ie: Rendering Thread) calls Trigger on the event.
+		AvailableEvent->Wait(~0);
+
+		// Return the sync event until we call Initialize again.
+		FPlatformProcess::ReturnSynchEventToPool(AvailableEvent);
+		AvailableEvent = nullptr;
+	}
+}
+
+void FMoviePipelineSurfaceReader::Reset()
+{
+	if (AvailableEvent)
+	{
+		AvailableEvent->Trigger();
+	}
+	BlockUntilAvailable();
+	bQueuedForCapture = false;
+}
+
+void FMoviePipelineSurfaceReader::ResolveSampleToReadbackTexture_RenderThread(const FTexture2DRHIRef& SourceSurfaceSample)
+{
+	// We use GetModuleChecked here to avoid accidentally loading the module from the non-main thread.
+	static const FName RendererModuleName("Renderer");
+	IRendererModule* RendererModule = &FModuleManager::GetModuleChecked<IRendererModule>(RendererModuleName);
+
+	bQueuedForCapture = true;
+	FRHICommandListImmediate& RHICmdList = GetImmediateCommandList_ForRenderCommand();
+
+	// Retrieve a temporary render target that matches the destination surface type/size.
+	const FIntPoint TargetSize(ReadbackTexture->GetSizeX(), ReadbackTexture->GetSizeY());
+
+	FPooledRenderTargetDesc OutputDesc = FPooledRenderTargetDesc::Create2DDesc(
+		TargetSize,
+		ReadbackTexture->GetFormat(),
+		FClearValueBinding::None,
+		TexCreate_None,
+		TexCreate_RenderTargetable,
+		false);
+	TRefCountPtr<IPooledRenderTarget> ResampleTexturePooledRenderTarget;
+	GRenderTargetPool.FindFreeElement(RHICmdList, OutputDesc, ResampleTexturePooledRenderTarget, TEXT("ResampleTexture"));
+	check(ResampleTexturePooledRenderTarget);
+
+	// Enqueue a render pass which uses a simple shader and a fullscreen quad to copy the SourceSurfaceSample to this SurfaceReader's cpu-readback enabled texture.
+	// This RenderPass resolves to our ReadbackTexture 
+	FRHIRenderPassInfo RPInfo(ResampleTexturePooledRenderTarget->GetRenderTargetItem().TargetableTexture, ERenderTargetActions::Load_Store, ReadbackTexture);
+	RHICmdList.BeginRenderPass(RPInfo, TEXT("MoviePipelineSurfaceResolveRenderTarget"));
+	{
+		RHICmdList.SetViewport(0, 0, 0.0f, TargetSize.X, TargetSize.Y, 1.0f);
+
+		FGraphicsPipelineStateInitializer GraphicsPSOInit;
+		RHICmdList.ApplyCachedRenderTargets(GraphicsPSOInit);
+		GraphicsPSOInit.BlendState = TStaticBlendState<>::GetRHI();
+		GraphicsPSOInit.RasterizerState = TStaticRasterizerState<>::GetRHI();
+		GraphicsPSOInit.DepthStencilState = TStaticDepthStencilState<false, CF_Always>::GetRHI();
+
+		const ERHIFeatureLevel::Type FeatureLevel = GMaxRHIFeatureLevel;
+
+		FGlobalShaderMap* ShaderMap = GetGlobalShaderMap(FeatureLevel);
+		TShaderMapRef<FScreenVS> VertexShader(ShaderMap);
+		TShaderMapRef<FScreenPS> PixelShader(ShaderMap);
+
+		GraphicsPSOInit.BoundShaderState.VertexDeclarationRHI = GFilterVertexDeclaration.VertexDeclarationRHI;
+		GraphicsPSOInit.BoundShaderState.VertexShaderRHI = VertexShader.GetVertexShader();
+		GraphicsPSOInit.BoundShaderState.PixelShaderRHI = PixelShader.GetPixelShader();
+		GraphicsPSOInit.PrimitiveType = PT_TriangleList;
+
+		SetGraphicsPipelineState(RHICmdList, GraphicsPSOInit);
+
+		// Bind the SourceSurfaceSample to our texture sampler with a point sample (since the resolutions match).
+		PixelShader->SetParameters(RHICmdList, TStaticSamplerState<SF_Point>::GetRHI(), SourceSurfaceSample);
+
+		RendererModule->DrawRectangle(
+			RHICmdList,
+			0, 0,									// Dest X, Y
+			TargetSize.X,							// Dest Width
+			TargetSize.Y,							// Dest Height
+			0, 0,									// Source U, V
+			1, 1,									// Source USize, VSize
+			TargetSize,								// Target buffer size
+			FIntPoint(1, 1),						// Source texture size
+			VertexShader,
+			EDRF_Default);
+	}
+	RHICmdList.EndRenderPass();
+}
+
+void FMoviePipelineSurfaceReader::CopyReadbackTexture_RenderThread(const FMoviePipelineRenderPassMetrics& InSampleState, const FMoviePipelineSampleReady& InCallback)
+{
+	// We use GetModuleChecked here to avoid accidentally loading the module from the non-main thread.
+	static const FName RendererModuleName("Renderer");
+	IRendererModule* RendererModule = &FModuleManager::GetModuleChecked<IRendererModule>(RendererModuleName);
+
+	FRHICommandListImmediate& RHICmdList = GetImmediateCommandList_ForRenderCommand();
+	{
+		void* ColorDataBuffer = nullptr;
+
+		int32 Width = 0, Height = 0;
+		RHICmdList.MapStagingSurface(ReadbackTexture, ColorDataBuffer, Width, Height);
+
+		// Do a quick block copy
+		TArray<FFloat16Color> OutputPixels;
+		OutputPixels.SetNumUninitialized(Width * Height);
+		FMemory::BigBlockMemcpy(OutputPixels.GetData(), ColorDataBuffer, (Width * Height) * sizeof(FFloat16Color));
+
+		// Enqueue the Unmap before we broadcast the resulting pixels, though the broadcast shouldn't do anything blocking.
+		RHICmdList.UnmapStagingSurface(ReadbackTexture);
+
+		InCallback.Broadcast(OutputPixels, InSampleState);
+		
+		// Now that we've successfully used the surface, we trigger the Available event so that we can reuse this surface. This
+		// triggers the Available event and then returns it to the pool.
+		Reset();
+	}
+}
+
+FMoviePipelineSurfaceQueue::FMoviePipelineSurfaceQueue(FIntPoint InSurfaceSize, EPixelFormat InPixelFormat, uint32 InNumSurfaces)
+{
+	// TargetSize = InSurfaceSize;
+	CurrentFrameIndex = 0;
+	check(InNumSurfaces > 0);
+
+	// The Surface array can't be reallocated as we're storing indexes into it.
+	Surfaces.Reserve(InNumSurfaces);
+	for (uint32 Index = 0; Index < InNumSurfaces; Index++)
+	{
+		Surfaces.Emplace(InPixelFormat, InSurfaceSize);
+	}
+
+	FrameResolveLatency = 1;
+}
+
+FMoviePipelineSurfaceQueue::~FMoviePipelineSurfaceQueue()
+{
+}
+
+void FMoviePipelineSurfaceQueue::BlockUntilAnyAvailable()
+{
+	bool bAnyAvailable = false;
+	for(FResolveSurface& ResolveSurface : Surfaces)
+	{
+		if (ResolveSurface.Surface.IsAvailable())
+		{
+			bAnyAvailable = true;
+			break;
+		}
+	}
+
+	// We only wait if none of them are available, and the oldest one is the most likely one to
+	// be available. 
+	if (!bAnyAvailable)
+	{
+		const int32 OldestIndex = (CurrentFrameIndex + 1) % Surfaces.Num();
+		Surfaces[OldestIndex].Surface.BlockUntilAvailable();
+	}
+}
+
+void FMoviePipelineSurfaceQueue::Shutdown()
+{
+	for (int32 Index = 0; Index < Surfaces.Num(); Index++)
+	{
+		FResolveSurface* ResolveSurface = &Surfaces[Index];
+		ENQUEUE_RENDER_COMMAND(PerformReadback)(
+			[ResolveSurface](FRHICommandListImmediate& RHICmdList)
+			{
+				// Ensure that all surfaces have queued up their readback.
+				if (ResolveSurface->Surface.WasEverQueued())
+				{
+					ResolveSurface->Surface.CopyReadbackTexture_RenderThread(ResolveSurface->SampleState, ResolveSurface->Callback);
+				}
+			});
+	}
+
+	// Flush the rendering thread so that the readbacks happen and trigger AvailableEvents.
+	FlushRenderingCommands();
+
+	for (FResolveSurface& ResolveSurface : Surfaces)
+	{
+		// These should now all be available.
+		ensureMsgf(ResolveSurface.Surface.IsAvailable(), TEXT("Flushed rendering commands but surface reader didn't perform the readback on a surface!"));
+	}
+}
+
+
+void FMoviePipelineSurfaceQueue::OnRenderTargetReady_RenderThread(const FTexture2DRHIRef InRenderTarget, const FMoviePipelineRenderPassMetrics& InSampleState, const FMoviePipelineSampleReady& InCallback)
+{
+	ensure(IsInRenderingThread());
+	
+	// Pick the next destination surface and ensure it's available.
+	FResolveSurface* NextResolveTarget = &Surfaces[CurrentFrameIndex];
+	if(!NextResolveTarget->Surface.IsAvailable())
+	{
+		// BlockUntilAnyAvailable should have been called before submitting more work. We can't block
+		// until a surface is available in this callback because we'd be waiting on the render thread
+		// from the render thread.
+		check(false);
+	}
+	
+	NextResolveTarget->Surface.Initialize();
+	NextResolveTarget->SampleState = InSampleState;
+	NextResolveTarget->Callback = InCallback;
+
+	// Queue this sample to be copied to the target surface.
+	NextResolveTarget->Surface.ResolveSampleToReadbackTexture_RenderThread(InRenderTarget);
+	
+	// By the time we get to this point, our oldest surface should have successfully been rendered to, and no longer be in use by the GPU.
+	// We can now safely map the surface and copy the data out of it without causing a GPU stall.
+	{
+		const int32 PrevCaptureIndexOffset = FMath::Clamp(FrameResolveLatency, 0, Surfaces.Num() - 1);
+
+		// Get PrevCaptureIndexOffset surfaces back, handling wraparound at 0.
+		const int32 PrevCaptureIndex = (CurrentFrameIndex - PrevCaptureIndexOffset) < 0 ? Surfaces.Num() - (PrevCaptureIndexOffset - CurrentFrameIndex) : (CurrentFrameIndex - PrevCaptureIndexOffset);
+
+		FResolveSurface* OldestResolveTarget = &Surfaces[PrevCaptureIndex];
+
+		// Only try to do the readback if the target has ever been written to.
+		if (OldestResolveTarget->Surface.WasEverQueued())
+		{
+			OldestResolveTarget->Surface.CopyReadbackTexture_RenderThread(OldestResolveTarget->SampleState, OldestResolveTarget->Callback);
+		}
+	}
+	
+
+	// Write to the next available surface next time.
+	CurrentFrameIndex = (CurrentFrameIndex + 1) % Surfaces.Num();
+}
\ No newline at end of file
diff --git a/Engine/Plugins/MovieScene/MovieRenderPipeline/Source/MovieRenderPipelineCore/Public/MoviePipelineSurfaceReader.h b/Engine/Plugins/MovieScene/MovieRenderPipeline/Source/MovieRenderPipelineCore/Public/MoviePipelineSurfaceReader.h
new file mode 100644
--- /dev/null
+++ b/Engine/Plugins/MovieScene/MovieRenderPipeline/Source/MovieRenderPipelineCore/Public/MoviePipelineSurfaceReader.h
@@ -0,0 +1,106 @@
+// Copyright Epic Games, Inc. All Rights Reserved.
+#pragma once
+
+#include "CoreMinimal.h"
+#include "PixelFormat.h"
+#include "RHI.h"
+#include "MovieRenderPipelineDataTypes.h"
+
+/**
+* When the GPU finishes rendering a view we can copy the resulting data back to the CPU. Because the
+* buffer is still in use this requires flushing the RHI and stalling the GPU from doing any additional
+* work while the copy is in progress. To resolve this issue, when the data is returned we schedule a copy
+* to a surface that isn't in use which allows the GPU to resume work. Then when the surface is written
+* to, we can copy it to the CPU without blocking the GPU. This requires ensuring that the surface is
+* available, so this SurfaceReader implements a trigger system to ensure we don't try to access the 
+* surface until it has been written to, combined with a round-robin of surfaces to avoid stalls.
+*/
+struct MOVIERENDERPIPELINECORE_API FMoviePipelineSurfaceReader
+{
+	/** Construct a surface reader with the given pixel format and size. This will be the output specification. */
+	FMoviePipelineSurfaceReader(EPixelFormat InPixelFormat, FIntPoint InSurfaceSize);
+
+	~FMoviePipelineSurfaceReader();
+
+	/** Initialize this reader so that it can be waited on. */
+	void Initialize();
+
+	/** Wait until this surface is available for reuse. Used to stall if we've run out of available surfaces. */
+	void BlockUntilAvailable();
+
+	/**
+	* Safely resets the state of the wait event. When doing latent surface reading sometimes we may want to just
+	* bail on reading a given frame. Should only be performed after flushing rendering commands.
+	*/
+	void Reset();
+
+	bool WasEverQueued() const { return bQueuedForCapture; }
+	bool IsAvailable() const { return AvailableEvent == nullptr; }
+
+	/**
+	* Issues a command to the GPU to copy the given SourceSurfaceSample to our local ReadbackTexture for this surface.
+	*/
+	void ResolveSampleToReadbackTexture_RenderThread(const FTexture2DRHIRef& SourceSurfaceSample);
+
+	/**
+	* Maps the ReadbackTexture to the CPU (which should have been resolved to before this point) and copies the data over.
+	*/
+	void CopyReadbackTexture_RenderThread(const FMoviePipelineRenderPassMetrics& InSampleState, const FMoviePipelineSampleReady& InCallback);
+protected:
+
+
+	/** Set up this surface to the specified width/height */
+	void Resize(uint32 Width, uint32 Height);
+
+	/** Optional event that is triggered when the surface is no longer in use */
+	FEvent* AvailableEvent;
+
+	/** Texture used to store the resolved render target */
+	FTexture2DRHIRef ReadbackTexture;
+
+	/** The desired pixel format of the resolved textures */
+	EPixelFormat PixelFormat;
+
+	bool bQueuedForCapture;
+};
+
+struct MOVIERENDERPIPELINECORE_API FMoviePipelineSurfaceQueue
+{
+public:
+	FMoviePipelineSurfaceQueue(FIntPoint InSurfaceSize, EPixelFormat InPixelFormat, uint32 InNumSurfaces);
+	~FMoviePipelineSurfaceQueue();
+
+	// Movable
+	FMoviePipelineSurfaceQueue(FMoviePipelineSurfaceQueue&&) = default;
+	FMoviePipelineSurfaceQueue& operator=(FMoviePipelineSurfaceQueue&&) = default;
+
+	// Non Copyable
+	FMoviePipelineSurfaceQueue(const FMoviePipelineSurfaceQueue&) = delete;
+	FMoviePipelineSurfaceQueue& operator=(const FMoviePipelineSurfaceQueue&) = delete;
+public:
+
+	void OnRenderTargetReady_RenderThread(const FTexture2DRHIRef InRenderTarget, const FMoviePipelineRenderPassMetrics& InSampleState, const FMoviePipelineSampleReady& InCallback);
+	void BlockUntilAnyAvailable();
+	void Shutdown();
+private:
+	struct FResolveSurface
+	{
+		FResolveSurface(const EPixelFormat InPixelFormat, const FIntPoint InSurfaceSize)
+			: Surface(InPixelFormat, InSurfaceSize)
+		{
+		}
+
+		FMoviePipelineSurfaceReader Surface;
+		FMoviePipelineRenderPassMetrics SampleState;
+		FMoviePipelineSampleReady Callback;
+	};
+	
+	/** Which surface is the next target to be copied to. */
+	int32 CurrentFrameIndex;
+
+	/** The array of surfaces we can resolve a render target to. More peformant if there's more than one. */
+	TArray<FResolveSurface> Surfaces;
+	
+	/** How many frames do we wait before we resolve a frame? */
+	int32 FrameResolveLatency;
+};
\ No newline at end of file
diff --git a/Engine/Plugins/MovieScene/MovieRenderPipeline/Source/MovieRenderPipelineCore/Public/MovieRenderPipelineDataTypes.h b/Engine/Plugins/MovieScene/MovieRenderPipeline/Source/MovieRenderPipelineCore/Public/MovieRenderPipelineDataTypes.h
--- a/Engine/Plugins/MovieScene/MovieRenderPipeline/Source/MovieRenderPipelineCore/Public/MovieRenderPipelineDataTypes.h
+++ b/Engine/Plugins/MovieScene/MovieRenderPipeline/Source/MovieRenderPipelineCore/Public/MovieRenderPipelineDataTypes.h
@@ -820,6 +820,8 @@ public:
 	MoviePipeline::FMoviePipelineFrameInfo FrameInfo;
 };
 
+DECLARE_MULTICAST_DELEGATE_TwoParams(FMoviePipelineSampleReady, TArray<FFloat16Color>&, FMoviePipelineRenderPassMetrics);
+
 namespace MoviePipeline
 {
 	struct FMoviePipelineRenderPassInitSettings
diff --git a/Engine/Plugins/MovieScene/MovieRenderPipeline/Source/MovieRenderPipelineRenderPasses/Private/MoviePipelineDeferredPasses.cpp b/Engine/Plugins/MovieScene/MovieRenderPipeline/Source/MovieRenderPipelineRenderPasses/Private/MoviePipelineDeferredPasses.cpp
--- a/Engine/Plugins/MovieScene/MovieRenderPipeline/Source/MovieRenderPipelineRenderPasses/Private/MoviePipelineDeferredPasses.cpp
+++ b/Engine/Plugins/MovieScene/MovieRenderPipeline/Source/MovieRenderPipelineRenderPasses/Private/MoviePipelineDeferredPasses.cpp
@@ -31,13 +31,17 @@
 #include "MovieRenderPipelineCoreModule.h"
 
 DECLARE_CYCLE_STAT(TEXT("MoviePipeline_SampleReadback"), STAT_SampleReadback, STATGROUP_MoviePipeline);
+DECLARE_CYCLE_STAT(TEXT("MoviePipeline_SampleReadbackAlloc"), STAT_SampleReadbackAlloc, STATGROUP_MoviePipeline);
+DECLARE_CYCLE_STAT(TEXT("MoviePipeline_SampleReadbackRHI"), STAT_SampleReadbackRHI, STATGROUP_MoviePipeline);
+DECLARE_CYCLE_STAT(TEXT("MoviePipeline_SampleReadbackBroadcast"), STAT_SampleReadbackBroadcast, STATGROUP_MoviePipeline);
 DECLARE_CYCLE_STAT(TEXT("MoviePipeline_AccumulateSample_RT"), STAT_AccumulateSample_RenderThread, STATGROUP_MoviePipeline);
+DECLARE_CYCLE_STAT(TEXT("MoviePipeline_OnBackbufferSampleReady_RT"), STAT_OnBackbufferSampleReady_RenderThread, STATGROUP_MoviePipeline);
 
+DECLARE_CYCLE_STAT(TEXT("STAT_MoviePipeline_WaitForAvailableSurface"), STAT_MoviePipeline_WaitForAvailableSurface, STATGROUP_MoviePipeline);
 
 namespace MoviePipeline
 {
 	// Forward Declare
-	void ReadBackbufferAndBroadcast_RenderThread(FRHICommandListImmediate &RHICmdList, FRenderTarget* InFromRenderTarget, const FMoviePipelineRenderPassMetrics& InSampleState, const FMoviePipelineSampleReady& OnReadDelegate);
 	void AccumulateSample_RenderThread(TUniquePtr<FImagePixelData>&& InPixelData, TSharedRef<FImagePixelDataPayload, ESPMode::ThreadSafe> InFrameData, MoviePipeline::FSampleRenderThreadParams& InParams);
 
 	void FDeferredRenderEnginePass::Setup(TWeakObjectPtr<UMoviePipeline> InOwningPipeline, const FMoviePipelineRenderPassInitSettings& InInitSettings)
@@ -53,12 +57,17 @@ namespace MoviePipeline
 
 		// Allocate 
 		ViewState.Allocate();
+
+		SurfaceQueue = MakeShared<FMoviePipelineSurfaceQueue>(InInitSettings.BackbufferResolution, EPixelFormat::PF_FloatRGBA, 3);
 	}
 
 	void FDeferredRenderEnginePass::Teardown()
 	{
 		FMoviePipelineEnginePass::Teardown();
 
+		// This may call FlushRenderingCommands if there are outstanding readbacks that need to happen.
+		SurfaceQueue->Shutdown();
+
 		FSceneViewStateInterface* Ref = ViewState.GetReference();
 		if (Ref)
 		{
@@ -175,57 +184,44 @@ namespace MoviePipeline
 			View->MaterialTextureMipBias += InSampleState.TextureSharpnessBias;
 		}
 
+		// Wait for a surface to be available to write to. This will stall the game thread while the RHI/Render Thread catch up.
+		{
+			SCOPE_CYCLE_COUNTER(STAT_MoviePipeline_WaitForAvailableSurface);
+			SurfaceQueue->BlockUntilAnyAvailable();
+		}
+
 		FCanvas Canvas = FCanvas(TileRenderTarget->GameThread_GetRenderTargetResource(), nullptr, GetPipeline()->GetWorld(), ERHIFeatureLevel::SM5, FCanvas::CDM_DeferDrawing, 1.0f);
 		// SetupViewDelegate.Broadcast(ViewFamily, *View, InSampleState);
 
 		// Draw the world into this View Family
 		GetRendererModule().BeginRenderingViewFamily(&Canvas, &ViewFamily);
 
+		// If this was just to contribute to the history buffer, no need to go any further.
+		if (InSampleState.bDiscardResult)
+		{
+			return;
+		}
+		
+		TSharedPtr<FMoviePipelineSurfaceQueue> LocalSurfaceQueue = SurfaceQueue;
+
 		FRenderTarget* BackbufferRenderTarget = TileRenderTarget->GameThread_GetRenderTargetResource();
 		FMoviePipelineSampleReady& OnBackbufferReadDelegate = BackbufferReadyDelegate;
 
 		ENQUEUE_RENDER_COMMAND(CanvasRenderTargetResolveCommand)(
-			[InSampleState, BackbufferRenderTarget, OnBackbufferReadDelegate](FRHICommandListImmediate& RHICmdList)
+			[LocalSurfaceQueue, InSampleState, BackbufferRenderTarget, OnBackbufferReadDelegate](FRHICommandListImmediate& RHICmdList)
 		{
-			// If they only needed this frame for the history, we don't actually care about trying to output it,
-			// so we can skip the copy and pushing to output.
-			if (InSampleState.bDiscardResult)
-			{
-				return;
-			}
-
-			// If they are using tiles, we'll output some more fine-grained progress via logging since it freezes the UI for such a long time.
-			if (InSampleState.GetTileCount() > 1)
-			{
-				UE_LOG(LogMovieRenderPipeline, Log, TEXT("Frame: %d TemporalSample: %d/%d Rendered Tile: %d/%d Sample: %d/%d"),
-					InSampleState.OutputState.OutputFrameNumber, InSampleState.TemporalSampleIndex + 1, InSampleState.TemporalSampleCount,
-					InSampleState.GetTileIndex() + 1, InSampleState.GetTileCount(), InSampleState.SpatialSampleIndex + 1, InSampleState.SpatialSampleCount);
-			}
+				// If they are using tiles, we'll output some more fine-grained progress via logging since it freezes the UI for such a long time.
+				if (InSampleState.GetTileCount() > 1)
+				{
+					UE_LOG(LogMovieRenderPipeline, Log, TEXT("Frame: %d TemporalSample: %d/%d Rendered Tile: %d/%d Sample: %d/%d"),
+						InSampleState.OutputState.OutputFrameNumber, InSampleState.TemporalSampleIndex + 1, InSampleState.TemporalSampleCount,
+						InSampleState.GetTileIndex() + 1, InSampleState.GetTileCount(), InSampleState.SpatialSampleIndex + 1, InSampleState.SpatialSampleCount);
+				}
 
-			ReadBackbufferAndBroadcast_RenderThread(RHICmdList, BackbufferRenderTarget, InSampleState, OnBackbufferReadDelegate);
+				LocalSurfaceQueue->OnRenderTargetReady_RenderThread(BackbufferRenderTarget->GetRenderTargetTexture(), InSampleState, OnBackbufferReadDelegate);
 		});
 	}
 
-	void ReadBackbufferAndBroadcast_RenderThread(FRHICommandListImmediate &RHICmdList, FRenderTarget* InFromRenderTarget, const FMoviePipelineRenderPassMetrics& InSampleState, const FMoviePipelineSampleReady& OnReadDelegate)
-	{
-		SCOPE_CYCLE_COUNTER(STAT_SampleReadback);
-
-		check(IsInRenderingThread());
-
-		FIntRect SourceRect = FIntRect(0, 0, InFromRenderTarget->GetSizeXY().X, InFromRenderTarget->GetSizeXY().Y);
-
-		// Read the data back to the CPU
-		TArray<FLinearColor> RawPixels;
-		RawPixels.SetNum(SourceRect.Width() * SourceRect.Height());
-
-		FReadSurfaceDataFlags ReadDataFlags(ERangeCompressionMode::RCM_MinMax);
-		ReadDataFlags.SetLinearToGamma(false);
-
-		RHICmdList.ReadSurfaceData(InFromRenderTarget->GetRenderTargetTexture(), SourceRect, RawPixels, ReadDataFlags);
-
-		OnReadDelegate.Broadcast(RawPixels, InSampleState);
-	}
-
 	FSceneView* FDeferredRenderEnginePass::CalcSceneView(FSceneViewFamily* ViewFamily, const FMoviePipelineRenderPassMetrics& InSampleState)
 	{
 		APlayerController* LocalPlayerController = GetPipeline()->GetWorld()->GetFirstPlayerController();
@@ -446,8 +442,10 @@ void UMoviePipelineDeferredPassBase::GatherOutputPassesImpl(TArray<FMoviePipelin
 	}
 }
 
-void UMoviePipelineDeferredPassBase::OnBackbufferSampleReady(TArray<FLinearColor> InPixelData, FMoviePipelineRenderPassMetrics InSampleState)
+void UMoviePipelineDeferredPassBase::OnBackbufferSampleReady(TArray<FFloat16Color>& InPixelData, FMoviePipelineRenderPassMetrics InSampleState)
 {
+	SCOPE_CYCLE_COUNTER(STAT_OnBackbufferSampleReady_RenderThread);
+
 	MoviePipeline::FSampleRenderThreadParams AccumulationParams;
 	{
 		AccumulationParams.OutputMerger = GetPipeline()->OutputBuilder;
@@ -461,7 +459,7 @@ void UMoviePipelineDeferredPassBase::OnBackbufferSampleReady(TArray<FLinearColor
 	FrameData->PassIdentifier = FMoviePipelinePassIdentifier(TEXT("Backbuffer"));
 	FrameData->SampleState = InSampleState;
 
-	TUniquePtr<TImagePixelData<FLinearColor>> PixelData = MakeUnique<TImagePixelData<FLinearColor>>(InSampleState.BackbufferSize, TArray64<FLinearColor>(MoveTemp(InPixelData)), FrameData);
+	TUniquePtr<TImagePixelData<FFloat16Color>> PixelData = MakeUnique<TImagePixelData<FFloat16Color>>(InSampleState.BackbufferSize, TArray64<FFloat16Color>(MoveTemp(InPixelData)), FrameData);
 
 	AccumulateSample_RenderThread(MoveTemp(PixelData), FrameData, AccumulationParams);
 }
diff --git a/Engine/Plugins/MovieScene/MovieRenderPipeline/Source/MovieRenderPipelineRenderPasses/Public/MoviePipelineDeferredPasses.h b/Engine/Plugins/MovieScene/MovieRenderPipeline/Source/MovieRenderPipelineRenderPasses/Public/MoviePipelineDeferredPasses.h
--- a/Engine/Plugins/MovieScene/MovieRenderPipeline/Source/MovieRenderPipelineRenderPasses/Public/MoviePipelineDeferredPasses.h
+++ b/Engine/Plugins/MovieScene/MovieRenderPipeline/Source/MovieRenderPipelineRenderPasses/Public/MoviePipelineDeferredPasses.h
@@ -5,6 +5,7 @@
 #include "MovieRenderPipelineDataTypes.h"
 #include "SceneTypes.h"
 #include "SceneView.h"
+#include "MoviePipelineSurfaceReader.h"
 #include "UObject/GCObject.h"
 #include "MoviePipelineDeferredPasses.generated.h"
 
@@ -29,7 +30,7 @@ protected:
 	virtual FText GetDisplayText() const override { return NSLOCTEXT("MovieRenderPipeline", "DeferredBasePassSettingDisplayName", "Deferred Rendering"); }
 
 
-	void OnBackbufferSampleReady(TArray<FLinearColor> InPixelData, FMoviePipelineRenderPassMetrics InSampleState);
+	void OnBackbufferSampleReady(TArray<FFloat16Color>& InPixelData, FMoviePipelineRenderPassMetrics InSampleState);
 	void OnSetupView(FSceneViewFamily& InViewFamily, FSceneView& InView, const FMoviePipelineRenderPassMetrics& InSampleState);
 
 private:
@@ -42,7 +43,6 @@ private:
 	TArray<TSharedPtr<FImageOverlappedAccumulator, ESPMode::ThreadSafe>> ImageTileAccumulator;
 };
 
-DECLARE_MULTICAST_DELEGATE_TwoParams(FMoviePipelineSampleReady, TArray<FLinearColor>, FMoviePipelineRenderPassMetrics);
 DECLARE_MULTICAST_DELEGATE_ThreeParams(FMoviePipelineSetupView, FSceneViewFamily&, FSceneView&, const FMoviePipelineRenderPassMetrics&);
 
 namespace MoviePipeline
@@ -70,6 +70,9 @@ namespace MoviePipeline
 
 	protected:
 		FSceneViewStateReference ViewState;
+		
+		/** A queue of surfaces that the render targets can be copied to. If no surface is available the game thread should hold off on submitting more samples. */
+		TSharedPtr<FMoviePipelineSurfaceQueue> SurfaceQueue;
 
 		TWeakObjectPtr<UTextureRenderTarget2D> TileRenderTarget;
 	};
