MoviePipeline: Changed some logging to VeryVerbose to avoid console spam.

#rnx

diff --git a/Engine/Plugins/MovieScene/MovieRenderPipeline/Source/MovieRenderPipelineCore/Private/MoviePipelineRendering.cpp b/Engine/Plugins/MovieScene/MovieRenderPipeline/Source/MovieRenderPipelineCore/Private/MoviePipelineRendering.cpp
--- a/Engine/Plugins/MovieScene/MovieRenderPipeline/Source/MovieRenderPipelineCore/Private/MoviePipelineRendering.cpp
+++ b/Engine/Plugins/MovieScene/MovieRenderPipeline/Source/MovieRenderPipelineCore/Private/MoviePipelineRendering.cpp
@@ -353,7 +353,7 @@ void UMoviePipeline::RenderFrame()
 				// only allow a spatial jitter if we have more than one sample
 				bool bAllowSpatialJitter = !(NumSpatialSamples == 1 && NumTemporalSamples == 1);
 
-				UE_LOG(LogTemp, Warning, TEXT("FrameIndex: %d HaltonIndex: %d Offset: (%f,%f)"), FrameIndex, HaltonIndex, HaltonOffsetX, HaltonOffsetY);
+				UE_LOG(LogTemp, VeryVerbose, TEXT("FrameIndex: %d HaltonIndex: %d Offset: (%f,%f)"), FrameIndex, HaltonIndex, HaltonOffsetX, HaltonOffsetY);
 				float SpatialShiftX = 0.0f;
 				float SpatialShiftY = 0.0f;
 
@@ -375,7 +375,6 @@ void UMoviePipeline::RenderFrame()
 					float Theta = 2.0f * PI * HaltonOffsetY;
 					float r = Sigma * FMath::Sqrt(-2.0f * FMath::Loge((1.0f - HaltonOffsetX) * InWindow + HaltonOffsetX));
 
-					UE_LOG(LogTemp, Log, TEXT("r: %f t: %f"), r, Theta);
 					SpatialShiftX = r * FMath::Cos(Theta);
 					SpatialShiftY = r * FMath::Sin(Theta);
 				}
diff --git a/Engine/Plugins/MovieScene/MovieRenderPipeline/Source/MovieRenderPipelineRenderPasses/Private/MoviePipelineDeferredPasses.cpp b/Engine/Plugins/MovieScene/MovieRenderPipeline/Source/MovieRenderPipelineRenderPasses/Private/MoviePipelineDeferredPasses.cpp
--- a/Engine/Plugins/MovieScene/MovieRenderPipeline/Source/MovieRenderPipelineRenderPasses/Private/MoviePipelineDeferredPasses.cpp
+++ b/Engine/Plugins/MovieScene/MovieRenderPipeline/Source/MovieRenderPipelineRenderPasses/Private/MoviePipelineDeferredPasses.cpp
@@ -594,7 +594,7 @@ namespace MoviePipeline
 			const double AccumulateEndTime = FPlatformTime::Seconds();
 			const float ElapsedMs = float((AccumulateEndTime - AccumulateBeginTime)*1000.0f);
 
-			UE_LOG(LogMovieRenderPipeline, Log, TEXT("Accumulation time: %8.2fms"), ElapsedMs);
+			UE_LOG(LogMovieRenderPipeline, VeryVerbose, TEXT("Accumulation time: %8.2fms"), ElapsedMs);
 
 		}
 
