MoviePipeline: Suppress screen messages globally when rendering to prevent messages that are written directly to the FSceneView.

#rnx

diff --git a/Engine/Plugins/MovieScene/MovieRenderPipeline/Source/MovieRenderPipelineCore/Private/MoviePipeline.cpp b/Engine/Plugins/MovieScene/MovieRenderPipeline/Source/MovieRenderPipelineCore/Private/MoviePipeline.cpp
--- a/Engine/Plugins/MovieScene/MovieRenderPipeline/Source/MovieRenderPipelineCore/Private/MoviePipeline.cpp
+++ b/Engine/Plugins/MovieScene/MovieRenderPipeline/Source/MovieRenderPipelineCore/Private/MoviePipeline.cpp
@@ -103,6 +103,9 @@ void UMoviePipeline::Initialize(UMoviePipelineExecutorJob* InJob)
 	// for FCoreDelegateS::OnBeginFrame to get called to actually start processing.
 	UE_LOG(LogMovieRenderPipeline, Log, TEXT("[%d] Initializing overall Movie Pipeline"), GFrameCounter);
 
+	bPrevGScreenMessagesEnabled = GAreScreenMessagesEnabled;
+	GAreScreenMessagesEnabled = false;
+
 	if (!ensureAlwaysMsgf(InJob, TEXT("MoviePipeline cannot be initialized with null job. Aborting.")))
 	{
 		OnMoviePipelineErroredDelegate.Broadcast(this, true, LOCTEXT("MissingJob", "Job was not specified, movie render is aborting."));
@@ -262,6 +265,8 @@ void UMoviePipeline::Shutdown()
 
 	TeardownAudioRendering();
 
+	GAreScreenMessagesEnabled = bPrevGScreenMessagesEnabled;
+
 	UE_LOG(LogMovieRenderPipeline, Log, TEXT("Movie Pipeline completed. Duration: %s"), *(FDateTime::UtcNow() - InitializationTime).ToString());
 	PipelineState = EMovieRenderPipelineState::Shutdown;
 }
diff --git a/Engine/Plugins/MovieScene/MovieRenderPipeline/Source/MovieRenderPipelineCore/Public/MoviePipeline.h b/Engine/Plugins/MovieScene/MovieRenderPipeline/Source/MovieRenderPipelineCore/Public/MoviePipeline.h
--- a/Engine/Plugins/MovieScene/MovieRenderPipeline/Source/MovieRenderPipelineCore/Public/MoviePipeline.h
+++ b/Engine/Plugins/MovieScene/MovieRenderPipeline/Source/MovieRenderPipelineCore/Public/MoviePipeline.h
@@ -249,6 +249,9 @@ private:
 
 	MoviePipeline::FAudioState AudioState;
 
+	/** Cached state of GAreScreenMessagesEnabled. We disable them since some messages are written to the FSceneView directly otherwise. */
+	bool bPrevGScreenMessagesEnabled;
+
 	/** 
 	* Have we hit the callback for the BeginFrame at least once? This solves an issue where the delegates
 	* get registered mid-frame so you end up calling EndFrame before BeginFrame which is undesirable.
