MoviePipeline: Exposed Burley SSS sample count overrides to the UI for better quality sub surface scattering in high resolution (tiled) frames.

diff --git a/Engine/Plugins/MovieScene/MovieRenderPipeline/Source/MovieRenderPipelineCore/Public/MoviePipelineHighResSetting.h b/Engine/Plugins/MovieScene/MovieRenderPipeline/Source/MovieRenderPipelineCore/Public/MoviePipelineHighResSetting.h
--- a/Engine/Plugins/MovieScene/MovieRenderPipeline/Source/MovieRenderPipelineCore/Public/MoviePipelineHighResSetting.h
+++ b/Engine/Plugins/MovieScene/MovieRenderPipeline/Source/MovieRenderPipelineCore/Public/MoviePipelineHighResSetting.h
@@ -14,6 +14,8 @@ public:
 		: TileCount(1)
 		, TextureSharpnessBias(0.f)
 		, OverlapRatio(0.f)
+		, bOverrideSubSurfaceScattering(false)
+		, BurleySampleCount(64)
 		, bWriteAllSamples(false)
 	{
 	}
@@ -53,6 +55,37 @@ public:
 		InOutFormatArgs.Arguments.Add(TEXT("overlap_percent"), OverlapRatio);
 	}
 
+	virtual void SetupForPipelineImpl(UMoviePipeline* InPipeline) override
+	{
+		if (!(IsEnabled() && GetIsUserCustomized()))
+		{
+			return;
+		}
+
+		IConsoleVariable* BurleyOverride = IConsoleManager::Get().FindConsoleVariable(TEXT("r.SSS.Burley.NumSamplesOverride"));
+		if (BurleyOverride)
+		{
+			PrevBurleyOverride = BurleyOverride->GetInt();
+
+			int32 NumSamples = bOverrideSubSurfaceScattering ? BurleySampleCount : 0;
+			BurleyOverride->Set(NumSamples, EConsoleVariableFlags::ECVF_SetByConsole);
+		}
+	}
+
+	virtual void TeardownForPipelineImpl(UMoviePipeline* InPipeline) override
+	{
+		if (!(IsEnabled() && GetIsUserCustomized()))
+		{
+			return;
+		}
+
+		IConsoleVariable* BurleyOverride = IConsoleManager::Get().FindConsoleVariable(TEXT("r.SSS.Burley.NumSamplesOverride"));
+		if (BurleyOverride)
+		{
+			BurleyOverride->Set(PrevBurleyOverride, EConsoleVariableFlags::ECVF_SetByConsole);
+		}
+	}
+
 public:
 	/**
 	* How many tiles should the resulting movie render be broken into? A tile should be no larger than
@@ -80,6 +113,19 @@ public:
 	*/
 	UPROPERTY(EditAnywhere, BlueprintReadWrite, meta = (UIMin = 0, ClampMin = 0, UIMax = 0.5, ClampMax=1), Category = "Movie Pipeline")
 	float OverlapRatio;
+
+	/**
+	* Sub Surface Scattering relies on history which is not available when using tiling. This can be overriden to use more samples
+	* to improve the quality.
+	*/
+	UPROPERTY(EditAnywhere, BlueprintReadWrite, Category = "Movie Pipeline")
+	bool bOverrideSubSurfaceScattering;
+
+	/*
+	* How many samples should the Burley Sub Surface Scattering use?
+	*/
+	UPROPERTY(EditAnywhere, BlueprintReadWrite, meta = (UIMin = 64, ClampMin = 0, UIMax = 1024, EditCondition="bOverrideSubSurfaceScattering"), Category = "Movie Pipeline")
+	int32 BurleySampleCount;
 	
 	/**
 	* If true, we will write all samples that get generated to disk individually. This can be useful for debugging or if you need to accumulate
@@ -88,4 +134,8 @@ public:
 	UPROPERTY(EditAnywhere, BlueprintReadWrite, AdvancedDisplay, Category = "Movie Pipeline")
 	bool bWriteAllSamples;
 
+
+private:
+	/** What was the Burley CVar before we overrode it? */
+	int32 PrevBurleyOverride;
 };
\ No newline at end of file
