MoviePipeline: Fixed total world seconds continually getting reset which prevented time-based effects from working correctly. This also fixes render warmup frames not working as the reset time was causing a negative delta time which resets camera history.

diff --git a/Engine/Plugins/MovieScene/MovieRenderPipeline/Source/MovieRenderPipelineCore/Private/MoviePipelineRendering.cpp b/Engine/Plugins/MovieScene/MovieRenderPipeline/Source/MovieRenderPipelineCore/Private/MoviePipelineRendering.cpp
--- a/Engine/Plugins/MovieScene/MovieRenderPipeline/Source/MovieRenderPipelineCore/Private/MoviePipelineRendering.cpp
+++ b/Engine/Plugins/MovieScene/MovieRenderPipeline/Source/MovieRenderPipelineCore/Private/MoviePipelineRendering.cpp
@@ -192,12 +192,6 @@ void UMoviePipeline::RenderFrame()
 	FMoviePipelineCameraCutInfo& CurrentCameraCut = CurrentShot.GetCurrentCameraCut();
 	APlayerController* LocalPlayerController = GetWorld()->GetFirstPlayerController();
 
-	// Capture the camera location if this is a motion blur frame, so that we get camera motion blur on next frame.
-	if (CurrentCameraCut.State == EMovieRenderShotState::MotionBlur)
-	{
-		LocalPlayerController->GetPlayerViewPoint(FrameInfo.CurrViewLocation, FrameInfo.CurrViewRotation);
-		LocalPlayerController->GetPlayerViewPoint(FrameInfo.PrevViewLocation, FrameInfo.PrevViewRotation);
-	}
 
 	// If we don't want to render this frame, then we will skip processing - engine warmup frames,
 	// render every nTh frame, etc. In other cases, we may wish to render the frame but discard the
diff --git a/Engine/Plugins/MovieScene/MovieRenderPipeline/Source/MovieRenderPipelineCore/Public/MovieRenderPipelineDataTypes.h b/Engine/Plugins/MovieScene/MovieRenderPipeline/Source/MovieRenderPipelineCore/Public/MovieRenderPipelineDataTypes.h
--- a/Engine/Plugins/MovieScene/MovieRenderPipeline/Source/MovieRenderPipelineCore/Public/MovieRenderPipelineDataTypes.h
+++ b/Engine/Plugins/MovieScene/MovieRenderPipeline/Source/MovieRenderPipelineCore/Public/MovieRenderPipelineDataTypes.h
@@ -563,6 +563,13 @@ public:
 		* may be duplicate frame Source/Effective frame numbers as they find the closest ideal time to the current.
 		*/
 		bool bWasAffectedByTimeDilation;
+
+		void ResetPerFrameData()
+		{
+			MotionBlurFraction = 0.f;
+			FrameDeltaTime = 0.0;
+			bWasAffectedByTimeDilation = false;
+		}
 	};
 
 	FMoviePipelineFrameOutputState()
@@ -614,11 +621,14 @@ public:
 	/** The total number of samples (including warm ups) that have been sent to the GPU for this shot. */
 	int32 ShotSamplesRendered;
 
+	/** What time data should this frame use? Can vary between samples when TemporalSampleCount > 1. */
+	FTimeData TimeData;
+
 	/** INFORMATION BELOW HERE SHOULD NOT GET PERSISTED BETWEEN FRAMES */
 
 	void ResetPerFrameData()
 	{
-		TimeData = FTimeData();
+		TimeData.ResetPerFrameData();
 		bSkipRendering = false;
 		bDiscardRenderResult = false;
 		SourceFrameNumber = 0;
@@ -637,9 +647,6 @@ public:
 		TemporalSampleCount = 0;
 	}
 
-	/** What time data should this frame use? Can vary between samples when TemporalSampleCount > 1. */
-	FTimeData TimeData;
-
 	/**
 	* If true, then the rendering for this frame should be skipped (ie: nothing submitted to the gpu, and the output merger not told to expect this frame).
 	* This is used for rendering every Nth frame for rendering drafts. We still run the game thread logic for the skipped frames (which is relatively cheap)
