MoviePipeline: Added a UI warning for large sample counts being less effective if TAA is still enabled.

diff --git a/Engine/Plugins/MovieScene/MovieRenderPipeline/Source/MovieRenderPipelineCore/Private/MoviePipelineRendering.cpp b/Engine/Plugins/MovieScene/MovieRenderPipeline/Source/MovieRenderPipelineCore/Private/MoviePipelineRendering.cpp
--- a/Engine/Plugins/MovieScene/MovieRenderPipeline/Source/MovieRenderPipelineCore/Private/MoviePipelineRendering.cpp
+++ b/Engine/Plugins/MovieScene/MovieRenderPipeline/Source/MovieRenderPipelineCore/Private/MoviePipelineRendering.cpp
@@ -307,14 +307,7 @@ void UMoviePipeline::RenderFrame()
 				const bool bCameraCut = CachedOutputState.ShotSamplesRendered == 0;
 				CachedOutputState.ShotSamplesRendered++;
 
-				EAntiAliasingMethod AntiAliasingMethod = EAntiAliasingMethod::AAM_TemporalAA; // temp
-
-				// Anti-aliasing appears to reasonably work with temporal and spatial sampling so we want it on by default 
-				// (otherwise going from 1 sample > 2 is a 4x reduction in quality due to TAA being turned off).
-				if (AntiAliasingSettings->bOverrideAntiAliasing)
-				{
-					AntiAliasingMethod = AntiAliasingSettings->AntiAliasingMethod;
-				}
+				EAntiAliasingMethod AntiAliasingMethod = UE::MovieRenderPipeline::GetEffectiveAntiAliasingMethod(AntiAliasingSettings);
 
 				// Now to check if we have to force it off (at which point we warn the user).
 				bool bMultipleTiles = (TileCount.X > 1) || (TileCount.Y > 1);
diff --git a/Engine/Plugins/MovieScene/MovieRenderPipeline/Source/MovieRenderPipelineCore/Private/MoviePipelineUtils.cpp b/Engine/Plugins/MovieScene/MovieRenderPipeline/Source/MovieRenderPipelineCore/Private/MoviePipelineUtils.cpp
--- a/Engine/Plugins/MovieScene/MovieRenderPipeline/Source/MovieRenderPipelineCore/Private/MoviePipelineUtils.cpp
+++ b/Engine/Plugins/MovieScene/MovieRenderPipeline/Source/MovieRenderPipelineCore/Private/MoviePipelineUtils.cpp
@@ -7,6 +7,8 @@
 #include "ARFilter.h"
 #include "AssetData.h"
 #include "UObject/UObjectIterator.h"
+#include "HAL/IConsoleManager.h"
+#include "MoviePipelineAntiAliasingSetting.h"
 
 namespace UE
 {
@@ -46,5 +48,35 @@ namespace UE
 
 			return Classes;
 		}
+	
+		/**
+		* Returns the anti-aliasing setting that we should use. This defaults to the project setting, and then
+		* uses the one specified by the setting if overriden.
+		*/
+		EAntiAliasingMethod GetEffectiveAntiAliasingMethod(const UMoviePipelineAntiAliasingSetting* InSetting)
+		{
+			EAntiAliasingMethod AntiAliasingMethod = EAntiAliasingMethod::AAM_None;
+
+			IConsoleVariable* AntiAliasingCVar = IConsoleManager::Get().FindConsoleVariable(TEXT("r.DefaultFeature.AntiAliasing"));
+			if (AntiAliasingCVar)
+			{
+				int32 Value = AntiAliasingCVar->GetInt();
+				if (Value >= 0 && Value < AAM_MAX)
+				{
+					AntiAliasingMethod = (EAntiAliasingMethod)Value;
+				}
+			}
+
+			if (InSetting)
+			{
+				if (InSetting->bOverrideAntiAliasing)
+				{
+					AntiAliasingMethod = InSetting->AntiAliasingMethod;
+				}
+			}
+
+			return AntiAliasingMethod;
+		}
 	}
+
 }
\ No newline at end of file
diff --git a/Engine/Plugins/MovieScene/MovieRenderPipeline/Source/MovieRenderPipelineCore/Public/MoviePipelineAntiAliasingSetting.h b/Engine/Plugins/MovieScene/MovieRenderPipeline/Source/MovieRenderPipelineCore/Public/MoviePipelineAntiAliasingSetting.h
--- a/Engine/Plugins/MovieScene/MovieRenderPipeline/Source/MovieRenderPipelineCore/Public/MoviePipelineAntiAliasingSetting.h
+++ b/Engine/Plugins/MovieScene/MovieRenderPipeline/Source/MovieRenderPipelineCore/Public/MoviePipelineAntiAliasingSetting.h
@@ -4,9 +4,9 @@
 #include "MoviePipelineSetting.h"
 #include "MovieRenderPipelineDataTypes.h"
 #include "Engine/Scene.h"
+#include "MoviePipelineUtils.h"
 #include "MoviePipelineAntiAliasingSetting.generated.h"
 
-
 UCLASS(Blueprintable)
 class MOVIERENDERPIPELINECORE_API UMoviePipelineAntiAliasingSetting : public UMoviePipelineSetting
 {
@@ -35,8 +35,14 @@ protected:
 	{
 		Super::ValidateStateImpl();
 
-		if (bOverrideAntiAliasing && (AntiAliasingMethod == EAntiAliasingMethod::AAM_TemporalAA))
+		if (UE::MovieRenderPipeline::GetEffectiveAntiAliasingMethod(this) == EAntiAliasingMethod::AAM_TemporalAA)
 		{
+			if ((TemporalSampleCount*SpatialSampleCount) > 8)
+			{
+				ValidationResults.Add(NSLOCTEXT("MovieRenderPipeline", "AntiAliasing_BetterOffWithoutTAA", "If the product of Temporal and Spatial counts is greater than the number of TAA samples then TAA is ineffective and you should consider overriding AA to None for better quality."));
+				ValidationState = EMoviePipelineValidationState::Warnings;
+			}
+
 			if (SpatialSampleCount % 2 == 0)
 			{
 				ValidationResults.Add(NSLOCTEXT("MovieRenderPipeline", "AntiAliasing_InsufficientJitters", "TAA does not converge when using an even number of samples. Disable TAA or increase sample count."));
@@ -106,4 +112,4 @@ public:
 	* is applied and pow(x,1/AccumulationGamma) is applied after accumulation is finished
 	*/
 	float AccumulationGamma;
-};
\ No newline at end of file
+};
diff --git a/Engine/Plugins/MovieScene/MovieRenderPipeline/Source/MovieRenderPipelineCore/Public/MoviePipelineUtils.h b/Engine/Plugins/MovieScene/MovieRenderPipeline/Source/MovieRenderPipelineCore/Public/MoviePipelineUtils.h
--- a/Engine/Plugins/MovieScene/MovieRenderPipeline/Source/MovieRenderPipelineCore/Public/MoviePipelineUtils.h
+++ b/Engine/Plugins/MovieScene/MovieRenderPipeline/Source/MovieRenderPipelineCore/Public/MoviePipelineUtils.h
@@ -2,14 +2,17 @@
 #pragma once
 
 #include "CoreMinimal.h"
+#include "Engine/Scene.h"
 
 // Forward Declare
 class UClass;
+class UMoviePipelineAntiAliasingSetting;
 
 namespace UE
 {
 	namespace MovieRenderPipeline
 	{
 		MOVIERENDERPIPELINECORE_API TArray<UClass*> FindMoviePipelineSettingClasses();
+		MOVIERENDERPIPELINECORE_API EAntiAliasingMethod GetEffectiveAntiAliasingMethod(const UMoviePipelineAntiAliasingSetting* InSetting);
 	}
 }
\ No newline at end of file
