MoviePipeline: Flush the rendering commands before nulling out the Widget Renderer in Burn Ins so that the renderer doesn't later try to access the now destroyed Widget Renderer.

diff --git a/Engine/Plugins/MovieScene/MovieRenderPipeline/Source/MovieRenderPipelineSettings/Private/MoviePipelineBurnInSetting.cpp b/Engine/Plugins/MovieScene/MovieRenderPipeline/Source/MovieRenderPipelineSettings/Private/MoviePipelineBurnInSetting.cpp
--- a/Engine/Plugins/MovieScene/MovieRenderPipeline/Source/MovieRenderPipelineSettings/Private/MoviePipelineBurnInSetting.cpp
+++ b/Engine/Plugins/MovieScene/MovieRenderPipeline/Source/MovieRenderPipelineSettings/Private/MoviePipelineBurnInSetting.cpp
@@ -107,6 +107,14 @@ void UMoviePipelineBurnInSetting::SetupImpl(TArray<TSharedPtr<MoviePipeline::FMo
 
 void UMoviePipelineBurnInSetting::TeardownImpl() 
 {
+	// If this was transiently added, don't make a burn-in.
+	if (!GetIsUserCustomized() || !IsEnabled())
+	{
+		return;
+	}
+	
+	FlushRenderingCommands();
+
 	if (FSlateApplication::IsInitialized() && VirtualWindow.IsValid())
 	{
 		FSlateApplication::Get().UnregisterVirtualWindow(VirtualWindow.ToSharedRef());
