MoviePipeline: Reduced code duplication in settings that overrode cvars.

#rnx

diff --git a/Engine/Plugins/MovieScene/MovieRenderPipeline/Source/MovieRenderPipelineCore/Public/MoviePipelineHighResSetting.h b/Engine/Plugins/MovieScene/MovieRenderPipeline/Source/MovieRenderPipelineCore/Public/MoviePipelineHighResSetting.h
--- a/Engine/Plugins/MovieScene/MovieRenderPipeline/Source/MovieRenderPipelineCore/Public/MoviePipelineHighResSetting.h
+++ b/Engine/Plugins/MovieScene/MovieRenderPipeline/Source/MovieRenderPipelineCore/Public/MoviePipelineHighResSetting.h
@@ -3,6 +3,7 @@
 
 #include "MoviePipelineSetting.h"
 #include "MovieRenderPipelineDataTypes.h"
+#include "MoviePipelineUtils.h"
 #include "MoviePipelineHighResSetting.generated.h"
 
 UCLASS(Blueprintable)
@@ -62,14 +63,8 @@ public:
 			return;
 		}
 
-		IConsoleVariable* BurleyOverride = IConsoleManager::Get().FindConsoleVariable(TEXT("r.SSS.Burley.NumSamplesOverride"));
-		if (BurleyOverride)
-		{
-			PrevBurleyOverride = BurleyOverride->GetInt();
-
-			int32 NumSamples = bOverrideSubSurfaceScattering ? BurleySampleCount : 0;
-			BurleyOverride->Set(NumSamples, EConsoleVariableFlags::ECVF_SetByConsole);
-		}
+		int32 NumSamples = bOverrideSubSurfaceScattering ? BurleySampleCount : 0;
+		MOVIEPIPELINE_STORE_AND_OVERRIDE_CVAR_INT(PrevBurleyOverride, TEXT("r.SSS.Burley.NumSamplesOverride"), NumSamples, true);
 	}
 
 	virtual void TeardownForPipelineImpl(UMoviePipeline* InPipeline) override
@@ -79,11 +74,8 @@ public:
 			return;
 		}
 
-		IConsoleVariable* BurleyOverride = IConsoleManager::Get().FindConsoleVariable(TEXT("r.SSS.Burley.NumSamplesOverride"));
-		if (BurleyOverride)
-		{
-			BurleyOverride->Set(PrevBurleyOverride, EConsoleVariableFlags::ECVF_SetByConsole);
-		}
+		int32 NumSamples = 0; // Dummy
+		MOVIEPIPELINE_STORE_AND_OVERRIDE_CVAR_INT(PrevBurleyOverride, TEXT("r.SSS.Burley.NumSamplesOverride"), NumSamples, false);
 	}
 
 public:
diff --git a/Engine/Plugins/MovieScene/MovieRenderPipeline/Source/MovieRenderPipelineCore/Public/MoviePipelineUtils.h b/Engine/Plugins/MovieScene/MovieRenderPipeline/Source/MovieRenderPipelineCore/Public/MoviePipelineUtils.h
--- a/Engine/Plugins/MovieScene/MovieRenderPipeline/Source/MovieRenderPipelineCore/Public/MoviePipelineUtils.h
+++ b/Engine/Plugins/MovieScene/MovieRenderPipeline/Source/MovieRenderPipelineCore/Public/MoviePipelineUtils.h
@@ -8,6 +8,34 @@
 class UClass;
 class UMoviePipelineAntiAliasingSetting;
 
+#define MOVIEPIPELINE_STORE_AND_OVERRIDE_CVAR_INT(InOutVariable, CVarName, OverrideValue, bUseOverride) \
+{ \
+	IConsoleVariable* CVar = IConsoleManager::Get().FindConsoleVariable(CVarName); \
+	if(bUseOverride) \
+	{ \
+		InOutVariable = CVar->GetInt(); \
+		CVar->Set(OverrideValue, EConsoleVariableFlags::ECVF_SetByConsole); \
+	} \
+	else \
+	{ \
+		CVar->Set(InOutVariable, EConsoleVariableFlags::ECVF_SetByConsole); \
+	} \
+}
+
+#define MOVIEPIPELINE_STORE_AND_OVERRIDE_CVAR_FLOAT(InOutVariable, CVarName, OverrideValue, bUseOverride) \
+{ \
+	IConsoleVariable* CVar = IConsoleManager::Get().FindConsoleVariable(CVarName); \
+	if(bUseOverride) \
+	{ \
+		InOutVariable = CVar->GetFloat(); \
+		CVar->Set(OverrideValue, EConsoleVariableFlags::ECVF_SetByConsole); \
+	} \
+	else \
+	{ \
+		CVar->Set(InOutVariable, EConsoleVariableFlags::ECVF_SetByConsole); \
+	} \
+}
+
 namespace UE
 {
 	namespace MovieRenderPipeline
diff --git a/Engine/Plugins/MovieScene/MovieRenderPipeline/Source/MovieRenderPipelineSettings/Private/MoviePipelineGameOverrideSetting.cpp b/Engine/Plugins/MovieScene/MovieRenderPipeline/Source/MovieRenderPipelineSettings/Private/MoviePipelineGameOverrideSetting.cpp
--- a/Engine/Plugins/MovieScene/MovieRenderPipeline/Source/MovieRenderPipelineSettings/Private/MoviePipelineGameOverrideSetting.cpp
+++ b/Engine/Plugins/MovieScene/MovieRenderPipeline/Source/MovieRenderPipelineSettings/Private/MoviePipelineGameOverrideSetting.cpp
@@ -4,50 +4,23 @@
 #include "Scalability.h"
 #include "Engine/World.h"
 #include "MovieRenderPipelineCoreModule.h"
-
-#define STORE_AND_OVERRIDE_INT_SETTING(InOutVariable, CVarName, OverrideValue, bApply) \
-{ \
-	IConsoleVariable* CVar = IConsoleManager::Get().FindConsoleVariable(CVarName); \
-	if(bApply) \
-	{ \
-		CVar->Set(InOutVariable, EConsoleVariableFlags::ECVF_SetByConsole); \
-	} \
-	else \
-	{ \
-		InOutVariable = CVar->GetInt(); \
-		CVar->Set(OverrideValue, EConsoleVariableFlags::ECVF_SetByConsole); \
-	} \
-}
-
-#define STORE_AND_OVERRIDE_FLOAT_SETTING(InOutVariable, CVarName, OverrideValue, bApply) \
-{ \
-	IConsoleVariable* CVar = IConsoleManager::Get().FindConsoleVariable(CVarName); \
-	if(bApply) \
-	{ \
-		CVar->Set(InOutVariable, EConsoleVariableFlags::ECVF_SetByConsole); \
-	} \
-	else \
-	{ \
-		InOutVariable = CVar->GetFloat(); \
-		CVar->Set(OverrideValue, EConsoleVariableFlags::ECVF_SetByConsole); \
-	} \
-}
+#include "MoviePipelineUtils.h"
 
 void UMoviePipelineGameOverrideSetting::SetupForPipelineImpl(UMoviePipeline* InPipeline)
 {
 	// Store the cvar values and apply the ones from this setting
 	UE_LOG(LogMovieRenderPipeline, Log, TEXT("Applying Game Override quality settings and cvars."));
-	ApplyCVarSettings(false);
+	ApplyCVarSettings(true);
 }
 
 void UMoviePipelineGameOverrideSetting::TeardownForPipelineImpl(UMoviePipeline* InPipeline)
 {
 	// Restore the previous cvar values the user had
 	UE_LOG(LogMovieRenderPipeline, Log, TEXT("Restoring Game Override quality settings and cvars."));
-	ApplyCVarSettings(true);
+	ApplyCVarSettings(false);
 }
 
-void UMoviePipelineGameOverrideSetting::ApplyCVarSettings(const bool bRestoreOldValues)
+void UMoviePipelineGameOverrideSetting::ApplyCVarSettings(const bool bOverrideValues)
 {
 	if (!IsEnabled())
 	{
@@ -56,7 +29,7 @@ void UMoviePipelineGameOverrideSetting::ApplyCVarSettings(const bool bRestoreOld
 
 	if (bCinematicQualitySettings)
 	{
-		if (!bRestoreOldValues)
+		if (bOverrideValues)
 		{
 			// Store their previous Scalability settings so we can revert back to them
 			PreviousQualityLevels = Scalability::GetQualityLevels();
@@ -77,24 +50,11 @@ void UMoviePipelineGameOverrideSetting::ApplyCVarSettings(const bool bRestoreOld
 	switch (TextureStreaming)
 	{
 	case EMoviePipelineTextureStreamingMethod::FullyLoad:
-		{ 
-			IConsoleVariable* CVar = IConsoleManager::Get().FindConsoleVariable(TEXT("r.Streaming.FramesForFullUpdate")); 
-			if(bRestoreOldValues) 
-			{ 
-				CVar->Set(PreviousFramesForFullUpdate, EConsoleVariableFlags::ECVF_SetByConsole); 
-			} 
-			else 
-			{ 
-				PreviousFramesForFullUpdate = CVar->GetInt(); 
-				CVar->Set(0, EConsoleVariableFlags::ECVF_SetByConsole);
-			} 
-		}
-
-		STORE_AND_OVERRIDE_INT_SETTING(PreviousFramesForFullUpdate, TEXT("r.Streaming.FramesForFullUpdate"), 0, bRestoreOldValues);
-		STORE_AND_OVERRIDE_INT_SETTING(PreviousFullyLoadUsedTextures, TEXT("r.Streaming.FullyLoadUsedTextures"), 1, bRestoreOldValues);
+		MOVIEPIPELINE_STORE_AND_OVERRIDE_CVAR_INT(PreviousFramesForFullUpdate, TEXT("r.Streaming.FramesForFullUpdate"), 0, bOverrideValues);
+		MOVIEPIPELINE_STORE_AND_OVERRIDE_CVAR_INT(PreviousFullyLoadUsedTextures, TEXT("r.Streaming.FullyLoadUsedTextures"), 1, bOverrideValues);
 		break;
 	case EMoviePipelineTextureStreamingMethod::Disabled:
-		STORE_AND_OVERRIDE_INT_SETTING(PreviousTextureStreaming, TEXT("r.TextureStreaming"), 0, bRestoreOldValues);
+		MOVIEPIPELINE_STORE_AND_OVERRIDE_CVAR_INT(PreviousTextureStreaming, TEXT("r.TextureStreaming"), 0, bOverrideValues);
 		break;
 	default:
 		// We don't change their texture streaming settings.
@@ -103,9 +63,9 @@ void UMoviePipelineGameOverrideSetting::ApplyCVarSettings(const bool bRestoreOld
 
 	if (bUseLODZero)
 	{
-		STORE_AND_OVERRIDE_INT_SETTING(PreviousForceLOD, TEXT("r.ForceLOD"), 0, bRestoreOldValues);
-		STORE_AND_OVERRIDE_INT_SETTING(PreviousSkeletalMeshBias, TEXT("r.SkeletalMeshLODBias"), -10, bRestoreOldValues);
-		STORE_AND_OVERRIDE_INT_SETTING(PreviousParticleLODBias, TEXT("r.ParticleLODBias"), -10, bRestoreOldValues);
+		MOVIEPIPELINE_STORE_AND_OVERRIDE_CVAR_INT(PreviousForceLOD, TEXT("r.ForceLOD"), 0, bOverrideValues);
+		MOVIEPIPELINE_STORE_AND_OVERRIDE_CVAR_INT(PreviousSkeletalMeshBias, TEXT("r.SkeletalMeshLODBias"), -10, bOverrideValues);
+		MOVIEPIPELINE_STORE_AND_OVERRIDE_CVAR_INT(PreviousParticleLODBias, TEXT("r.ParticleLODBias"), -10, bOverrideValues);
 	}
 
 	if (bDisableHLODs)
@@ -116,14 +76,14 @@ void UMoviePipelineGameOverrideSetting::ApplyCVarSettings(const bool bRestoreOld
 
 	if (bUseHighQualityShadows)
 	{
-		STORE_AND_OVERRIDE_INT_SETTING(PreviousShadowDistanceScale, TEXT("r.Shadow.DistanceScale"), ShadowDistanceScale, bRestoreOldValues);
-		STORE_AND_OVERRIDE_INT_SETTING(PreviousShadowQuality, TEXT("r.ShadowQuality"), 5, bRestoreOldValues);
-		STORE_AND_OVERRIDE_FLOAT_SETTING(PreviousShadowRadiusThreshold, TEXT("r.Shadow.RadiusThreshold"), ShadowRadiusThreshold, bRestoreOldValues);
+		MOVIEPIPELINE_STORE_AND_OVERRIDE_CVAR_INT(PreviousShadowDistanceScale, TEXT("r.Shadow.DistanceScale"), ShadowDistanceScale, bOverrideValues);
+		MOVIEPIPELINE_STORE_AND_OVERRIDE_CVAR_INT(PreviousShadowQuality, TEXT("r.ShadowQuality"), 5, bOverrideValues);
+		MOVIEPIPELINE_STORE_AND_OVERRIDE_CVAR_FLOAT(PreviousShadowRadiusThreshold, TEXT("r.Shadow.RadiusThreshold"), ShadowRadiusThreshold, bOverrideValues);
 	}
 
 	if (bOverrideViewDistanceScale)
 	{
-		STORE_AND_OVERRIDE_INT_SETTING(PreviousViewDistanceScale, TEXT("r.ViewDistanceScale"), ViewDistanceScale, bRestoreOldValues);
+		MOVIEPIPELINE_STORE_AND_OVERRIDE_CVAR_INT(PreviousViewDistanceScale, TEXT("r.ViewDistanceScale"), ViewDistanceScale, bOverrideValues);
 	}
 }
 
