MoviePipeline: Fixed a case where having a Camera Cut track but no sections could skip creating a range to render, instead of falling back to the playback bounds.

#rnx

diff --git a/Engine/Plugins/MovieScene/MovieRenderPipeline/Source/MovieRenderPipelineCore/Private/MoviePipeline.cpp b/Engine/Plugins/MovieScene/MovieRenderPipeline/Source/MovieRenderPipelineCore/Private/MoviePipeline.cpp
--- a/Engine/Plugins/MovieScene/MovieRenderPipeline/Source/MovieRenderPipelineCore/Private/MoviePipeline.cpp
+++ b/Engine/Plugins/MovieScene/MovieRenderPipeline/Source/MovieRenderPipelineCore/Private/MoviePipeline.cpp
@@ -606,9 +606,10 @@ FMoviePipelineShotInfo CreateShotFromMovieScene(const UMovieScene* InMovieScene,
 			NewRange.Section = CameraCutSection;
 		}
 	}
-	else
+
+	if(IntersectedRanges.Num() == 0)
 	{
-		// No camera cut track was found inside. We'll treat the whole shot as the desired range.
+		// No camera cut track (or section) was found inside. We'll treat the whole shot as the desired range.
 		FCameraCutRange& NewRange = IntersectedRanges.AddDefaulted_GetRef();
 		NewRange.Range = InIntersectionRange;
 		NewRange.Section = nullptr;
