MoviePipeline: Fixed Remote Render not respecting the Game Mode override.

diff --git a/Engine/Plugins/MovieScene/MovieRenderPipeline/Source/MovieRenderPipelineSettings/Private/MoviePipelineGameOverrideSetting.cpp b/Engine/Plugins/MovieScene/MovieRenderPipeline/Source/MovieRenderPipelineSettings/Private/MoviePipelineGameOverrideSetting.cpp
--- a/Engine/Plugins/MovieScene/MovieRenderPipeline/Source/MovieRenderPipelineSettings/Private/MoviePipelineGameOverrideSetting.cpp
+++ b/Engine/Plugins/MovieScene/MovieRenderPipeline/Source/MovieRenderPipelineSettings/Private/MoviePipelineGameOverrideSetting.cpp
@@ -100,7 +100,7 @@ void UMoviePipelineGameOverrideSetting::BuildNewProcessCommandLineImpl(FString&
 	// If they've left the setting enabled but set it to None, we don't do anything.
 	if (GameModeOverride)
 	{
-		InOutUnrealURLParams += FString::Printf(TEXT("?game=%s"), *GameModeOverride->GetName());
+		InOutUnrealURLParams += FString::Printf(TEXT("?game=%s"), *GameModeOverride->GetPathName());
 	}
 
 	if (bCinematicQualitySettings)
