MoviePipeline: Projects can now be configured with a project-wide setting for the number of frames to delay after startup before the movie pipeline starts processing. This works around projects which have systems that are still loading when BeginPlay is being run and need to be waited on.

diff --git a/Engine/Plugins/MovieScene/MovieRenderPipeline/Source/MovieRenderPipelineCore/Private/MoviePipeline.cpp b/Engine/Plugins/MovieScene/MovieRenderPipeline/Source/MovieRenderPipelineCore/Private/MoviePipeline.cpp
--- a/Engine/Plugins/MovieScene/MovieRenderPipeline/Source/MovieRenderPipelineCore/Private/MoviePipeline.cpp
+++ b/Engine/Plugins/MovieScene/MovieRenderPipeline/Source/MovieRenderPipelineCore/Private/MoviePipeline.cpp
@@ -140,12 +140,6 @@ void UMoviePipeline::Initialize(UMoviePipelineExecutorJob* InJob)
 	}
 
 	CurrentJob = InJob;
-
-	// When start up we want to override the engine's Custom Timestep with our own.
-	// This gives us the ability to completely control the engine tick/delta time before the frame
-	// is started so that we don't have to always be thinking of delta times one frame ahead.
-	CachedPrevCustomTimeStep = GEngine->GetCustomTimeStep();
-	GEngine->SetCustomTimeStep(CustomTimeStep);
 	
 	ULevelSequence* OriginalSequence = Cast<ULevelSequence>(InJob->Sequence.TryLoad());
 	if (!ensureAlwaysMsgf(OriginalSequence, TEXT("Failed to load Sequence Asset from specified path, aborting movie render! Attempted to load Path: %s"), *InJob->Sequence.ToString()))
diff --git a/Engine/Plugins/MovieScene/MovieRenderPipeline/Source/MovieRenderPipelineCore/Private/MoviePipelineInProcessExecutor.cpp b/Engine/Plugins/MovieScene/MovieRenderPipeline/Source/MovieRenderPipelineCore/Private/MoviePipelineInProcessExecutor.cpp
--- a/Engine/Plugins/MovieScene/MovieRenderPipeline/Source/MovieRenderPipelineCore/Private/MoviePipelineInProcessExecutor.cpp
+++ b/Engine/Plugins/MovieScene/MovieRenderPipeline/Source/MovieRenderPipelineCore/Private/MoviePipelineInProcessExecutor.cpp
@@ -4,11 +4,22 @@
 #include "Engine/World.h"
 #include "MoviePipelineQueue.h"
 #include "Kismet/KismetSystemLibrary.h"
+#include "MoviePipelineBlueprintLibrary.h"
+#include "MoviePipelineInProcessExecutorSettings.h"
 
 #define LOCTEXT_NAMESPACE "MoviePipelineInProcessExecutor"
 void UMoviePipelineInProcessExecutor::Start(const UMoviePipelineExecutorJob* InJob)
 {
 	FCoreUObjectDelegates::PostLoadMapWithWorld.AddUObject(this, &UMoviePipelineInProcessExecutor::OnMapLoadFinished);
+
+	// Force the engine into fixed timestep mode. There may be a global delay on the job that passes a fixed
+	// number of frames, so we want those frames to always pass the same amount of time for determinism. 
+	ULevelSequence* LevelSequence = CastChecked<ULevelSequence>(InJob->Sequence.TryLoad());
+	if (LevelSequence)
+	{
+		FApp::SetUseFixedTimeStep(true);
+		FApp::SetFixedDeltaTime(InJob->GetConfiguration()->GetEffectiveFrameRate(LevelSequence).AsInterval());
+	}
 }
 
 void UMoviePipelineInProcessExecutor::UpdateWindowTitle()
@@ -17,17 +28,14 @@ void UMoviePipelineInProcessExecutor::UpdateWindowTitle()
 	PercentFormatOptions.MinimumIntegralDigits = 2;
 	PercentFormatOptions.MaximumIntegralDigits = 3;
 
-	FText PercentCompleteText = FText::AsNumber(0.f, &PercentFormatOptions);
+	float CompletionPercentage = 0.f;
 	if (ActiveMoviePipeline)
 	{
-		float PercentComplete = 12.f;
-		PercentCompleteText = FText::AsNumber(PercentComplete, &PercentFormatOptions);
+		CompletionPercentage = UMoviePipelineBlueprintLibrary::GetCompletionPercentage(ActiveMoviePipeline);
 	}
 
-	FText TitleFormatString = LOCTEXT("WindowTitleFormat", "Movie Pipeline Render (Preview) [{CurrentCount}/{TotalCount} Total] {PercentComplete}% Completed.");
-	FText WindowTitle = FText::FormatNamed(TitleFormatString, 
-		TEXT("CurrentCount"), FText::AsNumber(CurrentPipelineIndex + 1), TEXT("TotalCount"),
-		FText::AsNumber(Queue->GetJobs().Num()), TEXT("PercentComplete"), PercentCompleteText);
+	FText TitleFormatString = LOCTEXT("MoviePreviewWindowTitleFormat", "Movie Pipeline Render (Preview) [Job {CurrentCount}/{TotalCount} Total] Current Job: {PercentComplete}% Completed.");
+	FText WindowTitle = FText::FormatNamed(TitleFormatString, TEXT("CurrentCount"), FText::AsNumber(CurrentPipelineIndex + 1), TEXT("TotalCount"), FText::AsNumber(Queue->GetJobs().Num()), TEXT("PercentComplete"), FText::AsNumber(12.f, &PercentFormatOptions));
 
 	UKismetSystemLibrary::SetWindowTitle(WindowTitle);
 }
@@ -37,6 +45,7 @@ void UMoviePipelineInProcessExecutor::OnMapLoadFinished(UWorld* NewWorld)
 	// NewWorld can be null if a world is being destroyed.
 	if (!NewWorld)
 	{
+		FCoreDelegates::OnBeginFrame.RemoveAll(this);
 		return;
 	}
 	
@@ -46,7 +55,22 @@ void UMoviePipelineInProcessExecutor::OnMapLoadFinished(UWorld* NewWorld)
 	UMoviePipelineExecutorJob* CurrentJob = Queue->GetJobs()[CurrentPipelineIndex];
 
 	ActiveMoviePipeline = NewObject<UMoviePipeline>(NewWorld, TargetPipelineClass);
-	ActiveMoviePipeline->Initialize(CurrentJob);
+
+	// We allow users to set a multi-frame delay before we actually run the Initialization function and start thinking.
+	// This solves cases where there are engine systems that need to finish loading before we do anything.
+	const UMoviePipelineInProcessExecutorSettings* ExecutorSettings = GetDefault<UMoviePipelineInProcessExecutorSettings>();
+
+	if (ExecutorSettings->InitialDelayFrameCount == 0)
+	{
+		ActiveMoviePipeline->Initialize(Queue->GetJobs()[CurrentPipelineIndex]);
+		RemainingInitializationFrames = -1;
+	}
+	else
+	{
+		RemainingInitializationFrames = ExecutorSettings->InitialDelayFrameCount;
+		FCoreDelegates::OnBeginFrame.AddUObject(this, &UMoviePipelineInProcessExecutor::DelayedInitializationCounter);
+	}
+
 	UpdateWindowTitle();
 
 	// Listen for when the pipeline thinks it has finished.
@@ -56,6 +80,17 @@ void UMoviePipelineInProcessExecutor::OnMapLoadFinished(UWorld* NewWorld)
 	bIsRendering = true;
 }
 
+void UMoviePipelineInProcessExecutor::DelayedInitializationCounter()
+{
+	if (RemainingInitializationFrames == 0)
+	{
+		ActiveMoviePipeline->Initialize(Queue->GetJobs()[CurrentPipelineIndex]);
+		FCoreDelegates::OnBeginFrame.RemoveAll(this);
+	}
+
+	RemainingInitializationFrames--;
+}
+
 void UMoviePipelineInProcessExecutor::OnMoviePipelineFinished(UMoviePipeline* InMoviePipeline)
 {
 	UMoviePipeline* MoviePipeline = ActiveMoviePipeline;
diff --git a/Engine/Plugins/MovieScene/MovieRenderPipeline/Source/MovieRenderPipelineCore/Private/MoviePipelineTiming.cpp b/Engine/Plugins/MovieScene/MovieRenderPipeline/Source/MovieRenderPipelineCore/Private/MoviePipelineTiming.cpp
--- a/Engine/Plugins/MovieScene/MovieRenderPipeline/Source/MovieRenderPipelineCore/Private/MoviePipelineTiming.cpp
+++ b/Engine/Plugins/MovieScene/MovieRenderPipeline/Source/MovieRenderPipelineCore/Private/MoviePipelineTiming.cpp
@@ -28,6 +28,18 @@ void UMoviePipeline::TickProducingFrames()
 	// We should not be calling this once we have completed all the shots.
 	check(CurrentShotIndex >= 0 && CurrentShotIndex < ShotList.Num());
 
+	// When start up we want to override the engine's Custom Timestep with our own.
+	// This gives us the ability to completely control the engine tick/delta time before the frame
+	// is started so that we don't have to always be thinking of delta times one frame ahead. We need
+	// to do this only once we're ready to set the timestep though, as Initialize can be called as
+	// a result of a OnBeginFrame, meaning that Initialize is called on the frame before TickProducingFrames
+	// so there would be one frame where it used the custom timestep (after initialize) before TPF was called.
+	if (GEngine->GetCustomTimeStep() != CustomTimeStep)
+	{
+		CachedPrevCustomTimeStep = GEngine->GetCustomTimeStep();
+		GEngine->SetCustomTimeStep(CustomTimeStep);
+	}
+
 	// If we're debug stepping one frame at a time, this will return true
 	// and we early out so that we don't end up advancing the pipeline at all.
 	// This handles setting the delta time for us to a fixed number.
diff --git a/Engine/Plugins/MovieScene/MovieRenderPipeline/Source/MovieRenderPipelineCore/Public/MoviePipelineInProcessExecutor.h b/Engine/Plugins/MovieScene/MovieRenderPipeline/Source/MovieRenderPipelineCore/Public/MoviePipelineInProcessExecutor.h
--- a/Engine/Plugins/MovieScene/MovieRenderPipeline/Source/MovieRenderPipelineCore/Public/MoviePipelineInProcessExecutor.h
+++ b/Engine/Plugins/MovieScene/MovieRenderPipeline/Source/MovieRenderPipelineCore/Public/MoviePipelineInProcessExecutor.h
@@ -19,6 +19,7 @@ class MOVIERENDERPIPELINECORE_API UMoviePipelineInProcessExecutor : public UMovi
 public:
 	UMoviePipelineInProcessExecutor()
 		: UMoviePipelineLinearExecutorBase()
+		, RemainingInitializationFrames(-1)
 	{
 	}
 
@@ -29,4 +30,9 @@ protected:
 private:
 	void OnMapLoadFinished(UWorld* NewWorld);
 	void OnMoviePipelineFinished(UMoviePipeline* InMoviePipeline);
+	void DelayedInitializationCounter();
+
+private:
+	/** If using delayed initialization, how many frames are left before we call Initialize. Will be -1 if not actively counting down. */
+	int32 RemainingInitializationFrames;
 };
\ No newline at end of file
diff --git a/Engine/Plugins/MovieScene/MovieRenderPipeline/Source/MovieRenderPipelineCore/Public/MoviePipelineInProcessExecutorSettings.h b/Engine/Plugins/MovieScene/MovieRenderPipeline/Source/MovieRenderPipelineCore/Public/MoviePipelineInProcessExecutorSettings.h
new file mode 100644
--- /dev/null
+++ b/Engine/Plugins/MovieScene/MovieRenderPipeline/Source/MovieRenderPipelineCore/Public/MoviePipelineInProcessExecutorSettings.h
@@ -0,0 +1,60 @@
+// Copyright Epic Games, Inc. All Rights Reserved.
+#pragma once
+
+#include "Engine/DeveloperSettings.h"
+#include "MoviePipelineInProcessExecutorSettings.generated.h"
+
+/**
+* This is the implementation responsible for executing the rendering of
+* multiple movie pipelines after being launched via the command line.
+*/
+UCLASS(BlueprintType, config = Engine, defaultconfig, meta=(DisplayName = "Movie Pipeline New Process") )
+class MOVIERENDERPIPELINECORE_API UMoviePipelineInProcessExecutorSettings : public UDeveloperSettings
+{
+	GENERATED_BODY()
+	
+public:
+	UMoviePipelineInProcessExecutorSettings()
+		: Super()
+	{
+		bCloseEditor = false;
+		AdditionalCommandLineArguments = TEXT("-NoLoadingScreen -FixedSeed -log -Unattended");
+		InitialDelayFrameCount = 0;
+
+		// Find all arguments from the command line and set them as the InheritedCommandLineArguments.
+		TArray<FString> Tokens, Switches;
+		FCommandLine::Parse(FCommandLine::Get(), Tokens, Switches);
+		for (auto& Switch : Switches)
+		{
+			InheritedCommandLineArguments.AppendChar('-');
+			InheritedCommandLineArguments.Append(Switch);
+			InheritedCommandLineArguments.AppendChar(' ');
+		}
+	}
+	
+		/** Gets the settings container name for the settings, either Project or Editor */
+	virtual FName GetContainerName() const override { return FName("Project"); }
+	/** Gets the category for the settings, some high level grouping like, Editor, Engine, Game...etc. */
+	virtual FName GetCategoryName() const override { return FName("Plugins"); }
+
+	/** If enabled the editor will close itself when a new process is started. This can be used to gain some performance. */
+	UPROPERTY(EditAnywhere, BlueprintReadWrite, config, Category = "Startup")
+	bool bCloseEditor;
+
+	/** A list of additional command line arguments to be appended to the new process startup. In the form of "-foo -bar=baz". Can be useful if your game requires certain arguments to start such as disabling log-in screens. */
+	UPROPERTY(EditAnywhere, BlueprintReadWrite, config, Category = "Startup")
+	FString AdditionalCommandLineArguments;
+
+	/** A list of command line arguments which are inherited from the currently running Editor instance that will be automatically appended to the new process. */
+	UPROPERTY(Transient, VisibleAnywhere, BlueprintReadOnly, config, Category = "Startup")
+	FString InheritedCommandLineArguments;
+
+	/** 
+	* How long should we wait after being initialized to start doing any work? This can be used
+	* to work around situations where the game is not fully loaded by the time the pipeline
+	* is automatically started and it is important that the game is fully loaded before we do
+	* any work (such as evaluating frames for warm-up).
+	*/
+	UPROPERTY(EditAnywhere, BlueprintReadWrite, config, meta = (UIMin = "0", MinValue = "0", UIMax = "150"), Category = "Startup")
+	int32 InitialDelayFrameCount;
+};
\ No newline at end of file
diff --git a/Engine/Plugins/MovieScene/MovieRenderPipeline/Source/MovieRenderPipelineEditor/Private/MoviePipelineNewProcessExecutor.cpp b/Engine/Plugins/MovieScene/MovieRenderPipeline/Source/MovieRenderPipelineEditor/Private/MoviePipelineNewProcessExecutor.cpp
--- a/Engine/Plugins/MovieScene/MovieRenderPipeline/Source/MovieRenderPipelineEditor/Private/MoviePipelineNewProcessExecutor.cpp
+++ b/Engine/Plugins/MovieScene/MovieRenderPipeline/Source/MovieRenderPipelineEditor/Private/MoviePipelineNewProcessExecutor.cpp
@@ -7,7 +7,7 @@
 #include "Serialization/StructuredArchive.h"
 #include "MoviePipelineQueue.h"
 #include "ObjectTools.h"
-#include "MoviePipelineNewProcessExecutorSettings.h"
+#include "MoviePipelineInProcessExecutorSettings.h"
 #include "MovieRenderPipelineCoreModule.h"
 #include "FileHelpers.h"
 #include "MovieRenderPipelineSettings.h"
@@ -84,7 +84,7 @@ void UMoviePipelineNewProcessExecutor::ExecuteImpl(UMoviePipelineQueue* InPipeli
 	FString UnrealURLParams;
 
 	// Append all of our inherited command line arguments from the editor.
-	const UMoviePipelineNewProcessExecutorSettings* ExecutorSettings = GetDefault<UMoviePipelineNewProcessExecutorSettings>();
+	const UMoviePipelineInProcessExecutorSettings* ExecutorSettings = GetDefault<UMoviePipelineInProcessExecutorSettings>();
 
 	CommandLineArgs += ExecutorSettings->InheritedCommandLineArguments;
 	CommandLineArgs += TEXT(" ") + ExecutorSettings->AdditionalCommandLineArguments;
diff --git a/Engine/Plugins/MovieScene/MovieRenderPipeline/Source/MovieRenderPipelineEditor/Private/MoviePipelinePIEExecutor.cpp b/Engine/Plugins/MovieScene/MovieRenderPipeline/Source/MovieRenderPipelineEditor/Private/MoviePipelinePIEExecutor.cpp
--- a/Engine/Plugins/MovieScene/MovieRenderPipeline/Source/MovieRenderPipelineEditor/Private/MoviePipelinePIEExecutor.cpp
+++ b/Engine/Plugins/MovieScene/MovieRenderPipeline/Source/MovieRenderPipelineEditor/Private/MoviePipelinePIEExecutor.cpp
@@ -12,6 +12,8 @@
 #include "Editor.h"
 #include "Settings/LevelEditorPlaySettings.h"
 #include "MoviePipelineQueue.h"
+#include "MoviePipelineBlueprintLibrary.h"
+#include "MoviePipelinePIEExecutorSettings.h"
 
 #define LOCTEXT_NAMESPACE "MoviePipelinePIEExecutor"
 
@@ -21,7 +23,13 @@ FText UMoviePipelinePIEExecutor::GetWindowTitle(const int32 InConfigIndex, const
 	PercentFormatOptions.MinimumIntegralDigits = 2;
 	PercentFormatOptions.MaximumIntegralDigits = 3;
 
-	FText TitleFormatString = LOCTEXT("MoviePreviewWindowTitleFormat", "Movie Pipeline Render (Preview) [{CurrentCount}/{TotalCount} Total] {PercentComplete}% Completed.");
+	float CompletionPercentage = 0.f;
+	if (ActiveMoviePipeline)
+	{
+		CompletionPercentage = UMoviePipelineBlueprintLibrary::GetCompletionPercentage(ActiveMoviePipeline);
+	}
+
+	FText TitleFormatString = LOCTEXT("MoviePreviewWindowTitleFormat", "Movie Pipeline Render (Preview) [Job {CurrentCount}/{TotalCount} Total] Current Job: {PercentComplete}% Completed.");
 	FText WindowTitle = FText::FormatNamed(TitleFormatString, TEXT("CurrentCount"), FText::AsNumber(InConfigIndex + 1), TEXT("TotalCount"), FText::AsNumber(InNumConfigs), TEXT("PercentComplete"), FText::AsNumber(12.f, &PercentFormatOptions));
 	return WindowTitle;
 }
@@ -120,7 +128,20 @@ void UMoviePipelinePIEExecutor::OnPIEStartupFinished(bool)
 	// This Pipeline belongs to the world being created so that they have context for things they execute.
 	ActiveMoviePipeline = NewObject<UMoviePipeline>(ExecutingWorld, PipelineClass);
 	
-	ActiveMoviePipeline->Initialize(Queue->GetJobs()[CurrentPipelineIndex]);
+	// We allow users to set a multi-frame delay before we actually run the Initialization function and start thinking.
+	// This solves cases where there are engine systems that need to finish loading before we do anything.
+	const UMoviePipelinePIEExecutorSettings* ExecutorSettings = GetDefault<UMoviePipelinePIEExecutorSettings>();
+
+	if (ExecutorSettings->InitialDelayFrameCount == 0)
+	{
+		ActiveMoviePipeline->Initialize(Queue->GetJobs()[CurrentPipelineIndex]);
+		RemainingInitializationFrames = -1;
+	}
+	else
+	{
+		RemainingInitializationFrames = ExecutorSettings->InitialDelayFrameCount;
+		FCoreDelegates::OnBeginFrame.AddUObject(this, &UMoviePipelinePIEExecutor::DelayedInitializationCounter);
+	}
 
 	// Listen for when the pipeline thinks it has finished.
 	ActiveMoviePipeline->OnMoviePipelineFinished().AddUObject(this, &UMoviePipelinePIEExecutor::OnPIEMoviePipelineFinished);
@@ -130,6 +151,17 @@ void UMoviePipelinePIEExecutor::OnPIEStartupFinished(bool)
 	FEditorDelegates::EndPIE.AddUObject(this, &UMoviePipelinePIEExecutor::OnPIEEnded);
 }
 
+void UMoviePipelinePIEExecutor::DelayedInitializationCounter()
+{
+	if (RemainingInitializationFrames == 0)
+	{
+		ActiveMoviePipeline->Initialize(Queue->GetJobs()[CurrentPipelineIndex]);
+		FCoreDelegates::OnBeginFrame.RemoveAll(this);
+	}
+
+	RemainingInitializationFrames--;
+}
+
 void UMoviePipelinePIEExecutor::OnPIEMoviePipelineFinished(UMoviePipeline* InMoviePipeline)
 {
 	// Unsubscribe to the EndPIE event so we don't think the user canceled it.
@@ -151,9 +183,7 @@ void UMoviePipelinePIEExecutor::OnPIEMoviePipelineFinished(UMoviePipeline* InMov
 
 void UMoviePipelinePIEExecutor::OnPIEEnded(bool)
 {
-	// Restore the previous settings.
-	FApp::SetUseFixedTimeStep(bPreviousUseFixedTimeStep);
-	FApp::SetFixedDeltaTime(PreviousFixedTimeStepDelta);
+	FCoreDelegates::OnBeginFrame.RemoveAll(this);
 
 	// If the movie pipeline finishes naturally we unsubscribe from the EndPIE event.
 	// This means that if this gets called, the user has hit escape and wants to cancel.
@@ -166,6 +196,10 @@ void UMoviePipelinePIEExecutor::OnPIEEnded(bool)
 		// ToDo: bAnyJobHadFatalError
 		DelayedFinishNotification();
 	}
+	
+		// Restore the previous settings.
+	FApp::SetUseFixedTimeStep(bPreviousUseFixedTimeStep);
+	FApp::SetFixedDeltaTime(PreviousFixedTimeStepDelta);
 }
 
 void UMoviePipelinePIEExecutor::DelayedFinishNotification()
diff --git a/Engine/Plugins/MovieScene/MovieRenderPipeline/Source/MovieRenderPipelineEditor/Private/MoviePipelinePIEExecutorSettings.h b/Engine/Plugins/MovieScene/MovieRenderPipeline/Source/MovieRenderPipelineEditor/Private/MoviePipelinePIEExecutorSettings.h
new file mode 100644
--- /dev/null
+++ b/Engine/Plugins/MovieScene/MovieRenderPipeline/Source/MovieRenderPipelineEditor/Private/MoviePipelinePIEExecutorSettings.h
@@ -0,0 +1,30 @@
+// Copyright Epic Games, Inc. All Rights Reserved.
+#pragma once
+
+#include "Engine/DeveloperSettings.h"
+#include "MoviePipelinePIEExecutorSettings.generated.h"
+
+/**
+* This is the implementation responsible for executing the rendering of
+* multiple movie pipelines within the editor using PIE.
+*/
+UCLASS(BlueprintType, config = Editor, defaultconfig, meta = (DisplayName = "Movie Pipeline In Editor"))
+class MOVIERENDERPIPELINEEDITOR_API UMoviePipelinePIEExecutorSettings : public UDeveloperSettings
+{
+	GENERATED_BODY()
+
+public:
+	/** Gets the settings container name for the settings, either Project or Editor */
+	virtual FName GetContainerName() const override { return FName("Project"); }
+	/** Gets the category for the settings, some high level grouping like, Editor, Engine, Game...etc. */
+	virtual FName GetCategoryName() const override { return FName("Plugins"); }
+
+	/**
+	* How long should we wait after being initialized to start doing any work? This can be used
+	* to work around situations where the game is not fully loaded by the time the pipeline
+	* is automatically started and it is important that the game is fully loaded before we do
+	* any work (such as evaluating frames for warm-up).
+	*/
+	UPROPERTY(EditAnywhere, BlueprintReadWrite, config, meta = (UIMin = "0", MinValue = "0", UIMax = "150"), Category = "Startup")
+	int32 InitialDelayFrameCount;
+};
\ No newline at end of file
diff --git a/Engine/Plugins/MovieScene/MovieRenderPipeline/Source/MovieRenderPipelineEditor/Public/MoviePipelineNewProcessExecutorSettings.h b/Engine/Plugins/MovieScene/MovieRenderPipeline/Source/MovieRenderPipelineEditor/Public/MoviePipelineNewProcessExecutorSettings.h
deleted file mode 100644
diff --git a/Engine/Plugins/MovieScene/MovieRenderPipeline/Source/MovieRenderPipelineEditor/Public/MoviePipelinePIEExecutor.h b/Engine/Plugins/MovieScene/MovieRenderPipeline/Source/MovieRenderPipelineEditor/Public/MoviePipelinePIEExecutor.h
--- a/Engine/Plugins/MovieScene/MovieRenderPipeline/Source/MovieRenderPipelineEditor/Public/MoviePipelinePIEExecutor.h
+++ b/Engine/Plugins/MovieScene/MovieRenderPipeline/Source/MovieRenderPipelineEditor/Public/MoviePipelinePIEExecutor.h
@@ -20,6 +20,9 @@ class MOVIERENDERPIPELINEEDITOR_API UMoviePipelinePIEExecutor : public UMoviePip
 public:
 	UMoviePipelinePIEExecutor()
 		: UMoviePipelineLinearExecutorBase()
+		, RemainingInitializationFrames(-1)
+		, bPreviousUseFixedTimeStep(false)
+		, PreviousFixedTimeStepDelta(1/30.0)
 	{
 	}
 
@@ -29,6 +32,10 @@ protected:
 private:
 	/** Called when PIE finishes booting up and it is safe for us to spawn an object into that world. */
 	void OnPIEStartupFinished(bool);
+
+	/** If they're using delayed initialization, this is called each frame to process the countdown until start. */
+	void DelayedInitializationCounter();
+
 	/** Called before PIE tears down the world during shutdown. Used to detect cancel-via-escape/stop PIE. */
 	void OnPIEEnded(bool);
 	/** Called when the instance of the pipeline in the PIE world has finished. */
@@ -40,6 +47,8 @@ private:
 	FText GetWindowTitle(const int32 InConfigIndex, const int32 InNumConfigs) const;
 
 private:
+	/** If using delayed initialization, how many frames are left before we call Initialize. Will be -1 if not actively counting down. */
+	int32 RemainingInitializationFrames;
 	bool bPreviousUseFixedTimeStep;
 	double PreviousFixedTimeStepDelta;
 };
\ No newline at end of file
