MoviePipeline: CIS include fix.

diff --git a/Engine/Plugins/MovieScene/MovieRenderPipeline/Source/MovieRenderPipelineCore/Public/MoviePipelineInProcessExecutorSettings.h b/Engine/Plugins/MovieScene/MovieRenderPipeline/Source/MovieRenderPipelineCore/Public/MoviePipelineInProcessExecutorSettings.h
--- a/Engine/Plugins/MovieScene/MovieRenderPipeline/Source/MovieRenderPipelineCore/Public/MoviePipelineInProcessExecutorSettings.h
+++ b/Engine/Plugins/MovieScene/MovieRenderPipeline/Source/MovieRenderPipelineCore/Public/MoviePipelineInProcessExecutorSettings.h
@@ -2,6 +2,7 @@
 #pragma once
 
 #include "Engine/DeveloperSettings.h"
+#include "Misc/CommandLine.h"
 #include "MoviePipelineInProcessExecutorSettings.generated.h"
 
 /**
