MoviePipeline: Fixed the last camera cut in a shot getting skipped when rendering.

diff --git a/Engine/Plugins/MovieScene/MovieRenderPipeline/Source/MovieRenderPipelineCore/Public/MovieRenderPipelineDataTypes.h b/Engine/Plugins/MovieScene/MovieRenderPipeline/Source/MovieRenderPipelineCore/Public/MovieRenderPipelineDataTypes.h
--- a/Engine/Plugins/MovieScene/MovieRenderPipeline/Source/MovieRenderPipelineCore/Public/MovieRenderPipelineDataTypes.h
+++ b/Engine/Plugins/MovieScene/MovieRenderPipeline/Source/MovieRenderPipelineCore/Public/MovieRenderPipelineDataTypes.h
@@ -500,14 +500,14 @@ public:
 
 	bool SetNextShotActive()
 	{
-		if (CurrentCameraCutIndex + 1 < CameraCuts.Num() - 1)
+		if (CurrentCameraCutIndex == CameraCuts.Num() - 1)
 		{
-			CurrentCameraCutIndex++;
-			return false;
+			// We're on the last camera cut, we can't make another one active.
+			return true;
 		}
-
-		// We're on the last camera cut, we can't make another one active.
-		return true;
+		
+		CurrentCameraCutIndex++;
+		return false;
 	}
 
 	bool operator == (const FMoviePipelineShotInfo& InRHS) const
