MoviePipeline: Fixed canceling a job mid-job soft-locking the editor. This adds an async cancel (for eventual UI) and ensures that both normal finishes and force-finishes go through the same code flow to write all outstanding work to disk.

diff --git a/Engine/Plugins/MovieScene/MovieRenderPipeline/Source/MovieRenderPipelineCore/Private/MoviePipeline.cpp b/Engine/Plugins/MovieScene/MovieRenderPipeline/Source/MovieRenderPipelineCore/Private/MoviePipeline.cpp
--- a/Engine/Plugins/MovieScene/MovieRenderPipeline/Source/MovieRenderPipelineCore/Private/MoviePipeline.cpp
+++ b/Engine/Plugins/MovieScene/MovieRenderPipeline/Source/MovieRenderPipelineCore/Private/MoviePipeline.cpp
@@ -55,11 +55,16 @@ UMoviePipeline::UMoviePipeline()
 	, CachedPrevCustomTimeStep(nullptr)
 	, TargetSequence(nullptr)
 	, LevelSequenceActor(nullptr)
+	, DebugWidget(nullptr)
 	, PipelineState(EMovieRenderPipelineState::Uninitialized)
 	, CurrentShotIndex(-1)
+	, bPrevGScreenMessagesEnabled(true)
 	, bHasRunBeginFrameOnce(false)
 	, bPauseAtEndOfFrame(false)
+	, bShutdownRequested(false)
+	, bIsTransitioningState(false)
 	, AccumulatedTickSubFrameDeltas(0.f)
+	, CurrentJob(nullptr)
 {
 	CustomTimeStep = CreateDefaultSubobject<UMoviePipelineCustomTimeStep>("MoviePipelineCustomTimeStep");
 	CustomSequenceTimeController = MakeShared<FMoviePipelineTimeController>();
@@ -209,7 +214,7 @@ void UMoviePipeline::Initialize(UMoviePipelineExecutorJob* InJob)
 	SetupAudioRendering();
 
 
-	PipelineState = EMovieRenderPipelineState::ProducingFrames;
+	TransitionToState(EMovieRenderPipelineState::ProducingFrames);
 	CurrentShotIndex = 0;
 	CachedOutputState.ShotCount = ShotList.Num();
 	InitializationTime = FDateTime::UtcNow();
@@ -218,53 +223,188 @@ void UMoviePipeline::Initialize(UMoviePipelineExecutorJob* InJob)
 	// delta time not generated by us) so we'll wait until the next engine frame to start rendering.
 }
 
+
+void UMoviePipeline::RequestShutdown()
+{
+	// The user has requested a shutdown, it will be read the next available chance and possibly acted on.
+	bShutdownRequested = true;
+	switch (PipelineState)
+	{
+		// It is valid to call Shutdown at any point during these two states.
+	case EMovieRenderPipelineState::Uninitialized:
+	case EMovieRenderPipelineState::ProducingFrames:
+		break;
+		// You can call Shutdown during these two, but they won't do anything as we're already shutting down at that point.
+	case EMovieRenderPipelineState::Finalize:
+	case EMovieRenderPipelineState::Export:
+		UE_LOG(LogMovieRenderPipeline, Log, TEXT("[GFrameCounter: %d] Async Shutdown Requested, ignoring due to already being on the way to shutdown."), GFrameCounter);
+		break;
+	}
+}
+
 void UMoviePipeline::Shutdown()
 {
-	// Pipeline should not be shut down if already shut down. You probably have a delegate
-	// loop where we shutdown because finish and you responded by requesting we shutdown.
-	ensure(PipelineState != EMovieRenderPipelineState::Shutdown);
+	check(IsInGameThread());
 
-	// Uninitialize our master config settings.
-	for (UMoviePipelineSetting* Setting : GetPipelineMasterConfig()->GetAllSettings())
+	// This is a blocking operation which abandons any outstanding work to be submitted but finishes
+	// the existing work already processed. This is called at the end of the normal flow, or a blocking
+	// shutdown is requested (such as engine shutdown). If called at the end of the normal flow, most
+	// operations are a no-op.
+	if (PipelineState == EMovieRenderPipelineState::Uninitialized)
+	{
+		// If initialize is not called, no need to do anything.
+		return;
+	}
+	
+	if (PipelineState == EMovieRenderPipelineState::Finished)
 	{
-		Setting->OnMoviePipelineShutdown(this);
+		// If already shut down, no need to do anything.
+		return;
 	}
 
-	// Restore any custom Time Step that may have been set before.
-	GEngine->SetCustomTimeStep(CachedPrevCustomTimeStep);
+	if (PipelineState == EMovieRenderPipelineState::ProducingFrames)
+	{
 
-	// Ensure our delegates don't get called anymore as we're going to become null soon.
-	FCoreDelegates::OnBeginFrame.RemoveAll(this);
-	FCoreDelegates::OnEndFrame.RemoveAll(this);
+		// Teardown the currently active shot (if there is one). This will flush any outstanding rendering
+		// work that has already submitted - it cannot be canceled, so we may as well execute it and save the results.
+		TransitionToState(EMovieRenderPipelineState::Finalize);
 
-	// If we're not shutting down naturally (ie: we got canceled) then force the output
-	// containers to flush so that all in-flight write processes get finished. If we
-	// did shut down naturally, we should be in the export state.
-	if (PipelineState != EMovieRenderPipelineState::Export)
+		// Abandon the current frame. When using temporal sampling we may had canceled mid-frame, so the rendering
+		// commands were never submitted, thus the output builder will still be expecting a frame to come in.
+		if (CachedOutputState.TemporalSampleCount > 1)
+		{
+			OutputBuilder->AbandonOutstandingWork();
+		}
+	}
+
+	if (PipelineState == EMovieRenderPipelineState::Finalize)
 	{
-		const bool bForceFinish = true;
-		TickFinalizeOutputContainers(bForceFinish);
+		// We were either in the middle of writing frames to disk, or we have moved to Finalize as a result of the above block.
+		// Tick output containers until they report they have finished writing to disk. This is a blocking operation. 
+		// Finalize automatically switches our state to Export so no need to manually transition afterwards.
+		TickFinalizeOutputContainers(true);
 	}
 
-	if (DebugWidget)
+	if (PipelineState == EMovieRenderPipelineState::Export)
 	{
-		DebugWidget->RemoveFromParent();
-		DebugWidget = nullptr;
+		// All frames have been written to disk but we're doing a post-export step (such as encoding). Flush this operation as well.
+		// Export automatically switches our state to Finished so no need to manually transition afterwards.
+		TickPostFinalizeExport(true);
 	}
+}
 
-	for (UMoviePipelineOutputBase* Setting : GetPipelineMasterConfig()->GetOutputContainers())
+void UMoviePipeline::TransitionToState(const EMovieRenderPipelineState InNewState)
+{
+	// No re-entrancy. This isn't an error as tearing down a shot may try to move to
+	// Finalize on its own, but we don't want that.
+	if (bIsTransitioningState)
 	{
-		Setting->OnPipelineFinished();
+		return;
+	}
+
+	TGuardValue<bool> StateTransitionGuard(bIsTransitioningState, true);
+
+	bool bInvalidTransition = true;
+	switch (PipelineState)
+	{
+	case EMovieRenderPipelineState::Uninitialized:
+		PipelineState = InNewState;
+		bInvalidTransition = false;
+		break;
+	case EMovieRenderPipelineState::ProducingFrames:
+		if (InNewState == EMovieRenderPipelineState::Finalize)
+		{
+			bInvalidTransition = false;
+
+			// If we had naturally finished the last shot before doing this transition it will have
+			// already been torn down, so this only catches mid-shot transitions to ensure teardown.
+			if (CurrentShotIndex < ShotList.Num())
+			{
+				FMoviePipelineShotInfo& CurrentShot = ShotList[CurrentShotIndex];
+
+				// Ensures all in-flight work for that shot is handled.
+				TeardownShot(CurrentShot);
+			}
+
+			// Unregister our OnEngineTickEndFrame delegate. We don't unregister BeginFrame as we need
+			// to continue to call it to allow ticking the Finalization stage.
+			FCoreDelegates::OnEndFrame.RemoveAll(this);
+
+			// Reset the Custom Timestep because we don't care how long the engine takes now
+			GEngine->SetCustomTimeStep(CachedPrevCustomTimeStep);
+
+			// Ensure all frames have been processed by the GPU and sent to the Output Merger
+			FlushRenderingCommands();
+
+			// And then make sure all frames are sent to the Output Containers before we finalize.
+			ProcessOutstandingFinishedFrames();
+
+			// This is called once notifying output containers that all frames that will be submitted have been submitted.
+			PipelineState = EMovieRenderPipelineState::Finalize;
+			BeginFinalize();
+		}
+		break;
+	case EMovieRenderPipelineState::Finalize:
+		if (InNewState == EMovieRenderPipelineState::Export)
+		{
+			bInvalidTransition = false;
+
+			// This is called once notifying our export step that they can begin the export.
+			PipelineState = EMovieRenderPipelineState::Export;
+			BeginExport();
+		}
+		break;
+	case EMovieRenderPipelineState::Export:
+		if (InNewState == EMovieRenderPipelineState::Finished)
+		{
+			bInvalidTransition = false;
+			PipelineState = InNewState;
+
+			// Uninitialize our master config settings.
+			for (UMoviePipelineSetting* Setting : GetPipelineMasterConfig()->GetAllSettings())
+			{
+				Setting->OnMoviePipelineShutdown(this);
+			}
+
+			// Restore any custom Time Step that may have been set before.
+			GEngine->SetCustomTimeStep(CachedPrevCustomTimeStep);
+
+			// Ensure our delegates don't get called anymore as we're going to become null soon.
+			FCoreDelegates::OnBeginFrame.RemoveAll(this);
+			FCoreDelegates::OnEndFrame.RemoveAll(this);
+
+			if (DebugWidget)
+			{
+				DebugWidget->RemoveFromParent();
+				DebugWidget = nullptr;
+			}
+
+			for (UMoviePipelineOutputBase* Setting : GetPipelineMasterConfig()->GetOutputContainers())
+			{
+				Setting->OnPipelineFinished();
+			}
+
+			TeardownAudioRendering();
+
+			GAreScreenMessagesEnabled = bPrevGScreenMessagesEnabled;
+
+			UE_LOG(LogMovieRenderPipeline, Log, TEXT("Movie Pipeline completed. Duration: %s"), *(FDateTime::UtcNow() - InitializationTime).ToString());
+
+			OnMoviePipelineFinished().Broadcast(this);
+		}
+		break;
 	}
 
-	TeardownAudioRendering();
+	if (!ensureAlwaysMsgf(!bInvalidTransition, TEXT("[GFrameCounter: %d] An invalid transition was requested (from: %d to: %d), ignoring transition request."),
+		GFrameCounter, PipelineState, InNewState))
+	{
+		return;
+	}
 
-	GAreScreenMessagesEnabled = bPrevGScreenMessagesEnabled;
 
-	UE_LOG(LogMovieRenderPipeline, Log, TEXT("Movie Pipeline completed. Duration: %s"), *(FDateTime::UtcNow() - InitializationTime).ToString());
-	PipelineState = EMovieRenderPipelineState::Shutdown;
 }
 
+
 void UMoviePipeline::OnEngineTickBeginFrame()
 {
 	UE_LOG(LogMovieRenderPipeline, VeryVerbose, TEXT("OnEngineTickBeginFrame (Start) Engine Frame: [%d]"), GFrameCounter);
@@ -282,11 +422,12 @@ void UMoviePipeline::OnEngineTickBeginFrame()
 		TickProducingFrames();
 		break;
 	case EMovieRenderPipelineState::Finalize:
+		// Don't flush the finalize to keep the UI responsive.
 		TickFinalizeOutputContainers(false);
 		break;
 	case EMovieRenderPipelineState::Export:
-		// We support ticking export state as well to allow async long export steps.
-		TickPostFinalizeExport();
+		// Don't flush the export to keep the UI responsive.
+		TickPostFinalizeExport(false);
 		break;
 	}
 
@@ -340,7 +481,7 @@ void UMoviePipeline::OnEngineTickEndFrame()
 	UE_LOG(LogMovieRenderPipeline, VeryVerbose, TEXT("OnEngineTickEndFrame (End) Engine Frame: [%d]"), GFrameCounter);
 }
 
-bool UMoviePipeline::ProcessEndOfCameraCut(FMoviePipelineShotInfo &CurrentShot, FMoviePipelineCameraCutInfo &CurrentCameraCut)
+void UMoviePipeline::ProcessEndOfCameraCut(FMoviePipelineShotInfo &CurrentShot, FMoviePipelineCameraCutInfo &CurrentCameraCut)
 {
 	UE_LOG(LogMovieRenderPipeline, Log, TEXT("[%d] Finished processing Camera Cut [%d/%d] on Shot [%d] (%s) with stats [...]"), GFrameCounter,
 		CurrentShot.CurrentCameraCutIndex + 1, CurrentShot.CameraCuts.Num(),
@@ -362,41 +503,7 @@ bool UMoviePipeline::ProcessEndOfCameraCut(FMoviePipelineShotInfo &CurrentShot,
 	if (bWasLastCameraCut)
 	{
 		TeardownShot(CurrentShot);
-
-		CurrentShotIndex++;
-
-		// Notify our containers that the current shot has ended.
-		for (UMoviePipelineOutputBase* Container : GetPipelineMasterConfig()->GetOutputContainers())
-		{
-			Container->OnShotFinished(CurrentShot);
-		}
-	}
-
-	// Check to see if this was the last shot in the Pipeline, otherwise on the next
-	// tick the new shot will be initialized and processed.
-	if (CurrentShotIndex >= ShotList.Num())
-	{
-		// Ensure all frames have been processed by the GPU and sent to the Output Merger
-		FlushRenderingCommands();
-		
-		// And then make sure all frames are sent to the Output Containers before we finalize.
-		ProcessOutstandingFinishedFrames();
-
-		UE_LOG(LogMovieRenderPipeline, Log, TEXT("[%d] Finished rendering last shot. Moving to Finalize to finish writing items to disk."), GFrameCounter);
-		PipelineState = EMovieRenderPipelineState::Finalize;
-
-		// Unregister our OnEngineTickEndFrame delegate. We don't unregister BeginFrame as we need
-		// to continue to call it to allow ticking the Finalization stage.
-		FCoreDelegates::OnEndFrame.RemoveAll(this);
-
-		// Reset the Custom Timestep because we don't care how long the engine takes now
-		GEngine->SetCustomTimeStep(CachedPrevCustomTimeStep);
-
-		BeginFinalize();
-		return true;
 	}
-
-	return false;
 }
 
 void UMoviePipeline::BeginFinalize()
@@ -409,6 +516,11 @@ void UMoviePipeline::BeginFinalize()
 	}
 }
 
+void UMoviePipeline::BeginExport()
+{
+	// ToDo: Add exports.
+}
+
 void UMoviePipeline::TickFinalizeOutputContainers(const bool bInForceFinish)
 {
 	// Tick all containers until they all report that they have finalized.
@@ -450,13 +562,10 @@ void UMoviePipeline::TickFinalizeOutputContainers(const bool bInForceFinish)
 		Container->Finalize();
 	}
 
-	// Move onto the Export stage and call the Export function directly. Most export steps are
-	// instantaneous and will complete here, but on the next tick we will check them all anyways
-	// to allow running long-running export processes (such as encoding).
-	PipelineState = EMovieRenderPipelineState::Export;
+	TransitionToState(EMovieRenderPipelineState::Export);
 }
 
-void UMoviePipeline::TickPostFinalizeExport()
+void UMoviePipeline::TickPostFinalizeExport(const bool bInForceFinish)
 {
 	// This step assumes you have produced data and filled the data structures.
 	check(PipelineState == EMovieRenderPipelineState::Export);
@@ -464,18 +573,35 @@ void UMoviePipeline::TickPostFinalizeExport()
 
 	// ToDo: Loop through any extensions (such as XML export) and let them export using all of the
 	// data that was generated during this run such as containers, output names and lengths.
+	// Tick all containers until they all report that they have finalized.
+	bool bAllContainsFinishedProcessing;
+
+	do
+	{
+		bAllContainsFinishedProcessing = true;
 
+		// Ask the containers if they're all done processing.
+		// for (UMoviePipelineOutputBase* Container : GetPipelineMasterConfig()->GetOutputContainers())
+		// {
+		// 	bAllContainsFinishedProcessing &= Container->HasFinishedProcessing();
+		// }
 
+		// If we aren't forcing a finish, early out after one loop to keep
+		// the editor/ui responsive.
+		if (!bInForceFinish || bAllContainsFinishedProcessing)
+		{
+			break;
+		}
 
-	bool bAllExportsFinishedProcessing = true;
+		// If they've reached here, they're forcing them to finish so we'll sleep for a touch to give
+		// everyone a chance to actually do work before asking them if they're done.
+		FPlatformProcess::Sleep(1.f);
 
-	// Ask the containers if they're all done processing.
-	//for (UMoviePipelineOutput* Container : Config->OutputContainers)
-	//{
-	//	bAllContainsFinishedProcessing &= Container->HasFinishedProcessing();
-	//}
+	} while (true);
 
-	if (!bAllExportsFinishedProcessing)
+	// If an output container is still working, we'll early out to keep the UI responsive.
+	// If they've forced a finish this will have to be true before we can reach this block.
+	if (!bAllContainsFinishedProcessing)
 	{
 		return;
 	}
@@ -483,14 +609,7 @@ void UMoviePipeline::TickPostFinalizeExport()
 
 	UE_LOG(LogMovieRenderPipeline, Log, TEXT("[%d] PostFinalize Export (End)."), GFrameCounter);
 
-	// Now that we have done the last step of the process we're going to call Shutdown on ourself
-	// to unregister all delegates, and then we'll broadcast to anyone listening that we're done.
-	// This allows an external controller to decide what to do, ie: start another one, shut down
-	// PIE or exit the game.
-	Shutdown();
-
-	UE_LOG(LogMovieRenderPipeline, Verbose, TEXT("[%d] Broadcasting Post-Finalize Finished event."), GFrameCounter);
-	OnMoviePipelineFinished().Broadcast(this);
+	TransitionToState(EMovieRenderPipelineState::Finished);
 }
 
 bool UMoviePipelineCustomTimeStep::UpdateTimeStep(UEngine* /*InEngine*/)
@@ -814,6 +933,12 @@ void UMoviePipeline::TeardownShot(FMoviePipelineShotInfo& InShot)
 	// Start in InitializeShot() because we don't want to record samples during warm up/motion blur.
 	StopAudioRecording();
 
+	// Notify our containers that the current shot has ended.
+	for (UMoviePipelineOutputBase* Container : GetPipelineMasterConfig()->GetOutputContainers())
+	{
+		Container->OnShotFinished(InShot);
+	}
+
 	if (InShot.ShotOverrideConfig != nullptr)
 	{
 		// Any shot-specific overrides haven't had first time initialization. So we'll do that now.
@@ -823,8 +948,20 @@ void UMoviePipeline::TeardownShot(FMoviePipelineShotInfo& InShot)
 		}
 	}
 
-	// Setup required rendering architecture for all passes in this shot.
+	// Teardown any rendering architecture for this shot.
 	TeardownRenderingPipelineForShot(InShot);
+
+	CurrentShotIndex++;
+
+	// Check to see if this was the last shot in the Pipeline, otherwise on the next
+	// tick the new shot will be initialized and processed.
+	if (CurrentShotIndex >= ShotList.Num())
+	{
+
+
+		UE_LOG(LogMovieRenderPipeline, Log, TEXT("[%d] Finished rendering last shot. Moving to Finalize to finish writing items to disk."), GFrameCounter);
+		TransitionToState(EMovieRenderPipelineState::Finalize);
+	}
 }
 
 void UMoviePipeline::SetSoloShot(const FMoviePipelineShotInfo& InShot)
diff --git a/Engine/Plugins/MovieScene/MovieRenderPipeline/Source/MovieRenderPipelineCore/Private/MoviePipelineOutputBuilder.cpp b/Engine/Plugins/MovieScene/MovieRenderPipeline/Source/MovieRenderPipelineCore/Private/MoviePipelineOutputBuilder.cpp
--- a/Engine/Plugins/MovieScene/MovieRenderPipeline/Source/MovieRenderPipelineCore/Private/MoviePipelineOutputBuilder.cpp
+++ b/Engine/Plugins/MovieScene/MovieRenderPipeline/Source/MovieRenderPipelineCore/Private/MoviePipelineOutputBuilder.cpp
@@ -90,4 +90,11 @@ void FMoviePipelineOutputMerger::OnCompleteRenderPassDataAvailable_AnyThread(TUn
 		PendingData.RemoveAndCopyValue(InFrameData->OutputState, FinalFrame);
 		FinishedFrames.Enqueue(MoveTemp(FinalFrame));
 	}
+}
+
+void FMoviePipelineOutputMerger::AbandonOutstandingWork()
+{
+	FScopeLock ScopeLock(&ActiveDataMutex);
+	FinishedFrames.Empty();
+	PendingData.Empty();
 }
\ No newline at end of file
diff --git a/Engine/Plugins/MovieScene/MovieRenderPipeline/Source/MovieRenderPipelineCore/Private/MoviePipelineTiming.cpp b/Engine/Plugins/MovieScene/MovieRenderPipeline/Source/MovieRenderPipelineCore/Private/MoviePipelineTiming.cpp
--- a/Engine/Plugins/MovieScene/MovieRenderPipeline/Source/MovieRenderPipelineCore/Private/MoviePipelineTiming.cpp
+++ b/Engine/Plugins/MovieScene/MovieRenderPipeline/Source/MovieRenderPipelineCore/Private/MoviePipelineTiming.cpp
@@ -28,6 +28,13 @@ void UMoviePipeline::TickProducingFrames()
 	// We should not be calling this once we have completed all the shots.
 	check(CurrentShotIndex >= 0 && CurrentShotIndex < ShotList.Num());
 
+	if (bShutdownRequested)
+	{
+		UE_LOG(LogMovieRenderPipeline, Log, TEXT("[GFrameCounter: %d] Async Shutdown Requested, abandoning remaining work and moving to Finalize."), GFrameCounter);
+		TransitionToState(EMovieRenderPipelineState::Finalize);
+		return;
+	}
+
 	// When start up we want to override the engine's Custom Timestep with our own.
 	// This gives us the ability to completely control the engine tick/delta time before the frame
 	// is started so that we don't have to always be thinking of delta times one frame ahead. We need
@@ -499,10 +506,14 @@ void UMoviePipeline::TickProducingFrames()
 			// If this isn't the last shot, we'll immediately call this function again to just
 			// determine a new setup/seek/etc. on the same engine tick. Otherwise we would have
 			// used a random delta time that didn't actually correlate to anything.
-			if (!ProcessEndOfCameraCut(CurrentShot, CurrentCameraCut))
+			ProcessEndOfCameraCut(CurrentShot, CurrentCameraCut);
+			if (PipelineState == EMovieRenderPipelineState::ProducingFrames)
 			{
 				CachedOutputState.TemporalSampleIndex = -1;
 				TickProducingFrames();
+			}
+			else
+			{
 				return;
 			}
 		}
diff --git a/Engine/Plugins/MovieScene/MovieRenderPipeline/Source/MovieRenderPipelineCore/Public/MoviePipeline.h b/Engine/Plugins/MovieScene/MovieRenderPipeline/Source/MovieRenderPipelineCore/Public/MoviePipeline.h
--- a/Engine/Plugins/MovieScene/MovieRenderPipeline/Source/MovieRenderPipelineCore/Public/MoviePipeline.h
+++ b/Engine/Plugins/MovieScene/MovieRenderPipeline/Source/MovieRenderPipelineCore/Public/MoviePipeline.h
@@ -39,15 +39,28 @@ public:
 
 	/** 
 	* Initialize the movie pipeline with the specified settings. This kicks off the rendering process. 
-	
 	* @param InJob	- This contains settings and sequence to render this Movie Pipeline with.
 	*/
 	void Initialize(UMoviePipelineExecutorJob* InJob);
 
+
+	/**
+	* Request the movie pipeline to shut down at the next available time. The pipeline will attempt to abandon
+	* the current frame (such as if there are more temporal samples pending) but may be forced into finishing if
+	* there are spatial samples already submitted to the GPU. The shutdown flow will be run to ensure already
+	* completed work is written to disk. This is a non-blocking operation, use Shutdown() instead if you need to
+	* block until it is fully shut down.
+	*
+	* This function is thread safe.
+	*/
+	void RequestShutdown();
+	
 	/** 
-	* Call to shut down the pipeline. This flushes any outstanding file writes and unregisters all delegates.
-	* This will not call the OnMoviePipelineFinished() delegate as it assumes the external party already knows
-	* it is "Finished". 
+	* Abandons any future work on this Movie Pipeline and runs through the shutdown flow to ensure already
+	* completed work is written to disk. This is a blocking-operation and will not return until all outstanding
+	* work has been completed.
+	*
+	* This function should only be called from the game thread.
 	*/
 	void Shutdown();
 
@@ -128,22 +141,27 @@ private:
 	/** Runs the per-tick logic when doing the ProducingFrames state. */
 	void TickProducingFrames();
 
-	bool ProcessEndOfCameraCut(FMoviePipelineShotInfo &CurrentShot, FMoviePipelineCameraCutInfo &CurrentCameraCut);
+	void ProcessEndOfCameraCut(FMoviePipelineShotInfo &CurrentShot, FMoviePipelineCameraCutInfo &CurrentCameraCut);
 
 
 	/** Called once when first moving to the Finalize state. */
 	void BeginFinalize();
+	/** Called once when first moving to the Export state. */
+	void BeginExport();
 
 	/** 
 	* Runs the per-tick logic when doing the Finalize state.
+	*
 	* @param bInForceFinish		If true, this function will not return until all Output Containers say they have finalized.
 	*/
 	void TickFinalizeOutputContainers(const bool bInForceFinish);
 	/** 
 	* Runs the per-tick logic when doing the Export state. This is spread over multiple ticks to allow non-blocking background 
-	* processes (such as extra encoding) 
+	* processes (such as extra encoding).
+	*
+	* @param bInForceFinish		If true, this function will not return until all exports say they have finished.
 	*/
-	void TickPostFinalizeExport();
+	void TickPostFinalizeExport(const bool bInForceFinish);
 
 	/** Return true if we should early out of the TickProducingFrames function. Decrements the remaining number of steps when false. */
 	bool DebugFrameStepPreTick();
@@ -209,6 +227,8 @@ private:
 	/** It can be useful to know where the data we're generating was relative to the original Timeline, so this calculates that. */
 	void CalculateFrameNumbersForOutputState(const MoviePipeline::FFrameConstantMetrics& InFrameMetrics, const FMoviePipelineCameraCutInfo& InCameraCut, FMoviePipelineFrameOutputState& InOutOutputState) const;
 
+	/** Handles transitioning between states, preventing reentrancy. Normal state flow should be respected, does not handle arbitrary x to y transitions. */
+	void TransitionToState(const EMovieRenderPipelineState InNewState);
 private:
 	/** Custom TimeStep used to drive the engine while rendering. */
 	UPROPERTY(Transient, Instanced)
@@ -260,6 +280,12 @@ private:
 	/** Should we pause the game at the end of the frame? Used to implement frame step debugger. */
 	bool bPauseAtEndOfFrame;
 
+	/** True if RequestShutdown() was called. At the start of the next frame we will stop producing frames (if needed) and start shutting down. */
+	FThreadSafeBool bShutdownRequested;
+
+	/** True if we're in a TransitionToState call. Used to prevent reentrancy. */
+	bool bIsTransitioningState;
+
 	/** When using temporal sub-frame stepping common counts (such as 3) don't result in whole ticks. We keep track of how many ticks we lose so we can add them the next time there's a chance. */
 	float AccumulatedTickSubFrameDeltas;
 
diff --git a/Engine/Plugins/MovieScene/MovieRenderPipeline/Source/MovieRenderPipelineCore/Public/MoviePipelineOutputBuilder.h b/Engine/Plugins/MovieScene/MovieRenderPipeline/Source/MovieRenderPipelineCore/Public/MoviePipelineOutputBuilder.h
--- a/Engine/Plugins/MovieScene/MovieRenderPipeline/Source/MovieRenderPipelineCore/Public/MoviePipelineOutputBuilder.h
+++ b/Engine/Plugins/MovieScene/MovieRenderPipeline/Source/MovieRenderPipelineCore/Public/MoviePipelineOutputBuilder.h
@@ -22,7 +22,9 @@ public:
 	FMoviePipelineMergerOutputFrame& QueueOutputFrame_GameThread(const FMoviePipelineFrameOutputState& CachedOutputState);
 	void OnCompleteRenderPassDataAvailable_AnyThread(TUniquePtr<FImagePixelData>&& InData, const TSharedRef<FImagePixelDataPayload, ESPMode::ThreadSafe> InFrameData);
 	void OnSingleSampleDataAvailable_AnyThread(TUniquePtr<FImagePixelData>&& InData, const TSharedRef<FImagePixelDataPayload, ESPMode::ThreadSafe> InFrameData);
+	void AbandonOutstandingWork();
 
+public:
 	TQueue<FMoviePipelineMergerOutputFrame> FinishedFrames;
 private:
 	/** The Movie Pipeline that owns us. */
diff --git a/Engine/Plugins/MovieScene/MovieRenderPipeline/Source/MovieRenderPipelineCore/Public/MovieRenderPipelineDataTypes.h b/Engine/Plugins/MovieScene/MovieRenderPipeline/Source/MovieRenderPipelineCore/Public/MovieRenderPipelineDataTypes.h
--- a/Engine/Plugins/MovieScene/MovieRenderPipeline/Source/MovieRenderPipelineCore/Public/MovieRenderPipelineDataTypes.h
+++ b/Engine/Plugins/MovieScene/MovieRenderPipeline/Source/MovieRenderPipelineCore/Public/MovieRenderPipelineDataTypes.h
@@ -47,7 +47,7 @@ enum class EMovieRenderPipelineState : uint8
 	/** All outputs have finished writing to disk or otherwise processing. Additional exports that may have needed information about the produced file can now be run. */
 	Export = 3,
 	/** The pipeline has been shut down. It is an error to shut it down again. */
-	Shutdown = 4,
+	Finished = 4,
 };
 
 /**
