MoviePipeline: Moved the Handle Frame setting from the Anti-Aliasing section to the Output section.

#rnx

diff --git a/Engine/Plugins/MovieScene/MovieRenderPipeline/Source/MovieRenderPipelineCore/Private/MoviePipeline.cpp b/Engine/Plugins/MovieScene/MovieRenderPipeline/Source/MovieRenderPipelineCore/Private/MoviePipeline.cpp
--- a/Engine/Plugins/MovieScene/MovieRenderPipeline/Source/MovieRenderPipelineCore/Private/MoviePipeline.cpp
+++ b/Engine/Plugins/MovieScene/MovieRenderPipeline/Source/MovieRenderPipelineCore/Private/MoviePipeline.cpp
@@ -840,9 +840,10 @@ TArray<FMoviePipelineShotInfo> UMoviePipeline::BuildShotListFromSequence(const U
 		UE_LOG(LogMovieRenderPipeline, Log, TEXT("Shot %d/%d has %d Camera Cuts."), ShotIndex, NewShotList.Num(), Shot.CameraCuts.Num());
 		ShotIndex++;
 
+		UMoviePipelineOutputSetting* OutputSettings = FindOrAddSetting<UMoviePipelineOutputSetting>(Shot);
 		UMoviePipelineAntiAliasingSetting* AntiAliasingSettings = FindOrAddSetting<UMoviePipelineAntiAliasingSetting>(Shot);
 		UMoviePipelineHighResSetting* HighResSettings = FindOrAddSetting<UMoviePipelineHighResSetting>(Shot);
-		Shot.NumHandleFrames = AntiAliasingSettings->HandleFrameCount;
+		Shot.NumHandleFrames = OutputSettings->HandleFrameCount;
 		
 		// Expand the shot to encompass handle frames. This will modify our Camera Cuts bounds.
 		ExpandShot(Shot);
diff --git a/Engine/Plugins/MovieScene/MovieRenderPipeline/Source/MovieRenderPipelineCore/Private/MoviePipelineOutputSetting.cpp b/Engine/Plugins/MovieScene/MovieRenderPipeline/Source/MovieRenderPipelineCore/Private/MoviePipelineOutputSetting.cpp
--- a/Engine/Plugins/MovieScene/MovieRenderPipeline/Source/MovieRenderPipelineCore/Private/MoviePipelineOutputSetting.cpp
+++ b/Engine/Plugins/MovieScene/MovieRenderPipeline/Source/MovieRenderPipelineCore/Private/MoviePipelineOutputSetting.cpp
@@ -10,9 +10,10 @@ UMoviePipelineOutputSetting::UMoviePipelineOutputSetting()
 	: OutputResolution(FIntPoint(1920, 1080))
 	, bUseCustomFrameRate(false)
 	, OutputFrameRate(FFrameRate(24, 1))
-	, OutputFrameStep(1)
 	, DEBUG_OutputFrameStepOffset(0)
 	, bOverrideExistingOutput(true)
+	, HandleFrameCount(0)
+	, OutputFrameStep(1)
 	, bUseCustomPlaybackRange(false)
 	, CustomStartFrame(0)
 	, CustomEndFrame(0)
diff --git a/Engine/Plugins/MovieScene/MovieRenderPipeline/Source/MovieRenderPipelineCore/Public/MoviePipelineAntiAliasingSetting.h b/Engine/Plugins/MovieScene/MovieRenderPipeline/Source/MovieRenderPipelineCore/Public/MoviePipelineAntiAliasingSetting.h
--- a/Engine/Plugins/MovieScene/MovieRenderPipeline/Source/MovieRenderPipelineCore/Public/MoviePipelineAntiAliasingSetting.h
+++ b/Engine/Plugins/MovieScene/MovieRenderPipeline/Source/MovieRenderPipelineCore/Public/MoviePipelineAntiAliasingSetting.h
@@ -17,7 +17,6 @@ public:
 		, TemporalSampleCount(1)
 		, bOverrideAntiAliasing(false)
 		, AntiAliasingMethod(EAntiAliasingMethod::AAM_None)
-		, HandleFrameCount(0)
 		, RenderWarmUpCount(32)
 		, EngineWarmUpCount(0)
 		, AccumulationGamma(1.f)
@@ -84,9 +83,6 @@ public:
 	UPROPERTY(EditAnywhere, BlueprintReadWrite, meta = (EditCondition="bOverrideAntiAliasing"), Category = "Movie Pipeline")
 	TEnumAsByte<EAntiAliasingMethod> AntiAliasingMethod;
 
-	UPROPERTY(EditAnywhere, BlueprintReadWrite, meta = (UIMin = 0, ClampMin = 0), Category = "Movie Pipeline")
-	int32 HandleFrameCount;
-
 	/**
 	* The number of frames at the start of each shot that the engine will render and then discard. This is useful for
 	* ensuring there is history for temporal effects (such as anti-aliasing). It can be set to a lower number if not
diff --git a/Engine/Plugins/MovieScene/MovieRenderPipeline/Source/MovieRenderPipelineCore/Public/MoviePipelineOutputSetting.h b/Engine/Plugins/MovieScene/MovieRenderPipeline/Source/MovieRenderPipelineCore/Public/MoviePipelineOutputSetting.h
--- a/Engine/Plugins/MovieScene/MovieRenderPipeline/Source/MovieRenderPipelineCore/Public/MoviePipelineOutputSetting.h
+++ b/Engine/Plugins/MovieScene/MovieRenderPipeline/Source/MovieRenderPipelineCore/Public/MoviePipelineOutputSetting.h
@@ -43,10 +43,6 @@ public:
 	UPROPERTY(EditAnywhere, BlueprintReadWrite, meta=(EditCondition=bUseCustomFrameRate), Category = "File Output")
 	FFrameRate OutputFrameRate;
 	
-	/** Render every Nth frame. ie: Setting this value to 2 renders every other frame. Game Thread is still evaluated on 'skipped' frames for accuracy between renders of different OutputFrameSteps. */
-	UPROPERTY(EditAnywhere, BlueprintReadWrite, meta=(UIMin="1", MinValue="1"), Category = "File Output")
-	int32 OutputFrameStep;
-
 	/** Experimental way to interleave renders between multiple computers. If FrameStep is 5, and each render is offset by 0-4, they'll render a different subset of frames to add up to all. */
 	int32 DEBUG_OutputFrameStepOffset;
 	
@@ -55,16 +51,30 @@ public:
 	bool bOverrideExistingOutput;
 	
 public:
+
+	/**
+	* Top level shot track sections will automatically be expanded by this many frames in both directions, and the resulting
+	* additional time will be rendered as part of that shot. The inner sequence should have sections long enough to cover
+	* this expanded range, otherwise these tracks will not evaluate during handle frames and may produce unexpected results.
+	* This can be used to generate handle frames for traditional non linear editing tools.
+	*/
+	UPROPERTY(EditAnywhere, BlueprintReadWrite, meta = (UIMin = 0, ClampMin = 0), Category = "Frames")
+	int32 HandleFrameCount;
+
+	/** Render every Nth frame. ie: Setting this value to 2 renders every other frame. Game Thread is still evaluated on 'skipped' frames for accuracy between renders of different OutputFrameSteps. */
+	UPROPERTY(EditAnywhere, BlueprintReadWrite, meta = (UIMin = "1", MinValue = "1"), Category = "Frames")
+	int32 OutputFrameStep;
+
 	/** If true, override the Playback Range start/end bounds with the bounds specified below.*/
-	UPROPERTY(EditAnywhere, BlueprintReadWrite, Category = "File Output")
+	UPROPERTY(EditAnywhere, BlueprintReadWrite, Category = "Frames")
 	bool bUseCustomPlaybackRange;
 	
 	/** Used when overriding the playback range. In Display Rate frames. If bUseCustomPlaybackRange is false range will come from Sequence. */
-	UPROPERTY(EditAnywhere, BlueprintReadWrite, meta=(EditCondition=bUseCustomPlaybackRange), Category = "File Output")
+	UPROPERTY(EditAnywhere, BlueprintReadWrite, meta=(EditCondition=bUseCustomPlaybackRange), Category = "Frames")
 	int32 CustomStartFrame;
 	
 	/** Used when overriding the playback range. In Display Rate frames. If bUseCustomPlaybackRange is false range will come from Sequence. */
-	UPROPERTY(EditAnywhere, BlueprintReadWrite, meta=(EditCondition=bUseCustomPlaybackRange), Category = "File Output")
+	UPROPERTY(EditAnywhere, BlueprintReadWrite, meta=(EditCondition=bUseCustomPlaybackRange), Category = "Frames")
 	int32 CustomEndFrame;
 
 public:
@@ -80,10 +90,10 @@ public:
 	int32 FrameNumberOffset;
 	
 	/** If true the Filmic Tone Curve will not be applied. Disabling this will allow you to export linear data for EXRs. */
-	UPROPERTY(EditAnywhere, BlueprintReadWrite, AdvancedDisplay, Category = "File Output")
+	UPROPERTY(EditAnywhere, BlueprintReadWrite, AdvancedDisplay, Category = "Misc")
 	bool bDisableToneCurve;
 
 	/** Should the normal player viewport be visible when rendering? This increases render cost and is only an approximation of what is actually being rendered. */
-	UPROPERTY(EditAnywhere, BlueprintReadWrite, AdvancedDisplay, Category = "File Output")
+	UPROPERTY(EditAnywhere, BlueprintReadWrite, AdvancedDisplay, Category = "Misc")
 	bool bHidePreview;
 };
\ No newline at end of file
