MoviePipeline: The sequence is now only told to Play once playback actually starts for rendering in a shot, otherwise the sequence uses Jump. It also jumps from the Motion Blur frame to the first frame before calling Play. Some Sequencer tracks (such as audio or media) which avoid time synchronization handle Jumping specifically as a synchronization point.

diff --git a/Engine/Plugins/MovieScene/MovieRenderPipeline/Source/MovieRenderPipelineCore/Private/MoviePipelineTiming.cpp b/Engine/Plugins/MovieScene/MovieRenderPipeline/Source/MovieRenderPipelineCore/Private/MoviePipelineTiming.cpp
--- a/Engine/Plugins/MovieScene/MovieRenderPipeline/Source/MovieRenderPipelineCore/Private/MoviePipelineTiming.cpp
+++ b/Engine/Plugins/MovieScene/MovieRenderPipeline/Source/MovieRenderPipelineCore/Private/MoviePipelineTiming.cpp
@@ -145,14 +145,6 @@ void UMoviePipeline::TickProducingFrames()
 		return;
 	}
 
-	// At this point, from now on, the Sequence should be in the Playing state. MotionBlur fixes only
-	// last for one frame, and rendering expects it to be playing anyways. 
-	if (!LevelSequenceActor->GetSequencePlayer()->IsPlaying())
-	{
-		UE_LOG(LogMovieRenderPipeline, Verbose, TEXT("[%d] Setting Sequence Player to Play due to being paused when motion started."), GFrameCounter);
-		LevelSequenceActor->GetSequencePlayer()->Play();
-	}
-
 	// Do the MotionBlur State exit condition before the MotionBlur loop, to allow the correct MotionBlur state to
 	// persist through to the post-frame render before switching to the next state.
 	if (CurrentCameraCut.State == EMovieRenderShotState::MotionBlur && CurrentCameraCut.bHasEvaluatedMotionBlurFrame)
@@ -207,6 +199,7 @@ void UMoviePipeline::TickProducingFrames()
 		double FrameDeltaTime = FrameMetrics.TickResolution.AsSeconds(FrameMetrics.TicksPerSample);
 		CachedOutputState.TimeData.FrameDeltaTime = FrameDeltaTime;
 		CachedOutputState.TimeData.WorldSeconds = CachedOutputState.TimeData.WorldSeconds + FrameDeltaTime;
+		CustomTimeStep->SetCachedFrameTiming(MoviePipeline::FFrameTimeStepCache(FrameDeltaTime));
 
 		// We subtract a single tick here so that if you're trying to render one frame, and it is the last frame of the
 		// camera cut, if we jump forward a whole frame the evaluation will fail (since end frames are exclusive). To
@@ -217,7 +210,9 @@ void UMoviePipeline::TickProducingFrames()
 		// offsets taken into account. When we evaluate for this frame, take those into account.
 		FFrameTime FinalEvalTime = FrameMetrics.GetFinalEvalTime(CurrentCameraCut.CurrentMasterSeqTick + MotionBlurDurationTicks);
 
-		CustomTimeStep->SetCachedFrameTiming(MoviePipeline::FFrameTimeStepCache(FrameDeltaTime));
+		// Jump to the motion blur frame
+		FFrameTime TimeInPlayRate = FFrameRate::TransformTime(FinalEvalTime, FrameMetrics.TickResolution, TargetSequence->GetMovieScene()->GetDisplayRate());
+		LevelSequenceActor->GetSequencePlayer()->JumpToFrame(TimeInPlayRate);
 		CustomSequenceTimeController->SetCachedFrameTiming(FQualifiedFrameTime(FinalEvalTime, FrameMetrics.TickResolution));
 		
 		// We early out on this loop so that at least one tick passes for motion blur. We need to leave our state in
@@ -481,6 +476,18 @@ void UMoviePipeline::TickProducingFrames()
 			GFrameCounter, *LexToString(FinalEvalTime.FloorToFrame()),
 			CurrentCameraCut.CurrentMasterSeqTick.Value, FrameMetrics.ShutterOffsetTicks.GetFrame().Value, FrameMetrics.MotionBlurCenteringOffsetTicks.GetFrame().Value);
 
+
+		if (!LevelSequenceActor->GetSequencePlayer()->IsPlaying())
+		{
+			// If the level sequence actor isn't playing then this is the first frame that we're trying to render for the shot.
+			// We want to trigger a jump and then a play (and then set the SetCachedFrameTiming below as normal). The jump is
+			// important because we're going backwards in time from the motion blur frame, and some tracks need to be notified
+			// of a jump when time goes backwards.
+			FFrameTime TimeInPlayRate = FFrameRate::TransformTime(FinalEvalTime, FrameMetrics.TickResolution, TargetSequence->GetMovieScene()->GetDisplayRate());
+			LevelSequenceActor->GetSequencePlayer()->JumpToFrame(TimeInPlayRate);
+			LevelSequenceActor->GetSequencePlayer()->Play();
+		}
+
 		
 		// If the user doesn't want to render every frame we need to increase the delta time so that enough
 		// time passes in the world. The framerate of the output media will be adjusted to match.
@@ -571,4 +578,4 @@ void UMoviePipeline::CalculateFrameNumbersForOutputState(const MoviePipeline::FF
 		InOutOutputState.CurrentShotSourceFrameNumber = FFrameRate::TransformTime(CenteredTick, InFrameMetrics.TickResolution, SourceDisplayRate).RoundToFrame().Value;
 		InOutOutputState.CurrentShotSourceTimeCode = FTimecode::FromFrameNumber(InOutOutputState.CurrentShotSourceFrameNumber, InFrameMetrics.FrameRate, false);
 	}
-}
\ No newline at end of file
+}
