MoviePipeline: {camera_name} and {shot_name} now work in the Output Filename format string.

#rnx

diff --git a/Engine/Plugins/MovieScene/MovieRenderPipeline/Source/MovieRenderPipelineCore/Private/MoviePipeline.cpp b/Engine/Plugins/MovieScene/MovieRenderPipeline/Source/MovieRenderPipelineCore/Private/MoviePipeline.cpp
--- a/Engine/Plugins/MovieScene/MovieRenderPipeline/Source/MovieRenderPipelineCore/Private/MoviePipeline.cpp
+++ b/Engine/Plugins/MovieScene/MovieRenderPipeline/Source/MovieRenderPipelineCore/Private/MoviePipeline.cpp
@@ -873,7 +873,7 @@ TArray<FMoviePipelineShotInfo> UMoviePipeline::BuildShotListFromSequence(const U
 				const FMovieSceneObjectBindingID& CameraObjectBindingId = CameraCut.CameraCutSection->GetCameraBindingID();
 				if (CameraObjectBindingId.IsValid())
 				{
-					UMovieScene* OwningMovieScene = Cast<UMovieScene>(CameraCut.CameraCutSection->GetOuter());
+					UMovieScene* OwningMovieScene = CameraCut.CameraCutSection->GetTypedOuter<UMovieScene>();
 					if (OwningMovieScene)
 					{
 						FMovieSceneBinding* Binding = OwningMovieScene->FindBinding(CameraObjectBindingId.GetGuid());
diff --git a/Engine/Plugins/MovieScene/MovieRenderPipeline/Source/MovieRenderPipelineCore/Private/MoviePipelineTiming.cpp b/Engine/Plugins/MovieScene/MovieRenderPipeline/Source/MovieRenderPipelineCore/Private/MoviePipelineTiming.cpp
--- a/Engine/Plugins/MovieScene/MovieRenderPipeline/Source/MovieRenderPipelineCore/Private/MoviePipelineTiming.cpp
+++ b/Engine/Plugins/MovieScene/MovieRenderPipeline/Source/MovieRenderPipelineCore/Private/MoviePipelineTiming.cpp
@@ -578,4 +578,7 @@ void UMoviePipeline::CalculateFrameNumbersForOutputState(const MoviePipeline::FF
 		InOutOutputState.CurrentShotSourceFrameNumber = FFrameRate::TransformTime(CenteredTick, InFrameMetrics.TickResolution, SourceDisplayRate).RoundToFrame().Value;
 		InOutOutputState.CurrentShotSourceTimeCode = FTimecode::FromFrameNumber(InOutOutputState.CurrentShotSourceFrameNumber, InFrameMetrics.FrameRate, false);
 	}
+
+	InOutOutputState.ShotName = InCameraCut.ShotName;
+	InOutOutputState.CameraName = InCameraCut.CameraName;
 }
diff --git a/Engine/Plugins/MovieScene/MovieRenderPipeline/Source/MovieRenderPipelineCore/Public/MovieRenderPipelineDataTypes.cpp b/Engine/Plugins/MovieScene/MovieRenderPipeline/Source/MovieRenderPipelineCore/Public/MovieRenderPipelineDataTypes.cpp
--- a/Engine/Plugins/MovieScene/MovieRenderPipeline/Source/MovieRenderPipelineCore/Public/MovieRenderPipelineDataTypes.cpp
+++ b/Engine/Plugins/MovieScene/MovieRenderPipeline/Source/MovieRenderPipelineCore/Public/MovieRenderPipelineDataTypes.cpp
@@ -205,6 +205,6 @@ void FMoviePipelineFrameOutputState::GetFilenameFormatArguments(FMoviePipelineFo
 	InOutFormatArgs.Arguments.Add(TEXT("frame_number_shot"), FrameNumberShot);
 	InOutFormatArgs.Arguments.Add(TEXT("frame_number_rel"), FrameNumberRel);
 	InOutFormatArgs.Arguments.Add(TEXT("frame_number_shot_rel"), FrameNumberShotRel);
-	InOutFormatArgs.Arguments.Add(TEXT("camera_name"), TEXT("DefaultCameraName"));
-	InOutFormatArgs.Arguments.Add(TEXT("shot_name"), TEXT("DefaultShotName"));
+	InOutFormatArgs.Arguments.Add(TEXT("camera_name"), CameraName.Len() > 0 ? CameraName : TEXT("NoCamera"));
+	InOutFormatArgs.Arguments.Add(TEXT("shot_name"), ShotName.Len() > 0 ? ShotName : TEXT("NoShot"));
 }
diff --git a/Engine/Plugins/MovieScene/MovieRenderPipeline/Source/MovieRenderPipelineCore/Public/MovieRenderPipelineDataTypes.h b/Engine/Plugins/MovieScene/MovieRenderPipeline/Source/MovieRenderPipelineCore/Public/MovieRenderPipelineDataTypes.h
--- a/Engine/Plugins/MovieScene/MovieRenderPipeline/Source/MovieRenderPipelineCore/Public/MovieRenderPipelineDataTypes.h
+++ b/Engine/Plugins/MovieScene/MovieRenderPipeline/Source/MovieRenderPipelineCore/Public/MovieRenderPipelineDataTypes.h
@@ -624,6 +624,12 @@ public:
 	/** What time data should this frame use? Can vary between samples when TemporalSampleCount > 1. */
 	FTimeData TimeData;
 
+	/** The name of the currently active camera being rendered. May be empty. */
+	FString CameraName;
+
+	/** THe name of the currently active shot. May be empty if there is no shot track. */
+	FString ShotName;
+
 	/** INFORMATION BELOW HERE SHOULD NOT GET PERSISTED BETWEEN FRAMES */
 
 	void ResetPerFrameData()
@@ -645,6 +651,8 @@ public:
 		ShotSamplesRendered = 0;
 		ShotIndex = 0;
 		TemporalSampleCount = 0;
+		CameraName.Reset();
+		ShotName.Reset();
 	}
 
 	/**
