MoviePipeline: Alpha channel is now written to opaquely if alpha output is not supported. This solves the issue where pngs were completely transparent when the rendering pipeline was not set up to output alpha.

#rnx

diff --git a/Engine/Plugins/MovieScene/MovieRenderPipeline/Source/MovieRenderPipelineRenderPasses/Private/MoviePipelineImageSequenceOutput.cpp b/Engine/Plugins/MovieScene/MovieRenderPipeline/Source/MovieRenderPipelineRenderPasses/Private/MoviePipelineImageSequenceOutput.cpp
--- a/Engine/Plugins/MovieScene/MovieRenderPipeline/Source/MovieRenderPipelineRenderPasses/Private/MoviePipelineImageSequenceOutput.cpp
+++ b/Engine/Plugins/MovieScene/MovieRenderPipeline/Source/MovieRenderPipelineRenderPasses/Private/MoviePipelineImageSequenceOutput.cpp
@@ -166,6 +166,22 @@ void UMoviePipelineImageSequenceOutputBase::OnRecieveImageDataImpl(FMoviePipelin
 				break;
 			}
 		}
+		else
+		{
+			// Fill the alpha channel when alpha is not supported/enabled.
+			switch (QuantizedPixelData->GetType())
+			{
+			case EImagePixelType::Color:
+				TileImageTask->PixelPreProcessors.Add(TAsyncAlphaWrite<FColor>(255));
+				break;
+			case EImagePixelType::Float16:
+				TileImageTask->PixelPreProcessors.Add(TAsyncAlphaWrite<FFloat16Color>(1.f));
+				break;
+			case EImagePixelType::Float32:
+				TileImageTask->PixelPreProcessors.Add(TAsyncAlphaWrite<FLinearColor>(1.f));
+				break;
+			}
+		}
 
 		TileImageTask->PixelData = MoveTemp(QuantizedPixelData);
 		ImageWriteQueue->Enqueue(MoveTemp(TileImageTask));
