MoviePipeline: Fixed Movie Render Queue showing negative frames when using engine warmup time and fixed window titles not being updated with progress.

diff --git a/Engine/Plugins/MovieScene/MovieRenderPipeline/Source/MovieRenderPipelineCore/Private/MoviePipelineBlueprintLibrary.cpp b/Engine/Plugins/MovieScene/MovieRenderPipeline/Source/MovieRenderPipelineCore/Private/MoviePipelineBlueprintLibrary.cpp
--- a/Engine/Plugins/MovieScene/MovieRenderPipeline/Source/MovieRenderPipelineCore/Private/MoviePipelineBlueprintLibrary.cpp
+++ b/Engine/Plugins/MovieScene/MovieRenderPipeline/Source/MovieRenderPipelineCore/Private/MoviePipelineBlueprintLibrary.cpp
@@ -196,7 +196,7 @@ bool UMoviePipelineBlueprintLibrary::GetEstimatedTimeRemaining(const UMoviePipel
 	int32 TotalOutputFrames;
 	GetOverallOutputFrames(InPipeline, OutputFrames, TotalOutputFrames);
 
-	if (OutputFrames == 0 || TotalOutputFrames == 0)
+	if (OutputFrames <= 0 || TotalOutputFrames <= 0)
 	{
 		OutEstimate = FTimespan();
 		return false;
@@ -221,14 +221,12 @@ float UMoviePipelineBlueprintLibrary::GetCompletionPercentage(const UMoviePipeli
 		return 0.f;
 	}
 
-	// Look at how many total samples we expect across all shots. This includes
-	// samples produced for warm-ups, motion blur fixes, and temporal/spatial samples.	
-	//int32 TotalExpectedSamples = GetTotalSampleCount(InPipeline->GetTotalWorkEstimate());
-	//int32 TotalCompletedSamples = GetTotalSampleCount(InPipeline->GetCompletedWork());
-	//
-	//float CompletionPercentage = FMath::Clamp(TotalCompletedSamples / (float)TotalExpectedSamples, 0.f, 1.f);
-	//return CompletionPercentage;
-	return 0.f;
+	int32 OutputFrames;
+	int32 TotalOutputFrames;
+	GetOverallOutputFrames(InPipeline, OutputFrames, TotalOutputFrames);
+
+	float CompletionPercentage = FMath::Clamp(OutputFrames / (float)TotalOutputFrames, 0.f, 1.f);
+	return CompletionPercentage;
 }
 
 FTimecode UMoviePipelineBlueprintLibrary::GetMasterTimecode(const UMoviePipeline* InMoviePipeline)
diff --git a/Engine/Plugins/MovieScene/MovieRenderPipeline/Source/MovieRenderPipelineCore/Private/MoviePipelineInProcessExecutor.cpp b/Engine/Plugins/MovieScene/MovieRenderPipeline/Source/MovieRenderPipelineCore/Private/MoviePipelineInProcessExecutor.cpp
--- a/Engine/Plugins/MovieScene/MovieRenderPipeline/Source/MovieRenderPipelineCore/Private/MoviePipelineInProcessExecutor.cpp
+++ b/Engine/Plugins/MovieScene/MovieRenderPipeline/Source/MovieRenderPipelineCore/Private/MoviePipelineInProcessExecutor.cpp
@@ -3,9 +3,8 @@
 #include "MoviePipeline.h"
 #include "Engine/World.h"
 #include "MoviePipelineQueue.h"
-#include "Kismet/KismetSystemLibrary.h"
-#include "MoviePipelineBlueprintLibrary.h"
 #include "MoviePipelineInProcessExecutorSettings.h"
+#include "Kismet/KismetSystemLibrary.h"
 
 #define LOCTEXT_NAMESPACE "MoviePipelineInProcessExecutor"
 void UMoviePipelineInProcessExecutor::Start(const UMoviePipelineExecutorJob* InJob)
@@ -22,24 +21,6 @@ void UMoviePipelineInProcessExecutor::Start(const UMoviePipelineExecutorJob* InJ
 	}
 }
 
-void UMoviePipelineInProcessExecutor::UpdateWindowTitle()
-{
-	FNumberFormattingOptions PercentFormatOptions;
-	PercentFormatOptions.MinimumIntegralDigits = 2;
-	PercentFormatOptions.MaximumIntegralDigits = 3;
-
-	float CompletionPercentage = 0.f;
-	if (ActiveMoviePipeline)
-	{
-		CompletionPercentage = UMoviePipelineBlueprintLibrary::GetCompletionPercentage(ActiveMoviePipeline);
-	}
-
-	FText TitleFormatString = LOCTEXT("MoviePreviewWindowTitleFormat", "Movie Pipeline Render (Preview) [Job {CurrentCount}/{TotalCount} Total] Current Job: {PercentComplete}% Completed.");
-	FText WindowTitle = FText::FormatNamed(TitleFormatString, TEXT("CurrentCount"), FText::AsNumber(CurrentPipelineIndex + 1), TEXT("TotalCount"), FText::AsNumber(Queue->GetJobs().Num()), TEXT("PercentComplete"), FText::AsNumber(12.f, &PercentFormatOptions));
-
-	UKismetSystemLibrary::SetWindowTitle(WindowTitle);
-}
-
 void UMoviePipelineInProcessExecutor::OnMapLoadFinished(UWorld* NewWorld)
 {
 	// NewWorld can be null if a world is being destroyed.
@@ -60,6 +41,15 @@ void UMoviePipelineInProcessExecutor::OnMapLoadFinished(UWorld* NewWorld)
 	// This solves cases where there are engine systems that need to finish loading before we do anything.
 	const UMoviePipelineInProcessExecutorSettings* ExecutorSettings = GetDefault<UMoviePipelineInProcessExecutorSettings>();
 
+	// We tick each frame to update the Window Title, and kick off latent pipeling initialization.
+	FCoreDelegates::OnBeginFrame.AddUObject(this, &UMoviePipelineInProcessExecutor::OnTick);
+	
+	// Listen for when the pipeline thinks it has finished.
+	ActiveMoviePipeline->OnMoviePipelineFinished().AddUObject(this, &UMoviePipelineInProcessExecutor::OnMoviePipelineFinished);
+
+	// Wait until we actually recieved the right map and created the pipeline before saying that we're actively rendering
+	bIsRendering = true;
+	
 	if (ExecutorSettings->InitialDelayFrameCount == 0)
 	{
 		ActiveMoviePipeline->Initialize(Queue->GetJobs()[CurrentPipelineIndex]);
@@ -67,32 +57,29 @@ void UMoviePipelineInProcessExecutor::OnMapLoadFinished(UWorld* NewWorld)
 	}
 	else
 	{
-		RemainingInitializationFrames = ExecutorSettings->InitialDelayFrameCount;
-		FCoreDelegates::OnBeginFrame.AddUObject(this, &UMoviePipelineInProcessExecutor::DelayedInitializationCounter);
+		RemainingInitializationFrames = ExecutorSettings->InitialDelayFrameCount;	
 	}
-
-	UpdateWindowTitle();
-
-	// Listen for when the pipeline thinks it has finished.
-	ActiveMoviePipeline->OnMoviePipelineFinished().AddUObject(this, &UMoviePipelineInProcessExecutor::OnMoviePipelineFinished);
-
-	// Wait until we actually recieved the right map and created the pipeline before saying that we're actively rendering
-	bIsRendering = true;
 }
 
-void UMoviePipelineInProcessExecutor::DelayedInitializationCounter()
+void UMoviePipelineInProcessExecutor::OnTick()
 {
-	if (RemainingInitializationFrames == 0)
+	if (RemainingInitializationFrames >= 0)
 	{
-		ActiveMoviePipeline->Initialize(Queue->GetJobs()[CurrentPipelineIndex]);
-		FCoreDelegates::OnBeginFrame.RemoveAll(this);
+		if (RemainingInitializationFrames == 0)
+		{
+			ActiveMoviePipeline->Initialize(Queue->GetJobs()[CurrentPipelineIndex]);
+		}
+
+		RemainingInitializationFrames--;
 	}
 
-	RemainingInitializationFrames--;
+	FText WindowTitle = GetWindowTitle();
+	UKismetSystemLibrary::SetWindowTitle(WindowTitle);
 }
 
 void UMoviePipelineInProcessExecutor::OnMoviePipelineFinished(UMoviePipeline* InMoviePipeline)
 {
+	FCoreDelegates::OnBeginFrame.RemoveAll(this);
 	UMoviePipeline* MoviePipeline = ActiveMoviePipeline;
 
 	if (ActiveMoviePipeline)
diff --git a/Engine/Plugins/MovieScene/MovieRenderPipeline/Source/MovieRenderPipelineCore/Private/MoviePipelineLinearExecutor.cpp b/Engine/Plugins/MovieScene/MovieRenderPipeline/Source/MovieRenderPipelineCore/Private/MoviePipelineLinearExecutor.cpp
--- a/Engine/Plugins/MovieScene/MovieRenderPipeline/Source/MovieRenderPipelineCore/Private/MoviePipelineLinearExecutor.cpp
+++ b/Engine/Plugins/MovieScene/MovieRenderPipeline/Source/MovieRenderPipelineCore/Private/MoviePipelineLinearExecutor.cpp
@@ -2,6 +2,7 @@
 #include "MoviePipelineLinearExecutor.h"
 #include "MovieRenderPipelineCoreModule.h"
 #include "MoviePipelineQueue.h"
+#include "MoviePipelineBlueprintLibrary.h"
 
 #define LOCTEXT_NAMESPACE "MoviePipelineLinearExecutorBase"
 
@@ -72,4 +73,23 @@ void UMoviePipelineLinearExecutorBase::OnExecutorFinishedImpl()
 	Super::OnExecutorFinishedImpl();
 }
 
+
+FText UMoviePipelineLinearExecutorBase::GetWindowTitle()
+{
+	FNumberFormattingOptions PercentFormatOptions;
+	PercentFormatOptions.MinimumIntegralDigits = 1;
+	PercentFormatOptions.MaximumIntegralDigits = 3;
+	PercentFormatOptions.MaximumFractionalDigits = 0;
+	PercentFormatOptions.RoundingMode = ERoundingMode::HalfFromZero;
+
+	float CompletionPercentage = 0.f;
+	if (ActiveMoviePipeline)
+	{
+		CompletionPercentage = UMoviePipelineBlueprintLibrary::GetCompletionPercentage(ActiveMoviePipeline) * 100.f;
+	}
+
+	FText TitleFormatString = LOCTEXT("MoviePreviewWindowTitleFormat", "Movie Pipeline Render (Preview) [Job {CurrentCount}/{TotalCount} Total] Current Job: {PercentComplete}% Completed.");
+	return FText::FormatNamed(TitleFormatString, TEXT("CurrentCount"), FText::AsNumber(CurrentPipelineIndex + 1), TEXT("TotalCount"), FText::AsNumber(Queue->GetJobs().Num()), TEXT("PercentComplete"), FText::AsNumber(CompletionPercentage, &PercentFormatOptions));
+}
+
 #undef LOCTEXT_NAMESPACE // "MoviePipelineLinearExecutorBase"
diff --git a/Engine/Plugins/MovieScene/MovieRenderPipeline/Source/MovieRenderPipelineCore/Public/MoviePipelineInProcessExecutor.h b/Engine/Plugins/MovieScene/MovieRenderPipeline/Source/MovieRenderPipelineCore/Public/MoviePipelineInProcessExecutor.h
--- a/Engine/Plugins/MovieScene/MovieRenderPipeline/Source/MovieRenderPipelineCore/Public/MoviePipelineInProcessExecutor.h
+++ b/Engine/Plugins/MovieScene/MovieRenderPipeline/Source/MovieRenderPipelineCore/Public/MoviePipelineInProcessExecutor.h
@@ -25,12 +25,10 @@ public:
 
 protected:
 	virtual void Start(const UMoviePipelineExecutorJob* InJob) override;
-
-	virtual void UpdateWindowTitle();
 private:
 	void OnMapLoadFinished(UWorld* NewWorld);
 	void OnMoviePipelineFinished(UMoviePipeline* InMoviePipeline);
-	void DelayedInitializationCounter();
+	void OnTick();
 
 private:
 	/** If using delayed initialization, how many frames are left before we call Initialize. Will be -1 if not actively counting down. */
diff --git a/Engine/Plugins/MovieScene/MovieRenderPipeline/Source/MovieRenderPipelineCore/Public/MoviePipelineLinearExecutor.h b/Engine/Plugins/MovieScene/MovieRenderPipeline/Source/MovieRenderPipelineCore/Public/MoviePipelineLinearExecutor.h
--- a/Engine/Plugins/MovieScene/MovieRenderPipeline/Source/MovieRenderPipelineCore/Public/MoviePipelineLinearExecutor.h
+++ b/Engine/Plugins/MovieScene/MovieRenderPipeline/Source/MovieRenderPipelineCore/Public/MoviePipelineLinearExecutor.h
@@ -35,7 +35,7 @@ protected:
 
 	virtual void StartPipelineByIndex(int32 InPipelineIndex);
 	virtual void Start(const UMoviePipelineExecutorJob* InJob) {}
-
+	virtual FText GetWindowTitle();
 public:
 	virtual void OnIndividualPipelineFinished(UMoviePipeline* /* FinishedPipeline */);
 	virtual void OnExecutorFinishedImpl();
diff --git a/Engine/Plugins/MovieScene/MovieRenderPipeline/Source/MovieRenderPipelineEditor/Private/MoviePipelinePIEExecutor.cpp b/Engine/Plugins/MovieScene/MovieRenderPipeline/Source/MovieRenderPipelineEditor/Private/MoviePipelinePIEExecutor.cpp
--- a/Engine/Plugins/MovieScene/MovieRenderPipeline/Source/MovieRenderPipelineEditor/Private/MoviePipelinePIEExecutor.cpp
+++ b/Engine/Plugins/MovieScene/MovieRenderPipeline/Source/MovieRenderPipelineEditor/Private/MoviePipelinePIEExecutor.cpp
@@ -12,27 +12,10 @@
 #include "Editor.h"
 #include "Settings/LevelEditorPlaySettings.h"
 #include "MoviePipelineQueue.h"
-#include "MoviePipelineBlueprintLibrary.h"
 #include "MoviePipelinePIEExecutorSettings.h"
 
 #define LOCTEXT_NAMESPACE "MoviePipelinePIEExecutor"
 
-FText UMoviePipelinePIEExecutor::GetWindowTitle(const int32 InConfigIndex, const int32 InNumConfigs) const
-{
-	FNumberFormattingOptions PercentFormatOptions;
-	PercentFormatOptions.MinimumIntegralDigits = 2;
-	PercentFormatOptions.MaximumIntegralDigits = 3;
-
-	float CompletionPercentage = 0.f;
-	if (ActiveMoviePipeline)
-	{
-		CompletionPercentage = UMoviePipelineBlueprintLibrary::GetCompletionPercentage(ActiveMoviePipeline);
-	}
-
-	FText TitleFormatString = LOCTEXT("MoviePreviewWindowTitleFormat", "Movie Pipeline Render (Preview) [Job {CurrentCount}/{TotalCount} Total] Current Job: {PercentComplete}% Completed.");
-	FText WindowTitle = FText::FormatNamed(TitleFormatString, TEXT("CurrentCount"), FText::AsNumber(InConfigIndex + 1), TEXT("TotalCount"), FText::AsNumber(InNumConfigs), TEXT("PercentComplete"), FText::AsNumber(12.f, &PercentFormatOptions));
-	return WindowTitle;
-}
 
 void UMoviePipelinePIEExecutor::Start(const UMoviePipelineExecutorJob* InJob)
 {
@@ -40,7 +23,6 @@ void UMoviePipelinePIEExecutor::Start(const UMoviePipelineExecutorJob* InJob)
 
 	// Create a Slate window to hold our preview.
 	TSharedRef<SWindow> CustomWindow = SNew(SWindow)
-		.Title_UObject(this, &UMoviePipelinePIEExecutor::GetWindowTitle, CurrentPipelineIndex, Queue->GetJobs().Num())
 		.ClientSize(FVector2D(1280, 720))
 		.AutoCenter(EAutoCenter::PrimaryWorkArea)
 		.UseOSWindowBorder(true)
@@ -51,6 +33,7 @@ void UMoviePipelinePIEExecutor::Start(const UMoviePipelineExecutorJob* InJob)
 		.SupportsMinimize(true)
 		.SizingRule(ESizingRule::UserSized);
 
+	WeakCustomWindow = CustomWindow;
 	FSlateApplication::Get().AddWindow(CustomWindow);
 
 	// Initialize our own copy of the Editor Play settings which we will adjust defaults on.
@@ -132,6 +115,13 @@ void UMoviePipelinePIEExecutor::OnPIEStartupFinished(bool)
 	// This solves cases where there are engine systems that need to finish loading before we do anything.
 	const UMoviePipelinePIEExecutorSettings* ExecutorSettings = GetDefault<UMoviePipelinePIEExecutorSettings>();
 
+	// We tick each frame to update the Window Title, and kick off latent pipeling initialization.
+	FCoreDelegates::OnBeginFrame.AddUObject(this, &UMoviePipelinePIEExecutor::OnTick);
+
+	// Listen for when the pipeline thinks it has finished.
+	ActiveMoviePipeline->OnMoviePipelineFinished().AddUObject(this, &UMoviePipelinePIEExecutor::OnPIEMoviePipelineFinished);
+	ActiveMoviePipeline->OnMoviePipelineErrored().AddUObject(this, &UMoviePipelinePIEExecutor::OnPipelineErrored);
+	
 	if (ExecutorSettings->InitialDelayFrameCount == 0)
 	{
 		ActiveMoviePipeline->Initialize(Queue->GetJobs()[CurrentPipelineIndex]);
@@ -140,32 +130,37 @@ void UMoviePipelinePIEExecutor::OnPIEStartupFinished(bool)
 	else
 	{
 		RemainingInitializationFrames = ExecutorSettings->InitialDelayFrameCount;
-		FCoreDelegates::OnBeginFrame.AddUObject(this, &UMoviePipelinePIEExecutor::DelayedInitializationCounter);
 	}
-
-	// Listen for when the pipeline thinks it has finished.
-	ActiveMoviePipeline->OnMoviePipelineFinished().AddUObject(this, &UMoviePipelinePIEExecutor::OnPIEMoviePipelineFinished);
-	ActiveMoviePipeline->OnMoviePipelineErrored().AddUObject(this, &UMoviePipelinePIEExecutor::OnPipelineErrored);
-
+	
 	// Listen for PIE shutdown in case the user hits escape to close it. 
 	FEditorDelegates::EndPIE.AddUObject(this, &UMoviePipelinePIEExecutor::OnPIEEnded);
 }
 
-void UMoviePipelinePIEExecutor::DelayedInitializationCounter()
+void UMoviePipelinePIEExecutor::OnTick()
 {
-	if (RemainingInitializationFrames == 0)
+	if (RemainingInitializationFrames >= 0)
 	{
-		ActiveMoviePipeline->Initialize(Queue->GetJobs()[CurrentPipelineIndex]);
-		FCoreDelegates::OnBeginFrame.RemoveAll(this);
+		if (RemainingInitializationFrames == 0)
+		{
+			ActiveMoviePipeline->Initialize(Queue->GetJobs()[CurrentPipelineIndex]);
+		}
+
+		RemainingInitializationFrames--;
 	}
 
-	RemainingInitializationFrames--;
+	FText WindowTitle = GetWindowTitle();
+	TSharedPtr<SWindow> CustomWindow = WeakCustomWindow.Pin();
+	if (CustomWindow)
+	{
+		CustomWindow->SetTitle(WindowTitle);
+	}
 }
 
 void UMoviePipelinePIEExecutor::OnPIEMoviePipelineFinished(UMoviePipeline* InMoviePipeline)
 {
 	// Unsubscribe to the EndPIE event so we don't think the user canceled it.
 	FEditorDelegates::EndPIE.RemoveAll(this);
+	FCoreDelegates::OnBeginFrame.RemoveAll(this);
 
 	if (ActiveMoviePipeline)
 	{
@@ -214,4 +209,4 @@ void UMoviePipelinePIEExecutor::DelayedFinishNotification()
 	OnIndividualPipelineFinished(MoviePipeline);
 }
 
-#undef LOCTEXT_NAMESPACE // "MoviePipelinePIEExecutor"
+#undef LOCTEXT_NAMESPACE // "MoviePipelinePIEExecutor"
\ No newline at end of file
diff --git a/Engine/Plugins/MovieScene/MovieRenderPipeline/Source/MovieRenderPipelineEditor/Public/MoviePipelinePIEExecutor.h b/Engine/Plugins/MovieScene/MovieRenderPipeline/Source/MovieRenderPipelineEditor/Public/MoviePipelinePIEExecutor.h
--- a/Engine/Plugins/MovieScene/MovieRenderPipeline/Source/MovieRenderPipelineEditor/Public/MoviePipelinePIEExecutor.h
+++ b/Engine/Plugins/MovieScene/MovieRenderPipeline/Source/MovieRenderPipelineEditor/Public/MoviePipelinePIEExecutor.h
@@ -33,8 +33,8 @@ private:
 	/** Called when PIE finishes booting up and it is safe for us to spawn an object into that world. */
 	void OnPIEStartupFinished(bool);
 
-	/** If they're using delayed initialization, this is called each frame to process the countdown until start. */
-	void DelayedInitializationCounter();
+	/** If they're using delayed initialization, this is called each frame to process the countdown until start, also updates Window Title each frame. */
+	void OnTick();
 
 	/** Called before PIE tears down the world during shutdown. Used to detect cancel-via-escape/stop PIE. */
 	void OnPIEEnded(bool);
@@ -42,13 +42,10 @@ private:
 	void OnPIEMoviePipelineFinished(UMoviePipeline* InMoviePipeline);
 	/** Called a short period of time after OnPIEMoviePipelineFinished to allow Editor the time to fully close PIE before we make a new request. */
 	void DelayedFinishNotification();
-
-	/** Gets the title for the Window. Used to show progress in the title. */
-	FText GetWindowTitle(const int32 InConfigIndex, const int32 InNumConfigs) const;
-
 private:
 	/** If using delayed initialization, how many frames are left before we call Initialize. Will be -1 if not actively counting down. */
 	int32 RemainingInitializationFrames;
 	bool bPreviousUseFixedTimeStep;
 	double PreviousFixedTimeStepDelta;
+	TWeakPtr<class SWindow> WeakCustomWindow;
 };
\ No newline at end of file
