MoviePipeline: Fix Fade track rendering

diff --git a/Engine/Plugins/MovieScene/MovieRenderPipeline/Source/MovieRenderPipelineRenderPasses/Private/MoviePipelineDeferredPasses.cpp b/Engine/Plugins/MovieScene/MovieRenderPipeline/Source/MovieRenderPipelineRenderPasses/Private/MoviePipelineDeferredPasses.cpp
--- a/Engine/Plugins/MovieScene/MovieRenderPipeline/Source/MovieRenderPipelineRenderPasses/Private/MoviePipelineDeferredPasses.cpp
+++ b/Engine/Plugins/MovieScene/MovieRenderPipeline/Source/MovieRenderPipelineRenderPasses/Private/MoviePipelineDeferredPasses.cpp
@@ -345,6 +345,18 @@ namespace MoviePipeline
 			TArray<float> const* CameraAnimPPBlendWeights;
 			LocalPlayerController->PlayerCameraManager->GetCachedPostProcessBlends(CameraAnimPPSettings, CameraAnimPPBlendWeights);
 
+			if (LocalPlayerController->PlayerCameraManager->bEnableFading)
+			{
+				View->OverlayColor = LocalPlayerController->PlayerCameraManager->FadeColor;
+				View->OverlayColor.A = FMath::Clamp(LocalPlayerController->PlayerCameraManager->FadeAmount, 0.f, 1.f);
+			}
+
+			if (LocalPlayerController->PlayerCameraManager->bEnableColorScaling)
+			{
+				FVector ColorScale = LocalPlayerController->PlayerCameraManager->ColorScale;
+				View->ColorScale = FLinearColor(ColorScale.X, ColorScale.Y, ColorScale.Z);
+			}
+
 			FMinimalViewInfo ViewInfo = LocalPlayerController->PlayerCameraManager->GetCameraCachePOV();
 			for (int32 PPIdx = 0; PPIdx < CameraAnimPPBlendWeights->Num(); ++PPIdx)
 			{
diff --git a/Engine/Shaders/Private/PostProcessCombineLUTs.usf b/Engine/Shaders/Private/PostProcessCombineLUTs.usf
--- a/Engine/Shaders/Private/PostProcessCombineLUTs.usf
+++ b/Engine/Shaders/Private/PostProcessCombineLUTs.usf
@@ -380,6 +380,9 @@ float4 CombineLUTsCommon(float2 InUV, uint InLayerIndex)
 
 	// blend with custom LDR color, used for Fade track in Matinee
 	float3 FilmColorNoGamma = lerp( FilmColor * ColorScale, OverlayColor.rgb, OverlayColor.a );
+	// Apply Fade track to linear outputs also
+	GradedColor = lerp(GradedColor * ColorScale, OverlayColor.rgb, OverlayColor.a);
+
 
 	// Apply "gamma" curve adjustment.
 	FilmColor = pow( max(0, FilmColorNoGamma), InverseGamma.y );
