MoviePipeline: Added a generic setting to apply arbitrary cvars/console commands at the start of each shot.

diff --git a/Engine/Plugins/MovieScene/MovieRenderPipeline/Source/MovieRenderPipelineSettings/Public/MoviePipelineConsoleVariableSetting.h b/Engine/Plugins/MovieScene/MovieRenderPipeline/Source/MovieRenderPipelineSettings/Public/MoviePipelineConsoleVariableSetting.h
new file mode 100644
--- /dev/null
+++ b/Engine/Plugins/MovieScene/MovieRenderPipeline/Source/MovieRenderPipelineSettings/Public/MoviePipelineConsoleVariableSetting.h
@@ -0,0 +1,115 @@
+// Copyright Epic Games, Inc. All Rights Reserved.
+#pragma once
+
+#include "MoviePipelineSetting.h"
+#include "MovieRenderPipelineCoreModule.h"
+#include "MoviePipelineConsoleVariableSetting.generated.h"
+
+UCLASS(BlueprintType)
+class MOVIERENDERPIPELINESETTINGS_API UMoviePipelineConsoleVariableSetting : public UMoviePipelineSetting
+{
+	GENERATED_BODY()
+public:
+
+public:
+#if WITH_EDITOR
+	virtual FText GetDisplayText() const override { return NSLOCTEXT("MovieRenderPipeline", "ConsoleVariableSettingDisplayName", "Console Variables"); }
+#endif
+	virtual bool IsValidOnShots() const override { return true; }
+	virtual bool IsValidOnMaster() const override { return true; }
+	virtual void SetupForPipelineImpl(UMoviePipeline* InPipeline) override
+	{
+		ApplyCVarSettings(true);
+	}
+
+	virtual void TeardownForPipelineImpl(UMoviePipeline* InPipeline) override
+	{
+		ApplyCVarSettings(false);
+	}
+	
+protected:
+	void ApplyCVarSettings(const bool bOverrideValues)
+	{
+		if (!(GetIsUserCustomized() && IsEnabled()))
+		{
+			return;
+		}
+
+		PreviousConsoleVariableValues.Reset();
+		PreviousConsoleVariableValues.SetNumZeroed(ConsoleVariables.Num());
+
+		int32 Index = 0;
+		for(const TPair<FString, float>& KVP : ConsoleVariables)
+		{
+			// We don't use the shared macro here because we want to soft-warn the user instead of tripping an ensure over missing cvar values.
+			IConsoleVariable* CVar = IConsoleManager::Get().FindConsoleVariable(*KVP.Key); 
+			if (CVar)
+			{
+				if (bOverrideValues)
+				{
+					PreviousConsoleVariableValues[Index] = CVar->GetFloat();
+					CVar->Set(KVP.Value, EConsoleVariableFlags::ECVF_SetByConsole);
+				}
+				else
+				{
+					CVar->Set(PreviousConsoleVariableValues[Index], EConsoleVariableFlags::ECVF_SetByConsole);
+				}
+
+				UE_LOG(LogMovieRenderPipeline, Log, TEXT("Applying CVar \"%s\" PreviousValue: %f NewValue: %f"),
+					*KVP.Key, PreviousConsoleVariableValues[Index], KVP.Value);
+			}
+			else
+			{
+				UE_LOG(LogMovieRenderPipeline, Warning, TEXT("Failed to apply CVar \"%s\" due to no cvar by that name. Ignoring."), *KVP.Key);
+			}
+
+			Index++;
+		}
+
+		if (bOverrideValues)
+		{
+			for (const FString& Command : StartConsoleCommands)
+			{
+				UE_LOG(LogMovieRenderPipeline, Log, TEXT("Executing Console Command \"%s\" before shot starts."), *Command);
+
+				GetWorld()->Exec(GetWorld(), *Command);
+			}
+		}
+		else
+		{
+			for (const FString& Command : EndConsoleCommands)
+			{
+				UE_LOG(LogMovieRenderPipeline, Log, TEXT("Executing Console Command \"%s\" after shot ends."), *Command);
+
+				GetWorld()->Exec(GetWorld(), *Command);
+			}
+		}
+		
+	}
+
+public:
+	/** 
+	* An array of key/value pairs for console variable name and the value you wish to set for that cvar.
+	* The existing value will automatically be cached and restored afterwards.
+	*/
+	UPROPERTY(EditAnywhere, BlueprintReadWrite, Category = "Game")
+	TMap<FString, float> ConsoleVariables;
+
+	/**
+	* An array of console commands to execute when this shot is started. If you need to restore the value 
+	* after the shot, add a matching entry in the EndConsoleCommands array. Because they are commands
+	* and not values we cannot save the preivous value automatically.
+	*/
+	UPROPERTY(EditAnywhere, BlueprintReadWrite, Category = "Game")
+	TArray<FString> StartConsoleCommands;
+
+	/**
+	* An array of console commands to execute when this shot is finished. Used to restore changes made by
+	* StartConsoleCommands.
+	*/
+	UPROPERTY(EditAnywhere, BlueprintReadWrite, Category = "Game")
+	TArray<FString> EndConsoleCommands;
+
+private:
+	TArray<float> PreviousConsoleVariableValues;
+};
\ No newline at end of file
