MoviePipeline: Fixed a shutter angle of 0 causing a divide by zero crash. We need at least one degree that it is open to determine time passing, it will just be small enough to effectively be zero.

diff --git a/Engine/Plugins/MovieScene/MovieRenderPipeline/Source/MovieRenderPipelineCore/Private/MoviePipeline.cpp b/Engine/Plugins/MovieScene/MovieRenderPipeline/Source/MovieRenderPipelineCore/Private/MoviePipeline.cpp
--- a/Engine/Plugins/MovieScene/MovieRenderPipeline/Source/MovieRenderPipelineCore/Private/MoviePipeline.cpp
+++ b/Engine/Plugins/MovieScene/MovieRenderPipeline/Source/MovieRenderPipelineCore/Private/MoviePipeline.cpp
@@ -1150,7 +1150,7 @@ MoviePipeline::FFrameConstantMetrics UMoviePipeline::CalculateShotFrameMetrics(c
 	UMoviePipelineAntiAliasingSetting* AntiAliasingSettings = FindOrAddSetting<UMoviePipelineAntiAliasingSetting>(InShot);
 
 	// (CameraShutterAngle/360) gives us the fraction-of-the-output-frame the accumulation frames should cover.
-	Output.ShutterAnglePercentage = CameraSettings->CameraShutterAngle / 360.0;
+	Output.ShutterAnglePercentage = FMath::Max(CameraSettings->CameraShutterAngle / 360.0, 1 / 360.0);
 
 	{
 		/*
