MoviePipeline: Fixed Game Override setting using incorrect UCLASS tag (Blueprintable instead of BlueprintType)

diff --git a/Engine/Plugins/MovieScene/MovieRenderPipeline/Source/MovieRenderPipelineSettings/Public/MoviePipelineGameOverrideSetting.h b/Engine/Plugins/MovieScene/MovieRenderPipeline/Source/MovieRenderPipelineSettings/Public/MoviePipelineGameOverrideSetting.h
--- a/Engine/Plugins/MovieScene/MovieRenderPipeline/Source/MovieRenderPipelineSettings/Public/MoviePipelineGameOverrideSetting.h
+++ b/Engine/Plugins/MovieScene/MovieRenderPipeline/Source/MovieRenderPipelineSettings/Public/MoviePipelineGameOverrideSetting.h
@@ -19,7 +19,7 @@ enum class EMoviePipelineTextureStreamingMethod : uint8
 	FullyLoad UMETA(DisplayName = "Fully Load Used Textures")
 };
 
-UCLASS(Blueprintable)
+UCLASS(BlueprintType)
 class MOVIERENDERPIPELINESETTINGS_API UMoviePipelineGameOverrideSetting : public UMoviePipelineSetting
 {
 	GENERATED_BODY()
@@ -48,7 +48,7 @@ public:
 	virtual void SetupForPipelineImpl(UMoviePipeline* InPipeline) override;
 	virtual void TeardownForPipelineImpl(UMoviePipeline* InPipeline) override;
 protected:
-	void ApplyCVarSettings(const bool bRestoreOldValues);
+	void ApplyCVarSettings(const bool bOverrideValues);
 
 public:
 	/** Optional Game Mode to override the map's default game mode with. This can be useful if the game's normal mode displays UI elements or loading screens that you don't want captured. */
