CIS: Fixed missing includes.

#rnx

diff --git a/Engine/Plugins/MovieScene/MovieRenderPipeline/Source/MovieRenderPipelineSettings/Public/MoviePipelineConsoleVariableSetting.h b/Engine/Plugins/MovieScene/MovieRenderPipeline/Source/MovieRenderPipelineSettings/Public/MoviePipelineConsoleVariableSetting.h
--- a/Engine/Plugins/MovieScene/MovieRenderPipeline/Source/MovieRenderPipelineSettings/Public/MoviePipelineConsoleVariableSetting.h
+++ b/Engine/Plugins/MovieScene/MovieRenderPipeline/Source/MovieRenderPipelineSettings/Public/MoviePipelineConsoleVariableSetting.h
@@ -3,6 +3,8 @@
 
 #include "MoviePipelineSetting.h"
 #include "MovieRenderPipelineCoreModule.h"
+#include "HAL/IConsoleManager.h"
+#include "Engine/World.h"
 #include "MoviePipelineConsoleVariableSetting.generated.h"
 
 UCLASS(BlueprintType)
@@ -81,7 +83,8 @@ protected:
 			{
 				UE_LOG(LogMovieRenderPipeline, Log, TEXT("Executing Console Command \"%s\" after shot ends."), *Command);
 
-				GetWorld()->Exec(GetWorld(), *Command);
+				GetWorld()->Exec(GetWorld(), 
+					*Command);
 			}
 		}
 		
