Fixing movie render pipeline when render target surface has padding.

diff --git a/Engine/Plugins/MovieScene/MovieRenderPipeline/Source/MovieRenderPipelineCore/Private/MoviePipelineSurfaceReader.cpp b/Engine/Plugins/MovieScene/MovieRenderPipeline/Source/MovieRenderPipelineCore/Private/MoviePipelineSurfaceReader.cpp
--- a/Engine/Plugins/MovieScene/MovieRenderPipeline/Source/MovieRenderPipelineCore/Private/MoviePipelineSurfaceReader.cpp
+++ b/Engine/Plugins/MovieScene/MovieRenderPipeline/Source/MovieRenderPipelineCore/Private/MoviePipelineSurfaceReader.cpp
@@ -159,13 +159,41 @@ void FMoviePipelineSurfaceReader::CopyReadbackTexture_RenderThread(const FMovieP
 	{
 		void* ColorDataBuffer = nullptr;
 
-		int32 Width = 0, Height = 0;
-		RHICmdList.MapStagingSurface(ReadbackTexture, ColorDataBuffer, Width, Height);
+		int32 ActualSizeX = 0, ActualSizeY = 0;
+		RHICmdList.MapStagingSurface(ReadbackTexture, ColorDataBuffer, ActualSizeX, ActualSizeY);
 
-		// Do a quick block copy
 		TArray<FFloat16Color> OutputPixels;
-		OutputPixels.SetNumUninitialized(Width * Height);
-		FMemory::BigBlockMemcpy(OutputPixels.GetData(), ColorDataBuffer, (Width * Height) * sizeof(FFloat16Color));
+
+		int32 ExpectedSizeX = InSampleState.TileSize.X;
+		int32 ExpectedSizeY = InSampleState.TileSize.Y;
+		OutputPixels.SetNumUninitialized(ExpectedSizeX * ExpectedSizeY);
+
+		// Due to padding, the actual size might be larger than the expected size. If they are the same, do a block copy. Otherwise copy
+		// line by line.
+		if (ExpectedSizeX == ActualSizeX && ExpectedSizeY == ActualSizeY)
+		{
+			FMemory::BigBlockMemcpy(OutputPixels.GetData(), ColorDataBuffer, (ExpectedSizeX * ExpectedSizeY) * sizeof(FFloat16Color));
+		}
+		else
+		{
+			// Make sure the target is larger than expected size.
+			check(ExpectedSizeX <= ActualSizeX);
+			check(ExpectedSizeY <= ActualSizeY);
+
+			int32 SrcPitchElem = ActualSizeX;
+			int32 DstPitchElem = ExpectedSizeX;
+
+			const FFloat16Color * SrcColorData = (const FFloat16Color*)ColorDataBuffer;
+			FFloat16Color * DstColorData = (FFloat16Color*)OutputPixels.GetData();
+
+			// Copy one line at a time
+			for (int32 RowIndex = 0; RowIndex < ExpectedSizeY; RowIndex++)
+			{
+				const FFloat16Color * SrcPtr = &SrcColorData[RowIndex * SrcPitchElem];
+				FFloat16Color * DstPtr = &DstColorData[RowIndex * DstPitchElem];
+				FMemory::Memcpy(DstPtr, SrcPtr, DstPitchElem * sizeof(FFloat16Color));
+			}
+		}
 
 		// Enqueue the Unmap before we broadcast the resulting pixels, though the broadcast shouldn't do anything blocking.
 		RHICmdList.UnmapStagingSurface(ReadbackTexture);
