MoviePipeline: Fixed zero-length camera cuts getting a camera cut section generated which just wasted work.

#rnx

diff --git a/Engine/Plugins/MovieScene/MovieRenderPipeline/Source/MovieRenderPipelineCore/Private/MoviePipeline.cpp b/Engine/Plugins/MovieScene/MovieRenderPipeline/Source/MovieRenderPipelineCore/Private/MoviePipeline.cpp
--- a/Engine/Plugins/MovieScene/MovieRenderPipeline/Source/MovieRenderPipelineCore/Private/MoviePipeline.cpp
+++ b/Engine/Plugins/MovieScene/MovieRenderPipeline/Source/MovieRenderPipelineCore/Private/MoviePipeline.cpp
@@ -699,9 +699,13 @@ FMoviePipelineShotInfo CreateShotFromMovieScene(const UMovieScene* InMovieScene,
 			UMovieSceneCameraCutSection* CameraCutSection = CastChecked<UMovieSceneCameraCutSection>(Section);
 
 			// ToDo: Inner vs. Outer resolution differences.
-			// Intersect this cut with the outer range in the likely event that the section goes past the bounds.
 			
 			TRange<FFrameNumber> LocalSectionRange = Section->GetRange(); // Section in local space
+			if (LocalSectionRange.IsEmpty())
+			{
+				UE_LOG(LogMovieRenderPipeline, Warning, TEXT("Found zero-length section in CameraCutTrack: %s Skipping..."), *CameraCutSection->GetPathName());
+				continue;
+			}
 			
 			LocalSectionRange = MovieScene::TranslateRange(LocalSectionRange, -InMovieScene->GetPlaybackRange().GetLowerBoundValue()); // Section relative to zero
 
@@ -712,6 +716,7 @@ FMoviePipelineShotInfo CreateShotFromMovieScene(const UMovieScene* InMovieScene,
 				SectionRangeInMaster = MovieScene::TranslateRange(LocalSectionRange, InSubSection->GetRange().GetLowerBoundValue());
 			}
 
+			// Intersect this cut with the outer range in the likely event that the section goes past the bounds.
 			TRange<FFrameNumber> IntersectingRange = TRange<FFrameNumber>::Intersection(SectionRangeInMaster, InIntersectionRange);
 
 			FCameraCutRange& NewRange = IntersectedRanges.AddDefaulted_GetRef();
