MoviePipeline: Level Sequences are no longer duplicated before rendering, and changes made by the rendering process are restored afterwards. This solves issues with complex binding hierarchies and blueprint-based binding modifications no longer working.

diff --git a/Engine/Plugins/MovieScene/MovieRenderPipeline/Source/MovieRenderPipelineCore/Private/MoviePipeline.cpp b/Engine/Plugins/MovieScene/MovieRenderPipeline/Source/MovieRenderPipelineCore/Private/MoviePipeline.cpp
--- a/Engine/Plugins/MovieScene/MovieRenderPipeline/Source/MovieRenderPipelineCore/Private/MoviePipeline.cpp
+++ b/Engine/Plugins/MovieScene/MovieRenderPipeline/Source/MovieRenderPipelineCore/Private/MoviePipeline.cpp
@@ -155,14 +155,29 @@ void UMoviePipeline::Initialize(UMoviePipelineExecutorJob* InJob)
 		return;
 	}
 
-	// Duplicate the target level sequence. This way modifications don't need to be undone.
-	TargetSequence = Cast<ULevelSequence>(UMoviePipelineBlueprintLibrary::DuplicateSequence(this, OriginalSequence));
+	TargetSequence = Cast<ULevelSequence>(GetCurrentJob()->TryLoadSequence());
+
+	// Disable some user settings that conflict with our need to mutate the data.
+	{
+		// Movie Scene Read Only
+		SequenceChanges.bSequenceReadOnly = TargetSequence->GetMovieScene()->IsReadOnly();
+		TargetSequence->GetMovieScene()->SetReadOnly(false);
+		
+		// Playback Range locked
+		SequenceChanges.bSequencePlaybackRangeLocked = TargetSequence->GetMovieScene()->IsPlaybackRangeLocked();
+		TargetSequence->GetMovieScene()->SetPlaybackRangeLocked(false);
+
+		// Force Frame-locked evaluation off on the sequence. We control time and will respect that, but need it off for subsampling.
+		SequenceChanges.EvaluationType = TargetSequence->GetMovieScene()->GetEvaluationType();
+		TargetSequence->GetMovieScene()->SetEvaluationType(EMovieSceneEvaluationType::WithSubFrames);
+	}
+
 
-	// Force Frame-locked evaluation off on the sequence. We control time and will respect that, but need it off for subsampling.
-	TargetSequence->GetMovieScene()->SetEvaluationType(EMovieSceneEvaluationType::WithSubFrames);
 
 	// Override the frame range on the target sequence if needed first before anyone has a chance to modify it.
 	{
+		SequenceChanges.PlaybackRange = TargetSequence->GetMovieScene()->GetPlaybackRange();
+
 		UMoviePipelineOutputSetting* OutputSetting = GetPipelineMasterConfig()->FindSetting<UMoviePipelineOutputSetting>();
 		if (OutputSetting->bUseCustomPlaybackRange)
 		{
@@ -191,12 +206,12 @@ void UMoviePipeline::Initialize(UMoviePipelineExecutorJob* InJob)
 
 	// Now that we've fixed up the sequence and validated it, we're going to build a list of shots that we need
 	// to produce in a simplified data structure. The simplified structure makes the flow/debugging easier.
-	ShotList = BuildShotListFromSequence(TargetSequence);
+	BuildShotListFromSequence();
 
 	// Finally, we're going to create a Level Sequence Actor in the world that has its settings configured by us.
 	// Because this callback is at the end of startup (and before tick) we should be able to spawn the actor
 	// and give it a chance to tick once (where it should do nothing) before we start manually manipulating it.
-	InitializeLevelSequenceActor(OriginalSequence, TargetSequence);
+	InitializeLevelSequenceActor();
 
 	// Register any additional engine callbacks needed.
 	{
@@ -213,14 +228,50 @@ void UMoviePipeline::Initialize(UMoviePipelineExecutorJob* InJob)
 
 	SetupAudioRendering();
 
-
 	TransitionToState(EMovieRenderPipelineState::ProducingFrames);
 	CurrentShotIndex = 0;
 	CachedOutputState.ShotCount = ShotList.Num();
-	InitializationTime = FDateTime::UtcNow();
 
 	// Initialization is complete. This engine frame is a wash (because the tick started with a 
 	// delta time not generated by us) so we'll wait until the next engine frame to start rendering.
+	InitializationTime = FDateTime::UtcNow();
+}
+
+void UMoviePipeline::RestoreTargetSequenceToOriginalState()
+{
+	if (PipelineState == EMovieRenderPipelineState::Uninitialized)
+	{
+		return;
+	}
+
+	if (!TargetSequence)
+	{
+		return;
+	}
+
+	TargetSequence->GetMovieScene()->SetEvaluationType(SequenceChanges.EvaluationType);
+	TargetSequence->GetMovieScene()->SetPlaybackRange(SequenceChanges.PlaybackRange);
+	TargetSequence->GetMovieScene()->SetReadOnly(SequenceChanges.bSequenceReadOnly);
+	TargetSequence->GetMovieScene()->SetPlaybackRangeLocked(SequenceChanges.bSequencePlaybackRangeLocked);
+
+	for (FMovieSceneChanges::FSegmentChange& ModifiedSegment : SequenceChanges.Segments)
+	{
+		if (ModifiedSegment.MovieScene.IsValid())
+		{
+			ModifiedSegment.MovieScene->SetPlaybackRange(ModifiedSegment.MovieScenePlaybackRange);
+			ModifiedSegment.MovieScene->SetReadOnly(ModifiedSegment.bMovieSceneReadOnly);
+		}
+
+		if (ModifiedSegment.ShotSection.IsValid())
+		{
+			ModifiedSegment.ShotSection->SetRange(ModifiedSegment.ShotSectionRange);
+			ModifiedSegment.ShotSection->SetIsLocked(ModifiedSegment.bShotSectionIsLocked);
+			
+			// We only make modified segments for shots that were previously active
+			ModifiedSegment.ShotSection->SetIsActive(true);
+			ModifiedSegment.ShotSection->MarkAsChanged();
+		}
+	}
 }
 
 
@@ -247,9 +298,7 @@ void UMoviePipeline::Shutdown()
 	check(IsInGameThread());
 
 	// This is a blocking operation which abandons any outstanding work to be submitted but finishes
-	// the existing work already processed. This is called at the end of the normal flow, or a blocking
-	// shutdown is requested (such as engine shutdown). If called at the end of the normal flow, most
-	// operations are a no-op.
+	// the existing work already processed.
 	if (PipelineState == EMovieRenderPipelineState::Uninitialized)
 	{
 		// If initialize is not called, no need to do anything.
@@ -385,6 +434,7 @@ void UMoviePipeline::TransitionToState(const EMovieRenderPipelineState InNewStat
 			}
 
 			TeardownAudioRendering();
+			RestoreTargetSequenceToOriginalState();
 
 			GAreScreenMessagesEnabled = bPrevGScreenMessagesEnabled;
 
@@ -634,37 +684,37 @@ void UMoviePipeline::ModifySequenceViaExtensions(ULevelSequence* InSequence)
 }
 
 
-void UMoviePipeline::InitializeLevelSequenceActor(ULevelSequence* OriginalLevelSequence, ULevelSequence* InSequenceToApply)
+void UMoviePipeline::InitializeLevelSequenceActor()
 {
-	// There is a reasonable chance that there exists a Level Sequence Actor in the world
-	// already set up to play this sequence. We don't want that however as it may have unknown
-	// settings on it.
+	// There is a reasonable chance that there exists a Level Sequence Actor in the world already set up to play this sequence.
+	ALevelSequenceActor* ExistingActor = nullptr;
+
 	for (auto It = TActorIterator<ALevelSequenceActor>(GetWorld()); It; ++It)
 	{
 		// Iterate through all of them in the event someone has multiple copies in the world on accident.
-		if (It->LevelSequence == OriginalLevelSequence)
+		if (It->LevelSequence == TargetSequence)
 		{
 			// Found it!
-			ALevelSequenceActor* ExistingActor = *It;
+			ExistingActor = *It;
 
-			//  Clear the sequence ref instead of killing actor to reduce likelyhood of game code throwing a NRE
-			// that may have expected this actor to exist.
+			// Stop it from playing if it's already playing.
 			if (ExistingActor->GetSequencePlayer())
 			{
 				ExistingActor->GetSequencePlayer()->Stop();
 			}
-
-			// Sequence must have been stopped before we can switch sequences.
-			ExistingActor->SetSequence(nullptr);
 		}
 	}
 
-	// Now we can construct a nice new actor which has the settings we need.
-	LevelSequenceActor = GetWorld()->SpawnActor<ALevelSequenceActor>();
-	check(LevelSequenceActor);
+	LevelSequenceActor = ExistingActor;
+	if (!LevelSequenceActor)
+	{
+		// Spawn a new level sequence
+		LevelSequenceActor = GetWorld()->SpawnActor<ALevelSequenceActor>();
+		check(LevelSequenceActor);
+	}
 
 	// Use our duplicated sequence
-	LevelSequenceActor->SetSequence(InSequenceToApply);
+	LevelSequenceActor->SetSequence(TargetSequence);
 
 	// Enforce settings.
 	LevelSequenceActor->PlaybackSettings.LoopCount.Value = 0;
@@ -716,6 +766,12 @@ FMoviePipelineShotInfo CreateShotFromMovieScene(const UMovieScene* InMovieScene,
 				SectionRangeInMaster = MovieScene::TranslateRange(LocalSectionRange, InSubSection->GetRange().GetLowerBoundValue());
 			}
 
+			if (!SectionRangeInMaster.Overlaps(InIntersectionRange))
+			{
+				UE_LOG(LogMovieRenderPipeline, Verbose, TEXT("Skipping camera cut section due to no overlap with playback range. CameraCutTrack: %s"), *CameraCutSection->GetPathName());
+				continue;
+			}
+			
 			// Intersect this cut with the outer range in the likely event that the section goes past the bounds.
 			TRange<FFrameNumber> IntersectingRange = TRange<FFrameNumber>::Intersection(SectionRangeInMaster, InIntersectionRange);
 
@@ -750,12 +806,10 @@ FMoviePipelineShotInfo CreateShotFromMovieScene(const UMovieScene* InMovieScene,
 	return NewShot;
 }
 
-TArray<FMoviePipelineShotInfo> UMoviePipeline::BuildShotListFromSequence(const ULevelSequence* InSequence)
+void UMoviePipeline::BuildShotListFromSequence()
 {
-	TArray<FMoviePipelineShotInfo> NewShotList;
-	
 	// Shot Tracks take precedent over camera cuts, as settings can only be applied as granular as a shot.
-	UMovieSceneCinematicShotTrack* CinematicShotTrack = InSequence->GetMovieScene()->FindMasterTrack<UMovieSceneCinematicShotTrack>();
+	UMovieSceneCinematicShotTrack* CinematicShotTrack = TargetSequence->GetMovieScene()->FindMasterTrack<UMovieSceneCinematicShotTrack>();
 	
 	if (CinematicShotTrack)
 	{
@@ -797,10 +851,9 @@ TArray<FMoviePipelineShotInfo> UMoviePipeline::BuildShotListFromSequence(const U
 
 			// The Shot Section may extend past our Sequence's Playback Bounds. We intersect the two bounds to ensure that
 			// the Playback Start/Playback End of the overall sequence is respected.
-			TRange<FFrameNumber> CinematicShotSectionRange = TRange<FFrameNumber>::Intersection(ShotSection->GetRange(), InSequence->GetMovieScene()->GetPlaybackRange());
+			TRange<FFrameNumber> CinematicShotSectionRange = TRange<FFrameNumber>::Intersection(ShotSection->GetRange(), TargetSequence->GetMovieScene()->GetPlaybackRange());
 			FMoviePipelineShotInfo NewShot = CreateShotFromMovieScene(ShotSection->GetSequence()->GetMovieScene(), CinematicShotSectionRange, ShotSection);
 
-
 			// Convert the offset time to ticks
 			FFrameTime OffsetTime = ShotSection->GetOffsetTime().GetValue();
 
@@ -813,7 +866,7 @@ TArray<FMoviePipelineShotInfo> UMoviePipeline::BuildShotListFromSequence(const U
 			// There should always be a shot config as the Pipeline default is returned in the event they didn't customize.
 			// check(NewShot.ShotConfig);
 
-			NewShotList.Add(MoveTemp(NewShot));
+			ShotList.Add(MoveTemp(NewShot));
 		}
 	}
 	else
@@ -822,12 +875,12 @@ TArray<FMoviePipelineShotInfo> UMoviePipeline::BuildShotListFromSequence(const U
 		FMoviePipelineShotInfo NewShot = CreateShotFromMovieScene(TargetSequence->GetMovieScene(), TargetSequence->GetMovieScene()->GetPlaybackRange(), nullptr);
 		// NewShot.ShotConfig = GetPipelineMasterConfig()->DefaultShotConfig;
 		
-		NewShotList.Add(MoveTemp(NewShot));
+		ShotList.Add(MoveTemp(NewShot));
 	}
 
 	// If they don't have a cinematic shot track, or a camera cut track then they want to control the camera
 	// through their own logic. We'll just use the duration of the Sequence as the render, plus warn them.
-	if (NewShotList.Num() == 0)
+	if (ShotList.Num() == 0)
 	{
 		UE_LOG(LogMovieRenderPipeline, Warning, TEXT("No Cinematic Shot Tracks found, and no Camera Cut Tracks found. Playback Range will be used but camera will render from Pawns perspective."));
 		FMoviePipelineShotInfo NewShot;
@@ -840,9 +893,9 @@ TArray<FMoviePipelineShotInfo> UMoviePipeline::BuildShotListFromSequence(const U
 	// Now that we've gathered at least one or more shots with one or more cuts, we can apply settings. It's easier to
 	// debug when all of the shots are calculated up front and debug info is printed in a block instead of as it is reached.
 	int32 ShotIndex = 1;
-	for (FMoviePipelineShotInfo& Shot : NewShotList)
+	for (FMoviePipelineShotInfo& Shot : ShotList)
 	{
-		UE_LOG(LogMovieRenderPipeline, Log, TEXT("Shot %d/%d has %d Camera Cuts."), ShotIndex, NewShotList.Num(), Shot.CameraCuts.Num());
+		UE_LOG(LogMovieRenderPipeline, Log, TEXT("Shot %d/%d has %d Camera Cuts."), ShotIndex, ShotList.Num(), Shot.CameraCuts.Num());
 		ShotIndex++;
 
 		UMoviePipelineOutputSetting* OutputSettings = FindOrAddSetting<UMoviePipelineOutputSetting>(Shot);
@@ -907,13 +960,12 @@ TArray<FMoviePipelineShotInfo> UMoviePipeline::BuildShotListFromSequence(const U
 	}
 
 	// Sort the shot list overall. The correct order is required for relative frame counts to put the shots on disk in the right order.
-	NewShotList.Sort([](const FMoviePipelineShotInfo& A, const FMoviePipelineShotInfo& B)
+	ShotList.Sort([](const FMoviePipelineShotInfo& A, const FMoviePipelineShotInfo& B)
 	{
 		return A.TotalOutputRange.GetLowerBoundValue() < B.TotalOutputRange.GetLowerBoundValue();
 	});
-
-	return NewShotList;
 }
+
 void UMoviePipeline::InitializeShot(FMoviePipelineShotInfo& InShot)
 {
 	// Set the new shot as the active shot. This enables the specified shot section and disables all other shot sections.
@@ -1058,6 +1110,14 @@ void UMoviePipeline::ExpandShot(FMoviePipelineShotInfo& InShot)
 		TRange<FFrameNumber> OldPlaybackRange = ShotMovieScene->GetPlaybackRange();
 		TRange<FFrameNumber> NewPlaybackRange = MovieScene::DilateRange(ShotMovieScene->GetPlaybackRange(), -DilatedStartOffset, DilatedEndOffset);
 
+		FMovieSceneChanges::FSegmentChange& ModifiedSegment = SequenceChanges.Segments.AddDefaulted_GetRef();
+		ModifiedSegment.MovieScene = ShotMovieScene;
+		ModifiedSegment.MovieScenePlaybackRange = OldPlaybackRange;
+		ModifiedSegment.bMovieSceneReadOnly = ShotMovieScene->IsReadOnly();
+		ModifiedSegment.ShotSection = ShotSection;
+		ModifiedSegment.bShotSectionIsLocked = ShotSection->IsLocked();
+		ModifiedSegment.ShotSectionRange = ShotSection->GetRange();
+
 		FFrameNumber TotalExpansionSizeInFrames = FFrameRate::TransformTime(FFrameTime(HandleFrameTicks) + FrameMetrics.TicksPerOutputFrame, FrameMetrics.TickResolution, FrameMetrics.FrameRate).CeilToFrame();
 		// ToDo: ShotList isn't assigned yet so IndexOfByKey fails.
 		UE_LOG(LogMovieRenderPipeline, Log, TEXT("ShotSection ([%d] %s) Expanding Shot by [-%d, %d] Ticks (+/-%d Frames). Old Inner Playback Range: %s New Inner Playback Range: %s"),
@@ -1066,6 +1126,7 @@ void UMoviePipeline::ExpandShot(FMoviePipelineShotInfo& InShot)
 			*LexToString(OldPlaybackRange), *LexToString(NewPlaybackRange));
 
 		// Expand the inner scene
+		ShotMovieScene->SetReadOnly(false);
 		ShotMovieScene->SetPlaybackRange(NewPlaybackRange, false);
 
 		// Expand the outer owning section
diff --git a/Engine/Plugins/MovieScene/MovieRenderPipeline/Source/MovieRenderPipelineCore/Public/MoviePipeline.h b/Engine/Plugins/MovieScene/MovieRenderPipeline/Source/MovieRenderPipelineCore/Public/MoviePipeline.h
--- a/Engine/Plugins/MovieScene/MovieRenderPipeline/Source/MovieRenderPipelineCore/Public/MoviePipeline.h
+++ b/Engine/Plugins/MovieScene/MovieRenderPipeline/Source/MovieRenderPipelineCore/Public/MoviePipeline.h
@@ -206,11 +206,14 @@ private:
 
 
 private:
+	/** Iterates through the changes we've made to a shot and applies the original settings. */
+	void RestoreTargetSequenceToOriginalState();
+
 	/** Initialize a new Level Sequence Actor to evaluate our target sequence. Disables any existing Level Sequences pointed at our original sequence. */
-	void InitializeLevelSequenceActor(ULevelSequence* OriginalLevelSequence, ULevelSequence* InSequenceToApply);
+	void InitializeLevelSequenceActor();
 
-	/** This converts the sequence into a Shot List and expands bounds.*/
-	TArray<FMoviePipelineShotInfo> BuildShotListFromSequence(const ULevelSequence* InSequence);
+	/** This builds the shot list from the target sequence, and expands Playback Bounds to cover any future evaluation we may need. */
+	void BuildShotListFromSequence();
 
 	/** 
 	* Modifies the TargetSequence to ensure that only the specified Shot has it's associated Cinematic Shot Section enabled.
@@ -315,6 +318,32 @@ private:
 	/** Keep track of which job we're working on. This holds our Configuration + which shots we're supposed to render from it. */
 	UPROPERTY(Transient)
 	UMoviePipelineExecutorJob* CurrentJob;
+
+
+	/** Previous values for data that we modified in the sequence for restoration in shutdown. */
+	struct FMovieSceneChanges
+	{
+		// Master level settings
+		EMovieSceneEvaluationType EvaluationType;
+		TRange<FFrameNumber> PlaybackRange;
+		bool bSequenceReadOnly;
+		bool bSequencePlaybackRangeLocked;
+
+		struct FSegmentChange
+		{
+			TWeakObjectPtr<class UMovieScene> MovieScene;
+			TRange<FFrameNumber> MovieScenePlaybackRange;
+			bool bMovieSceneReadOnly;
+			TWeakObjectPtr<UMovieSceneCinematicShotSection> ShotSection;
+			bool bShotSectionIsLocked;
+			TRange<FFrameNumber> ShotSectionRange;
+		};
+
+		// Shot-specific settings
+		TArray<FSegmentChange> Segments;
+	};
+
+	FMovieSceneChanges SequenceChanges;
 };
 
 UCLASS()
diff --git a/Engine/Plugins/MovieScene/MovieRenderPipeline/Source/MovieRenderPipelineCore/Public/MoviePipelineBlueprintLibrary.h b/Engine/Plugins/MovieScene/MovieRenderPipeline/Source/MovieRenderPipelineCore/Public/MoviePipelineBlueprintLibrary.h
--- a/Engine/Plugins/MovieScene/MovieRenderPipeline/Source/MovieRenderPipelineCore/Public/MoviePipelineBlueprintLibrary.h
+++ b/Engine/Plugins/MovieScene/MovieRenderPipeline/Source/MovieRenderPipelineCore/Public/MoviePipelineBlueprintLibrary.h
@@ -23,6 +23,10 @@ public:
 	* recursively duplicates the given sequence, shot and subsequence and then fixes up the references to
 	* point to newly duplicated sequences.
 	*
+	* Use at your own risk. Some features may not work when duplicated (complex object binding arrangements,
+	* blueprint GetSequenceBinding nodes, etc.) but can be useful when wanting to create a bunch of variations
+	* with minor differences (such as swapping out an actor, track, etc.)
+	*
 	* This does not duplicate any assets that the sequence points to outside of Shots/Subsequences.
 	*
 	* @param	Outer		- The Outer of the newly duplicated object. Leave null for TransientPackage();
