MoviePipeline: Single-frame renders now work correctly on sequences that only have one camera cut track and no shots.

diff --git a/Engine/Plugins/MovieScene/MovieRenderPipeline/Source/MovieRenderPipelineCore/Private/MoviePipeline.cpp b/Engine/Plugins/MovieScene/MovieRenderPipeline/Source/MovieRenderPipelineCore/Private/MoviePipeline.cpp
--- a/Engine/Plugins/MovieScene/MovieRenderPipeline/Source/MovieRenderPipelineCore/Private/MoviePipeline.cpp
+++ b/Engine/Plugins/MovieScene/MovieRenderPipeline/Source/MovieRenderPipelineCore/Private/MoviePipeline.cpp
@@ -746,24 +746,23 @@ FMoviePipelineShotInfo CreateShotFromMovieScene(const UMovieScene* InMovieScene,
 	{
 		for (UMovieSceneSection* Section : CameraCutTrack->GetAllSections())
 		{
-			UMovieSceneCameraCutSection* CameraCutSection = CastChecked<UMovieSceneCameraCutSection>(Section);
-
 			// ToDo: Inner vs. Outer resolution differences.
+			UMovieSceneCameraCutSection* CameraCutSection = CastChecked<UMovieSceneCameraCutSection>(Section);
 			
-			TRange<FFrameNumber> LocalSectionRange = Section->GetRange(); // Section in local space
-			if (LocalSectionRange.IsEmpty())
+			if (Section->GetRange().IsEmpty())
 			{
 				UE_LOG(LogMovieRenderPipeline, Warning, TEXT("Found zero-length section in CameraCutTrack: %s Skipping..."), *CameraCutSection->GetPathName());
 				continue;
 			}
-			
-			LocalSectionRange = MovieScene::TranslateRange(LocalSectionRange, -InMovieScene->GetPlaybackRange().GetLowerBoundValue()); // Section relative to zero
 
-			TRange<FFrameNumber> SectionRangeInMaster = LocalSectionRange;
-			// Add the offset from the parent subsection to put it back into parent space.
+			TRange<FFrameNumber> SectionRangeInMaster = Section->GetRange();
+			
+			// If this camera cut track is inside of a shot subsection, we need to take the parent section into account.
 			if (InSubSection)
 			{
-				SectionRangeInMaster = MovieScene::TranslateRange(LocalSectionRange, InSubSection->GetRange().GetLowerBoundValue());
+				TRange<FFrameNumber> LocalSectionRange = Section->GetRange(); // Section in local space
+				LocalSectionRange = MovieScene::TranslateRange(LocalSectionRange, -InMovieScene->GetPlaybackRange().GetLowerBoundValue()); // Section relative to zero
+				SectionRangeInMaster = MovieScene::TranslateRange(LocalSectionRange, InSubSection->GetRange().GetLowerBoundValue()); // Convert to master sequence space.
 			}
 
 			if (!SectionRangeInMaster.Overlaps(InIntersectionRange))
