MoviePipeline: Settings in the "Output" setting are now correctly saved/re-imported when using shot presets.

diff --git a/Engine/Plugins/MovieScene/MovieRenderPipeline/Source/MovieRenderPipelineCore/Public/MoviePipelineMasterConfig.h b/Engine/Plugins/MovieScene/MovieRenderPipeline/Source/MovieRenderPipelineCore/Public/MoviePipelineMasterConfig.h
--- a/Engine/Plugins/MovieScene/MovieRenderPipeline/Source/MovieRenderPipelineCore/Public/MoviePipelineMasterConfig.h
+++ b/Engine/Plugins/MovieScene/MovieRenderPipeline/Source/MovieRenderPipelineCore/Public/MoviePipelineMasterConfig.h
@@ -87,7 +87,7 @@ public:
 	TMap<FString, UMoviePipelineShotConfig*> PerShotConfigMapping;
 
 private:
-	UPROPERTY(Transient, Instanced)
+	UPROPERTY(Instanced)
 	UMoviePipelineOutputSetting* OutputSetting;
 
 	/** An array of settings that are available in the engine and have not been edited by the user. */
