MoviePipeline: Output Directory now defaults to the Project's Saved directory if it was saved with an empty string. This allows shipping presets that can be used in multiple projects.

diff --git a/Engine/Plugins/MovieScene/MovieRenderPipeline/Source/MovieRenderPipelineCore/Private/MoviePipelineOutputSetting.cpp b/Engine/Plugins/MovieScene/MovieRenderPipeline/Source/MovieRenderPipelineCore/Private/MoviePipelineOutputSetting.cpp
--- a/Engine/Plugins/MovieScene/MovieRenderPipeline/Source/MovieRenderPipelineCore/Private/MoviePipelineOutputSetting.cpp
+++ b/Engine/Plugins/MovieScene/MovieRenderPipeline/Source/MovieRenderPipelineCore/Private/MoviePipelineOutputSetting.cpp
@@ -26,6 +26,20 @@ UMoviePipelineOutputSetting::UMoviePipelineOutputSetting()
 	OutputDirectory.Path = FPaths::ProjectSavedDir() / TEXT("MovieRenders/");
 }
 
+void UMoviePipelineOutputSetting::PostLoad()
+{
+	Super::PostLoad();
+
+	// In order to ship presets that work with any project, we can't use a relative path because it is
+	// relative to the executable and thus has the project name embedded in. To solve this we will save
+	// an empty string into the Output Directory and convert it to their relative directory in Post Load.
+	// This isn't done in the CDO so that resetting to default value works as expected.
+	if (OutputDirectory.Path.Len() == 0)
+	{
+		OutputDirectory.Path = FPaths::ProjectSavedDir() / TEXT("MovieRenders/");
+	}
+}
+
 FText UMoviePipelineOutputSetting::GetFooterText(UMoviePipelineExecutorJob* InJob) const 
 {
 	FTextBuilder TextBuilder;
diff --git a/Engine/Plugins/MovieScene/MovieRenderPipeline/Source/MovieRenderPipelineCore/Public/MoviePipelineOutputSetting.h b/Engine/Plugins/MovieScene/MovieRenderPipeline/Source/MovieRenderPipelineCore/Public/MoviePipelineOutputSetting.h
--- a/Engine/Plugins/MovieScene/MovieRenderPipeline/Source/MovieRenderPipelineCore/Public/MoviePipelineOutputSetting.h
+++ b/Engine/Plugins/MovieScene/MovieRenderPipeline/Source/MovieRenderPipelineCore/Public/MoviePipelineOutputSetting.h
@@ -21,6 +21,10 @@ public:
 	virtual bool IsValidOnMaster() const override { return true; }
 	virtual void GetFilenameFormatArguments(FMoviePipelineFormatArgs& InOutFormatArgs) const override;
 	virtual bool CanBeDisabled() const override { return false; }
+
+	// UObject Interface
+	virtual void PostLoad() override;
+	// ~UObject Interface
 public:
 	/** What directory should all of our output files be relative to. */
 	UPROPERTY(EditAnywhere, BlueprintReadWrite, Category = "File Output")
