MoviePipeline: Fixed importing presets not copying the Output setting correctly. Fixed an issue where importing a preset always registered additional callbacks resulting in the callback getting called multiple times. Polished the spacing/padding/visibility on these UIs as well.

diff --git a/Engine/Plugins/MovieScene/MovieRenderPipeline/Source/MovieRenderPipelineCore/Private/MoviePipelineMasterConfig.cpp b/Engine/Plugins/MovieScene/MovieRenderPipeline/Source/MovieRenderPipelineCore/Private/MoviePipelineMasterConfig.cpp
--- a/Engine/Plugins/MovieScene/MovieRenderPipeline/Source/MovieRenderPipelineCore/Private/MoviePipelineMasterConfig.cpp
+++ b/Engine/Plugins/MovieScene/MovieRenderPipeline/Source/MovieRenderPipelineCore/Private/MoviePipelineMasterConfig.cpp
@@ -276,7 +276,7 @@ void UMoviePipelineMasterConfig::CopyFrom(UMoviePipelineConfigBase* InConfig)
 	if (InConfig->IsA<UMoviePipelineMasterConfig>())
 	{
 		UMoviePipelineMasterConfig* MasterConfig = CastChecked<UMoviePipelineMasterConfig>(InConfig);
-		OutputSetting = Cast<UMoviePipelineOutputSetting>(StaticDuplicateObject(MasterConfig->OutputSetting, this, MasterConfig->OutputSetting->GetFName()));
+		OutputSetting = Cast<UMoviePipelineOutputSetting>(StaticDuplicateObject(MasterConfig->OutputSetting, this, NAME_None));
 	}
 }
 
@@ -298,4 +298,4 @@ TArray<UMoviePipelineOutputBase*> UMoviePipelineMasterConfig::GetOutputContainer
 	return OutputContainers;
 }
 
-#undef LOCTEXT_NAMESPACE // "MovieRenderPipelineConfig"
\ No newline at end of file
+#undef LOCTEXT_NAMESPACE // "MovieRenderPipelineConfig"
diff --git a/Engine/Plugins/MovieScene/MovieRenderPipeline/Source/MovieRenderPipelineEditor/Private/Widgets/SMoviePipelineConfigEditor.cpp b/Engine/Plugins/MovieScene/MovieRenderPipeline/Source/MovieRenderPipelineEditor/Private/Widgets/SMoviePipelineConfigEditor.cpp
--- a/Engine/Plugins/MovieScene/MovieRenderPipeline/Source/MovieRenderPipelineEditor/Private/Widgets/SMoviePipelineConfigEditor.cpp
+++ b/Engine/Plugins/MovieScene/MovieRenderPipeline/Source/MovieRenderPipelineEditor/Private/Widgets/SMoviePipelineConfigEditor.cpp
@@ -73,7 +73,7 @@ void SMoviePipelineConfigEditor::Construct(const FArguments& InArgs)
 		.Orientation(Orient_Horizontal)
     
 		+ SSplitter::Slot()
-		.Value(.40f)
+		.Value(.33f)
 		[
 			SNew(SBorder)
 			.Padding(4)
@@ -84,7 +84,7 @@ void SMoviePipelineConfigEditor::Construct(const FArguments& InArgs)
 		]
     
 		+ SSplitter::Slot()
-		.Value(.60f)
+		.Value(.67f)
 		[
 			SNew(SVerticalBox)
 			+ SVerticalBox::Slot()
@@ -94,15 +94,24 @@ void SMoviePipelineConfigEditor::Construct(const FArguments& InArgs)
 			]
 
 			+ SVerticalBox::Slot()
+			.Padding(FMargin(0, 4, 0, 0))
 			.AutoHeight()
 			[
 				SNew(SBorder)
-				.BorderImage(FEditorStyle::GetBrush("ToolPanel.GroupBorder"))
 				.Visibility(this, &SMoviePipelineConfigEditor::IsSettingFooterVisible)
+				.BorderImage(FEditorStyle::GetBrush("ToolPanel.GroupBorder"))
 				[
-					SNew(STextBlock)
-					.AutoWrapText(true)
-					.Text(this, &SMoviePipelineConfigEditor::GetSettingsFooterText)
+					SNew(SBox)
+					.MaxDesiredHeight(96)
+					[
+						SNew(SScrollBox)
+						+ SScrollBox::Slot()
+						[
+							SNew(STextBlock)
+							.AutoWrapText(true)
+							.Text(this, &SMoviePipelineConfigEditor::GetSettingsFooterText)
+						]
+					]	
 				]
 			]
 
@@ -306,7 +315,8 @@ void SMoviePipelineConfigEditor::UpdateDetails()
 	FPropertyEditorModule& PropertyEditorModule = FModuleManager::Get().LoadModuleChecked<FPropertyEditorModule>("PropertyEditor");
 	FDetailsViewArgs DetailsViewArgs(false, false, false, FDetailsViewArgs::HideNameArea, true);
 	DetailsViewArgs.bShowScrollBar = false;
-	 
+	DetailsViewArgs.ColumnWidth = 0.5f;
+
 	TArray<UMoviePipelineSetting*> SelectedSources;
 	SettingsWidget->GetSelectedSettings(SelectedSources);
 
@@ -389,7 +399,11 @@ FText SMoviePipelineConfigEditor::GetSettingsFooterText() const
 	{
 		for (const UMoviePipelineSetting* Setting : SelectedSettings)
 		{
-			TextBuilder.AppendLine(Setting->GetFooterText(CachedOwningJob.Get()));
+			const FText FooterText = Setting->GetFooterText(CachedOwningJob.Get());
+			if (!FooterText.IsEmpty())
+			{
+				TextBuilder.AppendLine(FooterText);
+			}
 		}
 	}
 
@@ -410,7 +424,7 @@ EVisibility SMoviePipelineConfigEditor::IsValidationWarningVisible() const
 		}
 	}
 
-	return (ValidationResult != EMoviePipelineValidationState::Valid) ? EVisibility::Visible : EVisibility::Hidden;
+	return (ValidationResult != EMoviePipelineValidationState::Valid) ? EVisibility::Visible : EVisibility::Collapsed;
 }
 
 FText SMoviePipelineConfigEditor::GetValidationWarningText() const
diff --git a/Engine/Plugins/MovieScene/MovieRenderPipeline/Source/MovieRenderPipelineEditor/Private/Widgets/SMoviePipelineConfigPanel.cpp b/Engine/Plugins/MovieScene/MovieRenderPipeline/Source/MovieRenderPipelineEditor/Private/Widgets/SMoviePipelineConfigPanel.cpp
--- a/Engine/Plugins/MovieScene/MovieRenderPipeline/Source/MovieRenderPipelineEditor/Private/Widgets/SMoviePipelineConfigPanel.cpp
+++ b/Engine/Plugins/MovieScene/MovieRenderPipeline/Source/MovieRenderPipelineEditor/Private/Widgets/SMoviePipelineConfigPanel.cpp
@@ -57,12 +57,15 @@ void SMoviePipelineConfigPanel::Construct(const FArguments& InArgs, TSubclassOf<
 	{
 		TransientPreset->CopyFrom(InArgs._BasePreset);
 		PresetUsedIfNotModified = InArgs._BasePreset;
-		FCoreUObjectDelegates::OnObjectModified.AddSP(this, &SMoviePipelineConfigPanel::OnAnyObjectModified);
 	}
 	else if (InArgs._BaseConfig)
 	{
 		TransientPreset->CopyFrom(InArgs._BaseConfig);
 	}
+
+	// Listen for modificatiosn to our edited configuration so we can determine if they've modified the preset
+	// which only changes the display (the actual preset asset is not modified)
+	FCoreUObjectDelegates::OnObjectModified.AddSP(this, &SMoviePipelineConfigPanel::OnAnyObjectModified);
 	
 	WeakJob = InArgs._Job;
 	OnConfigurationModified = InArgs._OnConfigurationModified;
@@ -138,6 +141,7 @@ void SMoviePipelineConfigPanel::Construct(const FArguments& InArgs, TSubclassOf<
 			SNew(STextBlock)
 			.TextStyle(FMovieRenderPipelineStyle::Get(), "MovieRenderPipeline.Config.TypeLabel")
 			.Text(this, &SMoviePipelineConfigPanel::GetConfigTypeLabel)
+			.Visibility(EVisibility::Collapsed)
 		]
 
 		// Main Editor 
@@ -473,11 +477,6 @@ void SMoviePipelineConfigPanel::OnSaveAsPreset()
 			// the dialog.
 			PresetUsedIfNotModified = NewPreset;
 
-			// Instead, we track if our transient object has Modify() called on it after this
-			// point, since all changes should call Modify for undo/redo to work. When we close
-			// the dialog we can see if it has been modified since the last preset change. 
-			FCoreUObjectDelegates::OnObjectModified.AddSP(this, &SMoviePipelineConfigPanel::OnAnyObjectModified);
-
 			// We also want to rename our transient preset. The name is used to indicate what preset
 			// it originally came from, and now the closest preset to it is not the original but the
 			// newly exported preset.
@@ -514,7 +513,6 @@ void SMoviePipelineConfigPanel::OnImportPreset(const FAssetData& InPresetAsset)
 		TransientPreset->CopyFrom(Preset);
 
 		PresetUsedIfNotModified = Preset;
-		FCoreUObjectDelegates::OnObjectModified.AddSP(this, &SMoviePipelineConfigPanel::OnAnyObjectModified);
 		// CockpitWidget->GetMetaData()->SetPresetOrigin(Take);
 	}
 }
