Adding debug information for movie render pipeline when data is not well formed.

diff --git a/Engine/Plugins/MovieScene/MovieRenderPipeline/Source/MovieRenderPipelineRenderPasses/Private/MoviePipelineDeferredPasses.cpp b/Engine/Plugins/MovieScene/MovieRenderPipeline/Source/MovieRenderPipelineRenderPasses/Private/MoviePipelineDeferredPasses.cpp
--- a/Engine/Plugins/MovieScene/MovieRenderPipeline/Source/MovieRenderPipelineRenderPasses/Private/MoviePipelineDeferredPasses.cpp
+++ b/Engine/Plugins/MovieScene/MovieRenderPipeline/Source/MovieRenderPipelineRenderPasses/Private/MoviePipelineDeferredPasses.cpp
@@ -553,7 +553,27 @@ namespace MoviePipeline
 	{
 		SCOPE_CYCLE_COUNTER(STAT_AccumulateSample_RenderThread);
 
-		check(InPixelData->IsDataWellFormed());
+		bool isWellFormed = InPixelData->IsDataWellFormed();
+
+		if (!isWellFormed)
+		{
+			// figure out why it is not well formed, and print a warning.
+			int64 RawSize = InPixelData->GetRawDataSizeInBytes();
+
+			int64 SizeX = InPixelData->GetSize().X;
+			int64 SizeY = InPixelData->GetSize().Y;
+			int64 ByteDepth = int64(InPixelData->GetBitDepth() / 8);
+			int64 NumChannels = int64(InPixelData->GetNumChannels());
+			int64 ExpectedTotalSize = SizeX * SizeY * ByteDepth * NumChannels;
+			int64 ActualTotalSize = InPixelData->GetRawDataSizeInBytes();
+
+			UE_LOG(LogMovieRenderPipeline, Log, TEXT("AccumulateSample_RenderThread: Data is not well formed."));
+			UE_LOG(LogMovieRenderPipeline, Log, TEXT("Image dimension: %lldx%lld, %lld, %lld"), SizeX, SizeY, ByteDepth, NumChannels);
+			UE_LOG(LogMovieRenderPipeline, Log, TEXT("Expected size: %lld"), ExpectedTotalSize);
+			UE_LOG(LogMovieRenderPipeline, Log, TEXT("Actual size:   %lld"), ActualTotalSize);
+		}
+
+		check(isWellFormed);
 
 		// Writing tiles can be useful for debug reasons. These get passed onto the output every frame.
 		if (InParams.SampleState.bWriteSampleToDisk)
diff --git a/Engine/Source/Runtime/ImageWriteQueue/Public/ImagePixelData.h b/Engine/Source/Runtime/ImageWriteQueue/Public/ImagePixelData.h
--- a/Engine/Source/Runtime/ImageWriteQueue/Public/ImagePixelData.h
+++ b/Engine/Source/Runtime/ImageWriteQueue/Public/ImagePixelData.h
@@ -100,6 +100,18 @@ struct FImagePixelData
 		return false;
 	}
 
+	/**
+	 * Get the size in bytes, regardless of if it is well formed.
+	 */
+	int64 GetRawDataSizeInBytes() const
+	{
+		const void* RawPtr    = nullptr;
+		int64       SizeBytes = 0;
+
+		RetrieveData(RawPtr, SizeBytes);
+		return SizeBytes;
+	}
+
 	/**
 	 * Copy this whole image buffer. This can be very costly for large images.
 	 */
