Debugging info on high-res tiling render targets.

diff --git a/Engine/Plugins/MovieScene/MovieRenderPipeline/Source/MovieRenderPipelineCore/Private/MoviePipelineSurfaceReader.cpp b/Engine/Plugins/MovieScene/MovieRenderPipeline/Source/MovieRenderPipelineCore/Private/MoviePipelineSurfaceReader.cpp
--- a/Engine/Plugins/MovieScene/MovieRenderPipeline/Source/MovieRenderPipelineCore/Private/MoviePipelineSurfaceReader.cpp
+++ b/Engine/Plugins/MovieScene/MovieRenderPipeline/Source/MovieRenderPipelineCore/Private/MoviePipelineSurfaceReader.cpp
@@ -11,6 +11,7 @@
 #include "Shader.h"
 #include "StaticBoundShaderState.h"
 #include "Modules/ModuleManager.h"
+#include "MovieRenderPipelineCoreModule.h"
 
 FMoviePipelineSurfaceReader::FMoviePipelineSurfaceReader(EPixelFormat InPixelFormat, FIntPoint InSurfaceSize)
 {
@@ -112,6 +113,8 @@ void FMoviePipelineSurfaceReader::ResolveSampleToReadbackTexture_RenderThread(co
 	{
 		RHICmdList.SetViewport(0, 0, 0.0f, TargetSize.X, TargetSize.Y, 1.0f);
 
+		UE_LOG(LogMovieRenderPipeline, Log, TEXT("Rendering Tile: %d x %d"), TargetSize.X, TargetSize.Y);
+
 		FGraphicsPipelineStateInitializer GraphicsPSOInit;
 		RHICmdList.ApplyCachedRenderTargets(GraphicsPSOInit);
 		GraphicsPSOInit.BlendState = TStaticBlendState<>::GetRHI();
@@ -176,6 +179,10 @@ void FMoviePipelineSurfaceReader::CopyReadbackTexture_RenderThread(const FMovieP
 		}
 		else
 		{
+			UE_LOG(LogMovieRenderPipeline, Log, TEXT("Unexpected size in FMoviePipelineSurfaceReader::CopyReadbackTexture_RenderThread."));
+			UE_LOG(LogMovieRenderPipeline, Log, TEXT("    Expected size: %d x %d"), ExpectedSizeX, ExpectedSizeY);
+			UE_LOG(LogMovieRenderPipeline, Log, TEXT("    Actual size:   %d x %d"), ActualSizeX, ActualSizeY);
+
 			// Make sure the target is larger than expected size.
 			check(ExpectedSizeX <= ActualSizeX);
 			check(ExpectedSizeY <= ActualSizeY);
