Fixing overlapped tile buffers for small sizes.

diff --git a/Engine/Plugins/MovieScene/MovieRenderPipeline/Source/MovieRenderPipelineCore/Private/MoviePipelineSurfaceReader.cpp b/Engine/Plugins/MovieScene/MovieRenderPipeline/Source/MovieRenderPipelineCore/Private/MoviePipelineSurfaceReader.cpp
--- a/Engine/Plugins/MovieScene/MovieRenderPipeline/Source/MovieRenderPipelineCore/Private/MoviePipelineSurfaceReader.cpp
+++ b/Engine/Plugins/MovieScene/MovieRenderPipeline/Source/MovieRenderPipelineCore/Private/MoviePipelineSurfaceReader.cpp
@@ -167,8 +167,8 @@ void FMoviePipelineSurfaceReader::CopyReadbackTexture_RenderThread(const FMovieP
 
 		TArray<FFloat16Color> OutputPixels;
 
-		int32 ExpectedSizeX = InSampleState.TileSize.X;
-		int32 ExpectedSizeY = InSampleState.TileSize.Y;
+		int32 ExpectedSizeX = InSampleState.BackbufferSize.X;
+		int32 ExpectedSizeY = InSampleState.BackbufferSize.Y;
 		OutputPixels.SetNumUninitialized(ExpectedSizeX * ExpectedSizeY);
 
 		// Due to padding, the actual size might be larger than the expected size. If they are the same, do a block copy. Otherwise copy
@@ -180,6 +180,7 @@ void FMoviePipelineSurfaceReader::CopyReadbackTexture_RenderThread(const FMovieP
 		else
 		{
 			UE_LOG(LogMovieRenderPipeline, Log, TEXT("Unexpected size in FMoviePipelineSurfaceReader::CopyReadbackTexture_RenderThread."));
+			UE_LOG(LogMovieRenderPipeline, Log, TEXT("    Tile size:     %d x %d"), InSampleState.TileSize.X, InSampleState.TileSize.Y);
 			UE_LOG(LogMovieRenderPipeline, Log, TEXT("    Expected size: %d x %d"), ExpectedSizeX, ExpectedSizeY);
 			UE_LOG(LogMovieRenderPipeline, Log, TEXT("    Actual size:   %d x %d"), ActualSizeX, ActualSizeY);
 
