MoviePipeline: Fixed Output Settings not properly restoring upon import after editor restart due to default sub object name mismatches.

diff --git a/Engine/Plugins/MovieScene/MovieRenderPipeline/Source/MovieRenderPipelineCore/Private/MoviePipelineMasterConfig.cpp b/Engine/Plugins/MovieScene/MovieRenderPipeline/Source/MovieRenderPipelineCore/Private/MoviePipelineMasterConfig.cpp
--- a/Engine/Plugins/MovieScene/MovieRenderPipeline/Source/MovieRenderPipelineCore/Private/MoviePipelineMasterConfig.cpp
+++ b/Engine/Plugins/MovieScene/MovieRenderPipeline/Source/MovieRenderPipelineCore/Private/MoviePipelineMasterConfig.cpp
@@ -276,7 +276,17 @@ void UMoviePipelineMasterConfig::CopyFrom(UMoviePipelineConfigBase* InConfig)
 	if (InConfig->IsA<UMoviePipelineMasterConfig>())
 	{
 		UMoviePipelineMasterConfig* MasterConfig = CastChecked<UMoviePipelineMasterConfig>(InConfig);
-		OutputSetting = Cast<UMoviePipelineOutputSetting>(StaticDuplicateObject(MasterConfig->OutputSetting, this, NAME_None));
+
+		// Rename our current default subobject so we can duplicate the incoming config ontop the existing
+		// name which is required for it to be loaded correctly.
+		if (UMoviePipelineOutputSetting* ExistingObj = FindObject<UMoviePipelineOutputSetting>(this, TEXT("DefaultOutputSetting")))
+		{
+
+			FName UniqueName = MakeUniqueObjectName(GetTransientPackage(), UMoviePipelineMasterConfig::StaticClass(), "DEAD_MoviePipelineConfig_DefaultOutputSetting");
+			ExistingObj->Rename(*UniqueName.ToString(), GetTransientPackage(), REN_DontCreateRedirectors);
+		}
+
+		OutputSetting = Cast<UMoviePipelineOutputSetting>(StaticDuplicateObject(MasterConfig->OutputSetting, this, FName("DefaultOutputSetting")));
 	}
 }
 
