MoviePipeline: Guard against potential null cvar crash in the event a cvar is used that doesn't exist.

#rnx

diff --git a/Engine/Plugins/MovieScene/MovieRenderPipeline/Source/MovieRenderPipelineCore/Public/MoviePipelineUtils.h b/Engine/Plugins/MovieScene/MovieRenderPipeline/Source/MovieRenderPipelineCore/Public/MoviePipelineUtils.h
--- a/Engine/Plugins/MovieScene/MovieRenderPipeline/Source/MovieRenderPipelineCore/Public/MoviePipelineUtils.h
+++ b/Engine/Plugins/MovieScene/MovieRenderPipeline/Source/MovieRenderPipelineCore/Public/MoviePipelineUtils.h
@@ -11,28 +11,34 @@ class UMoviePipelineAntiAliasingSetting;
 #define MOVIEPIPELINE_STORE_AND_OVERRIDE_CVAR_INT(InOutVariable, CVarName, OverrideValue, bUseOverride) \
 { \
 	IConsoleVariable* CVar = IConsoleManager::Get().FindConsoleVariable(CVarName); \
-	if(bUseOverride) \
+	if(ensureMsgf(CVar, TEXT("Failed to find CVar " #CVarName " to override."))) \
 	{ \
-		InOutVariable = CVar->GetInt(); \
-		CVar->Set(OverrideValue, EConsoleVariableFlags::ECVF_SetByConsole); \
-	} \
-	else \
-	{ \
-		CVar->Set(InOutVariable, EConsoleVariableFlags::ECVF_SetByConsole); \
+		if(bUseOverride) \
+		{ \
+			InOutVariable = CVar->GetInt(); \
+			CVar->Set(OverrideValue, EConsoleVariableFlags::ECVF_SetByConsole); \
+		} \
+		else \
+		{ \
+			CVar->Set(InOutVariable, EConsoleVariableFlags::ECVF_SetByConsole); \
+		} \
 	} \
 }
 
 #define MOVIEPIPELINE_STORE_AND_OVERRIDE_CVAR_FLOAT(InOutVariable, CVarName, OverrideValue, bUseOverride) \
 { \
 	IConsoleVariable* CVar = IConsoleManager::Get().FindConsoleVariable(CVarName); \
-	if(bUseOverride) \
-	{ \
-		InOutVariable = CVar->GetFloat(); \
-		CVar->Set(OverrideValue, EConsoleVariableFlags::ECVF_SetByConsole); \
-	} \
-	else \
+	if(ensureMsgf(CVar, TEXT("Failed to find CVar " #CVarName " to override."))) \
 	{ \
-		CVar->Set(InOutVariable, EConsoleVariableFlags::ECVF_SetByConsole); \
+		if(bUseOverride) \
+		{ \
+			InOutVariable = CVar->GetFloat(); \
+			CVar->Set(OverrideValue, EConsoleVariableFlags::ECVF_SetByConsole); \
+		} \
+		else \
+		{ \
+			CVar->Set(InOutVariable, EConsoleVariableFlags::ECVF_SetByConsole); \
+		} \
 	} \
 }
 
