MoviePipeline: New movie renderer now respects r.screenpercentage and post processing volumes screen percentage settings.

#rnx

diff --git a/Engine/Plugins/MovieScene/MovieRenderPipeline/Source/MovieRenderPipelineCore/Private/MoviePipelineRendering.cpp b/Engine/Plugins/MovieScene/MovieRenderPipeline/Source/MovieRenderPipelineCore/Private/MoviePipelineRendering.cpp
--- a/Engine/Plugins/MovieScene/MovieRenderPipeline/Source/MovieRenderPipelineCore/Private/MoviePipelineRendering.cpp
+++ b/Engine/Plugins/MovieScene/MovieRenderPipeline/Source/MovieRenderPipelineCore/Private/MoviePipelineRendering.cpp
@@ -21,6 +21,7 @@
 #include "Modules/ModuleManager.h"
 #include "MoviePipelineCameraSetting.h"
 #include "Engine/GameViewportClient.h"
+#include "LegacyScreenPercentageDriver.h"
 
 // For flushing async systems
 #include "RendererInterface.h"
@@ -395,6 +396,7 @@ void UMoviePipeline::RenderFrame()
 				SampleState.bWriteSampleToDisk = HighResSettings->bWriteAllSamples;
 				SampleState.ExposureCompensation = CameraSettings->bManualExposure ? CameraSettings->ExposureCompensation : TOptional<float>();
 				SampleState.TextureSharpnessBias = HighResSettings->TextureSharpnessBias;
+				SampleState.GlobalScreenPercentageFraction = FLegacyScreenPercentageDriver::GetCVarResolutionFraction();
 				{
 					SampleState.OverlappedPad = FIntPoint(FMath::CeilToInt(TileResolution.X * HighResSettings->OverlapRatio), 
 														   FMath::CeilToInt(TileResolution.Y * HighResSettings->OverlapRatio));
diff --git a/Engine/Plugins/MovieScene/MovieRenderPipeline/Source/MovieRenderPipelineCore/Public/MovieRenderPipelineDataTypes.h b/Engine/Plugins/MovieScene/MovieRenderPipeline/Source/MovieRenderPipelineCore/Public/MovieRenderPipelineDataTypes.h
--- a/Engine/Plugins/MovieScene/MovieRenderPipeline/Source/MovieRenderPipelineCore/Public/MovieRenderPipelineDataTypes.h
+++ b/Engine/Plugins/MovieScene/MovieRenderPipeline/Source/MovieRenderPipelineCore/Public/MovieRenderPipelineDataTypes.h
@@ -820,10 +820,15 @@ public:
 	TOptional<float> ExposureCompensation;
 
 	/**
-	* Any additional texture mip map bias that should be added when rendering. Can be used to force extra sharpness.
+	* Any additional texture mip map bias that should be added when rendering. Can be used to force extra sharpness. A more negative number means more likely to use a higher quality mip map.
 	*/
 	float TextureSharpnessBias;
 
+	/**
+	* What screen percentage should each view be rendered at? 1.0 = standard size. Can gain more detail (at the cost of longer renders, more device timeouts), consider trying to use TextureSharpnessBias if possible instead.
+	*/
+	float GlobalScreenPercentageFraction;
+
 
 	FMoviePipelineFrameOutputState OutputState;
 
diff --git a/Engine/Plugins/MovieScene/MovieRenderPipeline/Source/MovieRenderPipelineRenderPasses/Private/MoviePipelineDeferredPasses.cpp b/Engine/Plugins/MovieScene/MovieRenderPipeline/Source/MovieRenderPipelineRenderPasses/Private/MoviePipelineDeferredPasses.cpp
--- a/Engine/Plugins/MovieScene/MovieRenderPipeline/Source/MovieRenderPipelineRenderPasses/Private/MoviePipelineDeferredPasses.cpp
+++ b/Engine/Plugins/MovieScene/MovieRenderPipeline/Source/MovieRenderPipelineRenderPasses/Private/MoviePipelineDeferredPasses.cpp
@@ -114,7 +114,7 @@ namespace MoviePipeline
 			.SetRealtimeUpdate(true));
 
 		ViewFamily.SceneCaptureSource = InSampleState.SceneCaptureSource;
-		ViewFamily.SetScreenPercentageInterface(new FLegacyScreenPercentageDriver(ViewFamily, 1.0f, false));
+		ViewFamily.SetScreenPercentageInterface(new FLegacyScreenPercentageDriver(ViewFamily, InSampleState.GlobalScreenPercentageFraction, true));
 		ViewFamily.bWorldIsPaused = InSampleState.bWorldIsPaused;
 
 		// View is added as a child of the ViewFamily
