MoviePipeline: Fixed views not being frozen when using spatial sampling.

#rnx

diff --git a/Engine/Plugins/MovieScene/MovieRenderPipeline/Source/MovieRenderPipelineCore/Private/MoviePipelineRendering.cpp b/Engine/Plugins/MovieScene/MovieRenderPipeline/Source/MovieRenderPipelineCore/Private/MoviePipelineRendering.cpp
--- a/Engine/Plugins/MovieScene/MovieRenderPipeline/Source/MovieRenderPipelineCore/Private/MoviePipelineRendering.cpp
+++ b/Engine/Plugins/MovieScene/MovieRenderPipeline/Source/MovieRenderPipelineCore/Private/MoviePipelineRendering.cpp
@@ -301,7 +301,7 @@ void UMoviePipeline::RenderFrame()
 				// Our spatial samples need to be different positional takes on the same world, thus pausing it.
 				const bool bAllowPause = CurrentCameraCut.State == EMovieRenderShotState::Rendering;
 				const bool bIsLastTile = FIntPoint(TileX, TileY) == FIntPoint(TileCount.X - 1, TileCount.Y - 1);
-				const bool bWorldIsPaused = bAllowPause && !(bIsLastTile && CachedOutputState.IsLastTemporalSample());
+				const bool bWorldIsPaused = bAllowPause && !(bIsLastTile && (RenderSampleIndex == (NumSamplesToRender - 1)));
 
 				// We need to pass camera cut flag on the first sample that gets rendered for a given camera cut. If you don't have any render
 				// warm up frames, we do this on the first render sample because we no longer render the motion blur frame (just evaluate it).
