MoviePipeline: Added a post-finalize Export step for settings which is called after all files are written to disk and finalized. This can be used to implement post-render hooks. Fixed quitting PIE not setting the bShutdownRequested flag so you can differentiate between the two.

#rnx

diff --git a/Engine/Plugins/MovieScene/MovieRenderPipeline/Source/MovieRenderPipelineCore/Private/MoviePipeline.cpp b/Engine/Plugins/MovieScene/MovieRenderPipeline/Source/MovieRenderPipelineCore/Private/MoviePipeline.cpp
--- a/Engine/Plugins/MovieScene/MovieRenderPipeline/Source/MovieRenderPipelineCore/Private/MoviePipeline.cpp
+++ b/Engine/Plugins/MovieScene/MovieRenderPipeline/Source/MovieRenderPipelineCore/Private/MoviePipeline.cpp
@@ -568,7 +568,10 @@ void UMoviePipeline::BeginFinalize()
 
 void UMoviePipeline::BeginExport()
 {
-	// ToDo: Add exports.
+	for (UMoviePipelineSetting* Setting : GetPipelineMasterConfig()->GetAllSettings())
+	{
+		Setting->BeginExport();
+	}
 }
 
 void UMoviePipeline::TickFinalizeOutputContainers(const bool bInForceFinish)
@@ -619,7 +622,7 @@ void UMoviePipeline::TickPostFinalizeExport(const bool bInForceFinish)
 {
 	// This step assumes you have produced data and filled the data structures.
 	check(PipelineState == EMovieRenderPipelineState::Export);
-	UE_LOG(LogMovieRenderPipeline, Log, TEXT("[%d] PostFinalize Export (Start)."), GFrameCounter);
+	UE_LOG(LogMovieRenderPipeline, Verbose, TEXT("[%d] PostFinalize Export (Start)."), GFrameCounter);
 
 	// ToDo: Loop through any extensions (such as XML export) and let them export using all of the
 	// data that was generated during this run such as containers, output names and lengths.
@@ -631,11 +634,11 @@ void UMoviePipeline::TickPostFinalizeExport(const bool bInForceFinish)
 		bAllContainsFinishedProcessing = true;
 
 		// Ask the containers if they're all done processing.
-		// for (UMoviePipelineOutputBase* Container : GetPipelineMasterConfig()->GetOutputContainers())
-		// {
-		// 	bAllContainsFinishedProcessing &= Container->HasFinishedProcessing();
-		// }
-
+		for (UMoviePipelineSetting* Setting : GetPipelineMasterConfig()->GetAllSettings())
+		{
+			bAllContainsFinishedProcessing &= Setting->HasFinishedExporting();
+		}
+		
 		// If we aren't forcing a finish, early out after one loop to keep
 		// the editor/ui responsive.
 		if (!bInForceFinish || bAllContainsFinishedProcessing)
@@ -649,6 +652,8 @@ void UMoviePipeline::TickPostFinalizeExport(const bool bInForceFinish)
 
 	} while (true);
 
+	UE_LOG(LogMovieRenderPipeline, Verbose, TEXT("[%d] PostFinalize Export (End)."), GFrameCounter);
+
 	// If an output container is still working, we'll early out to keep the UI responsive.
 	// If they've forced a finish this will have to be true before we can reach this block.
 	if (!bAllContainsFinishedProcessing)
@@ -656,9 +661,6 @@ void UMoviePipeline::TickPostFinalizeExport(const bool bInForceFinish)
 		return;
 	}
 
-
-	UE_LOG(LogMovieRenderPipeline, Log, TEXT("[%d] PostFinalize Export (End)."), GFrameCounter);
-
 	TransitionToState(EMovieRenderPipelineState::Finished);
 }
 
diff --git a/Engine/Plugins/MovieScene/MovieRenderPipeline/Source/MovieRenderPipelineCore/Private/MoviePipelineInProcessExecutor.cpp b/Engine/Plugins/MovieScene/MovieRenderPipeline/Source/MovieRenderPipelineCore/Private/MoviePipelineInProcessExecutor.cpp
--- a/Engine/Plugins/MovieScene/MovieRenderPipeline/Source/MovieRenderPipelineCore/Private/MoviePipelineInProcessExecutor.cpp
+++ b/Engine/Plugins/MovieScene/MovieRenderPipeline/Source/MovieRenderPipelineCore/Private/MoviePipelineInProcessExecutor.cpp
@@ -5,6 +5,7 @@
 #include "MoviePipelineQueue.h"
 #include "MoviePipelineInProcessExecutorSettings.h"
 #include "Kismet/KismetSystemLibrary.h"
+#include "MovieRenderPipelineCoreModule.h"
 
 #define LOCTEXT_NAMESPACE "MoviePipelineInProcessExecutor"
 void UMoviePipelineInProcessExecutor::Start(const UMoviePipelineExecutorJob* InJob)
@@ -46,6 +47,7 @@ void UMoviePipelineInProcessExecutor::OnMapLoadFinished(UWorld* NewWorld)
 	
 	// Listen for when the pipeline thinks it has finished.
 	ActiveMoviePipeline->OnMoviePipelineFinished().AddUObject(this, &UMoviePipelineInProcessExecutor::OnMoviePipelineFinished);
+	FCoreDelegates::OnEnginePreExit.AddUObject(this, &UMoviePipelineInProcessExecutor::OnApplicationQuit);
 
 	// Wait until we actually recieved the right map and created the pipeline before saying that we're actively rendering
 	bIsRendering = true;
@@ -77,6 +79,21 @@ void UMoviePipelineInProcessExecutor::OnTick()
 	UKismetSystemLibrary::SetWindowTitle(WindowTitle);
 }
 
+void UMoviePipelineInProcessExecutor::OnApplicationQuit()
+{
+	// Only call Shutdown if the pipeline hasn't been finished.
+	if (ActiveMoviePipeline && ActiveMoviePipeline->GetPipelineState() != EMovieRenderPipelineState::Finished)
+	{
+		UE_LOG(LogMovieRenderPipeline, Log, TEXT("MoviePipelineInProcessExecutor: Application quit while Movie Pipeline was still active. Stalling to do full shutdown."));
+
+		// This will flush any outstanding work on the movie pipeline (file writes) immediately
+		ActiveMoviePipeline->RequestShutdown(); // Set the Shutdown Requested flag.
+		ActiveMoviePipeline->Shutdown(); // Flush the shutdown.
+
+		UE_LOG(LogMovieRenderPipeline, Log, TEXT("MoviePipelineInProcessExecutor: Stalling finished, pipeline has shut down."));
+	}
+}
+
 void UMoviePipelineInProcessExecutor::OnMoviePipelineFinished(UMoviePipeline* InMoviePipeline)
 {
 	FCoreDelegates::OnBeginFrame.RemoveAll(this);
diff --git a/Engine/Plugins/MovieScene/MovieRenderPipeline/Source/MovieRenderPipelineCore/Private/MoviePipelineLinearExecutor.cpp b/Engine/Plugins/MovieScene/MovieRenderPipeline/Source/MovieRenderPipelineCore/Private/MoviePipelineLinearExecutor.cpp
--- a/Engine/Plugins/MovieScene/MovieRenderPipeline/Source/MovieRenderPipelineCore/Private/MoviePipelineLinearExecutor.cpp
+++ b/Engine/Plugins/MovieScene/MovieRenderPipeline/Source/MovieRenderPipelineCore/Private/MoviePipelineLinearExecutor.cpp
@@ -43,7 +43,7 @@ void UMoviePipelineLinearExecutorBase::StartPipelineByIndex(int32 InPipelineInde
 		return;
 	}
 	
-	UE_LOG(LogMovieRenderPipeline, Log, TEXT("MoviePipelineLinearExecutorBase starting jobs [%d/%d]"), CurrentPipelineIndex, Queue->GetJobs().Num());
+	UE_LOG(LogMovieRenderPipeline, Log, TEXT("MoviePipelineLinearExecutorBase starting job [%d/%d]"), CurrentPipelineIndex + 1, Queue->GetJobs().Num());
 	Start(Queue->GetJobs()[CurrentPipelineIndex]);
 }
 
diff --git a/Engine/Plugins/MovieScene/MovieRenderPipeline/Source/MovieRenderPipelineCore/Public/MoviePipeline.h b/Engine/Plugins/MovieScene/MovieRenderPipeline/Source/MovieRenderPipelineCore/Public/MoviePipeline.h
--- a/Engine/Plugins/MovieScene/MovieRenderPipeline/Source/MovieRenderPipelineCore/Public/MoviePipeline.h
+++ b/Engine/Plugins/MovieScene/MovieRenderPipeline/Source/MovieRenderPipelineCore/Public/MoviePipeline.h
@@ -64,6 +64,11 @@ public:
 	*/
 	void Shutdown();
 
+	/**
+	* Has RequestShutdown() been called?
+	*/
+	bool IsShutdownRequested() const { return bShutdownRequested; }
+
 	/** 
 	* Called when we have completely finished this pipeline. This means that all frames have been rendered,
 	* all files written to disk, and any post-finalize exports have finished. This Pipeline will call
diff --git a/Engine/Plugins/MovieScene/MovieRenderPipeline/Source/MovieRenderPipelineCore/Public/MoviePipelineInProcessExecutor.h b/Engine/Plugins/MovieScene/MovieRenderPipeline/Source/MovieRenderPipelineCore/Public/MoviePipelineInProcessExecutor.h
--- a/Engine/Plugins/MovieScene/MovieRenderPipeline/Source/MovieRenderPipelineCore/Public/MoviePipelineInProcessExecutor.h
+++ b/Engine/Plugins/MovieScene/MovieRenderPipeline/Source/MovieRenderPipelineCore/Public/MoviePipelineInProcessExecutor.h
@@ -28,6 +28,7 @@ protected:
 private:
 	void OnMapLoadFinished(UWorld* NewWorld);
 	void OnMoviePipelineFinished(UMoviePipeline* InMoviePipeline);
+	void OnApplicationQuit();
 	void OnTick();
 
 private:
diff --git a/Engine/Plugins/MovieScene/MovieRenderPipeline/Source/MovieRenderPipelineCore/Public/MoviePipelineSetting.h b/Engine/Plugins/MovieScene/MovieRenderPipeline/Source/MovieRenderPipelineCore/Public/MoviePipelineSetting.h
--- a/Engine/Plugins/MovieScene/MovieRenderPipeline/Source/MovieRenderPipelineCore/Public/MoviePipelineSetting.h
+++ b/Engine/Plugins/MovieScene/MovieRenderPipeline/Source/MovieRenderPipelineCore/Public/MoviePipelineSetting.h
@@ -58,6 +58,11 @@ public:
 	virtual UWorld* GetWorld() const override;
 	// ~UObject Interface
 
+	// Post Finalize Export
+	bool HasFinishedExporting() const { return HasFinishedExportingImpl(); }
+	void BeginExport() { BeginExportImpl(); }
+	// ~Post Finalize Export
+
 protected:
 	UMoviePipeline* GetPipeline() const;
 
@@ -115,6 +120,12 @@ public:
 
 	virtual void SetIsUserCustomized(bool bIsUserCustomized) { bUserCustomized = bIsUserCustomized; }
 	virtual bool GetIsUserCustomized() const { return bUserCustomized; }
+
+	/** Has this setting finished any export-related things it needs to do post-finalize? */
+	bool HasFinishedExportingImpl() const { return true; }
+	/** Called once when all files have been finalized. */
+	void BeginExportImpl() { }
+	
 private:
 	UPROPERTY(Transient)
 	TWeakObjectPtr<UMoviePipeline> CachedPipeline;
diff --git a/Engine/Plugins/MovieScene/MovieRenderPipeline/Source/MovieRenderPipelineEditor/Private/MoviePipelinePIEExecutor.cpp b/Engine/Plugins/MovieScene/MovieRenderPipeline/Source/MovieRenderPipelineEditor/Private/MoviePipelinePIEExecutor.cpp
--- a/Engine/Plugins/MovieScene/MovieRenderPipeline/Source/MovieRenderPipelineEditor/Private/MoviePipelinePIEExecutor.cpp
+++ b/Engine/Plugins/MovieScene/MovieRenderPipeline/Source/MovieRenderPipelineEditor/Private/MoviePipelinePIEExecutor.cpp
@@ -159,7 +159,6 @@ void UMoviePipelinePIEExecutor::OnTick()
 void UMoviePipelinePIEExecutor::OnPIEMoviePipelineFinished(UMoviePipeline* InMoviePipeline)
 {
 	// Unsubscribe to the EndPIE event so we don't think the user canceled it.
-	FEditorDelegates::EndPIE.RemoveAll(this);
 	FCoreDelegates::OnBeginFrame.RemoveAll(this);
 
 	if (ActiveMoviePipeline)
@@ -170,27 +169,27 @@ void UMoviePipelinePIEExecutor::OnPIEMoviePipelineFinished(UMoviePipeline* InMov
 
 	// The End Play will happen on the next frame.
 	GEditor->RequestEndPlayMap();
-
-	// Delay for one frame so that PIE can finish shut down. It's not a huge fan of us starting up on the same frame.
-	FTimerHandle Handle;
-	GEditor->GetTimerManager()->SetTimer(Handle, FTimerDelegate::CreateUObject(this, &UMoviePipelinePIEExecutor::DelayedFinishNotification), 1.f, false);
 }
 
 void UMoviePipelinePIEExecutor::OnPIEEnded(bool)
 {
-	FCoreDelegates::OnBeginFrame.RemoveAll(this);
+	FEditorDelegates::EndPIE.RemoveAll(this);
 
-	// If the movie pipeline finishes naturally we unsubscribe from the EndPIE event.
-	// This means that if this gets called, the user has hit escape and wants to cancel.
-	if (ActiveMoviePipeline)
+	// Only call Shutdown if the pipeline hasn't been finished.
+	if (ActiveMoviePipeline && ActiveMoviePipeline->GetPipelineState() != EMovieRenderPipelineState::Finished)
 	{
-		// This will flush any outstanding work on the movie pipeline (file writes) immediately
-		ActiveMoviePipeline->Shutdown();
+		UE_LOG(LogMovieRenderPipeline, Log, TEXT("PIE Ended while Movie Pipeline was still active. Stalling to do full shutdown."));
 
-		// Broadcast that we finished.
-		// ToDo: bAnyJobHadFatalError
-		DelayedFinishNotification();
+		// This will flush any outstanding work on the movie pipeline (file writes) immediately
+		ActiveMoviePipeline->RequestShutdown(); // Set the Shutdown Requested flag.
+		ActiveMoviePipeline->Shutdown(); // Flush the shutdown.
+		
+		UE_LOG(LogMovieRenderPipeline, Log, TEXT("MoviePipelinePIEExecutor: Stalling finished, pipeline has shut down."));
 	}
+	// ToDo: bAnyJobHadFatalError
+
+	// Delay for one frame so that PIE can finish shut down. It's not a huge fan of us starting up on the same frame.
+	GEditor->GetTimerManager()->SetTimerForNextTick(FTimerDelegate::CreateUObject(this, &UMoviePipelinePIEExecutor::DelayedFinishNotification));
 }
 
 void UMoviePipelinePIEExecutor::DelayedFinishNotification()
@@ -208,5 +207,4 @@ void UMoviePipelinePIEExecutor::DelayedFinishNotification()
 	// Now that another frame has passed and we should be OK to start another PIE session, notify our owner.
 	OnIndividualPipelineFinished(MoviePipeline);
 }
-
 #undef LOCTEXT_NAMESPACE // "MoviePipelinePIEExecutor"
\ No newline at end of file
