MoviePipeline: The last used preset is now only stored for the current editor session. This solves the problem where you couldn't go back to a 'blank'/unsaved preset and were always creating new jobs off of the last one you ever used, while maintaining the ease of use of adding multiple jobs and automatically using the previously set preset.

diff --git a/Engine/Plugins/MovieScene/MovieRenderPipeline/Source/MovieRenderPipelineEditor/Public/MovieRenderPipelineSettings.h b/Engine/Plugins/MovieScene/MovieRenderPipeline/Source/MovieRenderPipelineEditor/Public/MovieRenderPipelineSettings.h
--- a/Engine/Plugins/MovieScene/MovieRenderPipeline/Source/MovieRenderPipelineEditor/Public/MovieRenderPipelineSettings.h
+++ b/Engine/Plugins/MovieScene/MovieRenderPipeline/Source/MovieRenderPipelineEditor/Public/MovieRenderPipelineSettings.h
@@ -33,7 +33,7 @@ public:
 	/**
 	* What was the last configuration preset the user used? Can be null.
 	*/
-	UPROPERTY(config, EditAnywhere, BlueprintReadWrite, Category = "Movie Render Pipeline")
+	UPROPERTY(Transient, EditAnywhere, BlueprintReadWrite, Category = "Movie Render Pipeline")
 	TSoftObjectPtr<UMoviePipelineMasterConfig> LastPresetOrigin;
 	
 	/**
