MoviePipeline: Expanded the warning on high-res tiling to include more non-supported features.

diff --git a/Engine/Plugins/MovieScene/MovieRenderPipeline/Source/MovieRenderPipelineCore/Public/MoviePipelineHighResSetting.h b/Engine/Plugins/MovieScene/MovieRenderPipeline/Source/MovieRenderPipelineCore/Public/MoviePipelineHighResSetting.h
--- a/Engine/Plugins/MovieScene/MovieRenderPipeline/Source/MovieRenderPipelineCore/Public/MoviePipelineHighResSetting.h
+++ b/Engine/Plugins/MovieScene/MovieRenderPipeline/Source/MovieRenderPipelineCore/Public/MoviePipelineHighResSetting.h
@@ -4,6 +4,9 @@
 #include "MoviePipelineSetting.h"
 #include "MovieRenderPipelineDataTypes.h"
 #include "MoviePipelineUtils.h"
+#include "MoviePipelineQueue.h"
+#include "MoviePipelineOutputSetting.h"
+#include "MoviePipelineAntiAliasingSetting.h"
 #include "MoviePipelineHighResSetting.generated.h"
 
 UCLASS(Blueprintable)
@@ -21,7 +24,7 @@ public:
 	{
 	}
 	
-	FIntPoint CalculatePaddedBackbufferSize(const FIntPoint& InTileSize)
+	FIntPoint CalculatePaddedBackbufferSize(const FIntPoint& InTileSize) const
 	{
 		int32 OverlappedPadX = FMath::CeilToInt(InTileSize.X * OverlapRatio);
 		int32 OverlappedPadY = FMath::CeilToInt(InTileSize.Y * OverlapRatio);
@@ -41,7 +44,7 @@ public:
 
 		if (TileCount > 1)
 		{
-			ValidationResults.Add(NSLOCTEXT("MovieRenderPipeline", "HighRes_UnsupportedFeatures", "Tiling does not support TAA or Auto Exposure. These will be forced off when rendering. Use Spatial/Temporal sampling and Manual Camera Exposure /w Exposure Compensation instead."));
+			ValidationResults.Add(NSLOCTEXT("MovieRenderPipeline", "HighRes_UnsupportedFeatures", "Tiling does not support all rendering features (bloom, some screen-space effects). Additionally, TAA and Auto Exposure are not supported and will be forced off when rendering. Use Spatial/Temporal sampling and Manual Camera Exposure /w Exposure Compensation instead."));
 			if (FMath::IsNearlyEqual(OverlapRatio, 0.f))
 			{
 				ValidationResults.Add(NSLOCTEXT("MovieRenderPipeline", "HighRes_OverlapNeeded", "Increase the overlap amount to avoid seams in the final image. More overlap results in longer renders so start at 0.1 and increase as needed."));
@@ -56,6 +59,55 @@ public:
 		InOutFormatArgs.Arguments.Add(TEXT("overlap_percent"), OverlapRatio);
 	}
 
+	virtual FText GetFooterText(UMoviePipelineExecutorJob* InJob) const override
+	{
+		if (!InJob || !InJob->GetConfiguration())
+		{
+			return FText();
+		}
+
+		if (TileCount > 1)
+		{
+			// Get the settings not from the job (those are the cached values) but from our owning configuration, which may
+			// exist transiently in the UI.
+			UMoviePipelineConfigBase* OwningConfig = GetTypedOuter<UMoviePipelineConfigBase>();
+			if (OwningConfig)
+			{
+				UMoviePipelineOutputSetting* OutputSettings = OwningConfig->FindSetting<UMoviePipelineOutputSetting>();
+				UMoviePipelineAntiAliasingSetting* AASettings = OwningConfig->FindSetting<UMoviePipelineAntiAliasingSetting>();
+
+				int32 NumTiles = TileCount * TileCount;
+				int32 NumSamplesPerTick = 1;
+				int32 NumSamplesPerTemporal = 1;
+			
+				if (AASettings)
+				{
+					NumSamplesPerTick = AASettings->SpatialSampleCount;
+					NumSamplesPerTemporal = AASettings->TemporalSampleCount;
+				}
+
+				FIntPoint TileResolution = FIntPoint(
+					FMath::CeilToInt(OutputSettings->OutputResolution.X / TileCount),
+					FMath::CeilToInt(OutputSettings->OutputResolution.Y / TileCount));
+
+				TileResolution = CalculatePaddedBackbufferSize(TileResolution);
+
+				FFormatNamedArguments FormatArgs;
+				FormatArgs.Add(TEXT("NumTiles")) = NumTiles;
+				FormatArgs.Add(TEXT("NumPerFrameSamples")) = NumTiles * NumSamplesPerTick;
+				FormatArgs.Add(TEXT("NumTotalFrameSamples")) = NumTiles * NumSamplesPerTick * NumSamplesPerTemporal;
+				FormatArgs.Add(TEXT("ResolutionX")) = TileResolution.X;
+				FormatArgs.Add(TEXT("ResolutionY")) = TileResolution.Y;
+
+				FText OutputFormatString = NSLOCTEXT("MovieRenderPipeline", "HighRes_OutputInfoFmt", "This will render {NumTiles} tiles, each with an individual resolution of {ResolutionX}x{ResolutionY}. There will be a total of {NumPerFrameSamples} renders per tick, and {NumTotalFrameSamples} renders per output frame. Large per-tick sample counts may freeze the editor for long periods of time.");
+				return FText::Format(OutputFormatString, FormatArgs);
+			}
+		}
+
+		return FText();
+	}
+
+
 	virtual void SetupForPipelineImpl(UMoviePipeline* InPipeline) override
 	{
 		if (!(IsEnabled() && GetIsUserCustomized()))
