MoviePipeline: Added a warning about < 8 total samples in anti-aliasing being degraded quality compared to TAA.

diff --git a/Engine/Plugins/MovieScene/MovieRenderPipeline/Source/MovieRenderPipelineCore/Public/MoviePipelineAntiAliasingSetting.h b/Engine/Plugins/MovieScene/MovieRenderPipeline/Source/MovieRenderPipelineCore/Public/MoviePipelineAntiAliasingSetting.h
--- a/Engine/Plugins/MovieScene/MovieRenderPipeline/Source/MovieRenderPipelineCore/Public/MoviePipelineAntiAliasingSetting.h
+++ b/Engine/Plugins/MovieScene/MovieRenderPipeline/Source/MovieRenderPipelineCore/Public/MoviePipelineAntiAliasingSetting.h
@@ -48,6 +48,24 @@ protected:
 				ValidationState = EMoviePipelineValidationState::Warnings;
 			}
 		}
+
+		if (UE::MovieRenderPipeline::GetEffectiveAntiAliasingMethod(this) == EAntiAliasingMethod::AAM_None)
+		{
+			if ((TemporalSampleCount * SpatialSampleCount) < 8)
+			{
+				ValidationResults.Add(NSLOCTEXT("MovieRenderPipeline", "AntiAliasing_InsufficientSamples", "Traditional TAA uses at least 8 samples. Increase sample count to maintain AA quality."));
+				ValidationState = EMoviePipelineValidationState::Warnings;
+			}
+		}
+
+	}
+
+	virtual void GetFilenameFormatArguments(FMoviePipelineFormatArgs& InOutFormatArgs) const override
+	{
+		Super::GetFilenameFormatArguments(InOutFormatArgs);
+
+		InOutFormatArgs.Arguments.Add(TEXT("ts_count"), TemporalSampleCount);
+		InOutFormatArgs.Arguments.Add(TEXT("ss_count"), SpatialSampleCount);
 	}
 
 public:
