MoviePipeline: Added support for rendering viewport widgets to a separate render pass texture. This is not composited onto the final image. This excludes the progress widget shown to the actual player viewport, and burn ins as they are not added to the viewport.

diff --git a/Engine/Plugins/MovieScene/MovieRenderPipeline/Source/MovieRenderPipelineCore/Private/MoviePipeline.cpp b/Engine/Plugins/MovieScene/MovieRenderPipeline/Source/MovieRenderPipelineCore/Private/MoviePipeline.cpp
--- a/Engine/Plugins/MovieScene/MovieRenderPipeline/Source/MovieRenderPipelineCore/Private/MoviePipeline.cpp
+++ b/Engine/Plugins/MovieScene/MovieRenderPipeline/Source/MovieRenderPipelineCore/Private/MoviePipeline.cpp
@@ -1377,4 +1377,12 @@ FString UMoviePipeline::ResolveFilenameFormatArguments(const FString& InFormatSt
 	return ThisTry;
 }
 
+void UMoviePipeline::SetProgressWidgetVisible(bool bVisible)
+{
+	if (DebugWidget)
+	{
+		DebugWidget->SetVisibility(bVisible ? ESlateVisibility::Visible : ESlateVisibility::Collapsed);
+	}
+}
+
 #undef LOCTEXT_NAMESPACE // "MoviePipeline"
diff --git a/Engine/Plugins/MovieScene/MovieRenderPipeline/Source/MovieRenderPipelineCore/Private/MoviePipelineRendering.cpp b/Engine/Plugins/MovieScene/MovieRenderPipeline/Source/MovieRenderPipelineCore/Private/MoviePipelineRendering.cpp
--- a/Engine/Plugins/MovieScene/MovieRenderPipeline/Source/MovieRenderPipelineCore/Private/MoviePipelineRendering.cpp
+++ b/Engine/Plugins/MovieScene/MovieRenderPipeline/Source/MovieRenderPipelineCore/Private/MoviePipelineRendering.cpp
@@ -202,6 +202,9 @@ void UMoviePipeline::RenderFrame()
 		return;
 	}
 	
+	// Hide the progress widget before we render anything. This allows widget captures to not include the progress bar.
+	SetProgressWidgetVisible(false);
+
 	// To produce a frame from the movie pipeline we may render many frames over a period of time, additively collecting the results
 	// together before submitting it for writing on the last result - this is referred to as an "output frame". The 1 (or more) samples
 	// that make up each output frame are referred to as "sample frames". Within each sample frame, we may need to render the scene many
@@ -426,6 +429,9 @@ void UMoviePipeline::RenderFrame()
 			}
 		}
 	}
+
+	// Re-enable the progress widget so when the player viewport is drawn to the preview window, it shows.
+	SetProgressWidgetVisible(true);
 }
 
 void UMoviePipeline::ProcessOutstandingFinishedFrames()
diff --git a/Engine/Plugins/MovieScene/MovieRenderPipeline/Source/MovieRenderPipelineCore/Public/MoviePipeline.h b/Engine/Plugins/MovieScene/MovieRenderPipeline/Source/MovieRenderPipelineCore/Public/MoviePipeline.h
--- a/Engine/Plugins/MovieScene/MovieRenderPipeline/Source/MovieRenderPipelineCore/Public/MoviePipeline.h
+++ b/Engine/Plugins/MovieScene/MovieRenderPipeline/Source/MovieRenderPipelineCore/Public/MoviePipeline.h
@@ -209,6 +209,12 @@ private:
 	/** Allow any Settings to modify the (already duplicated) sequence. This allows inserting automatic pre-roll, etc. */
 	void ModifySequenceViaExtensions(ULevelSequence* InSequence);
 
+	/**
+	* Should the Progress UI be visible on the player's screen?
+	*/
+	void SetProgressWidgetVisible(bool bVisible);
+
+
 
 private:
 	/** Iterates through the changes we've made to a shot and applies the original settings. */
diff --git a/Engine/Plugins/MovieScene/MovieRenderPipeline/Source/MovieRenderPipelineCore/Public/MovieRenderPipelineDataTypes.h b/Engine/Plugins/MovieScene/MovieRenderPipeline/Source/MovieRenderPipelineCore/Public/MovieRenderPipelineDataTypes.h
--- a/Engine/Plugins/MovieScene/MovieRenderPipeline/Source/MovieRenderPipelineCore/Public/MovieRenderPipelineDataTypes.h
+++ b/Engine/Plugins/MovieScene/MovieRenderPipeline/Source/MovieRenderPipelineCore/Public/MovieRenderPipelineDataTypes.h
@@ -921,6 +921,13 @@ struct FImagePixelDataPayload : IImagePixelDataPayload, public TSharedFromThis<F
 	FMoviePipelineRenderPassMetrics SampleState;
 
 	FMoviePipelinePassIdentifier PassIdentifier;
+	
+	/** Does this output data have to be transparent to be useful? Overrides output format to one that supports transparency. */
+	bool bRequireTransparentOutput;
+
+	FImagePixelDataPayload()
+		: bRequireTransparentOutput(false)
+	{}
 
 	/** Is this the first tile of an image and we should start accumulating? */
 	FORCEINLINE bool IsFirstTile() const
diff --git a/Engine/Plugins/MovieScene/MovieRenderPipeline/Source/MovieRenderPipelineRenderPasses/Private/MoviePipelineImageSequenceOutput.cpp b/Engine/Plugins/MovieScene/MovieRenderPipeline/Source/MovieRenderPipelineRenderPasses/Private/MoviePipelineImageSequenceOutput.cpp
--- a/Engine/Plugins/MovieScene/MovieRenderPipeline/Source/MovieRenderPipelineRenderPasses/Private/MoviePipelineImageSequenceOutput.cpp
+++ b/Engine/Plugins/MovieScene/MovieRenderPipeline/Source/MovieRenderPipelineRenderPasses/Private/MoviePipelineImageSequenceOutput.cpp
@@ -68,16 +68,6 @@ void UMoviePipelineImageSequenceOutputBase::OnRecieveImageDataImpl(FMoviePipelin
 
 	FString OutputDirectory = OutputSettings->OutputDirectory.Path;
 
-	const TCHAR* Extension = TEXT("");
-	switch (OutputFormat)
-	{
-	case EImageFormat::PNG: Extension = TEXT("png"); break;
-	case EImageFormat::JPEG: Extension = TEXT("jpeg"); break;
-	case EImageFormat::BMP: Extension = TEXT("bmp"); break;
-	case EImageFormat::EXR: Extension = TEXT("exr"); break;
-	}
-
-
 	for (TPair<FMoviePipelinePassIdentifier, TUniquePtr<FImagePixelData>>& RenderPassData : InMergedOutputFrame->ImageOutputData)
 	{
 		// Don't write out the burn in pass in this loop, it will get handled separately (or composited).
@@ -86,10 +76,34 @@ void UMoviePipelineImageSequenceOutputBase::OnRecieveImageDataImpl(FMoviePipelin
 			continue;
 		}
 
+		EImageFormat PreferredOutputFormat = OutputFormat;
+
+		FImagePixelDataPayload* Payload = RenderPassData.Value->GetPayload<FImagePixelDataPayload>();
+
+		// If the output requires a transparent output (to be useful) then we'll on a per-case basis override their intended
+		// filetype to something that makes that file useful.
+		if (Payload->bRequireTransparentOutput)
+		{
+			if (PreferredOutputFormat == EImageFormat::BMP ||
+				PreferredOutputFormat == EImageFormat::JPEG)
+			{
+				PreferredOutputFormat = EImageFormat::PNG;
+			}
+		}
+
+		const TCHAR* Extension = TEXT("");
+		switch (PreferredOutputFormat)
+		{
+		case EImageFormat::PNG: Extension = TEXT("png"); break;
+		case EImageFormat::JPEG: Extension = TEXT("jpeg"); break;
+		case EImageFormat::BMP: Extension = TEXT("bmp"); break;
+		case EImageFormat::EXR: Extension = TEXT("exr"); break;
+		}
+
 		TUniquePtr<FImagePixelData> QuantizedPixelData = nullptr;
 		
 
-		switch (OutputFormat)
+		switch (PreferredOutputFormat)
 		{
 		case EImageFormat::PNG:
 		case EImageFormat::JPEG:
@@ -129,7 +143,7 @@ void UMoviePipelineImageSequenceOutputBase::OnRecieveImageDataImpl(FMoviePipelin
 		}
 
 		TUniquePtr<FImageWriteTask> TileImageTask = MakeUnique<FImageWriteTask>();
-		TileImageTask->Format = OutputFormat;
+		TileImageTask->Format = PreferredOutputFormat;
 		TileImageTask->CompressionQuality = 100;
 		TileImageTask->Filename = FinalFilePath;
 
@@ -166,7 +180,8 @@ void UMoviePipelineImageSequenceOutputBase::OnRecieveImageDataImpl(FMoviePipelin
 				break;
 			}
 		}
-		else
+		// We don't flip these right now because we assume that it comes in with the correct Transparent vs. Opaque.
+		else if(!Payload->bRequireTransparentOutput)
 		{
 			// Fill the alpha channel when alpha is not supported/enabled.
 			switch (QuantizedPixelData->GetType())
diff --git a/Engine/Plugins/MovieScene/MovieRenderPipeline/Source/MovieRenderPipelineSettings/Private/MoviePipelineBurnInSetting.cpp b/Engine/Plugins/MovieScene/MovieRenderPipeline/Source/MovieRenderPipelineSettings/Private/MoviePipelineBurnInSetting.cpp
--- a/Engine/Plugins/MovieScene/MovieRenderPipeline/Source/MovieRenderPipelineSettings/Private/MoviePipelineBurnInSetting.cpp
+++ b/Engine/Plugins/MovieScene/MovieRenderPipeline/Source/MovieRenderPipelineSettings/Private/MoviePipelineBurnInSetting.cpp
@@ -65,6 +65,7 @@ void UMoviePipelineBurnInSetting::RenderSample_GameThreadImpl(const FMoviePipeli
 			FrameData->OutputState = InSampleState.OutputState;
 			FrameData->PassIdentifier = FMoviePipelinePassIdentifier(TEXT("BurnInOverlay"));
 			FrameData->SampleState = InSampleState;
+			FrameData->bRequireTransparentOutput = true;
 
 			TUniquePtr<FImagePixelData> PixelData = MakeUnique<TImagePixelData<FColor>>(InSampleState.BackbufferSize, TArray64<FColor>(MoveTemp(RawPixels)), FrameData);
 
diff --git a/Engine/Plugins/MovieScene/MovieRenderPipeline/Source/MovieRenderPipelineSettings/Private/MoviePipelineWidgetRenderSetting.cpp b/Engine/Plugins/MovieScene/MovieRenderPipeline/Source/MovieRenderPipelineSettings/Private/MoviePipelineWidgetRenderSetting.cpp
new file mode 100644
--- /dev/null
+++ b/Engine/Plugins/MovieScene/MovieRenderPipeline/Source/MovieRenderPipelineSettings/Private/MoviePipelineWidgetRenderSetting.cpp
@@ -0,0 +1,115 @@
+// Copyright Epic Games, Inc. All Rights Reserved.
+
+#include "MoviePipelineWidgetRenderSetting.h"
+#include "Slate/WidgetRenderer.h"
+#include "MovieRenderPipelineDataTypes.h"
+#include "MoviePipelineBurnInWidget.h"
+#include "MoviePipelineOutputSetting.h"
+#include "MoviePipelineMasterConfig.h"
+#include "MoviePipeline.h"
+#include "Engine/TextureRenderTarget2D.h"
+#include "MoviePipelineOutputBuilder.h"
+#include "ImagePixelData.h"
+#include "Widgets/SViewport.h"
+#include "Slate/SGameLayerManager.h"
+
+void UMoviePipelineWidgetRenderer::GatherOutputPassesImpl(TArray<FMoviePipelinePassIdentifier>& ExpectedRenderPasses)
+{
+	// If this was transiently added, don't render.
+	if (!GetIsUserCustomized() || !IsEnabled())
+	{
+		return;
+	}
+
+	ExpectedRenderPasses.Add(FMoviePipelinePassIdentifier(TEXT("ViewportUI")));
+}
+
+void UMoviePipelineWidgetRenderer::RenderSample_GameThreadImpl(const FMoviePipelineRenderPassMetrics& InSampleState)
+{
+	if (InSampleState.bDiscardResult)
+	{
+		return;
+	}
+
+	const bool bFirstTile = InSampleState.GetTileIndex() == 0;
+	const bool bFirstSpatial = InSampleState.SpatialSampleIndex == (InSampleState.SpatialSampleCount - 1);
+	const bool bFirstTemporal = InSampleState.TemporalSampleIndex == (InSampleState.TemporalSampleCount - 1);
+
+	if (bFirstTile && bFirstSpatial && bFirstTemporal)
+	{
+		// Draw the widget to the render target
+		FRenderTarget* BackbufferRenderTarget = RenderTarget->GameThread_GetRenderTargetResource();
+		
+		ULocalPlayer* LocalPlayer = GetWorld()->GetFirstLocalPlayerFromController();
+
+		// Cast the interface to a widget is a little yucky but the implementation is unlikely to change.
+		TSharedPtr<SGameLayerManager> GameLayerManager = StaticCastSharedPtr<SGameLayerManager>(LocalPlayer->ViewportClient->GetGameLayerManager());
+
+		WidgetRenderer->DrawWidget(BackbufferRenderTarget, GameLayerManager.ToSharedRef(), 1.f, FVector2D(RenderTarget->SizeX, RenderTarget->SizeY), (float)InSampleState.OutputState.TimeData.FrameDeltaTime);
+
+		TSharedPtr<FMoviePipelineOutputMerger, ESPMode::ThreadSafe> OutputBuilder = GetPipeline()->OutputBuilder;
+
+		ENQUEUE_RENDER_COMMAND(BurnInRenderTargetResolveCommand)(
+			[InSampleState, BackbufferRenderTarget, OutputBuilder](FRHICommandListImmediate& RHICmdList)
+		{
+			FIntRect SourceRect = FIntRect(0, 0, BackbufferRenderTarget->GetSizeXY().X, BackbufferRenderTarget->GetSizeXY().Y);
+
+			// Read the data back to the CPU
+			TArray<FColor> RawPixels;
+			RawPixels.SetNum(SourceRect.Width() * SourceRect.Height());
+
+			FReadSurfaceDataFlags ReadDataFlags(ERangeCompressionMode::RCM_MinMax);
+			ReadDataFlags.SetLinearToGamma(false);
+
+			RHICmdList.ReadSurfaceData(BackbufferRenderTarget->GetRenderTargetTexture(), SourceRect, RawPixels, ReadDataFlags);
+
+			TSharedRef<FImagePixelDataPayload, ESPMode::ThreadSafe> FrameData = MakeShared<FImagePixelDataPayload, ESPMode::ThreadSafe>();
+			FrameData->OutputState = InSampleState.OutputState;
+			FrameData->PassIdentifier = FMoviePipelinePassIdentifier(TEXT("ViewportUI"));
+			FrameData->SampleState = InSampleState;
+			FrameData->bRequireTransparentOutput = true;
+
+			TUniquePtr<FImagePixelData> PixelData = MakeUnique<TImagePixelData<FColor>>(InSampleState.BackbufferSize, TArray64<FColor>(MoveTemp(RawPixels)), FrameData);
+
+			OutputBuilder->OnCompleteRenderPassDataAvailable_AnyThread(MoveTemp(PixelData), FrameData);
+		});
+	}
+}
+
+void UMoviePipelineWidgetRenderer::SetupImpl(TArray<TSharedPtr<MoviePipeline::FMoviePipelineEnginePass>>& InEnginePasses, const MoviePipeline::FMoviePipelineRenderPassInitSettings& InPassInitSettings)
+{
+	// If this was transiently added, don't render.
+	if (!GetIsUserCustomized() || !IsEnabled())
+	{
+		return;
+	}
+
+	RenderTarget = NewObject<UTextureRenderTarget2D>();
+	RenderTarget->ClearColor = FLinearColor::Transparent;
+
+	bool bInForceLinearGamma = false;
+	FIntPoint OutputResolution = GetPipeline()->GetPipelineMasterConfig()->FindSetting<UMoviePipelineOutputSetting>()->OutputResolution;
+	RenderTarget->InitCustomFormat(OutputResolution.X, OutputResolution.Y, EPixelFormat::PF_B8G8R8A8, bInForceLinearGamma);
+
+	bool bApplyGammaCorrection = false;
+	WidgetRenderer = MakeShared<FWidgetRenderer>(bApplyGammaCorrection);
+}
+
+void UMoviePipelineWidgetRenderer::TeardownImpl() 
+{
+	// If this was transiently added, don't render.
+	if (!GetIsUserCustomized() || !IsEnabled())
+	{
+		return;
+	}
+	
+	FlushRenderingCommands();
+
+	WidgetRenderer = nullptr;
+	RenderTarget = nullptr;
+}
+
+FText UMoviePipelineWidgetRenderer::GetFooterText(UMoviePipelineExecutorJob* InJob) const
+{ 
+	return NSLOCTEXT("MovieRenderPipeline", "WidgetRenderSetting_NoCompositeWarning", "This will render widgets added to the Viewport to a separate texture with alpha. This is currently not composited onto the final image, and will need to be combined in post.");
+}
\ No newline at end of file
diff --git a/Engine/Plugins/MovieScene/MovieRenderPipeline/Source/MovieRenderPipelineSettings/Public/MoviePipelineWidgetRenderSetting.h b/Engine/Plugins/MovieScene/MovieRenderPipeline/Source/MovieRenderPipelineSettings/Public/MoviePipelineWidgetRenderSetting.h
new file mode 100644
--- /dev/null
+++ b/Engine/Plugins/MovieScene/MovieRenderPipeline/Source/MovieRenderPipelineSettings/Public/MoviePipelineWidgetRenderSetting.h
@@ -0,0 +1,41 @@
+// Copyright Epic Games, Inc. All Rights Reserved.
+#pragma once
+
+#include "MoviePipelineRenderPass.h"
+#include "MovieRenderPipelineDataTypes.h"
+#include "MoviePipelineWidgetRenderSetting.generated.h"
+
+class FWidgetRenderer;
+class UTextureRenderTarget2D;
+class UMoviePipelineBurnInWidget;
+
+namespace MoviePipeline { struct FMoviePipelineRenderPassInitSettings; }
+
+UCLASS(Blueprintable)
+class UMoviePipelineWidgetRenderer : public UMoviePipelineRenderPass
+{
+	GENERATED_BODY()
+
+protected:
+	// UMoviePipelineRenderPass Interface
+	virtual void SetupImpl(TArray<TSharedPtr<MoviePipeline::FMoviePipelineEnginePass>>& InEnginePasses, const MoviePipeline::FMoviePipelineRenderPassInitSettings& InPassInitSettings) override;
+	virtual void TeardownImpl() override;
+	virtual void GatherOutputPassesImpl(TArray<FMoviePipelinePassIdentifier>& ExpectedRenderPasses) override;
+	virtual void RenderSample_GameThreadImpl(const FMoviePipelineRenderPassMetrics& InSampleState) override;
+	virtual FText GetFooterText(UMoviePipelineExecutorJob* InJob) const override;
+	// ~UMoviePipelineRenderPass Interface
+
+
+public:
+#if WITH_EDITOR
+	virtual FText GetDisplayText() const override { return NSLOCTEXT("MovieRenderPipeline", "WidgetRendererSettingDisplayName", "UI Renderer (Non-Composited)"); }
+	virtual FText GetCategoryText() const { return NSLOCTEXT("MovieRenderPipeline", "DefaultCategoryName_Text", "Rendering"); }
+#endif
+	virtual bool IsValidOnShots() const override { return false; }
+	virtual bool IsValidOnMaster() const override { return true; }
+private:
+	TSharedPtr<FWidgetRenderer> WidgetRenderer;
+
+	UPROPERTY(Transient)
+	UTextureRenderTarget2D* RenderTarget;
+};
\ No newline at end of file
