MoviePipeline: An empty output string now defaults to the MovieRender pipeline directory only on load. This allows you to make presets that can work in any project by making the output directory empty before saving it as a preset.

diff --git a/Engine/Plugins/MovieScene/MovieRenderPipeline/Source/MovieRenderPipelineEditor/Private/Widgets/SMoviePipelineConfigPanel.cpp b/Engine/Plugins/MovieScene/MovieRenderPipeline/Source/MovieRenderPipelineEditor/Private/Widgets/SMoviePipelineConfigPanel.cpp
--- a/Engine/Plugins/MovieScene/MovieRenderPipeline/Source/MovieRenderPipelineEditor/Private/Widgets/SMoviePipelineConfigPanel.cpp
+++ b/Engine/Plugins/MovieScene/MovieRenderPipeline/Source/MovieRenderPipelineEditor/Private/Widgets/SMoviePipelineConfigPanel.cpp
@@ -7,6 +7,7 @@
 #include "MoviePipelineShotConfig.h"
 #include "MovieRenderPipelineSettings.h"
 #include "MoviePipelineConfigBase.h"
+#include "MoviePipelineOutputSetting.h"
 
 // Slate Includes
 #include "Widgets/Text/STextBlock.h"
@@ -462,6 +463,19 @@ void SMoviePipelineConfigPanel::OnSaveAsPreset()
 	{
 		NewPreset->CopyFrom(TransientPreset);
 		NewPreset->MarkPackageDirty();
+
+		// We need to be able to save an empty string into the file. When we duplicate the Output Setting in CopyFrom, the PostLoad is fired which fixes it back up
+		// to point to the local output directory. A bit unfortunate but we need to do this after the CopyFrom since it needs to only happen on asset save.
+		UMoviePipelineOutputSetting* ExistingOutputSetting = TransientPreset->FindSetting<UMoviePipelineOutputSetting>();
+		UMoviePipelineOutputSetting* NewOutputSetting = NewPreset->FindSetting<UMoviePipelineOutputSetting>();
+		if (ExistingOutputSetting && NewOutputSetting)
+		{
+			if (ExistingOutputSetting->OutputDirectory.Path.Len() == 0)
+			{
+				// If an empty string, it will get filled via PostLoad either next time the preset is brought into the UI or the asset is loaded.
+				NewOutputSetting->OutputDirectory.Path = ExistingOutputSetting->OutputDirectory.Path;
+			}
+		}
 		
 		// Make the Display Name match the preset name so when imported later they look like the asset.
 		NewPreset->DisplayName = NewAssetName;
