Movie Pipeline: Fix duplicate loc string

diff --git a/Engine/Plugins/MovieScene/MovieRenderPipeline/Source/MovieRenderPipelineSettings/Public/MoviePipelineWidgetRenderSetting.h b/Engine/Plugins/MovieScene/MovieRenderPipeline/Source/MovieRenderPipelineSettings/Public/MoviePipelineWidgetRenderSetting.h
--- a/Engine/Plugins/MovieScene/MovieRenderPipeline/Source/MovieRenderPipelineSettings/Public/MoviePipelineWidgetRenderSetting.h
+++ b/Engine/Plugins/MovieScene/MovieRenderPipeline/Source/MovieRenderPipelineSettings/Public/MoviePipelineWidgetRenderSetting.h
@@ -29,7 +29,7 @@ protected:
 public:
 #if WITH_EDITOR
 	virtual FText GetDisplayText() const override { return NSLOCTEXT("MovieRenderPipeline", "WidgetRendererSettingDisplayName", "UI Renderer (Non-Composited)"); }
-	virtual FText GetCategoryText() const { return NSLOCTEXT("MovieRenderPipeline", "DefaultCategoryName_Text", "Rendering"); }
+	virtual FText GetCategoryText() const { return NSLOCTEXT("MovieRenderPipeline", "WidgetRendererSettingCategoryName", "Rendering"); }
 #endif
 	virtual bool IsValidOnShots() const override { return false; }
 	virtual bool IsValidOnMaster() const override { return true; }
