MoviePipeline: Fixed crash where the accumulator wasn't passing the payload on to the final files.

diff --git a/Engine/Plugins/MovieScene/MovieRenderPipeline/Source/MovieRenderPipelineRenderPasses/Private/MoviePipelineDeferredPasses.cpp b/Engine/Plugins/MovieScene/MovieRenderPipeline/Source/MovieRenderPipelineRenderPasses/Private/MoviePipelineDeferredPasses.cpp
--- a/Engine/Plugins/MovieScene/MovieRenderPipeline/Source/MovieRenderPipelineRenderPasses/Private/MoviePipelineDeferredPasses.cpp
+++ b/Engine/Plugins/MovieScene/MovieRenderPipeline/Source/MovieRenderPipelineRenderPasses/Private/MoviePipelineDeferredPasses.cpp
@@ -553,9 +553,9 @@ namespace MoviePipeline
 	{
 		SCOPE_CYCLE_COUNTER(STAT_AccumulateSample_RenderThread);
 
-		bool isWellFormed = InPixelData->IsDataWellFormed();
+		bool bIsWellFormed = InPixelData->IsDataWellFormed();
 
-		if (!isWellFormed)
+		if (!bIsWellFormed)
 		{
 			// figure out why it is not well formed, and print a warning.
 			int64 RawSize = InPixelData->GetRawDataSizeInBytes();
@@ -573,7 +573,7 @@ namespace MoviePipeline
 			UE_LOG(LogMovieRenderPipeline, Log, TEXT("Actual size:   %lld"), ActualTotalSize);
 		}
 
-		check(isWellFormed);
+		check(bIsWellFormed);
 
 		// Writing tiles can be useful for debug reasons. These get passed onto the output every frame.
 		if (InParams.SampleState.bWriteSampleToDisk)
@@ -640,7 +640,7 @@ namespace MoviePipeline
 			if (InPixelData->GetType() == EImagePixelType::Float32)
 			{
 				// 32 bit FLinearColor
-				TUniquePtr<TImagePixelData<FLinearColor> > FinalPixelData = MakeUnique<TImagePixelData<FLinearColor>>(FIntPoint(FullSizeX, FullSizeY));
+				TUniquePtr<TImagePixelData<FLinearColor> > FinalPixelData = MakeUnique<TImagePixelData<FLinearColor>>(FIntPoint(FullSizeX, FullSizeY), InFrameData);
 				InParams.ImageAccumulator->FetchFinalPixelDataLinearColor(FinalPixelData->Pixels);
 
 				// Send the data to the Output Builder
@@ -649,7 +649,7 @@ namespace MoviePipeline
 			else if (InPixelData->GetType() == EImagePixelType::Float16)
 			{
 				// 32 bit FLinearColor
-				TUniquePtr<TImagePixelData<FFloat16Color> > FinalPixelData = MakeUnique<TImagePixelData<FFloat16Color>>(FIntPoint(FullSizeX, FullSizeY));
+				TUniquePtr<TImagePixelData<FFloat16Color> > FinalPixelData = MakeUnique<TImagePixelData<FFloat16Color>>(FIntPoint(FullSizeX, FullSizeY), InFrameData);
 				InParams.ImageAccumulator->FetchFinalPixelDataHalfFloat(FinalPixelData->Pixels);
 
 				// Send the data to the Output Builder
@@ -658,7 +658,7 @@ namespace MoviePipeline
 			else if (InPixelData->GetType() == EImagePixelType::Color)
 			{
 				// 8bit FColors
-				TUniquePtr<TImagePixelData<FColor>> FinalPixelData = MakeUnique<TImagePixelData<FColor>>(FIntPoint(FullSizeX, FullSizeY));
+				TUniquePtr<TImagePixelData<FColor>> FinalPixelData = MakeUnique<TImagePixelData<FColor>>(FIntPoint(FullSizeX, FullSizeY), InFrameData);
 				InParams.ImageAccumulator->FetchFinalPixelDataByte(FinalPixelData->Pixels);
 
 				// Send the data to the Output Builder
