MoviePipeline: Moved Game Override settings to Core instead of Settings module so that the setting can be respected when doing separate process renders.

diff --git a/Engine/Plugins/MovieScene/MovieRenderPipeline/Source/MovieRenderPipelineCore/Private/MoviePipelineGameOverrideSetting.cpp b/Engine/Plugins/MovieScene/MovieRenderPipeline/Source/MovieRenderPipelineCore/Private/MoviePipelineGameOverrideSetting.cpp
new file mode 100644
--- /dev/null
+++ b/Engine/Plugins/MovieScene/MovieRenderPipeline/Source/MovieRenderPipelineCore/Private/MoviePipelineGameOverrideSetting.cpp
@@ -0,0 +1,147 @@
+// Copyright Epic Games, Inc. All Rights Reserved.
+
+#include "MoviePipelineGameOverrideSetting.h"
+#include "Scalability.h"
+#include "Engine/World.h"
+#include "MovieRenderPipelineCoreModule.h"
+#include "MoviePipelineUtils.h"
+
+void UMoviePipelineGameOverrideSetting::SetupForPipelineImpl(UMoviePipeline* InPipeline)
+{
+	// Store the cvar values and apply the ones from this setting
+	UE_LOG(LogMovieRenderPipeline, Log, TEXT("Applying Game Override quality settings and cvars."));
+	ApplyCVarSettings(true);
+}
+
+void UMoviePipelineGameOverrideSetting::TeardownForPipelineImpl(UMoviePipeline* InPipeline)
+{
+	// Restore the previous cvar values the user had
+	UE_LOG(LogMovieRenderPipeline, Log, TEXT("Restoring Game Override quality settings and cvars."));
+	ApplyCVarSettings(false);
+}
+
+void UMoviePipelineGameOverrideSetting::ApplyCVarSettings(const bool bOverrideValues)
+{
+	if (!IsEnabled())
+	{
+		return;
+	}
+
+	if (bCinematicQualitySettings)
+	{
+		if (bOverrideValues)
+		{
+			// Store their previous Scalability settings so we can revert back to them
+			PreviousQualityLevels = Scalability::GetQualityLevels();
+
+			// Create a copy and override to the maximum level for each Scalability category
+			Scalability::FQualityLevels QualityLevels = PreviousQualityLevels;
+			QualityLevels.SetFromSingleQualityLevelRelativeToMax(0);
+
+			// Apply
+			Scalability::SetQualityLevels(QualityLevels);
+		}
+		else
+		{
+			Scalability::SetQualityLevels(PreviousQualityLevels);
+		}
+	}
+
+	switch (TextureStreaming)
+	{
+	case EMoviePipelineTextureStreamingMethod::FullyLoad:
+		MOVIEPIPELINE_STORE_AND_OVERRIDE_CVAR_INT(PreviousFramesForFullUpdate, TEXT("r.Streaming.FramesForFullUpdate"), 0, bOverrideValues);
+		MOVIEPIPELINE_STORE_AND_OVERRIDE_CVAR_INT(PreviousFullyLoadUsedTextures, TEXT("r.Streaming.FullyLoadUsedTextures"), 1, bOverrideValues);
+		break;
+	case EMoviePipelineTextureStreamingMethod::Disabled:
+		MOVIEPIPELINE_STORE_AND_OVERRIDE_CVAR_INT(PreviousTextureStreaming, TEXT("r.TextureStreaming"), 0, bOverrideValues);
+		break;
+	default:
+		// We don't change their texture streaming settings.
+		break;
+	}
+
+	if (bUseLODZero)
+	{
+		MOVIEPIPELINE_STORE_AND_OVERRIDE_CVAR_INT(PreviousForceLOD, TEXT("r.ForceLOD"), 0, bOverrideValues);
+		MOVIEPIPELINE_STORE_AND_OVERRIDE_CVAR_INT(PreviousSkeletalMeshBias, TEXT("r.SkeletalMeshLODBias"), -10, bOverrideValues);
+		MOVIEPIPELINE_STORE_AND_OVERRIDE_CVAR_INT(PreviousParticleLODBias, TEXT("r.ParticleLODBias"), -10, bOverrideValues);
+	}
+
+	if (bDisableHLODs)
+	{
+		// It's a command and not an integer cvar (despite taking 1/0), so we can't cache it 
+		GetWorld()->Exec(GetWorld(), TEXT("r.HLOD 0"));
+	}
+
+	if (bUseHighQualityShadows)
+	{
+		MOVIEPIPELINE_STORE_AND_OVERRIDE_CVAR_INT(PreviousShadowDistanceScale, TEXT("r.Shadow.DistanceScale"), ShadowDistanceScale, bOverrideValues);
+		MOVIEPIPELINE_STORE_AND_OVERRIDE_CVAR_INT(PreviousShadowQuality, TEXT("r.ShadowQuality"), 5, bOverrideValues);
+		MOVIEPIPELINE_STORE_AND_OVERRIDE_CVAR_FLOAT(PreviousShadowRadiusThreshold, TEXT("r.Shadow.RadiusThreshold"), ShadowRadiusThreshold, bOverrideValues);
+	}
+
+	if (bOverrideViewDistanceScale)
+	{
+		MOVIEPIPELINE_STORE_AND_OVERRIDE_CVAR_INT(PreviousViewDistanceScale, TEXT("r.ViewDistanceScale"), ViewDistanceScale, bOverrideValues);
+	}
+}
+
+void UMoviePipelineGameOverrideSetting::BuildNewProcessCommandLineImpl(FString& InOutUnrealURLParams, FString& InOutCommandLineArgs) const
+{
+	if (!IsEnabled())
+	{
+		return;
+	}
+
+	FString CVarCommandLineArgs;
+	FString CVarExecArgs;
+
+	// If they've left the setting enabled but set it to None, we don't do anything.
+	if (GameModeOverride)
+	{
+		InOutUnrealURLParams += FString::Printf(TEXT("?game=%s"), *GameModeOverride->GetPathName());
+	}
+
+	if (bCinematicQualitySettings)
+	{
+		CVarCommandLineArgs += TEXT("sg.ViewDistanceQuality=4,sg.AntiAliasingQuality=4,sg.ShadowQuality=4,sg.PostProcessQuality=4,sg.TextureQuality=4,sg.EffectsQuality=4,sg.FoliageQuality=4,sg.ShadingQuality=4,");
+	}
+
+	switch (TextureStreaming)
+	{
+	case EMoviePipelineTextureStreamingMethod::FullyLoad:
+		CVarCommandLineArgs += TEXT("r.Streaming.FramesForFullUpdate=0,r.Streaming.FullyLoadUsedTextures=1,");
+		break;
+	case EMoviePipelineTextureStreamingMethod::Disabled:
+		CVarCommandLineArgs += TEXT("r.TextureStreaming=0,");
+		break;
+	default:
+		// We don't change their texture streaming settings.
+		break;
+	}
+
+	if (bUseLODZero)
+	{
+		CVarCommandLineArgs += TEXT("r.ForceLOD=0,r.SkeletalMeshLODBias=-4,");
+	}
+
+	if (bDisableHLODs)
+	{
+		// It's a command and not an integer cvar (despite taking 1/0)
+		CVarExecArgs += TEXT("r.HLOD 0,");
+	}
+
+	if (bUseHighQualityShadows)
+	{
+		CVarCommandLineArgs += FString::Printf(TEXT("r.Shadow.DistanceScale=%d,r.ShadowQuality=5,r.Shadow.RadiusThreshold=%f,"), ShadowDistanceScale, ShadowRadiusThreshold);
+	}
+
+	if (bOverrideViewDistanceScale)
+	{
+		CVarCommandLineArgs += FString::Printf(TEXT("r.ViewDistanceScale=%d,"), ViewDistanceScale);
+	}
+
+	// Apply the cvar very early on (device profile) time instead of after the map has loaded if possible
+	InOutCommandLineArgs += FString::Printf(TEXT(" -dpcvars=%s -execcmds=%s"), *CVarCommandLineArgs, *CVarExecArgs);
+}
\ No newline at end of file
diff --git a/Engine/Plugins/MovieScene/MovieRenderPipeline/Source/MovieRenderPipelineCore/Public/MoviePipelineGameMode.h b/Engine/Plugins/MovieScene/MovieRenderPipeline/Source/MovieRenderPipelineCore/Public/MoviePipelineGameMode.h
new file mode 100644
--- /dev/null
+++ b/Engine/Plugins/MovieScene/MovieRenderPipeline/Source/MovieRenderPipelineCore/Public/MoviePipelineGameMode.h
@@ -0,0 +1,27 @@
+// Copyright Epic Games, Inc. All Rights Reserved.
+#pragma once
+
+#include "GameFramework/GameModeBase.h"
+#include "GameFramework/PlayerController.h"
+#include "Kismet/GameplayStatics.h"
+#include "MoviePipelineGameMode.generated.h"
+
+UCLASS()
+class MOVIERENDERPIPELINECORE_API AMoviePipelineGameMode : public AGameModeBase
+{
+	GENERATED_BODY()
+public:
+	virtual void BeginPlay() override
+	{
+		APlayerController* PlayerController = UGameplayStatics::GetPlayerController(this, 0);
+		if(PlayerController)
+		{
+			const bool bCinematicMode = true;
+			const bool bHidePlayer = true;
+			const bool bHideHUD = true;
+			const bool bPreventMovement = true;
+			const bool bPreventTurning = true;
+			PlayerController->SetCinematicMode(bCinematicMode, bHidePlayer, bHideHUD, bPreventMovement, bPreventTurning);
+		}
+	}
+};
\ No newline at end of file
diff --git a/Engine/Plugins/MovieScene/MovieRenderPipeline/Source/MovieRenderPipelineCore/Public/MoviePipelineGameOverrideSetting.h b/Engine/Plugins/MovieScene/MovieRenderPipeline/Source/MovieRenderPipelineCore/Public/MoviePipelineGameOverrideSetting.h
new file mode 100644
--- /dev/null
+++ b/Engine/Plugins/MovieScene/MovieRenderPipeline/Source/MovieRenderPipelineCore/Public/MoviePipelineGameOverrideSetting.h
@@ -0,0 +1,107 @@
+// Copyright Epic Games, Inc. All Rights Reserved.
+#pragma once
+
+#include "MoviePipelineSetting.h"
+#include "GameFramework/GameModeBase.h"
+#include "Templates/SubclassOf.h"
+#include "MoviePipelineGameMode.h"
+#include "Scalability.h"
+#include "MoviePipelineGameOverrideSetting.generated.h"
+
+UENUM(BlueprintType)
+enum class EMoviePipelineTextureStreamingMethod : uint8
+{
+	/** This will not change the texture streaming method / cvars the users has set. */
+	None UMETA(DisplayName="Don't Override" ),
+	/** Disable the Texture Streaming system. Requires higher RAM limits but can help if there are still blurry textures. */
+	Disabled UMETA(DisplayName = "Disable Streaming"),
+	/**  Fully load used textures instead of progressively streaming them in over multiple frames. Requires less VRAM. */
+	FullyLoad UMETA(DisplayName = "Fully Load Used Textures")
+};
+
+UCLASS(BlueprintType)
+class MOVIERENDERPIPELINECORE_API UMoviePipelineGameOverrideSetting : public UMoviePipelineSetting
+{
+	GENERATED_BODY()
+public:
+	UMoviePipelineGameOverrideSetting()
+		: GameModeOverride(AMoviePipelineGameMode::StaticClass())
+		, bCinematicQualitySettings(true)
+		, TextureStreaming(EMoviePipelineTextureStreamingMethod::FullyLoad)
+		, bUseLODZero(true)
+		, bDisableHLODs(true)
+		, bUseHighQualityShadows(true)
+		, ShadowDistanceScale(10)
+		, ShadowRadiusThreshold(0.001f)
+		, bOverrideViewDistanceScale(true)
+		, ViewDistanceScale(50)
+	{
+	}
+
+public:
+#if WITH_EDITOR
+	virtual FText GetDisplayText() const override { return NSLOCTEXT("MovieRenderPipeline", "GameOverrideSettingDisplayName", "Game Overrides"); }
+#endif
+	virtual bool IsValidOnShots() const override { return false; }
+	virtual bool IsValidOnMaster() const override { return true; }
+	virtual void BuildNewProcessCommandLineImpl(FString& InOutUnrealURLParams, FString& InOutCommandLineArgs) const override;
+	virtual void SetupForPipelineImpl(UMoviePipeline* InPipeline) override;
+	virtual void TeardownForPipelineImpl(UMoviePipeline* InPipeline) override;
+protected:
+	void ApplyCVarSettings(const bool bOverrideValues);
+
+public:
+	/** Optional Game Mode to override the map's default game mode with. This can be useful if the game's normal mode displays UI elements or loading screens that you don't want captured. */
+	UPROPERTY(EditAnywhere, BlueprintReadWrite, Category = "Game")
+	TSubclassOf<AGameModeBase> GameModeOverride;
+
+	/** If true, automatically set the engine to the Cinematic Scalability quality settings during render. See the Scalability Reference documentation for information on how to edit cvars to add/change default quality values.*/
+	UPROPERTY(EditAnywhere, BlueprintReadWrite, Category = "Game")
+	bool bCinematicQualitySettings;
+
+	/** Defines which If true, when using texture streaming fully load the required textures each frame instead of loading them in over time. This solves objects being blurry after camera cuts. */
+	UPROPERTY(EditAnywhere, BlueprintReadWrite, Category = "Rendering")
+	EMoviePipelineTextureStreamingMethod TextureStreaming;
+
+	/** Should we try to use the highest quality LOD for meshes and particle systems regardless of distance? */
+	UPROPERTY(EditAnywhere, BlueprintReadWrite, Category = "Rendering")
+	bool bUseLODZero;
+
+	/** Should we disable Hierarchical LODs and instead use their real meshes regardless of distance? */
+	UPROPERTY(EditAnywhere, BlueprintReadWrite, Category = "Rendering")
+	bool bDisableHLODs;
+
+	/** Should we override shadow-related CVars with some high quality preset settings? */
+	UPROPERTY(EditAnywhere, BlueprintReadWrite, Category = "Rendering")
+	bool bUseHighQualityShadows;
+	
+	/** Scalability option to trade shadow distance versus performance for directional lights  */
+	UPROPERTY(EditAnywhere, BlueprintReadWrite, Category = "Rendering", meta=(EditCondition=bUseHighQualityShadows, ClampMin=0.1, UIMin=0.1, UIMax=10))
+	int32 ShadowDistanceScale;
+
+	/** Cull shadow casters if they are too small, value is the minimal screen space bounding sphere radius */
+	UPROPERTY(EditAnywhere, BlueprintReadWrite, Category = "Rendering", meta = (EditCondition = bUseHighQualityShadows, UIMin=0.001))
+	float ShadowRadiusThreshold;
+
+	/** Should we override the View Distance Scale? Can be used in situations where MaxDrawDistance has been set before for in-game performance. */
+	UPROPERTY(EditAnywhere, BlueprintReadWrite, Category = "Rendering")
+	bool bOverrideViewDistanceScale;
+
+	/** Controls the view distance scale. A primitive's MaxDrawDistance is scaled by this value. */
+	UPROPERTY(EditAnywhere, BlueprintReadWrite, Category = "Rendering", meta = (EditCondition = bOverrideViewDistanceScale))
+	int32 ViewDistanceScale;
+
+private:
+	// To restore previous choices when we modify these at runtime. These will be unset if the user doesn't have the override enabled.
+	Scalability::FQualityLevels PreviousQualityLevels;
+	int32 PreviousFramesForFullUpdate;
+	int32 PreviousFullyLoadUsedTextures;
+	int32 PreviousTextureStreaming;
+	int32 PreviousForceLOD;
+	int32 PreviousSkeletalMeshBias;
+	int32 PreviousParticleLODBias;
+	int32 PreviousShadowDistanceScale;
+	int32 PreviousShadowQuality;
+	float PreviousShadowRadiusThreshold;
+	int32 PreviousViewDistanceScale;
+};
\ No newline at end of file
diff --git a/Engine/Plugins/MovieScene/MovieRenderPipeline/Source/MovieRenderPipelineSettings/Private/MoviePipelineGameOverrideSetting.cpp b/Engine/Plugins/MovieScene/MovieRenderPipeline/Source/MovieRenderPipelineSettings/Private/MoviePipelineGameOverrideSetting.cpp
deleted file mode 100644
diff --git a/Engine/Plugins/MovieScene/MovieRenderPipeline/Source/MovieRenderPipelineSettings/Public/MoviePipelineGameMode.h b/Engine/Plugins/MovieScene/MovieRenderPipeline/Source/MovieRenderPipelineSettings/Public/MoviePipelineGameMode.h
deleted file mode 100644
diff --git a/Engine/Plugins/MovieScene/MovieRenderPipeline/Source/MovieRenderPipelineSettings/Public/MoviePipelineGameOverrideSetting.h b/Engine/Plugins/MovieScene/MovieRenderPipeline/Source/MovieRenderPipelineSettings/Public/MoviePipelineGameOverrideSetting.h
deleted file mode 100644
