MoviePipeline: Fixed Render (Remote) never completing if there was more than one movie render in the Queue. Added a check for unsaved maps before launching the movie render queue to prevent the scenario where the engine is opened, a sequence is added to the queue and Render (Remote) is hit which resulted in a separate process that opened and then closed without warning.

diff --git a/Engine/Plugins/MovieScene/MovieRenderPipeline/Source/MovieRenderPipelineCore/Private/MoviePipelineInProcessExecutor.cpp b/Engine/Plugins/MovieScene/MovieRenderPipeline/Source/MovieRenderPipelineCore/Private/MoviePipelineInProcessExecutor.cpp
--- a/Engine/Plugins/MovieScene/MovieRenderPipeline/Source/MovieRenderPipelineCore/Private/MoviePipelineInProcessExecutor.cpp
+++ b/Engine/Plugins/MovieScene/MovieRenderPipeline/Source/MovieRenderPipelineCore/Private/MoviePipelineInProcessExecutor.cpp
@@ -4,8 +4,11 @@
 #include "Engine/World.h"
 #include "MoviePipelineQueue.h"
 #include "MoviePipelineInProcessExecutorSettings.h"
+#include "MoviePipelineGameOverrideSetting.h"
 #include "Kismet/KismetSystemLibrary.h"
+#include "Kismet/GameplayStatics.h"
 #include "MovieRenderPipelineCoreModule.h"
+#include "Misc/PackageName.h"
 
 #define LOCTEXT_NAMESPACE "MoviePipelineInProcessExecutor"
 void UMoviePipelineInProcessExecutor::Start(const UMoviePipelineExecutorJob* InJob)
@@ -112,4 +115,40 @@ void UMoviePipelineInProcessExecutor::OnMoviePipelineFinished(UMoviePipeline* In
 	// Now that another frame has passed and we should be OK to start another PIE session, notify our owner.
 	OnIndividualPipelineFinished(MoviePipeline);
 }
+
+void UMoviePipelineInProcessExecutor::OnIndividualPipelineFinished(UMoviePipeline* FinishedPipeline)
+{
+	if (CurrentPipelineIndex < Queue->GetJobs().Num() - 1)
+	{
+		// Get the next job in the queue
+		UMoviePipelineExecutorJob* NextJob = Queue->GetJobs()[CurrentPipelineIndex + 1];
+
+		FString MapOptions;
+
+		// Initialize the transient settings so that they will exist in time for the GameOverrides check.
+		NextJob->GetConfiguration()->InitializeTransientSettings();
+
+		TArray<UMoviePipelineSetting*> AllSettings = NextJob->GetConfiguration()->GetAllSettings();
+		UMoviePipelineSetting** GameOverridesPtr = AllSettings.FindByPredicate([](UMoviePipelineSetting* InSetting) { return InSetting->GetClass() == UMoviePipelineGameOverrideSetting::StaticClass(); });
+		if (GameOverridesPtr)
+		{
+			UMoviePipelineSetting* Setting = *GameOverridesPtr;
+			if (Setting)
+			{
+				UMoviePipelineGameOverrideSetting* GameOverrideSetting = CastChecked<UMoviePipelineGameOverrideSetting>(Setting);
+				if (GameOverrideSetting->GameModeOverride)
+				{
+					FString GameModeOverride = FPackageName::GetShortName(*GameOverrideSetting->GameModeOverride->GetPathName());
+					MapOptions = TEXT("?game=") + GameModeOverride;
+				}
+
+			}
+		}
+
+		UGameplayStatics::OpenLevel(FinishedPipeline->GetWorld(), NextJob->Map.GetAssetPathName(), true, MapOptions);
+	}
+
+	Super::OnIndividualPipelineFinished(FinishedPipeline);
+
+}
 #undef LOCTEXT_NAMESPACE // "MoviePipelineInProcessExecutor"
diff --git a/Engine/Plugins/MovieScene/MovieRenderPipeline/Source/MovieRenderPipelineCore/Public/MoviePipelineInProcessExecutor.h b/Engine/Plugins/MovieScene/MovieRenderPipeline/Source/MovieRenderPipelineCore/Public/MoviePipelineInProcessExecutor.h
--- a/Engine/Plugins/MovieScene/MovieRenderPipeline/Source/MovieRenderPipelineCore/Public/MoviePipelineInProcessExecutor.h
+++ b/Engine/Plugins/MovieScene/MovieRenderPipeline/Source/MovieRenderPipelineCore/Public/MoviePipelineInProcessExecutor.h
@@ -25,6 +25,8 @@ public:
 
 protected:
 	virtual void Start(const UMoviePipelineExecutorJob* InJob) override;
+	virtual void OnIndividualPipelineFinished(UMoviePipeline* FinishedPipeline) override;
+
 private:
 	void OnMapLoadFinished(UWorld* NewWorld);
 	void OnMoviePipelineFinished(UMoviePipeline* InMoviePipeline);
diff --git a/Engine/Plugins/MovieScene/MovieRenderPipeline/Source/MovieRenderPipelineEditor/Private/MoviePipelineNewProcessExecutor.cpp b/Engine/Plugins/MovieScene/MovieRenderPipeline/Source/MovieRenderPipelineEditor/Private/MoviePipelineNewProcessExecutor.cpp
--- a/Engine/Plugins/MovieScene/MovieRenderPipeline/Source/MovieRenderPipelineEditor/Private/MoviePipelineNewProcessExecutor.cpp
+++ b/Engine/Plugins/MovieScene/MovieRenderPipeline/Source/MovieRenderPipelineEditor/Private/MoviePipelineNewProcessExecutor.cpp
@@ -15,6 +15,7 @@
 #include "UnrealEdMisc.h"
 #include "HAL/PlatformFilemanager.h"
 #include "Misc/FileHelper.h"
+#include "Misc/MessageDialog.h"
 
 #define LOCTEXT_NAMESPACE "MoviePipelineNewProcessExecutor"
 
@@ -36,6 +37,28 @@ void UMoviePipelineNewProcessExecutor::ExecuteImpl(UMoviePipelineQueue* InPipeli
 		return;
 	}
 
+	// Make sure all of the maps in the queue exist on disk somewhere, otherwise the remote process boots up and then fails.
+	bool bHasValidMap = true;
+	for (const UMoviePipelineExecutorJob* Job : InPipelineQueue->GetJobs())
+	{
+		FName PackageName = Job->Map.GetLongPackageName();
+		if (!FPackageName::IsValidLongPackageName(PackageName.ToString()))
+		{
+			bHasValidMap = false;
+			break;
+		}
+	}
+
+	if (!bHasValidMap)
+	{
+		FText FailureReason = LOCTEXT("UnsavedMapFailureDialog", "One or more jobs in the queue have an unsaved map as their target map. These unsaved maps cannot be loaded by an external process, and the render has been aborted.");
+		FMessageDialog::Open(EAppMsgType::Ok, FailureReason);
+
+		OnExecutorFinishedImpl();
+		return;
+	}
+
+
 	if (!ensureMsgf(!ProcessHandle.IsValid(), TEXT("Attempted to start another New Process Executor without the last one quitting. Force killing. This executor only supports one at a time.")))
 	{
 		FPlatformProcess::TerminateProc(ProcessHandle, true);
