CIS: Fixed 2019 Strict compile issue.

diff --git a/Engine/Plugins/MovieScene/MovieRenderPipeline/Source/MovieRenderPipelineEditor/Private/MoviePipelineNewProcessExecutor.cpp b/Engine/Plugins/MovieScene/MovieRenderPipeline/Source/MovieRenderPipelineEditor/Private/MoviePipelineNewProcessExecutor.cpp
--- a/Engine/Plugins/MovieScene/MovieRenderPipeline/Source/MovieRenderPipelineEditor/Private/MoviePipelineNewProcessExecutor.cpp
+++ b/Engine/Plugins/MovieScene/MovieRenderPipeline/Source/MovieRenderPipelineEditor/Private/MoviePipelineNewProcessExecutor.cpp
@@ -41,8 +41,8 @@ void UMoviePipelineNewProcessExecutor::ExecuteImpl(UMoviePipelineQueue* InPipeli
 	bool bHasValidMap = true;
 	for (const UMoviePipelineExecutorJob* Job : InPipelineQueue->GetJobs())
 	{
-		FName PackageName = Job->Map.GetLongPackageName();
-		if (!FPackageName::IsValidLongPackageName(PackageName.ToString()))
+		FString PackageName = Job->Map.GetLongPackageName();
+		if (!FPackageName::IsValidLongPackageName(PackageName))
 		{
 			bHasValidMap = false;
 			break;
