MoviePipeline: Updated EXR Image Sequence to say 16bit instead of 32 bit to reflect the real precision of our data.

diff --git a/Engine/Plugins/MovieScene/MovieRenderPipeline/Source/MovieRenderPipelineRenderPasses/Public/MoviePipelineImageSequenceOutput.h b/Engine/Plugins/MovieScene/MovieRenderPipeline/Source/MovieRenderPipelineRenderPasses/Public/MoviePipelineImageSequenceOutput.h
--- a/Engine/Plugins/MovieScene/MovieRenderPipeline/Source/MovieRenderPipelineRenderPasses/Public/MoviePipelineImageSequenceOutput.h
+++ b/Engine/Plugins/MovieScene/MovieRenderPipeline/Source/MovieRenderPipelineRenderPasses/Public/MoviePipelineImageSequenceOutput.h
@@ -99,7 +99,7 @@ class MOVIERENDERPIPELINERENDERPASSES_API UMoviePipelineImageSequenceOutput_EXR
 	GENERATED_BODY()
 public:
 #if WITH_EDITOR
-	virtual FText GetDisplayText() const override { return NSLOCTEXT("MovieRenderPipeline", "ImgSequenceEXRSettingDisplayName", ".exr Sequence [32bit]"); }
+	virtual FText GetDisplayText() const override { return NSLOCTEXT("MovieRenderPipeline", "ImgSequenceEXRSettingDisplayName", ".exr Sequence [16bit]"); }
 #endif
 public:
 	UMoviePipelineImageSequenceOutput_EXR()
