MoviePipeline: Added a preview of the last render sample produced each frame on the UI. Made the UI not cover up the preview as well. Removed "Disable Preview" setting as it was only a temporary workaround for not having a visual preview of what was being rendered, and using that setting halved your performance due to the player view still being rendered.

diff --git a/Engine/Plugins/MovieScene/MovieRenderPipeline/Source/MovieRenderPipelineCore/Private/MoviePipeline.cpp b/Engine/Plugins/MovieScene/MovieRenderPipeline/Source/MovieRenderPipelineCore/Private/MoviePipeline.cpp
--- a/Engine/Plugins/MovieScene/MovieRenderPipeline/Source/MovieRenderPipelineCore/Private/MoviePipeline.cpp
+++ b/Engine/Plugins/MovieScene/MovieRenderPipeline/Source/MovieRenderPipelineCore/Private/MoviePipeline.cpp
@@ -225,6 +225,8 @@ void UMoviePipeline::Initialize(UMoviePipelineExecutorJob* InJob)
 
 	// Construct a debug UI and bind it to this instance.
 	LoadDebugWidget();
+	
+	GetWorld()->GetGameViewport()->bDisableWorldRendering = true;
 
 	SetupAudioRendering();
 
@@ -388,6 +390,8 @@ void UMoviePipeline::TransitionToState(const EMovieRenderPipelineState InNewStat
 			// And then make sure all frames are sent to the Output Containers before we finalize.
 			ProcessOutstandingFinishedFrames();
 
+			PreviewTexture = nullptr;
+
 			// This is called once notifying output containers that all frames that will be submitted have been submitted.
 			PipelineState = EMovieRenderPipelineState::Finalize;
 			BeginFinalize();
@@ -436,6 +440,7 @@ void UMoviePipeline::TransitionToState(const EMovieRenderPipelineState InNewStat
 			TeardownAudioRendering();
 			RestoreTargetSequenceToOriginalState();
 
+			GetWorld()->GetGameViewport()->bDisableWorldRendering = false;
 			GAreScreenMessagesEnabled = bPrevGScreenMessagesEnabled;
 
 			UE_LOG(LogMovieRenderPipeline, Log, TEXT("Movie Pipeline completed. Duration: %s"), *(FDateTime::UtcNow() - InitializationTime).ToString());
diff --git a/Engine/Plugins/MovieScene/MovieRenderPipeline/Source/MovieRenderPipelineCore/Private/MoviePipelineOutputSetting.cpp b/Engine/Plugins/MovieScene/MovieRenderPipeline/Source/MovieRenderPipelineCore/Private/MoviePipelineOutputSetting.cpp
--- a/Engine/Plugins/MovieScene/MovieRenderPipeline/Source/MovieRenderPipelineCore/Private/MoviePipelineOutputSetting.cpp
+++ b/Engine/Plugins/MovieScene/MovieRenderPipeline/Source/MovieRenderPipelineCore/Private/MoviePipelineOutputSetting.cpp
@@ -20,7 +20,6 @@ UMoviePipelineOutputSetting::UMoviePipelineOutputSetting()
 	, ZeroPadFrameNumbers(4)
 	, FrameNumberOffset(0)
 	, bDisableToneCurve(false)
-	, bHidePreview(true)
 {
 	FileNameFormat = TEXT("{sequence_name}.{frame_number}");
 	OutputDirectory.Path = FPaths::ProjectSavedDir() / TEXT("MovieRenders/");
diff --git a/Engine/Plugins/MovieScene/MovieRenderPipeline/Source/MovieRenderPipelineCore/Private/MoviePipelineRendering.cpp b/Engine/Plugins/MovieScene/MovieRenderPipeline/Source/MovieRenderPipelineCore/Private/MoviePipelineRendering.cpp
--- a/Engine/Plugins/MovieScene/MovieRenderPipeline/Source/MovieRenderPipelineCore/Private/MoviePipelineRendering.cpp
+++ b/Engine/Plugins/MovieScene/MovieRenderPipeline/Source/MovieRenderPipelineCore/Private/MoviePipelineRendering.cpp
@@ -156,11 +156,6 @@ void UMoviePipeline::SetupRenderingPipelineForShot(FMoviePipelineShotInfo& Shot)
 		NumOutputPasses++;
 	}
 
-	if (OutputSettings->bHidePreview)
-	{
-		GetWorld()->GetGameViewport()->bDisableWorldRendering = true;
-	}
-
 	UE_LOG(LogMovieRenderPipeline, Log, TEXT("Finished setting up rendering for shot. Shot has %d Engine Passes and %d Output Passes."), ActiveRenderPasses.Num(), NumOutputPasses);
 }
 
diff --git a/Engine/Plugins/MovieScene/MovieRenderPipeline/Source/MovieRenderPipelineCore/Public/MoviePipeline.h b/Engine/Plugins/MovieScene/MovieRenderPipeline/Source/MovieRenderPipelineCore/Public/MoviePipeline.h
--- a/Engine/Plugins/MovieScene/MovieRenderPipeline/Source/MovieRenderPipelineCore/Public/MoviePipeline.h
+++ b/Engine/Plugins/MovieScene/MovieRenderPipeline/Source/MovieRenderPipelineCore/Public/MoviePipeline.h
@@ -23,6 +23,7 @@ class FMoviePipelineOutputMerger;
 class IImageWriteQueue;
 class UMoviePipelineExecutorJob;
 class UMoviePipelineSetting;
+class UTexture;
 
 
 
@@ -98,6 +99,11 @@ public:
 public:
 	ULevelSequence* GetTargetSequence() const { return TargetSequence; }
 
+	UFUNCTION(BlueprintPure, Category = "Movie Render Pipeline")
+	UTexture* GetPreviewTexture() const { return PreviewTexture; }
+
+	void SetPreviewTexture(UTexture* InTexture) { PreviewTexture = InTexture; }
+
 	const TArray<FMoviePipelineShotInfo>& GetShotList() const { return ShotList; }
 
 	int32 GetCurrentShotIndex() const { return CurrentShotIndex; }
@@ -267,6 +273,9 @@ private:
 	UPROPERTY(Transient)
 	UMovieRenderDebugWidget* DebugWidget;
 
+	UPROPERTY(Transient)
+	UTexture* PreviewTexture;
+
 	/** A list of all of the shots we are going to render out from this sequence. */
 	TArray<FMoviePipelineShotInfo> ShotList;
 
diff --git a/Engine/Plugins/MovieScene/MovieRenderPipeline/Source/MovieRenderPipelineCore/Public/MoviePipelineOutputSetting.h b/Engine/Plugins/MovieScene/MovieRenderPipeline/Source/MovieRenderPipelineCore/Public/MoviePipelineOutputSetting.h
--- a/Engine/Plugins/MovieScene/MovieRenderPipeline/Source/MovieRenderPipelineCore/Public/MoviePipelineOutputSetting.h
+++ b/Engine/Plugins/MovieScene/MovieRenderPipeline/Source/MovieRenderPipelineCore/Public/MoviePipelineOutputSetting.h
@@ -96,8 +96,4 @@ public:
 	/** If true the Filmic Tone Curve will not be applied. Disabling this will allow you to export linear data for EXRs. */
 	UPROPERTY(EditAnywhere, BlueprintReadWrite, AdvancedDisplay, Category = "Misc")
 	bool bDisableToneCurve;
-
-	/** Should the normal player viewport be visible when rendering? This increases render cost and is only an approximation of what is actually being rendered. */
-	UPROPERTY(EditAnywhere, BlueprintReadWrite, AdvancedDisplay, Category = "Misc")
-	bool bHidePreview;
 };
\ No newline at end of file
diff --git a/Engine/Plugins/MovieScene/MovieRenderPipeline/Source/MovieRenderPipelineRenderPasses/Private/MoviePipelineDeferredPasses.cpp b/Engine/Plugins/MovieScene/MovieRenderPipeline/Source/MovieRenderPipelineRenderPasses/Private/MoviePipelineDeferredPasses.cpp
--- a/Engine/Plugins/MovieScene/MovieRenderPipeline/Source/MovieRenderPipelineRenderPasses/Private/MoviePipelineDeferredPasses.cpp
+++ b/Engine/Plugins/MovieScene/MovieRenderPipeline/Source/MovieRenderPipelineRenderPasses/Private/MoviePipelineDeferredPasses.cpp
@@ -55,6 +55,9 @@ namespace MoviePipeline
 		// Initialize to the tile size (not final size) and use a 16 bit back buffer to avoid precision issues when accumulating later
 		TileRenderTarget->InitCustomFormat(InInitSettings.BackbufferResolution.X, InInitSettings.BackbufferResolution.Y, EPixelFormat::PF_FloatRGBA, false);
 
+		GetPipeline()->SetPreviewTexture(TileRenderTarget.Get());
+
+
 		// Allocate 
 		ViewState.Allocate();
 
@@ -65,6 +68,8 @@ namespace MoviePipeline
 	{
 		FMoviePipelineEnginePass::Teardown();
 
+		GetPipeline()->SetPreviewTexture(nullptr);
+
 		// This may call FlushRenderingCommands if there are outstanding readbacks that need to happen.
 		SurfaceQueue->Shutdown();
 
