Sequencer: UnfocusedVolumeMultiplier should be fixed at 1.0f during movie rendering

diff --git a/Engine/Plugins/MovieScene/MovieRenderPipeline/Source/MovieRenderPipelineCore/Private/MoviePipelineAudioRendering.cpp b/Engine/Plugins/MovieScene/MovieRenderPipeline/Source/MovieRenderPipelineCore/Private/MoviePipelineAudioRendering.cpp
--- a/Engine/Plugins/MovieScene/MovieRenderPipeline/Source/MovieRenderPipelineCore/Private/MoviePipelineAudioRendering.cpp
+++ b/Engine/Plugins/MovieScene/MovieRenderPipeline/Source/MovieRenderPipelineCore/Private/MoviePipelineAudioRendering.cpp
@@ -62,6 +62,9 @@ void UMoviePipeline::SetupAudioRendering()
 
 void UMoviePipeline::TeardownAudioRendering()
 {
+	// Restore previous unfocused audio multiplier, to no longer force audio when unfocused
+	FApp::SetUnfocusedVolumeMultiplier(AudioState.PrevUnfocusedAudioMultiplier);
+
 	// Restore our cached CVar value.
 	IConsoleVariable* AudioRenderEveryTickCvar = IConsoleManager::Get().FindConsoleVariable(TEXT("au.nrt.RenderEveryTick"));
 
@@ -92,8 +95,6 @@ void UMoviePipeline::StartAudioRecording()
 			WeakSubmix.Pin()->OnStartRecordingOutput(ExpectedDuration);
 		}
 	}
-
-	FApp::SetUnfocusedVolumeMultiplier(AudioState.PrevUnfocusedAudioMultiplier);
 }
 
 void UMoviePipeline::StopAudioRecording()
