Movie Pipeline: Fixed the example editor utility widget not working due to API changes during development.

This removes an unused variable (to prevent confusion), adds UFUNCTIONs to some functions, and adds some safe guards against common scripting mistakes.

diff --git a/Engine/Plugins/MovieScene/MovieRenderPipeline/Source/MovieRenderPipelineCore/Private/MoviePipelineMasterConfig.cpp b/Engine/Plugins/MovieScene/MovieRenderPipeline/Source/MovieRenderPipelineCore/Private/MoviePipelineMasterConfig.cpp
--- a/Engine/Plugins/MovieScene/MovieRenderPipeline/Source/MovieRenderPipelineCore/Private/MoviePipelineMasterConfig.cpp
+++ b/Engine/Plugins/MovieScene/MovieRenderPipeline/Source/MovieRenderPipelineCore/Private/MoviePipelineMasterConfig.cpp
@@ -86,12 +86,6 @@ UMoviePipelineShotConfig* UMoviePipelineMasterConfig::GetConfigForShot(const FSt
 {
 	UMoviePipelineShotConfig* OutConfig = PerShotConfigMapping.FindRef(ShotName);
 
-	// They didn't customize this shot, return the global pipeline default
-	if (!OutConfig)
-	{
-		OutConfig = DefaultShotConfig;
-	}
-
 	return OutConfig;
 }
 
@@ -255,8 +249,6 @@ TRange<FFrameNumber> UMoviePipelineMasterConfig::GetEffectivePlaybackRange(const
 
 UMoviePipelineMasterConfig::UMoviePipelineMasterConfig()
 {
-	DefaultShotConfig = CreateDefaultSubobject<UMoviePipelineShotConfig>("DefaultShotConfig");
-	
 	// Always add at least the output settings block since having a framerate/directory/etc. is critical.
 	OutputSetting = CreateDefaultSubobject<UMoviePipelineOutputSetting>("DefaultOutputSetting");
 }
diff --git a/Engine/Plugins/MovieScene/MovieRenderPipeline/Source/MovieRenderPipelineCore/Private/MoviePipelineRendering.cpp b/Engine/Plugins/MovieScene/MovieRenderPipeline/Source/MovieRenderPipelineCore/Private/MoviePipelineRendering.cpp
--- a/Engine/Plugins/MovieScene/MovieRenderPipeline/Source/MovieRenderPipelineCore/Private/MoviePipelineRendering.cpp
+++ b/Engine/Plugins/MovieScene/MovieRenderPipeline/Source/MovieRenderPipelineCore/Private/MoviePipelineRendering.cpp
@@ -226,7 +226,10 @@ void UMoviePipeline::RenderFrame()
 
 	int32 NumSpatialSamples = AntiAliasingSettings->SpatialSampleCount;
 	int32 NumTemporalSamples = AntiAliasingSettings->TemporalSampleCount;
-	ensure(TileCount.X > 0 && TileCount .Y> 0 && NumSpatialSamples > 0);
+	if (!ensureAlways(TileCount.X > 0 && TileCount.Y > 0 && NumSpatialSamples > 0 && NumTemporalSamples > 0))
+	{
+		return;
+	}
 
 	FrameInfo.PrevViewLocation = FrameInfo.CurrViewLocation;
 	FrameInfo.PrevViewRotation = FrameInfo.CurrViewRotation;
diff --git a/Engine/Plugins/MovieScene/MovieRenderPipeline/Source/MovieRenderPipelineCore/Public/MoviePipelineExecutor.h b/Engine/Plugins/MovieScene/MovieRenderPipeline/Source/MovieRenderPipelineCore/Public/MoviePipelineExecutor.h
--- a/Engine/Plugins/MovieScene/MovieRenderPipeline/Source/MovieRenderPipelineCore/Public/MoviePipelineExecutor.h
+++ b/Engine/Plugins/MovieScene/MovieRenderPipeline/Source/MovieRenderPipelineCore/Public/MoviePipelineExecutor.h
@@ -45,6 +45,11 @@ public:
 	UFUNCTION(BlueprintCallable, Category = "Movie Render Pipeline")
 	void Execute(UMoviePipelineQueue* InPipelineQueue)
 	{
+		if (!InPipelineQueue)
+		{
+			FFrame::KismetExecutionMessage(TEXT("Cannot execute movie render queue with null queue object."), ELogVerbosity::Error);
+			return;
+		}
 		ExecuteImpl(InPipelineQueue);
 	}
 
diff --git a/Engine/Plugins/MovieScene/MovieRenderPipeline/Source/MovieRenderPipelineCore/Public/MoviePipelineMasterConfig.h b/Engine/Plugins/MovieScene/MovieRenderPipeline/Source/MovieRenderPipelineCore/Public/MoviePipelineMasterConfig.h
--- a/Engine/Plugins/MovieScene/MovieRenderPipeline/Source/MovieRenderPipelineCore/Public/MoviePipelineMasterConfig.h
+++ b/Engine/Plugins/MovieScene/MovieRenderPipeline/Source/MovieRenderPipelineCore/Public/MoviePipelineMasterConfig.h
@@ -75,15 +75,11 @@ protected:
 
 	virtual void OnSettingAdded(UMoviePipelineSetting* InSetting) override;
 	virtual void OnSettingRemoved(UMoviePipelineSetting* InSetting) override;
+
 	void AddTransientSettingByClass(const UClass* InSettingClass);
-public:
-	
-	/** The default shot-setup to use for any shot that doesn't a specific implementation. This is required! */
-	UPROPERTY(EditAnywhere, Instanced, BlueprintReadWrite, Category = "Movie Render Pipeline")
-	UMoviePipelineShotConfig* DefaultShotConfig;
-	
+public:	
 	/** A mapping of Shot Name -> Shot Config to use for rendering specific shots with specific configs. */
-	UPROPERTY(VisibleAnywhere, Instanced, BlueprintReadOnly, Category = "Movie Render Pipeline")
+	UPROPERTY(Instanced)
 	TMap<FString, UMoviePipelineShotConfig*> PerShotConfigMapping;
 
 private:
diff --git a/Engine/Plugins/MovieScene/MovieRenderPipeline/Source/MovieRenderPipelineCore/Public/MoviePipelineQueue.h b/Engine/Plugins/MovieScene/MovieRenderPipeline/Source/MovieRenderPipelineCore/Public/MoviePipelineQueue.h
--- a/Engine/Plugins/MovieScene/MovieRenderPipeline/Source/MovieRenderPipelineCore/Public/MoviePipelineQueue.h
+++ b/Engine/Plugins/MovieScene/MovieRenderPipeline/Source/MovieRenderPipelineCore/Public/MoviePipelineQueue.h
@@ -56,6 +56,7 @@ public:
 		return LoadedSequence;
 	}
 	
+	UFUNCTION(BlueprintCallable, Category = "Movie Render Pipeline")
 	void SetPresetOrigin(UMoviePipelineMasterConfig* InPreset)
 	{
 		if (InPreset)
@@ -65,16 +66,19 @@ public:
 		}
 	}
 
+	UFUNCTION(BlueprintPure, Category = "Movie Render Pipeline")
 	UMoviePipelineMasterConfig* GetPresetOrigin() const
 	{
 		return PresetOrigin.Get();
 	}
 
+	UFUNCTION(BlueprintPure, Category = "Movie Render Pipeline")
 	UMoviePipelineMasterConfig* GetConfiguration() const
 	{
 		return Configuration;
 	}
 
+	UFUNCTION(BlueprintCallable, Category = "Movie Render Pipeline")
 	void SetConfiguration(UMoviePipelineMasterConfig* InPreset)
 	{
 		if (InPreset)
diff --git a/Engine/Plugins/MovieScene/MovieRenderPipeline/Source/MovieRenderPipelineEditor/Private/MovieRenderPipelineEditorModule.cpp b/Engine/Plugins/MovieScene/MovieRenderPipeline/Source/MovieRenderPipelineEditor/Private/MovieRenderPipelineEditorModule.cpp
--- a/Engine/Plugins/MovieScene/MovieRenderPipeline/Source/MovieRenderPipelineEditor/Private/MovieRenderPipelineEditorModule.cpp
+++ b/Engine/Plugins/MovieScene/MovieRenderPipeline/Source/MovieRenderPipelineEditor/Private/MovieRenderPipelineEditorModule.cpp
@@ -72,45 +72,6 @@ namespace
 
 }
 
-void FMovieRenderPipelineEditorModule::PerformTestPipelineRender(const TArray<FString>& Args)
-{
-	if (Args.Num() == 0)
-	{
-		UE_LOG(LogMovieRenderPipeline, Warning, TEXT("No arguments specified. Specify the path to a sequence asset to render."));
-		return;
-	}
-
-	FSoftObjectPath SequencePath = FSoftObjectPath(FPackageName::ExportTextPathToObjectPath(Args[0]));
-
-	const UMovieRenderPipelineProjectSettings* ProjectSettings = GetDefault<UMovieRenderPipelineProjectSettings>();
-	TSubclassOf<UMoviePipelineExecutorBase> ExecutorClass = ProjectSettings->DefaultLocalExecutor;
-
-	check(false);
-	// Create an instance of the Executor and add it to the root so it never dies.
-	// Executor = NewObject<UMoviePipelineExecutorBase>(GetTransientPackage(), ExecutorClass);
-	// Executor->AddToRoot();
-	// Executor->OnExecutorFinished().AddRaw(this, &FMovieRenderPipelineEditorModule::OnTestPipelineExecutorFinished);
-	// 
-	// TArray<UMoviePipelineMasterConfig*> Pipelines = GenerateTestPipelineConfigs();
-	// 
-	// TArray<FMoviePipelineExecutorJobPrev> Jobs;
-	// for(UMoviePipelineMasterConfig* Pipeline : Pipelines)
-	// {
-	// 	Jobs.Add(FMoviePipelineExecutorJobPrev(SequencePath, Pipeline));
-	// }
-	// 
-	// Executor->Execute(Jobs);
-}
-
-void FMovieRenderPipelineEditorModule::OnTestPipelineExecutorFinished(UMoviePipelineExecutorBase* InExecutor, bool bSuccess)
-{
-	if (Executor)
-	{
-		Executor->RemoveFromRoot();
-		Executor = nullptr;
-	}
-}
-
 void FMovieRenderPipelineEditorModule::RegisterSettings()
 {
 	ISettingsModule& SettingsModule = FModuleManager::LoadModuleChecked<ISettingsModule>("Settings");
@@ -120,13 +81,6 @@ void FMovieRenderPipelineEditorModule::RegisterSettings()
 		LOCTEXT("ProjectSettings_Description", "Configure project-wide defaults for the movie render pipeline."),
 		GetMutableDefault<UMovieRenderPipelineProjectSettings>()
 	);
-
-	IConsoleManager::Get().RegisterConsoleCommand(
-		TEXT("MovieRenderPipeline.TestRenderSequence"),
-		TEXT("Renders the specified sequence asset using the default executor and a pre-made configuration."),
-		FConsoleCommandWithArgsDelegate::CreateRaw(this, &FMovieRenderPipelineEditorModule::PerformTestPipelineRender),
-		ECVF_Default
-	);
 }
 
 void FMovieRenderPipelineEditorModule::UnregisterSettings()
diff --git a/Engine/Plugins/MovieScene/MovieRenderPipeline/Source/MovieRenderPipelineEditor/Private/MovieRenderPipelineEditorModule.h b/Engine/Plugins/MovieScene/MovieRenderPipeline/Source/MovieRenderPipelineEditor/Private/MovieRenderPipelineEditorModule.h
--- a/Engine/Plugins/MovieScene/MovieRenderPipeline/Source/MovieRenderPipelineEditor/Private/MovieRenderPipelineEditorModule.h
+++ b/Engine/Plugins/MovieScene/MovieRenderPipeline/Source/MovieRenderPipelineEditor/Private/MovieRenderPipelineEditorModule.h
@@ -17,13 +17,4 @@ public:
 private:
 	void RegisterSettings();
 	void UnregisterSettings();
-	void PerformTestPipelineRender(const TArray<FString>& Args);
-	void OnTestPipelineExecutorFinished(UMoviePipelineExecutorBase* InExecutor, bool bSuccess);
-
-private:
-	TArray<UMoviePipelineMasterConfig*> GenerateTestPipelineConfigs();
-private:
-	class UMoviePipeline* ActiveMoviePipeline;
-	class UMoviePipelineMasterConfig* PipelineConfig;
-	UMoviePipelineExecutorBase* Executor;
 };
diff --git a/Engine/Plugins/MovieScene/MovieRenderPipeline/Source/MovieRenderPipelineEditor/Private/MovieRenderPipelineEditorTest.cpp b/Engine/Plugins/MovieScene/MovieRenderPipeline/Source/MovieRenderPipelineEditor/Private/MovieRenderPipelineEditorTest.cpp
deleted file mode 100644
