MoviePipeline: Fixed a crash if the Burn In or UI Renderer setting was added and then disabled.

#rnx

diff --git a/Engine/Plugins/MovieScene/MovieRenderPipeline/Source/MovieRenderPipelineCore/Private/MoviePipeline.cpp b/Engine/Plugins/MovieScene/MovieRenderPipeline/Source/MovieRenderPipelineCore/Private/MoviePipeline.cpp
--- a/Engine/Plugins/MovieScene/MovieRenderPipeline/Source/MovieRenderPipelineCore/Private/MoviePipeline.cpp
+++ b/Engine/Plugins/MovieScene/MovieRenderPipeline/Source/MovieRenderPipelineCore/Private/MoviePipeline.cpp
@@ -226,7 +226,10 @@ void UMoviePipeline::Initialize(UMoviePipelineExecutorJob* InJob)
 	// Construct a debug UI and bind it to this instance.
 	LoadDebugWidget();
 	
-	GetWorld()->GetGameViewport()->bDisableWorldRendering = true;
+	if (UGameViewportClient* Viewport = GetWorld()->GetGameViewport())
+	{
+		Viewport->bDisableWorldRendering = true;
+	}
 
 	SetupAudioRendering();
 
@@ -440,7 +443,11 @@ void UMoviePipeline::TransitionToState(const EMovieRenderPipelineState InNewStat
 			TeardownAudioRendering();
 			RestoreTargetSequenceToOriginalState();
 
-			GetWorld()->GetGameViewport()->bDisableWorldRendering = false;
+			if (UGameViewportClient* Viewport = GetWorld()->GetGameViewport())
+			{
+				Viewport->bDisableWorldRendering = false;
+			}
+
 			GAreScreenMessagesEnabled = bPrevGScreenMessagesEnabled;
 
 			UE_LOG(LogMovieRenderPipeline, Log, TEXT("Movie Pipeline completed. Duration: %s"), *(FDateTime::UtcNow() - InitializationTime).ToString());
diff --git a/Engine/Plugins/MovieScene/MovieRenderPipeline/Source/MovieRenderPipelineSettings/Private/MoviePipelineBurnInSetting.cpp b/Engine/Plugins/MovieScene/MovieRenderPipeline/Source/MovieRenderPipelineSettings/Private/MoviePipelineBurnInSetting.cpp
--- a/Engine/Plugins/MovieScene/MovieRenderPipeline/Source/MovieRenderPipelineSettings/Private/MoviePipelineBurnInSetting.cpp
+++ b/Engine/Plugins/MovieScene/MovieRenderPipeline/Source/MovieRenderPipelineSettings/Private/MoviePipelineBurnInSetting.cpp
@@ -27,6 +27,12 @@ void UMoviePipelineBurnInSetting::GatherOutputPassesImpl(TArray<FMoviePipelinePa
 
 void UMoviePipelineBurnInSetting::RenderSample_GameThreadImpl(const FMoviePipelineRenderPassMetrics& InSampleState)
 {
+	// If this was transiently added, don't make a burn-in.
+	if (!GetIsUserCustomized() || !IsEnabled())
+	{
+		return;
+	}
+
 	if (InSampleState.bDiscardResult)
 	{
 		return;
diff --git a/Engine/Plugins/MovieScene/MovieRenderPipeline/Source/MovieRenderPipelineSettings/Private/MoviePipelineWidgetRenderSetting.cpp b/Engine/Plugins/MovieScene/MovieRenderPipeline/Source/MovieRenderPipelineSettings/Private/MoviePipelineWidgetRenderSetting.cpp
--- a/Engine/Plugins/MovieScene/MovieRenderPipeline/Source/MovieRenderPipelineSettings/Private/MoviePipelineWidgetRenderSetting.cpp
+++ b/Engine/Plugins/MovieScene/MovieRenderPipeline/Source/MovieRenderPipelineSettings/Private/MoviePipelineWidgetRenderSetting.cpp
@@ -26,6 +26,12 @@ void UMoviePipelineWidgetRenderer::GatherOutputPassesImpl(TArray<FMoviePipelineP
 
 void UMoviePipelineWidgetRenderer::RenderSample_GameThreadImpl(const FMoviePipelineRenderPassMetrics& InSampleState)
 {
+	// If this was transiently added, don't make a burn-in.
+	if (!GetIsUserCustomized() || !IsEnabled())
+	{
+		return;
+	}
+
 	if (InSampleState.bDiscardResult)
 	{
 		return;
