MoviePipeline: Fixed two functions intended to be virtual for derived class implementations not being marked virtual.

diff --git a/Engine/Plugins/MovieScene/MovieRenderPipeline/Source/MovieRenderPipelineCore/Public/MoviePipelineSetting.h b/Engine/Plugins/MovieScene/MovieRenderPipeline/Source/MovieRenderPipelineCore/Public/MoviePipelineSetting.h
--- a/Engine/Plugins/MovieScene/MovieRenderPipeline/Source/MovieRenderPipelineCore/Public/MoviePipelineSetting.h
+++ b/Engine/Plugins/MovieScene/MovieRenderPipeline/Source/MovieRenderPipelineCore/Public/MoviePipelineSetting.h
@@ -122,9 +122,9 @@ public:
 	virtual bool GetIsUserCustomized() const { return bUserCustomized; }
 
 	/** Has this setting finished any export-related things it needs to do post-finalize? */
-	bool HasFinishedExportingImpl() const { return true; }
+	virtual bool HasFinishedExportingImpl() const { return true; }
 	/** Called once when all files have been finalized. */
-	void BeginExportImpl() { }
+	virtual void BeginExportImpl() { }
 	
 private:
 	UPROPERTY(Transient)
