UDN archive for 4.25 and very minor fixes to two header files.

diff --git a/Engine/Plugins/MovieScene/MovieRenderPipeline/Source/MovieRenderPipelineCore/Public/MoviePipelineCameraSetting.h b/Engine/Plugins/MovieScene/MovieRenderPipeline/Source/MovieRenderPipelineCore/Public/MoviePipelineCameraSetting.h
--- a/Engine/Plugins/MovieScene/MovieRenderPipeline/Source/MovieRenderPipelineCore/Public/MoviePipelineCameraSetting.h
+++ b/Engine/Plugins/MovieScene/MovieRenderPipeline/Source/MovieRenderPipelineCore/Public/MoviePipelineCameraSetting.h
@@ -42,7 +42,7 @@ public:
 	* A shutter angle of 360 means continuous movement, while a shutter angle of zero means no
 	* motion blur.
 	*/
-	UPROPERTY(EditAnywhere, BlueprintReadWrite, meta = (UIMin = 0, UIMax = 360, ClampMin = 0, ClampMax = 360), Category = "Movie Pipeline")
+	UPROPERTY(EditAnywhere, BlueprintReadWrite, meta = (UIMin = "0", UIMax = "360", ClampMin = "0", ClampMax = "360"), Category = "Movie Pipeline")
 	int32 CameraShutterAngle;
 	
 	/**
@@ -52,7 +52,7 @@ public:
 	* frame and half the time after the frame. When set to FrameOpen, the motion represents the time from 
 	* Frame N onwards.
 	*/
-	UPROPERTY(EditAnywhere, BlueprintReadWrite, meta = (UIMin = 0, UIMax = 360, ClampMin = 0, ClampMax = 360), Category = "Movie Pipeline")
+	UPROPERTY(EditAnywhere, BlueprintReadWrite, meta = (UIMin = "0", UIMax = "360", ClampMin = "0", ClampMax = "360"), Category = "Movie Pipeline")
 	EMoviePipelineShutterTiming ShutterTiming;
 	
 	/**
@@ -60,12 +60,12 @@ public:
 	* exposure between the different tiles of the image. Leaving this off lets the camera determine based on the previous frame rendered,
 	* which can require long warm up times between shots to allow the exposure to settle.
 	*/
-	UPROPERTY(EditAnywhere, BlueprintReadWrite, meta = (UIMin = 0, UIMax = 360, ClampMin = 0, ClampMax = 360), Category = "Movie Pipeline")
+	UPROPERTY(EditAnywhere, BlueprintReadWrite, meta = (UIMin = "0", UIMax = "360", ClampMin = "0", ClampMax = "360"), Category = "Movie Pipeline")
 	bool bManualExposure;
 	
 	/**
 	* What exposure should we use when using Manual Exposure? Same behavior as the Post Processing volume.
 	*/
-	UPROPERTY(EditAnywhere, BlueprintReadWrite, meta = (UIMin = -10, UIMax = 10, EditCondition="bManualExposure"), Category = "Movie Pipeline")
+	UPROPERTY(EditAnywhere, BlueprintReadWrite, meta = (UIMin = "-10", UIMax = "10", EditCondition="bManualExposure"), Category = "Movie Pipeline")
 	float ExposureCompensation;
 };
\ No newline at end of file
diff --git a/Engine/Plugins/MovieScene/MovieRenderPipeline/Source/MovieRenderPipelineCore/Public/MoviePipelineHighResSetting.h b/Engine/Plugins/MovieScene/MovieRenderPipeline/Source/MovieRenderPipelineCore/Public/MoviePipelineHighResSetting.h
--- a/Engine/Plugins/MovieScene/MovieRenderPipeline/Source/MovieRenderPipelineCore/Public/MoviePipelineHighResSetting.h
+++ b/Engine/Plugins/MovieScene/MovieRenderPipeline/Source/MovieRenderPipelineCore/Public/MoviePipelineHighResSetting.h
@@ -138,7 +138,7 @@ public:
 	* resolution which may help with gpu timeouts. Requires at least 1 tile. Tiling is applied evenly to
 	* both X and Y.
 	*/
-	UPROPERTY(EditAnywhere, BlueprintReadWrite, meta = (UIMin = 1, ClampMin = 1, UIMax = 16), Category = "Movie Pipeline")
+	UPROPERTY(EditAnywhere, BlueprintReadWrite, meta = (UIMin = "1", ClampMin = "1", UIMax = "16"), Category = "Movie Pipeline")
 	int32 TileCount;
 	
 	/**
@@ -147,7 +147,7 @@ public:
 	* (up to the detail resolution of your texture). Too much sharpness will cause visual grain/noise in the
 	* resulting image, but this can be mitigated with more spatial samples.
 	*/
-	UPROPERTY(EditAnywhere, BlueprintReadWrite, meta = (UIMin = -1, UIMax = 0), Category = "Movie Pipeline")
+	UPROPERTY(EditAnywhere, BlueprintReadWrite, meta = (UIMin = "-1", UIMax = "0"), Category = "Movie Pipeline")
 	float TextureSharpnessBias;
 	
 	/**
@@ -155,7 +155,7 @@ public:
 	* tiles (which means faster renders) but increases the likelyhood of edge-of-screen artifacts showing up which
 	* will become visible in the final image as a "grid" of repeated problem areas.
 	*/
-	UPROPERTY(EditAnywhere, BlueprintReadWrite, meta = (UIMin = 0, ClampMin = 0, UIMax = 0.5, ClampMax=1), Category = "Movie Pipeline")
+	UPROPERTY(EditAnywhere, BlueprintReadWrite, meta = (UIMin = "0", ClampMin = "0", UIMax = "0.5", ClampMax = "1"), Category = "Movie Pipeline")
 	float OverlapRatio;
 
 	/**
@@ -168,7 +168,7 @@ public:
 	/*
 	* How many samples should the Burley Sub Surface Scattering use?
 	*/
-	UPROPERTY(EditAnywhere, BlueprintReadWrite, meta = (UIMin = 64, ClampMin = 0, UIMax = 1024, EditCondition="bOverrideSubSurfaceScattering"), Category = "Movie Pipeline")
+	UPROPERTY(EditAnywhere, BlueprintReadWrite, meta = (UIMin = "64", ClampMin = "0", UIMax = "1024", EditCondition="bOverrideSubSurfaceScattering"), Category = "Movie Pipeline")
 	int32 BurleySampleCount;
 	
 	/**
