Additional casing issues

diff --git a/Engine/Plugins/MovieScene/MovieRenderPipeline/Source/MovieRenderPipelineRenderPasses/Private/MoviePipelineEXROutput.cpp b/Engine/Plugins/MovieScene/MovieRenderPipeline/Source/MovieRenderPipelineRenderPasses/Private/MoviePipelineEXROutput.cpp
--- a/Engine/Plugins/MovieScene/MovieRenderPipeline/Source/MovieRenderPipelineRenderPasses/Private/MoviePipelineEXROutput.cpp
+++ b/Engine/Plugins/MovieScene/MovieRenderPipeline/Source/MovieRenderPipelineRenderPasses/Private/MoviePipelineEXROutput.cpp
@@ -18,7 +18,7 @@
 #include "Modules/ModuleManager.h"
 
 THIRD_PARTY_INCLUDES_START
-#include "OpenEXR/ImfChannelList.h"
+#include "openexr/ImfChannelList.h"
 THIRD_PARTY_INCLUDES_END
 
 #if WITH_UNREALEXR
diff --git a/Engine/Plugins/MovieScene/MovieRenderPipeline/Source/MovieRenderPipelineRenderPasses/Private/MoviePipelineEXROutput.h b/Engine/Plugins/MovieScene/MovieRenderPipeline/Source/MovieRenderPipelineRenderPasses/Private/MoviePipelineEXROutput.h
--- a/Engine/Plugins/MovieScene/MovieRenderPipeline/Source/MovieRenderPipelineRenderPasses/Private/MoviePipelineEXROutput.h
+++ b/Engine/Plugins/MovieScene/MovieRenderPipeline/Source/MovieRenderPipelineRenderPasses/Private/MoviePipelineEXROutput.h
@@ -8,16 +8,16 @@
 
 #if WITH_UNREALEXR
 THIRD_PARTY_INCLUDES_START
-#include "OpenEXR/ImfIO.h"
-#include "OpenEXR/ImathBox.h"
-#include "OpenEXR/ImfChannelList.h"
-#include "OpenEXR/ImfInputFile.h"
-#include "OpenEXR/ImfOutputFile.h"
-#include "OpenEXR/ImfArray.h"
-#include "OpenEXR/ImfHeader.h"
-#include "OpenEXR/ImfStdIO.h"
-#include "OpenEXR/ImfChannelList.h"
-#include "OpenEXR/ImfRgbaFile.h"
+#include "openexr/ImfIO.h"
+#include "openexr/ImathBox.h"
+#include "openexr/ImfChannelList.h"
+#include "openexr/ImfInputFile.h"
+#include "openexr/ImfOutputFile.h"
+#include "openexr/ImfArray.h"
+#include "openexr/ImfHeader.h"
+#include "openexr/ImfStdIO.h"
+#include "openexr/ImfChannelList.h"
+#include "openexr/ImfRgbaFile.h"
 THIRD_PARTY_INCLUDES_END
 #endif // WITH_UNREALEXR
 
