CIS fixes for deprecation warnings.

diff --git a/Engine/Plugins/FX/Niagara/Source/NiagaraEditorWidgets/Private/Stack/SNiagaraStackParameterStoreEntryValue.cpp b/Engine/Plugins/FX/Niagara/Source/NiagaraEditorWidgets/Private/Stack/SNiagaraStackParameterStoreEntryValue.cpp
--- a/Engine/Plugins/FX/Niagara/Source/NiagaraEditorWidgets/Private/Stack/SNiagaraStackParameterStoreEntryValue.cpp
+++ b/Engine/Plugins/FX/Niagara/Source/NiagaraEditorWidgets/Private/Stack/SNiagaraStackParameterStoreEntryValue.cpp
@@ -164,8 +164,13 @@ TSharedRef<SWidget> SNiagaraStackParameterStoreEntryValue::ConstructValueStructW
 		{
 			FPropertyEditorModule& PropertyEditorModule = FModuleManager::GetModuleChecked<FPropertyEditorModule>("PropertyEditor");
 
+			FDetailsViewArgs DetailsViewArgs;
+			DetailsViewArgs.bAllowSearch = false;
+			DetailsViewArgs.NameAreaSettings = FDetailsViewArgs::HideNameArea;
+			DetailsViewArgs.bHideSelectionTip = true;
+
 			TSharedRef<IStructureDetailsView> StructureDetailsView = PropertyEditorModule.CreateStructureDetailView(
-				FDetailsViewArgs(false, false, false, FDetailsViewArgs::HideNameArea, true),
+				DetailsViewArgs,
 				FStructureDetailsViewArgs(),
 				nullptr);
 
diff --git a/Engine/Plugins/MovieScene/MovieRenderPipeline/Source/MovieRenderPipelineEditor/Private/Widgets/SMoviePipelineConfigEditor.cpp b/Engine/Plugins/MovieScene/MovieRenderPipeline/Source/MovieRenderPipelineEditor/Private/Widgets/SMoviePipelineConfigEditor.cpp
--- a/Engine/Plugins/MovieScene/MovieRenderPipeline/Source/MovieRenderPipelineEditor/Private/Widgets/SMoviePipelineConfigEditor.cpp
+++ b/Engine/Plugins/MovieScene/MovieRenderPipeline/Source/MovieRenderPipelineEditor/Private/Widgets/SMoviePipelineConfigEditor.cpp
@@ -313,7 +313,10 @@ TSharedRef<SWidget> SMoviePipelineConfigEditor::OnGenerateSettingsMenu()
 void SMoviePipelineConfigEditor::UpdateDetails()
 {
 	FPropertyEditorModule& PropertyEditorModule = FModuleManager::Get().LoadModuleChecked<FPropertyEditorModule>("PropertyEditor");
-	FDetailsViewArgs DetailsViewArgs(false, false, false, FDetailsViewArgs::HideNameArea, true);
+	FDetailsViewArgs DetailsViewArgs;
+	DetailsViewArgs.bAllowSearch = false;
+	DetailsViewArgs.NameAreaSettings = FDetailsViewArgs::HideNameArea;
+	DetailsViewArgs.bHideSelectionTip = true;
 	DetailsViewArgs.bShowScrollBar = false;
 	DetailsViewArgs.ColumnWidth = 0.5f;
 
diff --git a/Engine/Plugins/MovieScene/MovieRenderPipeline/Source/MovieRenderPipelineEditor/Private/Widgets/SMoviePipelineQueuePanel.cpp b/Engine/Plugins/MovieScene/MovieRenderPipeline/Source/MovieRenderPipelineEditor/Private/Widgets/SMoviePipelineQueuePanel.cpp
--- a/Engine/Plugins/MovieScene/MovieRenderPipeline/Source/MovieRenderPipelineEditor/Private/Widgets/SMoviePipelineQueuePanel.cpp
+++ b/Engine/Plugins/MovieScene/MovieRenderPipeline/Source/MovieRenderPipelineEditor/Private/Widgets/SMoviePipelineQueuePanel.cpp
@@ -51,7 +51,10 @@ PRAGMA_DISABLE_OPTIMIZATION
 void SMoviePipelineQueuePanel::Construct(const FArguments& InArgs)
 {
 	FPropertyEditorModule& PropertyEditorModule = FModuleManager::Get().LoadModuleChecked<FPropertyEditorModule>("PropertyEditor");
-	FDetailsViewArgs DetailsViewArgs(false, false, false, FDetailsViewArgs::HideNameArea, true);
+	FDetailsViewArgs DetailsViewArgs;
+	DetailsViewArgs.bAllowSearch = false;
+	DetailsViewArgs.NameAreaSettings = FDetailsViewArgs::HideNameArea;
+	DetailsViewArgs.bHideSelectionTip = true;
 	DetailsViewArgs.ColumnWidth = 0.7f;
 
 	JobDetailsPanelWidget = PropertyEditorModule.CreateDetailView(DetailsViewArgs);
