Movie Render Queue tiling support for path tracer

diff --git a/Engine/Plugins/MovieScene/MovieRenderPipeline/Source/MovieRenderPipelineRenderPasses/Private/MoviePipelineImagePassBase.cpp b/Engine/Plugins/MovieScene/MovieRenderPipeline/Source/MovieRenderPipelineRenderPasses/Private/MoviePipelineImagePassBase.cpp
--- a/Engine/Plugins/MovieScene/MovieRenderPipeline/Source/MovieRenderPipelineRenderPasses/Private/MoviePipelineImagePassBase.cpp
+++ b/Engine/Plugins/MovieScene/MovieRenderPipeline/Source/MovieRenderPipelineRenderPasses/Private/MoviePipelineImagePassBase.cpp
@@ -161,21 +161,12 @@ TSharedPtr<FSceneViewFamilyContext> UMoviePipelineImagePassBase::CalculateViewFa
 		// Overwrite whatever sampling count came from the PostProcessVolume
 		View->FinalPostProcessSettings.bOverride_PathTracingSamplesPerPixel = true;
 		View->FinalPostProcessSettings.PathTracingSamplesPerPixel = InOutSampleState.SpatialSampleCount;
-		if (InOutSampleState.SpatialSampleIndex == 0)
-		{
-			// reset path tracer's accumulation at the start of each temporal sample
-			View->bForceCameraVisibilityReset = true;
-		}
+
+		// reset path tracer's accumulation at the start of each spatial sample
+		View->bForcePathTracerReset = InOutSampleState.SpatialSampleIndex == 0;
 
 		// discard the result, unless its the last spatial sample
 		InOutSampleState.bDiscardResult |= !(InOutSampleState.SpatialSampleIndex == InOutSampleState.SpatialSampleCount - 1);
-
-		if (InOutSampleState.GetTileCount() > 1 && InOutSampleState.GetTileIndex() == 0 && SampleIndex == 0)
-		{
-			// warn the user
-			UE_LOG(LogMovieRenderPipeline, Warning, TEXT("Path tracer does not support tiling correctly yet!"));
-			// FIXME: find a way to disable tiling earlier so this doesn't happen
-		}
 	}
 
 	// Object Occlusion/Histories
diff --git a/Engine/Source/Runtime/Engine/Private/SceneView.cpp b/Engine/Source/Runtime/Engine/Private/SceneView.cpp
--- a/Engine/Source/Runtime/Engine/Private/SceneView.cpp
+++ b/Engine/Source/Runtime/Engine/Private/SceneView.cpp
@@ -701,6 +701,7 @@ FSceneView::FSceneView(const FSceneViewInitOptions& InitOptions)
 	, bShouldBindInstancedViewUB(false)
 	, UnderwaterDepth(-1.0f)
 	, bForceCameraVisibilityReset(false)
+	, bForcePathTracerReset(false)
 	, GlobalClippingPlane(FPlane(0, 0, 0, 0))
 	, LensPrincipalPointOffsetScale(0.0f, 0.0f, 1.0f, 1.0f)
 #if WITH_EDITOR
diff --git a/Engine/Source/Runtime/Engine/Public/SceneView.h b/Engine/Source/Runtime/Engine/Public/SceneView.h
--- a/Engine/Source/Runtime/Engine/Public/SceneView.h
+++ b/Engine/Source/Runtime/Engine/Public/SceneView.h
@@ -1120,6 +1120,9 @@ public:
 	 */
 	bool bForceCameraVisibilityReset;
 
+	/** True if we should force the path tracer to reset its internal accumulation state */
+	bool bForcePathTracerReset;
+
 	/** Global clipping plane being applied to the scene, or all 0's if disabled.  This is used when rendering the planar reflection pass. */
 	FPlane GlobalClippingPlane;
 
diff --git a/Engine/Source/Runtime/Renderer/Private/SceneVisibility.cpp b/Engine/Source/Runtime/Renderer/Private/SceneVisibility.cpp
--- a/Engine/Source/Runtime/Renderer/Private/SceneVisibility.cpp
+++ b/Engine/Source/Runtime/Renderer/Private/SceneVisibility.cpp
@@ -3438,8 +3438,10 @@ void FSceneRenderer::PreVisibilityFrameSetup(FRDGBuilder& GraphBuilder, const FS
 			
 			if (View.bIsOfflineRender)
 			{
-				// for offline renders, we want to allow motion blur so don't let geometry changes or camera moves invalidate the path tracer accumulation
-				bInvalidatePathTracer = View.bCameraCut || View.bForceCameraVisibilityReset;
+				// In the offline context, we want precise control over when to restart the path tracer's accumulation to allow for motion blur
+				// So we use the camera cut signal only. In particular - we should not use bForceCameraVisibilityReset since this has
+				// interactions with the motion blur post process effect in tiled rendering (see comment below).
+				bInvalidatePathTracer = View.bCameraCut || View.bForcePathTracerReset;
 			}
 			else
 			{
@@ -3447,7 +3449,8 @@ void FSceneRenderer::PreVisibilityFrameSetup(FRDGBuilder& GraphBuilder, const FS
 				bInvalidatePathTracer = bResetCamera ||
 					Scene->bPathTracingNeedsInvalidation ||
 					bIsProjMatrixDifferent ||
-					bIsCameraMove;
+					bIsCameraMove ||
+					View.bForcePathTracerReset;
 			}
 
 			if (bInvalidatePathTracer)
