Implement missing async flush for movie rendering

diff --git a/Engine/Plugins/MovieScene/MovieRenderPipeline/Source/MovieRenderPipelineCore/Private/MoviePipelineRendering.cpp b/Engine/Plugins/MovieScene/MovieRenderPipeline/Source/MovieRenderPipelineCore/Private/MoviePipelineRendering.cpp
--- a/Engine/Plugins/MovieScene/MovieRenderPipeline/Source/MovieRenderPipelineCore/Private/MoviePipelineRendering.cpp
+++ b/Engine/Plugins/MovieScene/MovieRenderPipeline/Source/MovieRenderPipelineCore/Private/MoviePipelineRendering.cpp
@@ -30,6 +30,8 @@
 #include "LandscapeProxy.h"
 #include "EngineModule.h"
 #include "DistanceFieldAtlas.h"
+#include "MeshCardRepresentation.h"
+#include "AssetCompilingManager.h"
 #include "ShaderCompiler.h"
 #include "EngineUtils.h"
 
@@ -487,6 +489,12 @@ void UMoviePipeline::FlushAsyncEngineSystems()
 	{
 		GetWorld()->BlockTillLevelStreamingCompleted();
 	}
+	
+#if WITH_EDITOR
+	// Flush all assets handled by the asset compiling manager (i.e. textures, static meshes).
+	// A progressbar is already in place so the user can get feedback while waiting for everything to settle.
+	FAssetCompilingManager::Get().FinishAllCompilation();
+#endif
 
 	// Now we can flush the shader compiler. ToDo: This should probably happen right before SendAllEndOfFrameUpdates() is normally called
 	if (GShaderCompilingManager)
@@ -540,6 +548,31 @@ void UMoviePipeline::FlushAsyncEngineSystems()
 		}
 	}
 
+	if (GCardRepresentationAsyncQueue)
+	{
+		bool bDidWork = false;
+		int32 NumTasks = GCardRepresentationAsyncQueue->GetNumOutstandingTasks();
+		if (NumTasks > 0)
+		{
+			UE_LOG(LogMovieRenderPipeline, Log, TEXT("[%d] Starting build for %d mesh cards."), GFrameCounter, NumTasks);
+		}
+
+		while (GCardRepresentationAsyncQueue->GetNumOutstandingTasks() > 0)
+		{
+			UE_LOG(LogMovieRenderPipeline, Log, TEXT("[%d] Waiting for %d Mesh Cards to finish building..."), GFrameCounter, GCardRepresentationAsyncQueue->GetNumOutstandingTasks());
+			GCardRepresentationAsyncQueue->ProcessAsyncTasks();
+
+			// Sleep for 1 second and then check again. This way we get an indication of progress as this works.
+			FPlatformProcess::Sleep(1.f);
+			bDidWork = true;
+		}
+
+		if (bDidWork)
+		{
+			UE_LOG(LogMovieRenderPipeline, Log, TEXT("[%d] Done building %d Mesh Cards."), GFrameCounter, NumTasks);
+		}
+	}
+
 	// Flush grass
 	{
 		for (TActorIterator<ALandscapeProxy> It(GetWorld()); It; ++It)
