Fixed an incorrect use of IsAudioThreadRunning() as it returns false when AudioThread is suspended

diff --git a/Engine/Plugins/MovieScene/MovieRenderPipeline/Source/MovieRenderPipelineRenderPasses/Private/MoviePipelineWaveOutput.cpp b/Engine/Plugins/MovieScene/MovieRenderPipeline/Source/MovieRenderPipelineRenderPasses/Private/MoviePipelineWaveOutput.cpp
--- a/Engine/Plugins/MovieScene/MovieRenderPipeline/Source/MovieRenderPipelineRenderPasses/Private/MoviePipelineWaveOutput.cpp
+++ b/Engine/Plugins/MovieScene/MovieRenderPipeline/Source/MovieRenderPipelineRenderPasses/Private/MoviePipelineWaveOutput.cpp
@@ -52,8 +52,8 @@ static bool IsMoviePipelineAudioOutputSupported(const UObject* WorldContextObjec
 		return true;
 	}
 
-	// If there is no audio thread running (e.g. we're in the editor), it's possible to create a new non-realtime audio mixer
-	if (!IsAudioThreadRunning())
+	// If there is no async audio processing (e.g. we're in the editor), it's possible to create a new non-realtime audio mixer
+	if (!FAudioThread::IsUsingThreadedAudio())
 	{
 		return true;
 	}
