Movie Render Queue: Removing unnecessary code that checks if the current world is the same as the executor's map if the executor has already been set to "UseCurrentLevel"

diff --git a/Engine/Plugins/MovieScene/MovieRenderPipeline/Source/MovieRenderPipelineCore/Private/MoviePipelineInProcessExecutor.cpp b/Engine/Plugins/MovieScene/MovieRenderPipeline/Source/MovieRenderPipelineCore/Private/MoviePipelineInProcessExecutor.cpp
--- a/Engine/Plugins/MovieScene/MovieRenderPipeline/Source/MovieRenderPipelineCore/Private/MoviePipelineInProcessExecutor.cpp
+++ b/Engine/Plugins/MovieScene/MovieRenderPipeline/Source/MovieRenderPipelineCore/Private/MoviePipelineInProcessExecutor.cpp
@@ -26,13 +26,6 @@ void UMoviePipelineInProcessExecutor::Start(const UMoviePipelineExecutorJob* InJ
 			OnIndividualPipelineFinished(nullptr);
 			return;
 		}
-
-		if (World != InJob->Map.ResolveObject())
-		{
-			UE_LOG(LogMovieRenderPipeline, Warning, TEXT("Unable to start movie pipeline job. Current map '%s' does not match job's map: '%s'"), *GetNameSafe(World), *InJob->Map.GetAssetPathName().ToString());
-			OnIndividualPipelineFinished(nullptr);
-			return;
-		}
 		
 		UE_LOG(LogMovieRenderPipeline, Log, TEXT("Starting %s"), *GetNameSafe(World));
 	}
