fix some more image color quantization conventions

diff --git a/Engine/Plugins/MovieScene/MovieRenderPipeline/Source/MovieRenderPipelineCore/Private/MoviePipelineImageQuantization.cpp b/Engine/Plugins/MovieScene/MovieRenderPipeline/Source/MovieRenderPipelineCore/Private/MoviePipelineImageQuantization.cpp
--- a/Engine/Plugins/MovieScene/MovieRenderPipeline/Source/MovieRenderPipelineCore/Private/MoviePipelineImageQuantization.cpp
+++ b/Engine/Plugins/MovieScene/MovieRenderPipeline/Source/MovieRenderPipelineCore/Private/MoviePipelineImageQuantization.cpp
@@ -20,18 +20,20 @@ TArray<TToColorBitDepthType> ConvertLinearToLinearBitDepth(TFromColorBitDepthTyp
 	// Convert all of our pixels.
 	TArray<TToColorBitDepthType> OutsRGBData;
 	OutsRGBData.SetNumUninitialized(InCount);
-	TColorChannelNumericType MaxValueInt = TNumericLimits<TColorChannelNumericType>::Max();
+	// TColorChannelNumericType = uint8 or uint16
+	// do clamp in int in case input float is out of [0,1] range
+	int MaxValueInt = TNumericLimits<TColorChannelNumericType>::Max();
 	float MaxValue = static_cast<float>(MaxValueInt);
 	for (int32 PixelIndex = 0; PixelIndex < InCount; PixelIndex++)
 	{
 		// Avoid the bounds checking of TArray[]
 		TToColorBitDepthType* OutColor = &OutsRGBData.GetData()[PixelIndex];
 
-		// We don't need sRGB color conversion. Flooring avoids an extra branch for Round.
-		OutColor->R = FMath::Clamp<TColorChannelNumericType>(FMath::FloorToInt((InColor[PixelIndex].R * MaxValue) + 0.5f), 0, MaxValueInt);
-		OutColor->G = FMath::Clamp<TColorChannelNumericType>(FMath::FloorToInt((InColor[PixelIndex].G * MaxValue) + 0.5f), 0, MaxValueInt);
-		OutColor->B = FMath::Clamp<TColorChannelNumericType>(FMath::FloorToInt((InColor[PixelIndex].B * MaxValue) + 0.5f), 0, MaxValueInt);
-		OutColor->A = FMath::Clamp<TColorChannelNumericType>(FMath::FloorToInt((InColor[PixelIndex].A * MaxValue) + 0.5f), 0, MaxValueInt);
+		// We don't need sRGB color conversion. Avoid FMath::Floor or Round
+		OutColor->R = (TColorChannelNumericType) FMath::Clamp( (int)((InColor[PixelIndex].R * MaxValue) + 0.5f), 0, MaxValueInt);
+		OutColor->G = (TColorChannelNumericType) FMath::Clamp( (int)((InColor[PixelIndex].G * MaxValue) + 0.5f), 0, MaxValueInt);
+		OutColor->B = (TColorChannelNumericType) FMath::Clamp( (int)((InColor[PixelIndex].B * MaxValue) + 0.5f), 0, MaxValueInt);
+		OutColor->A = (TColorChannelNumericType) FMath::Clamp( (int)((InColor[PixelIndex].A * MaxValue) + 0.5f), 0, MaxValueInt);
 	}
 
 	return OutsRGBData;
@@ -57,7 +59,8 @@ static TArray<uint8> GenerateSRGBTable(uint32 InPrecision)
 		}
 
 		// Flooring avoids an extra branch for Round. Using [] on GetData() avoids bounds checking cost.
-		OutsRGBTable.GetData()[TableIndex] = (uint8)FMath::Clamp(FMath::FloorToInt(ValueAsLinear * 256.f), 0, 255);
+		check( ValueAsLinear >= 0.f && ValueAsLinear <= 1.f );
+		OutsRGBTable.GetData()[TableIndex] = (uint8) ( ValueAsLinear * 255.f + 0.5f );
 	}
 	
 	return OutsRGBTable;
@@ -118,8 +121,9 @@ static TArray<uint8> GenerateSRGBTableFloat16to8()
 		{
 			ValueAsLinear = FMath::Pow(ValueAsLinear, 1.0f / 2.4f) * 1.055f - 0.055f;
 		}
-		// Flooring avoids an extra branch for Round. Using [] on GetData() avoids bounds checking cost.
-		OutsRGBTableData[TableIndex] = (uint8)FMath::FloorToInt((ValueAsLinear * 255.f) + 0.5f);
+		// Avoid FMath::FloorToInt or Round which call floorf. Using [] on GetData() avoids bounds checking cost.
+		check( ValueAsLinear >= 0.f && ValueAsLinear <= 1.f );
+		OutsRGBTableData[TableIndex] = (uint8) (ValueAsLinear * 255.f + 0.5f);
 	}
 
 	return OutsRGBTable;
@@ -173,8 +177,8 @@ static TArray<FColor> ConvertLinearTosRGB8bppViaLookupTable(FFloat16Color* InCol
 		OutColor->G = sRGBTableData[InColor[PixelIndex].G.Encoded];
 		OutColor->B = sRGBTableData[InColor[PixelIndex].B.Encoded];
 
-		// Alpha doesn't get sRGB conversion, it just gets linearly converted to 8 bit. Flooring avoids an extra branch for Round.
-		OutColor->A = (uint8)FMath::Clamp(FMath::FloorToInt((InColor[PixelIndex].A * 255.f) + 0.5f), 0, 255);
+		// Alpha doesn't get sRGB conversion, it just gets linearly converted to 8 bit.
+		OutColor->A = FColor::QuantizeUNormFloatTo8( InColor[PixelIndex].A );
 	}
 
 	return OutsRGBData;
@@ -203,8 +207,8 @@ static TArray<FColor> ConvertLinearTosRGB8bppViaLookupTable(FLinearColor* InColo
 		OutColor->G = sRGBTable.GetData()[FMath::Clamp(TableIndexG, 0, TableUpperBound)];
 		OutColor->B = sRGBTable.GetData()[FMath::Clamp(TableIndexB, 0, TableUpperBound)];
 		
-		// Alpha doesn't get sRGB conversion, it just gets linearly converted to 8 bit. Flooring avoids an extra branch for Round.
-		OutColor->A = (uint8)FMath::Clamp(FMath::FloorToInt((InColor[PixelIndex].A * 255.f) + 0.5f), 0, 255);
+		// Alpha doesn't get sRGB conversion, it just gets linearly converted to 8 bit.
+		OutColor->A = FColor::QuantizeUNormFloatTo8( InColor[PixelIndex].A );
 	}
 
 	return OutsRGBData;
diff --git a/Engine/Source/Runtime/Core/Public/Math/Color.h b/Engine/Source/Runtime/Core/Public/Math/Color.h
--- a/Engine/Source/Runtime/Core/Public/Math/Color.h
+++ b/Engine/Source/Runtime/Core/Public/Math/Color.h
@@ -535,12 +535,43 @@ public:
 	*/
 	static CORE_API FColor MakeFromColorTemperature( float Temp );
 
+	/**
+	* Conversions to/from GPU UNorm floats, U8, U16
+	* matches convention of FColor FLinearColor::QuantizeRound
+	*/
+
+	static uint8 QuantizeUNormFloatTo8( float UnitFloat )
+	{
+		UnitFloat = FMath::Clamp(UnitFloat,0.f,1.f);
+		return (uint8)( 0.5f + UnitFloat * 255.f );
+	}
+	
+	static uint16 QuantizeUNormFloatTo16( float UnitFloat )
+	{
+		UnitFloat = FMath::Clamp(UnitFloat,0.f,1.f);
+		return (uint16)( 0.5f + UnitFloat * 65535.f );
+	}
+
+	static float DequantizeUNorm8ToFloat( int Value8 )
+	{
+		check( Value8 >= 0 && Value8 <= 255 );
+
+		return Value8 / 255.f;
+	}
+	
+	static float DequantizeUNorm16ToFloat( int Value16 )
+	{
+		check( Value16 >= 0 && Value16 <= 65535 );
+
+		return Value16 / 65535.f;
+	}
+
 	static uint8 Requantize10to8( int Value10 )
 	{
 		check( Value10 >= 0 && Value10 <= 1023 );
 
 		// Dequantize from 10 bit (Value10/1023.f)
-		// requantize to 8 bit with rounding (GPU convention)
+		// requantize to 8 bit with rounding (GPU convention UNorm)
 		//  this is the computation we want :
 		// (int)( (Value10/1023.f)*255.f + 0.5f );
 		// this gives the exactly the same results :
@@ -554,7 +585,7 @@ public:
 		check( Value16 >= 0 && Value16 <= 65535 );
 
 		// Dequantize x from 16 bit (Value16/65535.f)
-		// then requantize to 8 bit with rounding (GPU convention)
+		// then requantize to 8 bit with rounding (GPU convention UNorm)
 
 		// matches exactly with :
 		//  (int)( (Value16/65535.f) * 255.f + 0.5f );
diff --git a/Engine/Source/Runtime/ImageCore/Private/ImageCore.cpp b/Engine/Source/Runtime/ImageCore/Private/ImageCore.cpp
--- a/Engine/Source/Runtime/ImageCore/Private/ImageCore.cpp
+++ b/Engine/Source/Runtime/ImageCore/Private/ImageCore.cpp
@@ -91,7 +91,7 @@ static void CopyImage(const FImage& SrcImage, FImage& DestImage)
 				const int64 EndIndex = FMath::Min(StartIndex + TexelsPerJob, NumTexels);
 				for (int64 TexelIndex = StartIndex; TexelIndex < EndIndex; ++TexelIndex)
 				{
-					DestLum[TexelIndex] = FMath::Clamp(FMath::FloorToInt(SrcColors[TexelIndex].R * 65535.999f), 0, 65535);
+					DestLum[TexelIndex] = FColor::QuantizeUNormFloatTo16( SrcColors[TexelIndex].R );
 				}
 			});
 		}
@@ -135,10 +135,10 @@ static void CopyImage(const FImage& SrcImage, FImage& DestImage)
 					for (int64 TexelIndex = 0; TexelIndex < NumTexels; ++TexelIndex)
 					{
 						int64 DestIndex = TexelIndex * 4;
-						DestColors[DestIndex + 0] = FMath::Clamp(FMath::FloorToInt(SrcColors[TexelIndex].R * 65535.999f), 0, 65535);
-						DestColors[DestIndex + 1] = FMath::Clamp(FMath::FloorToInt(SrcColors[TexelIndex].G * 65535.999f), 0, 65535);
-						DestColors[DestIndex + 2] = FMath::Clamp(FMath::FloorToInt(SrcColors[TexelIndex].B * 65535.999f), 0, 65535);
-						DestColors[DestIndex + 3] = FMath::Clamp(FMath::FloorToInt(SrcColors[TexelIndex].A * 65535.999f), 0, 65535);
+						DestColors[DestIndex + 0] = FColor::QuantizeUNormFloatTo16( SrcColors[TexelIndex].R );
+						DestColors[DestIndex + 1] = FColor::QuantizeUNormFloatTo16( SrcColors[TexelIndex].G );
+						DestColors[DestIndex + 2] = FColor::QuantizeUNormFloatTo16( SrcColors[TexelIndex].B );
+						DestColors[DestIndex + 3] = FColor::QuantizeUNormFloatTo16( SrcColors[TexelIndex].A );
 					}
 				});
 			}
@@ -205,7 +205,8 @@ static void CopyImage(const FImage& SrcImage, FImage& DestImage)
 			TArrayView64<const uint16> SrcLum = SrcImage.AsG16();
 			for (int64 TexelIndex = 0; TexelIndex < NumTexels; ++TexelIndex)
 			{
-				DestColors[TexelIndex] = FLinearColor(SrcLum[TexelIndex] / 65535.0f, SrcLum[TexelIndex] / 65535.0f, SrcLum[TexelIndex] / 65535.0f, 1.0f);
+				float Value = FColor::DequantizeUNorm16ToFloat( SrcLum[TexelIndex] );
+				DestColors[TexelIndex] = FLinearColor(Value,Value,Value, 1.0f);
 			}
 		}
 		break;
@@ -213,37 +214,35 @@ static void CopyImage(const FImage& SrcImage, FImage& DestImage)
 		case ERawImageFormat::BGRA8:
 			{
 				TArrayView64<const FColor> SrcColors = SrcImage.AsBGRA8();
-				switch ( SrcImage.GammaSpace )
+				EGammaSpace SrcGamma = SrcImage.GammaSpace;
+
+				ParallelFor(NumJobs, [DestColors, SrcColors, TexelsPerJob, NumTexels, SrcGamma](int64 JobIndex)
 				{
-				case EGammaSpace::Linear:
-					ParallelFor(NumJobs, [DestColors, SrcColors, TexelsPerJob, NumTexels](int64 JobIndex)
+					int64 StartIndex = JobIndex * TexelsPerJob;
+					int64 EndIndex = FMath::Min(StartIndex + TexelsPerJob, NumTexels);
+
+					switch ( SrcGamma )
 					{
-						int64 StartIndex = JobIndex * TexelsPerJob;
-						int64 EndIndex = FMath::Min(StartIndex + TexelsPerJob, NumTexels);
+					case EGammaSpace::Linear:
 						for (int64 TexelIndex = StartIndex; TexelIndex < EndIndex; ++TexelIndex)
 						{
 							DestColors[TexelIndex] = SrcColors[TexelIndex].ReinterpretAsLinear();
 						}
-					});
 					break;
-				case EGammaSpace::sRGB:
-					ParallelFor(NumJobs, [DestColors, SrcColors, TexelsPerJob, NumTexels](int64 JobIndex)
-					{
-						int64 StartIndex = JobIndex * TexelsPerJob;
-						int64 EndIndex = FMath::Min(StartIndex + TexelsPerJob, NumTexels);
+					case EGammaSpace::sRGB:
 						for (int64 TexelIndex = StartIndex; TexelIndex < EndIndex; ++TexelIndex)
 						{
 							DestColors[TexelIndex] = FLinearColor(SrcColors[TexelIndex]);
 						}
-					});
 					break;
-				case EGammaSpace::Pow22:
-					for (int64 TexelIndex = 0; TexelIndex < NumTexels; ++TexelIndex)
-					{
-						DestColors[TexelIndex] = FLinearColor::FromPow22Color(SrcColors[TexelIndex]);
-					}
+					case EGammaSpace::Pow22:
+						for (int64 TexelIndex = StartIndex; TexelIndex < EndIndex; ++TexelIndex)
+						{
+							DestColors[TexelIndex] = FLinearColor::FromPow22Color(SrcColors[TexelIndex]);
+						}
 					break;
-				}
+					}
+				});
 			}
 			break;
 
