Add DeleteAllJobs() to MoviePipelineQueue

diff --git a/Engine/Plugins/MovieScene/MovieRenderPipeline/Source/MovieRenderPipelineCore/Private/MoviePipelineQueue.cpp b/Engine/Plugins/MovieScene/MovieRenderPipeline/Source/MovieRenderPipelineCore/Private/MoviePipelineQueue.cpp
--- a/Engine/Plugins/MovieScene/MovieRenderPipeline/Source/MovieRenderPipelineCore/Private/MoviePipelineQueue.cpp
+++ b/Engine/Plugins/MovieScene/MovieRenderPipeline/Source/MovieRenderPipelineCore/Private/MoviePipelineQueue.cpp
@@ -41,6 +41,16 @@ void UMoviePipelineQueue::DeleteJob(UMoviePipelineExecutorJob* InJob)
 	QueueSerialNumber++;
 }
 
+void UMoviePipelineQueue::DeleteAllJobs()
+{
+#if WITH_EDITOR
+	Modify();
+#endif
+
+	Jobs.Empty();
+	QueueSerialNumber++;
+}
+
 UMoviePipelineExecutorJob* UMoviePipelineQueue::DuplicateJob(UMoviePipelineExecutorJob* InJob)
 {
 	if (!InJob)
diff --git a/Engine/Plugins/MovieScene/MovieRenderPipeline/Source/MovieRenderPipelineCore/Public/MoviePipelineQueue.h b/Engine/Plugins/MovieScene/MovieRenderPipeline/Source/MovieRenderPipelineCore/Public/MoviePipelineQueue.h
--- a/Engine/Plugins/MovieScene/MovieRenderPipeline/Source/MovieRenderPipelineCore/Public/MoviePipelineQueue.h
+++ b/Engine/Plugins/MovieScene/MovieRenderPipeline/Source/MovieRenderPipelineCore/Public/MoviePipelineQueue.h
@@ -398,6 +398,12 @@ public:
 	UFUNCTION(BlueprintCallable, Category = "Movie Render Pipeline|Queue")
 	void DeleteJob(UMoviePipelineExecutorJob* InJob);
 
+	/**
+	* Delete all jobs from the Queue.
+	*/
+	UFUNCTION(BlueprintCallable, Category = "Movie Render Pipeline|Queue", meta = (Keywords = "Clear Empty"))
+	void DeleteAllJobs();
+
 	/**
 	* Duplicate the specific job and return the duplicate. Configurations are duplicated and not shared.
 	*
