Adding GetCurrentSequence and GetCurrentExecutorShot to MoviePipelineBlueprintLibrary

diff --git a/Engine/Plugins/MovieScene/MovieRenderPipeline/Source/MovieRenderPipelineCore/Private/MoviePipelineBlueprintLibrary.cpp b/Engine/Plugins/MovieScene/MovieRenderPipeline/Source/MovieRenderPipelineCore/Private/MoviePipelineBlueprintLibrary.cpp
--- a/Engine/Plugins/MovieScene/MovieRenderPipeline/Source/MovieRenderPipelineCore/Private/MoviePipelineBlueprintLibrary.cpp
+++ b/Engine/Plugins/MovieScene/MovieRenderPipeline/Source/MovieRenderPipelineCore/Private/MoviePipelineBlueprintLibrary.cpp
@@ -910,3 +910,27 @@ void UMoviePipelineBlueprintLibrary::ResolveFilenameFormatArguments(const FStrin
 	}
 }
 
+ULevelSequence* UMoviePipelineBlueprintLibrary::GetCurrentSequence(const UMoviePipeline* InMoviePipeline)
+{
+	if (InMoviePipeline)
+	{
+		return InMoviePipeline->GetTargetSequence();
+	}
+
+	return nullptr;
+}
+
+UMoviePipelineExecutorShot* UMoviePipelineBlueprintLibrary::GetCurrentExecutorShot(const UMoviePipeline* InMoviePipeline)
+{
+	if (InMoviePipeline)
+	{
+		const TArray<UMoviePipelineExecutorShot*>& ActiveShotList = InMoviePipeline->GetActiveShotList();
+		int32 CurrentShotIndex = InMoviePipeline->GetCurrentShotIndex();
+		if (ActiveShotList.IsValidIndex(CurrentShotIndex))
+		{
+			return ActiveShotList[CurrentShotIndex];
+		}
+	}
+
+	return nullptr;
+}
\ No newline at end of file
diff --git a/Engine/Plugins/MovieScene/MovieRenderPipeline/Source/MovieRenderPipelineCore/Public/MoviePipelineBlueprintLibrary.h b/Engine/Plugins/MovieScene/MovieRenderPipeline/Source/MovieRenderPipelineCore/Public/MoviePipelineBlueprintLibrary.h
--- a/Engine/Plugins/MovieScene/MovieRenderPipeline/Source/MovieRenderPipelineCore/Public/MoviePipelineBlueprintLibrary.h
+++ b/Engine/Plugins/MovieScene/MovieRenderPipeline/Source/MovieRenderPipelineCore/Public/MoviePipelineBlueprintLibrary.h
@@ -141,4 +141,9 @@ public:
 	UFUNCTION(BlueprintCallable, Category = "Movie Render Pipeline")
 	static UMoviePipelineSetting* FindOrGetDefaultSettingForShot(TSubclassOf<UMoviePipelineSetting> InSettingType, const UMoviePipelineMasterConfig* InMasterConfig, const UMoviePipelineExecutorShot* InShot);
 
+	UFUNCTION(BlueprintPure, Category = "Movie Render Pipeline")
+	static ULevelSequence* GetCurrentSequence(const UMoviePipeline* InMoviePipeline);
+
+	UFUNCTION(BlueprintPure, Category = "Movie Render Pipeline")
+	static UMoviePipelineExecutorShot* GetCurrentExecutorShot(const UMoviePipeline* InMoviePipeline);
 };
\ No newline at end of file
