MovieRenderQueue: Add some extra information on where to find settings for the Path Tracer

diff --git a/Engine/Plugins/MovieScene/MovieRenderPipeline/Source/MovieRenderPipelineRenderPasses/Private/MoviePipelineDeferredPasses.cpp b/Engine/Plugins/MovieScene/MovieRenderPipeline/Source/MovieRenderPipelineRenderPasses/Private/MoviePipelineDeferredPasses.cpp
--- a/Engine/Plugins/MovieScene/MovieRenderPipeline/Source/MovieRenderPipelineRenderPasses/Private/MoviePipelineDeferredPasses.cpp
+++ b/Engine/Plugins/MovieScene/MovieRenderPipeline/Source/MovieRenderPipelineRenderPasses/Private/MoviePipelineDeferredPasses.cpp
@@ -742,6 +742,16 @@ void UMoviePipelineDeferredPassBase::PostRendererSubmission(const FMoviePipeline
 
 }
 
+#if WITH_EDITOR
+FText UMoviePipelineDeferredPass_PathTracer::GetFooterText(UMoviePipelineExecutorJob* InJob) const {
+	return NSLOCTEXT(
+		"MovieRenderPipeline",
+		"DeferredBasePassSetting_FooterText_PathTracer",
+		"Samples per Pixel for the Path Tracer are controlled by the Spatial Sample Count from the Anti-Aliasing settings.\n"
+		"All other Path Tracer settings are taken from the Post Process settings as usual.");
+}
+#endif
+
 namespace MoviePipeline
 {
 	static void AccumulateSample_TaskThread(TUniquePtr<FImagePixelData>&& InPixelData, const MoviePipeline::FImageSampleAccumulationArgs& InParams)
diff --git a/Engine/Plugins/MovieScene/MovieRenderPipeline/Source/MovieRenderPipelineRenderPasses/Public/MoviePipelineDeferredPasses.h b/Engine/Plugins/MovieScene/MovieRenderPipeline/Source/MovieRenderPipelineRenderPasses/Public/MoviePipelineDeferredPasses.h
--- a/Engine/Plugins/MovieScene/MovieRenderPipeline/Source/MovieRenderPipelineRenderPasses/Public/MoviePipelineDeferredPasses.h
+++ b/Engine/Plugins/MovieScene/MovieRenderPipeline/Source/MovieRenderPipelineRenderPasses/Public/MoviePipelineDeferredPasses.h
@@ -251,6 +251,7 @@ public:
 	}
 #if WITH_EDITOR
 	virtual FText GetDisplayText() const override { return NSLOCTEXT("MovieRenderPipeline", "DeferredBasePassSetting_DisplayName_PathTracer", "Path Tracer"); }
+	virtual FText GetFooterText(UMoviePipelineExecutorJob* InJob) const override;
 #endif
 	virtual void GetViewShowFlags(FEngineShowFlags& OutShowFlag, EViewModeIndex& OutViewModeIndex) const override
 	{
