MoviePipeline: Added a cvar for controlling how long the .wav writer waits for the async process to finish until the API allows us to do this properly.

diff --git a/Engine/Plugins/MovieScene/MovieRenderPipeline/Source/MovieRenderPipelineRenderPasses/Private/MoviePipelineWaveOutput.cpp b/Engine/Plugins/MovieScene/MovieRenderPipeline/Source/MovieRenderPipelineRenderPasses/Private/MoviePipelineWaveOutput.cpp
--- a/Engine/Plugins/MovieScene/MovieRenderPipeline/Source/MovieRenderPipelineRenderPasses/Private/MoviePipelineWaveOutput.cpp
+++ b/Engine/Plugins/MovieScene/MovieRenderPipeline/Source/MovieRenderPipelineRenderPasses/Private/MoviePipelineWaveOutput.cpp
@@ -15,6 +15,14 @@
 #include "AudioThread.h"
 #include "MoviePipelineUtils.h"
 
+static TAutoConsoleVariable<float> CVarWaveOutputDelay(
+	TEXT("MovieRenderPipeline.WaveOutput.WriteDelay"),
+	0.5f,
+	TEXT("How long (in seconds) should the .wav writer stall the main thread to wait for the async write to finish\n")
+	TEXT("before moving on? If your .wav files take a long time to write and they're not finished by the time the\n")
+	TEXT("encoder runs, the encoder may fail.\n"),
+	ECVF_Default);
+
 static FAudioDevice* GetAudioDeviceFromWorldContext(const UObject* WorldContextObject)
 {
 	UWorld* ThisWorld = GEngine->GetWorldFromContextObject(WorldContextObject, EGetWorldErrorMode::LogAndReturnNull);
@@ -180,7 +188,8 @@ void UMoviePipelineWaveOutput::BeginFinalizeImpl()
 	// can be spun (as it enqueues a callback onto the main thread). We're going to just cheat here and stall the main thread
 	// for 0.5s to give it a chance to write to disk. It'll only potentially be an issue with command line encoding if it takes
 	// longer than 0.5s to write to disk.
-	FPlatformProcess::Sleep(0.5f);
+	UE_LOG(LogMovieRenderPipeline, Log, TEXT("Delaying main thread for %f seconds while audio writes to disk."), CVarWaveOutputDelay.GetValueOnGameThread());
+	FPlatformProcess::Sleep(CVarWaveOutputDelay.GetValueOnGameThread());
 }
 
 void UMoviePipelineWaveOutput::OnShotFinishedImpl(const UMoviePipelineExecutorShot* InShot, const bool bFlushToDisk)
