Fix a crash in the movie render queue plugin when initializing the asset editor.
Fix an unknown toolbar tab that was getting added to the default layout in movie render queue plugin.

diff --git a/Engine/Plugins/MovieScene/MovieRenderPipeline/Source/MovieRenderPipelineEditor/Private/MoviePipelineConfigAssetEditor.cpp b/Engine/Plugins/MovieScene/MovieRenderPipeline/Source/MovieRenderPipelineEditor/Private/MoviePipelineConfigAssetEditor.cpp
--- a/Engine/Plugins/MovieScene/MovieRenderPipeline/Source/MovieRenderPipelineEditor/Private/MoviePipelineConfigAssetEditor.cpp
+++ b/Engine/Plugins/MovieScene/MovieRenderPipeline/Source/MovieRenderPipelineEditor/Private/MoviePipelineConfigAssetEditor.cpp
@@ -29,13 +29,6 @@ FMoviePipelineConfigAssetEditor::FMoviePipelineConfigAssetEditor(UAssetEditor* I
 		(
 			FTabManager::NewPrimaryArea()->SetOrientation(Orient_Vertical)
 			->Split
-			(
-				FTabManager::NewStack()
-				->SetSizeCoefficient(0.1f)
-				->AddTab(GetToolbarTabId(), ETabState::OpenedTab)
-				->SetHideTabWell(true)
-			)
-			->Split
 			(
 				FTabManager::NewStack()
 				->SetSizeCoefficient(1.f)
@@ -49,7 +42,7 @@ void FMoviePipelineConfigAssetEditor::RegisterTabSpawners(const TSharedRef<class
 {
 	WorkspaceMenuCategory = InTabManager->AddLocalWorkspaceMenuCategory(LOCTEXT("WorkspaceMenu_PipelineAssetEditor", "Movie Pipeline Asset Editor"));
 
-	FAssetEditorToolkit::RegisterTabSpawners(InTabManager);
+	FBaseAssetToolkit::RegisterTabSpawners(InTabManager);
 
 	InTabManager->RegisterTabSpawner(ContentTabId, FOnSpawnTab::CreateSP(this, &FMoviePipelineConfigAssetEditor::SpawnTab_ConfigEditor))
 		.SetDisplayName(LOCTEXT("ConfigTab", "Config"))
@@ -87,7 +80,6 @@ TSharedRef<SDockTab> FMoviePipelineConfigAssetEditor::SpawnTab_ConfigEditor(cons
 				Section.AddEntry(FToolMenuEntry::InitWidget(FName("Settings"), ConfigEditorPanel->MakeSettingsWidget(), LOCTEXT("AddSettings_Label", "Settings")));
 			}
 		}
-		RegenerateMenusAndToolbars();
 	}
 
 	// Tabs and Windows have different backgrounds, so we copy the window style to
