MoviePipeline: Exposed {job_name} to the format strings. {sequence_name} will be asset, and {job_name} the name set on the job (which defaults to the Sequence asest name)

diff --git a/Engine/Plugins/MovieScene/MovieRenderPipeline/Source/MovieRenderPipelineCore/Private/MoviePipelineMasterConfig.cpp b/Engine/Plugins/MovieScene/MovieRenderPipeline/Source/MovieRenderPipelineCore/Private/MoviePipelineMasterConfig.cpp
--- a/Engine/Plugins/MovieScene/MovieRenderPipeline/Source/MovieRenderPipelineCore/Private/MoviePipelineMasterConfig.cpp
+++ b/Engine/Plugins/MovieScene/MovieRenderPipeline/Source/MovieRenderPipelineCore/Private/MoviePipelineMasterConfig.cpp
@@ -94,12 +94,14 @@ void UMoviePipelineMasterConfig::GetFormatArguments(FMoviePipelineFormatArgs& In
 	{
 		FString LevelName = TEXT("Level Name");
 		FString SequenceName = TEXT("Sequence Name");
+		FString JobName = TEXT("Job Name");
 		double FrameRate = 0.0;
 
 		if (InOutFormatArgs.InJob)
 		{
 			LevelName = InOutFormatArgs.InJob->Map.GetAssetName();
 			SequenceName = InOutFormatArgs.InJob->Sequence.GetAssetName();
+			JobName = InOutFormatArgs.InJob->JobName;
 			
 			// FrameRate is a combination of Output Settings and Sequence so we do it here instead of in OutputSetting
 			FrameRate = GetEffectiveFrameRate(Cast<ULevelSequence>(InOutFormatArgs.InJob->Sequence.TryLoad())).AsDecimal();
@@ -107,10 +109,12 @@ void UMoviePipelineMasterConfig::GetFormatArguments(FMoviePipelineFormatArgs& In
 
 		InOutFormatArgs.FilenameArguments.Add(TEXT("level_name"), LevelName);
 		InOutFormatArgs.FilenameArguments.Add(TEXT("sequence_name"), SequenceName);
+		InOutFormatArgs.FilenameArguments.Add(TEXT("job_name"), JobName);
 		InOutFormatArgs.FilenameArguments.Add(TEXT("frame_rate"), FString::SanitizeFloat(FrameRate));
 
 		InOutFormatArgs.FileMetadata.Add(TEXT("unreal/levelName"), LevelName);
 		InOutFormatArgs.FileMetadata.Add(TEXT("unreal/sequenceName"), SequenceName);
+		InOutFormatArgs.FileMetadata.Add(TEXT("unreal/jobName"), JobName);
 		InOutFormatArgs.FileMetadata.Add(TEXT("unreal/frameRate"), FString::SanitizeFloat(FrameRate));
 
 		
