MoviePipeline: Refactored how command line arguments are generated so that Python pipelines can better hook into them in certain situations.

diff --git a/Engine/Plugins/MovieScene/MovieRenderPipeline/Source/MovieRenderPipelineCore/Private/MoviePipelineBlueprintLibrary.cpp b/Engine/Plugins/MovieScene/MovieRenderPipeline/Source/MovieRenderPipelineCore/Private/MoviePipelineBlueprintLibrary.cpp
--- a/Engine/Plugins/MovieScene/MovieRenderPipeline/Source/MovieRenderPipelineCore/Private/MoviePipelineBlueprintLibrary.cpp
+++ b/Engine/Plugins/MovieScene/MovieRenderPipeline/Source/MovieRenderPipelineCore/Private/MoviePipelineBlueprintLibrary.cpp
@@ -28,6 +28,8 @@
 #include "MovieSceneSequenceVisitor.h"
 #include "MovieSceneSequenceID.h"
 #include "MoviePipelineUtils.h"
+#include "UObject/UObjectHash.h"
+
 
 // For camera settings
 #include "GameFramework/PlayerController.h"
@@ -344,6 +346,32 @@ FString UMoviePipelineBlueprintLibrary::GetMapPackageName(UMoviePipelineExecutor
 	return InJob->Map.GetLongPackageName();
 }
 
+UMoviePipelineQueue* UMoviePipelineBlueprintLibrary::LoadManifestFileFromString(const FString& InManifestFilePath)
+{
+	FString InFileName = TEXT("QueueManifest");
+	FString InPackagePath = TEXT("/Engine/MovieRenderPipeline/Editor/Transient");
+
+	FString NewPackageName = FPackageName::GetLongPackagePath(InPackagePath) + TEXT("/") + InFileName;
+
+	// Relative paths are considered relative to the game's save directory to avoid having to hard code
+	// game names into relative paths.
+	FString ConfigAssetPath = InManifestFilePath;
+	if (FPaths::IsRelative(InManifestFilePath))
+	{
+		ConfigAssetPath = FPaths::Combine(FPaths::ProjectSavedDir(), InManifestFilePath);
+	}
+
+	UPackage* OuterPackage = CreatePackage(*NewPackageName);
+	UPackage* QueuePackage = LoadPackage(OuterPackage, *ConfigAssetPath, LOAD_None);
+	UMoviePipelineQueue* OutQueue = nullptr;
+	if (QueuePackage)
+	{
+		OutQueue = Cast<UMoviePipelineQueue>((UObject*)FindObjectWithOuter(QueuePackage, UMoviePipelineQueue::StaticClass()));
+	}
+
+	return OutQueue;
+}
+
 
 void UMoviePipelineBlueprintLibrary::UpdateJobShotListFromSequence(ULevelSequence* InSequence, UMoviePipelineExecutorJob* InJob, bool& bShotsChanged)
 {
diff --git a/Engine/Plugins/MovieScene/MovieRenderPipeline/Source/MovieRenderPipelineCore/Private/MoviePipelineGameOverrideSetting.cpp b/Engine/Plugins/MovieScene/MovieRenderPipeline/Source/MovieRenderPipelineCore/Private/MoviePipelineGameOverrideSetting.cpp
--- a/Engine/Plugins/MovieScene/MovieRenderPipeline/Source/MovieRenderPipelineCore/Private/MoviePipelineGameOverrideSetting.cpp
+++ b/Engine/Plugins/MovieScene/MovieRenderPipeline/Source/MovieRenderPipelineCore/Private/MoviePipelineGameOverrideSetting.cpp
@@ -108,30 +108,35 @@ void UMoviePipelineGameOverrideSetting::ApplyCVarSettings(const bool bOverrideVa
 
 }
 
-void UMoviePipelineGameOverrideSetting::BuildNewProcessCommandLineImpl(FString& InOutUnrealURLParams, FString& InOutCommandLineArgs) const
+void UMoviePipelineGameOverrideSetting::BuildNewProcessCommandLineArgsImpl(TArray<FString>& InOutUnrealURLParams, TArray<FString>& InOutCommandLineArgs, TArray<FString>& InOutDeviceProfileCvars, TArray<FString>& InOutExecCmds) const
 {
 	if (!IsEnabled())
 	{
 		return;
 	}
 
-	FString CVarCommandLineArgs;
-	FString CVarExecArgs;
-
 	// We don't provide the GameMode on the command line argument as we expect NewProcess to boot into an empty map and then it will
 	// transition into the correct map which will then use the GameModeOverride setting.
 	if (bCinematicQualitySettings)
 	{
-		CVarCommandLineArgs += TEXT("sg.ViewDistanceQuality=4,sg.AntiAliasingQuality=4,sg.ShadowQuality=4,sg.PostProcessQuality=4,sg.TextureQuality=4,sg.EffectsQuality=4,sg.FoliageQuality=4,sg.ShadingQuality=4,");
+		InOutDeviceProfileCvars.Add(TEXT("sg.ViewDistanceQuality=4"));
+		InOutDeviceProfileCvars.Add(TEXT("sg.AntiAliasingQuality=4"));
+		InOutDeviceProfileCvars.Add(TEXT("sg.ShadowQuality=4"));
+		InOutDeviceProfileCvars.Add(TEXT("sg.PostProcessQuality=4"));
+		InOutDeviceProfileCvars.Add(TEXT("sg.TextureQuality=4"));
+		InOutDeviceProfileCvars.Add(TEXT("sg.EffectsQuality=4"));
+		InOutDeviceProfileCvars.Add(TEXT("sg.FoliageQuality=4"));
+		InOutDeviceProfileCvars.Add(TEXT("sg.ShadingQuality=4"));
 	}
 
 	switch (TextureStreaming)
 	{
 	case EMoviePipelineTextureStreamingMethod::FullyLoad:
-		CVarCommandLineArgs += TEXT("r.Streaming.FramesForFullUpdate=0,r.Streaming.FullyLoadUsedTextures=1,");
+		InOutDeviceProfileCvars.Add(TEXT("r.Streaming.FramesForFullUpdate=0"));
+		InOutDeviceProfileCvars.Add(TEXT("r.Streaming.FullyLoadUsedTextures=1"));
 		break;
 	case EMoviePipelineTextureStreamingMethod::Disabled:
-		CVarCommandLineArgs += TEXT("r.TextureStreaming=0,");
+		InOutDeviceProfileCvars.Add(TEXT("r.TextureStreaming=0"));
 		break;
 	default:
 		// We don't change their texture streaming settings.
@@ -140,36 +145,39 @@ void UMoviePipelineGameOverrideSetting::BuildNewProcessCommandLineImpl(FString&
 
 	if (bUseLODZero)
 	{
-		CVarCommandLineArgs += TEXT("r.ForceLOD=0,r.SkeletalMeshLODBias=-10,r.ParticleLODBias=-10,foliage.DitheredLOD=0,foliage.ForceLOD=0");
+		InOutDeviceProfileCvars.Add(TEXT("r.ForceLOD=0"));
+		InOutDeviceProfileCvars.Add(TEXT("r.SkeletalMeshLODBias=-10"));
+		InOutDeviceProfileCvars.Add(TEXT("r.ParticleLODBias=-10"));
+		InOutDeviceProfileCvars.Add(TEXT("foliage.DitheredLOD=0"));
+		InOutDeviceProfileCvars.Add(TEXT("foliage.ForceLOD=0"));
 	}
 
 	if (bDisableHLODs)
 	{
 		// It's a command and not an integer cvar (despite taking 1/0)
-		CVarExecArgs += TEXT("r.HLOD 0,");
+		InOutExecCmds.Add(TEXT("r.HLOD 0"));
 	}
 
 	if (bUseHighQualityShadows)
 	{
-		CVarCommandLineArgs += FString::Printf(TEXT("r.Shadow.DistanceScale=%d,r.ShadowQuality=5,r.Shadow.RadiusThreshold=%f,"), ShadowDistanceScale, ShadowRadiusThreshold);
+		InOutDeviceProfileCvars.Add(FString::Printf(TEXT("r.Shadow.DistanceScale=%d"), ShadowDistanceScale));
+		InOutDeviceProfileCvars.Add(FString::Printf(TEXT("r.Shadow.RadiusThreshold=%f"), ShadowRadiusThreshold));
+		InOutDeviceProfileCvars.Add(TEXT("r.ShadowQuality=5"));
 	}
 
 	if (bOverrideViewDistanceScale)
 	{
-		CVarCommandLineArgs += FString::Printf(TEXT("r.ViewDistanceScale=%d,"), ViewDistanceScale);
+		InOutDeviceProfileCvars.Add(FString::Printf(TEXT("r.ViewDistanceScale=%d"), ViewDistanceScale));
 	}
 
 	if (bDisableGPUTimeout)
 	{
-		CVarCommandLineArgs += TEXT("r.D3D12.GPUTimeout=0,");
+		InOutDeviceProfileCvars.Add(TEXT("r.D3D12.GPUTimeout=0"));
 	}
 	
 	{
-		CVarCommandLineArgs += FString::Printf(TEXT("a.URO.Enable=%d,"), 0);
+		InOutDeviceProfileCvars.Add(FString::Printf(TEXT("a.URO.Enable=%d"), 0));
 	}
 
-	CVarCommandLineArgs += FString::Printf(TEXT("r.SkyLight.RealTimeReflectionCapture.TimeSlice=%d,"), 0);
-
-	// Apply the cvar very early on (device profile) time instead of after the map has loaded if possible
-	InOutCommandLineArgs += FString::Printf(TEXT(" -dpcvars=%s -execcmds=%s"), *CVarCommandLineArgs, *CVarExecArgs);
+	InOutDeviceProfileCvars.Add(FString::Printf(TEXT("r.SkyLight.RealTimeReflectionCapture.TimeSlice=%d"), 0));
 }
\ No newline at end of file
diff --git a/Engine/Plugins/MovieScene/MovieRenderPipeline/Source/MovieRenderPipelineCore/Private/MovieRenderPipelineCommandLine.cpp b/Engine/Plugins/MovieScene/MovieRenderPipeline/Source/MovieRenderPipelineCore/Private/MovieRenderPipelineCommandLine.cpp
--- a/Engine/Plugins/MovieScene/MovieRenderPipeline/Source/MovieRenderPipelineCore/Private/MovieRenderPipelineCommandLine.cpp
+++ b/Engine/Plugins/MovieScene/MovieRenderPipeline/Source/MovieRenderPipelineCore/Private/MovieRenderPipelineCommandLine.cpp
@@ -16,6 +16,7 @@
 #include "UObject/UObjectHash.h"
 #include "MoviePipelinePythonHostExecutor.h"
 #include "MoviePipelineUtils.h"
+#include "MoviePipelineBlueprintLibrary.h"
 
 #if WITH_EDITOR
 //#include "Editor.h"
@@ -282,25 +283,7 @@ uint8 FMovieRenderPipelineCoreModule::ParseMovieRenderData(const FString& InSequ
 		// If they didn't pass a path that started with /Game/, we'll try to see if it is a manifest file.
 		if (InConfigAssetPath.EndsWith(FPackageName::GetTextAssetPackageExtension()))
 		{
-			FString InFileName = TEXT("QueueManifest");
-			FString InPackagePath = TEXT("/Engine/MovieRenderPipeline/Editor/Transient");
-
-			FString NewPackageName = FPackageName::GetLongPackagePath(InPackagePath) + TEXT("/") + InFileName;
-
-			// Relative paths are considered relative to the game's save directory to avoid having to hard code
-			// game names into relative paths.
-			FString ConfigAssetPath = InConfigAssetPath;
-			if (FPaths::IsRelative(InConfigAssetPath))
-			{
-				ConfigAssetPath = FPaths::Combine(FPaths::ProjectSavedDir(), InConfigAssetPath);
-			}
-
-			UPackage* OuterPackage = CreatePackage(*NewPackageName);
-			UPackage* QueuePackage = LoadPackage(OuterPackage, *ConfigAssetPath, LOAD_None);
-			if (QueuePackage)
-			{
-				OutQueue = Cast<UMoviePipelineQueue>((UObject*)FindObjectWithOuter(QueuePackage, UMoviePipelineQueue::StaticClass()));
-			}
+			OutQueue = UMoviePipelineBlueprintLibrary::LoadManifestFileFromString(InConfigAssetPath);
 			
 			if(!OutQueue)
 			{
diff --git a/Engine/Plugins/MovieScene/MovieRenderPipeline/Source/MovieRenderPipelineCore/Public/MoviePipelineBlueprintLibrary.h b/Engine/Plugins/MovieScene/MovieRenderPipeline/Source/MovieRenderPipelineCore/Public/MoviePipelineBlueprintLibrary.h
--- a/Engine/Plugins/MovieScene/MovieRenderPipeline/Source/MovieRenderPipelineCore/Public/MoviePipelineBlueprintLibrary.h
+++ b/Engine/Plugins/MovieScene/MovieRenderPipeline/Source/MovieRenderPipelineCore/Public/MoviePipelineBlueprintLibrary.h
@@ -125,6 +125,10 @@ public:
 	UFUNCTION(BlueprintPure, Category = "Movie Render Pipeline")
 	static FString GetMapPackageName(UMoviePipelineExecutorJob* InJob);
 
+	/** Loads the specified manifest file and converts it into an UMoviePipelineQueue. Use in combination with SaveQueueToManifestFile. */
+	UFUNCTION(BlueprintCallable, Category = "Movie Render Pipeline")
+	static class UMoviePipelineQueue* LoadManifestFileFromString(const FString& InManifestFilePath);
+
 	/** Scan the provided sequence in the job to see which camera cut sections we would try to render and update the job's shotlist. */
 	UFUNCTION(BlueprintCallable, Category = "Movie Render Pipeline")
 	static void UpdateJobShotListFromSequence(ULevelSequence* InSequence, UMoviePipelineExecutorJob* InJob, bool& bShotsChanged);
diff --git a/Engine/Plugins/MovieScene/MovieRenderPipeline/Source/MovieRenderPipelineCore/Public/MoviePipelineGameOverrideSetting.h b/Engine/Plugins/MovieScene/MovieRenderPipeline/Source/MovieRenderPipelineCore/Public/MoviePipelineGameOverrideSetting.h
--- a/Engine/Plugins/MovieScene/MovieRenderPipeline/Source/MovieRenderPipelineCore/Public/MoviePipelineGameOverrideSetting.h
+++ b/Engine/Plugins/MovieScene/MovieRenderPipeline/Source/MovieRenderPipelineCore/Public/MoviePipelineGameOverrideSetting.h
@@ -46,7 +46,7 @@ public:
 #endif
 	virtual bool IsValidOnShots() const override { return false; }
 	virtual bool IsValidOnMaster() const override { return true; }
-	virtual void BuildNewProcessCommandLineImpl(FString& InOutUnrealURLParams, FString& InOutCommandLineArgs) const override;
+	virtual void BuildNewProcessCommandLineArgsImpl(TArray<FString>& InOutUnrealURLParams, TArray<FString>& InOutCommandLineArgs, TArray<FString>& InOutDeviceProfileCvars, TArray<FString>& InOutExecCmds) const override;
 	virtual void SetupForPipelineImpl(UMoviePipeline* InPipeline) override;
 	virtual void TeardownForPipelineImpl(UMoviePipeline* InPipeline) override;
 	// Used to ensure we set the default values even if the user forgets to add it.
diff --git a/Engine/Plugins/MovieScene/MovieRenderPipeline/Source/MovieRenderPipelineCore/Public/MoviePipelineSetting.h b/Engine/Plugins/MovieScene/MovieRenderPipeline/Source/MovieRenderPipelineCore/Public/MoviePipelineSetting.h
--- a/Engine/Plugins/MovieScene/MovieRenderPipeline/Source/MovieRenderPipelineCore/Public/MoviePipelineSetting.h
+++ b/Engine/Plugins/MovieScene/MovieRenderPipeline/Source/MovieRenderPipelineCore/Public/MoviePipelineSetting.h
@@ -47,8 +47,17 @@ public:
 	* allows a setting to provide these when the user wants to run in a separate process. This won't
 	* be used when running in the current process because it is too late to modify the command line.
 	*/
+	UE_DEPRECATED(5.0, "Use BuildNewProcessCommandLineArgs instead.")
 	UFUNCTION(BlueprintCallable, Category = "Movie Render Pipeline")
-	void BuildNewProcessCommandLine(UPARAM(ref) FString& InOutUnrealURLParams, UPARAM(ref) FString& InOutCommandLineArgs) const { BuildNewProcessCommandLineImpl(InOutUnrealURLParams, InOutCommandLineArgs); }
+	void BuildNewProcessCommandLine(UPARAM(ref) FString& InOutUnrealURLParams, UPARAM(ref) FString& InOutCommandLineArgs) const
+	{
+		PRAGMA_DISABLE_DEPRECATION_WARNINGS
+		BuildNewProcessCommandLineImpl(InOutUnrealURLParams, InOutCommandLineArgs); 
+		PRAGMA_ENABLE_DEPRECATION_WARNINGS
+	}
+
+	UFUNCTION(BlueprintCallable, Category = "Movie Render Pipeline")
+	void BuildNewProcessCommandLineArgs(UPARAM(ref) TArray<FString>& InOutUnrealURLParams, UPARAM(ref) TArray<FString>& InOutCommandLineArgs, UPARAM(ref) TArray<FString>& InOutDeviceProfileCvars, UPARAM(ref) TArray<FString>& InOutExecCmds) const { BuildNewProcessCommandLineArgsImpl(InOutUnrealURLParams, InOutCommandLineArgs, InOutDeviceProfileCvars, InOutExecCmds); }
 
 	/**
 	* Attempt to validate the configuration the user has chosen for this setting. Caches results for fast lookup in UI later.
@@ -114,8 +123,10 @@ public:
 	/** Return Key/Value pairs that you wish to be usable in the Output File Name format string or file metadata. This allows settings to add format strings based on their values. */
 	virtual void GetFormatArguments(FMoviePipelineFormatArgs& InOutFormatArgs) const {}
 	
-	/** Modify the Unreal URL and Command Line Arguments when preparing the setting to be run in a new process. */
+	UE_DEPRECATED(5.0, "Use BuildNewProcessCommandLineArgsImpl instead.")
 	virtual void BuildNewProcessCommandLineImpl(FString& InOutUnrealURLParams, FString& InOutCommandLineArgs) const { }
+	/** Modify the Unreal URL and Command Line Arguments when preparing the setting to be run in a new process. */
+	virtual void BuildNewProcessCommandLineArgsImpl(TArray<FString>& InOutUnrealURLParams, TArray<FString>& InOutCommandLineArgs, TArray<FString>& InOutDeviceProfileCvars, TArray<FString>& InOutExecCmds) const { }
 
 	/** Can only one of these settings objects be active in a valid pipeline? */
 	virtual bool IsSolo() const { return true; }
diff --git a/Engine/Plugins/MovieScene/MovieRenderPipeline/Source/MovieRenderPipelineEditor/Private/MoviePipelineNewProcessExecutor.cpp b/Engine/Plugins/MovieScene/MovieRenderPipeline/Source/MovieRenderPipelineEditor/Private/MoviePipelineNewProcessExecutor.cpp
--- a/Engine/Plugins/MovieScene/MovieRenderPipeline/Source/MovieRenderPipelineEditor/Private/MoviePipelineNewProcessExecutor.cpp
+++ b/Engine/Plugins/MovieScene/MovieRenderPipeline/Source/MovieRenderPipelineEditor/Private/MoviePipelineNewProcessExecutor.cpp
@@ -85,15 +85,50 @@ void UMoviePipelineNewProcessExecutor::Execute_Implementation(UMoviePipelineQueu
 
 	// Loop through our settings in the job and let them modify the command line arguments/params. Because we could have multiple jobs,
 	// we go through all jobs and all settings and hope the user doesn't have conflicting settings.
+	TArray<FString> OutUrlParams;
+	TArray<FString> OutCommandLineArgs;
+	TArray<FString> OutDeviceProfileCVars;
+	TArray<FString> OutExecCmds;
 	for (const UMoviePipelineExecutorJob* Job : DuplicatedQueue->GetJobs())
 	{
 		Job->GetConfiguration()->InitializeTransientSettings();
 		for (const UMoviePipelineSetting* Setting : Job->GetConfiguration()->GetAllSettings())
 		{
-			Setting->BuildNewProcessCommandLine(UnrealURLParams, CommandLineArgs);
+			Setting->BuildNewProcessCommandLineArgs(OutUrlParams, OutCommandLineArgs, OutDeviceProfileCVars, OutExecCmds);
 		}
 	}
 
+	for (const FString& UrlParam : OutUrlParams)
+	{
+		UnrealURLParams += UrlParam;
+	}
+
+	CommandLineArgs += TEXT(" ");
+	for (const FString& Arg : OutCommandLineArgs)
+	{
+		CommandLineArgs += FString::Printf(TEXT("%s,"), *Arg);
+	}
+
+	if (OutDeviceProfileCVars.Num() > 0)
+	{
+		CommandLineArgs += TEXT(" -dpcvars=\"");
+		for (const FString& Cvar : OutDeviceProfileCVars)
+		{
+			CommandLineArgs += FString::Printf(TEXT("%s,"), *Cvar);
+		}
+		CommandLineArgs += TEXT("\"");
+	}
+
+	if (OutExecCmds.Num() > 0)
+	{
+		CommandLineArgs += TEXT(" -execcmds=\"");
+		for (const FString& Cmd : OutExecCmds)
+		{
+			CommandLineArgs += FString::Printf(TEXT("%s,"), *Cmd);
+		}
+		CommandLineArgs += TEXT("\"");
+	}
+
 	FString GameNameOrProjectFile;
 	if (FPaths::IsProjectFilePathSet())
 	{
diff --git a/Engine/Plugins/MovieScene/MovieRenderPipeline/Source/MovieRenderPipelineEditor/Private/MoviePipelinePIEExecutor.cpp b/Engine/Plugins/MovieScene/MovieRenderPipeline/Source/MovieRenderPipelineEditor/Private/MoviePipelinePIEExecutor.cpp
--- a/Engine/Plugins/MovieScene/MovieRenderPipeline/Source/MovieRenderPipelineEditor/Private/MoviePipelinePIEExecutor.cpp
+++ b/Engine/Plugins/MovieScene/MovieRenderPipeline/Source/MovieRenderPipelineEditor/Private/MoviePipelinePIEExecutor.cpp
@@ -215,6 +215,11 @@ void UMoviePipelinePIEExecutor::OnTick()
 	{
 		if (RemainingInitializationFrames == 0)
 		{
+			if (CustomInitializationTime.IsSet())
+			{
+				ActiveMoviePipeline->SetInitializationTime(CustomInitializationTime.GetValue());
+			}
+
 			ActiveMoviePipeline->Initialize(Queue->GetJobs()[CurrentPipelineIndex]);
 		}
 
@@ -256,6 +261,14 @@ void UMoviePipelinePIEExecutor::OnJobShotFinished(FMoviePipelineOutputData InOut
 	OnIndividualShotWorkFinishedDelegate.Broadcast(InOutputData);
 }
 
+void UMoviePipelinePIEExecutor::BeginDestroy()
+{
+	// Ensure we're no longer gathering, otherwise it tries to call a callback on the now dead uobject.
+	ValidationMessageGatherer.StopGathering();
+
+	Super::BeginDestroy();
+}
+
 void UMoviePipelinePIEExecutor::OnPIEEnded(bool)
 {
 	FEditorDelegates::EndPIE.RemoveAll(this);
@@ -318,4 +331,4 @@ void UMoviePipelinePIEExecutor::OnIndividualJobFinishedImpl(FMoviePipelineOutput
 	OnIndividualJobWorkFinishedDelegate.Broadcast(InOutputData);
 }
 
-#undef LOCTEXT_NAMESPACE // "MoviePipelinePIEExecutor"
\ No newline at end of file
+#undef LOCTEXT_NAMESPACE // "MoviePipelinePIEExecutor"
diff --git a/Engine/Plugins/MovieScene/MovieRenderPipeline/Source/MovieRenderPipelineEditor/Public/MoviePipelinePIEExecutor.h b/Engine/Plugins/MovieScene/MovieRenderPipeline/Source/MovieRenderPipelineEditor/Public/MoviePipelinePIEExecutor.h
--- a/Engine/Plugins/MovieScene/MovieRenderPipeline/Source/MovieRenderPipelineEditor/Public/MoviePipelinePIEExecutor.h
+++ b/Engine/Plugins/MovieScene/MovieRenderPipeline/Source/MovieRenderPipelineEditor/Public/MoviePipelinePIEExecutor.h
@@ -52,10 +52,13 @@ public:
 		return OnIndividualShotWorkFinishedDelegateNative;
 	}
 
+	UFUNCTION(BlueprintCallable, Category = "Movie Render Pipeline")
+	void SetInitializationTime(const FDateTime& InInitializationTime) { CustomInitializationTime = InInitializationTime; }
+
 protected:
 	virtual void Start(const UMoviePipelineExecutorJob* InJob) override;
 
-	
+	virtual void BeginDestroy() override;
 	/**
 	* This should be called after PIE has been shut down for an individual job and it is generally safer to
 	* make modifications to the editor world.
@@ -83,7 +86,7 @@ private:
 	bool bPreviousUseFixedTimeStep;
 	double PreviousFixedTimeStepDelta;
 	TWeakPtr<class SWindow> WeakCustomWindow;
-
+	TOptional<FDateTime> CustomInitializationTime;
 	FMoviePipelineOutputData CachedOutputDataParams;
 
 	/** Deprecated. use OnIndividualJobWorkFinishedDelegate instead. */
diff --git a/Engine/Plugins/MovieScene/MovieRenderPipeline/Source/MovieRenderPipelineRenderPasses/Private/MoviePipelineWaveOutput.cpp b/Engine/Plugins/MovieScene/MovieRenderPipeline/Source/MovieRenderPipelineRenderPasses/Private/MoviePipelineWaveOutput.cpp
--- a/Engine/Plugins/MovieScene/MovieRenderPipeline/Source/MovieRenderPipelineRenderPasses/Private/MoviePipelineWaveOutput.cpp
+++ b/Engine/Plugins/MovieScene/MovieRenderPipeline/Source/MovieRenderPipelineRenderPasses/Private/MoviePipelineWaveOutput.cpp
@@ -223,8 +223,8 @@ void UMoviePipelineWaveOutput::ValidateStateImpl()
 	}
 }
 
-void UMoviePipelineWaveOutput::BuildNewProcessCommandLineImpl(FString& InOutUnrealURLParams, FString& InOutCommandLineArgs) const
+void UMoviePipelineWaveOutput::BuildNewProcessCommandLineArgsImpl(TArray<FString>& InOutUnrealURLParams, TArray<FString>& InOutCommandLineArgs, TArray<FString>& InOutDeviceProfileCvars, TArray<FString>& InOutExecCmds) const 
 {
 	// Always add this even if we're not disabled so that audio is muted, it'll never line up during preview anyways.
-	InOutCommandLineArgs += " -deterministicaudio";
+	InOutCommandLineArgs.Add("-deterministicaudio");
 }
diff --git a/Engine/Plugins/MovieScene/MovieRenderPipeline/Source/MovieRenderPipelineRenderPasses/Public/MoviePipelineWaveOutput.h b/Engine/Plugins/MovieScene/MovieRenderPipeline/Source/MovieRenderPipelineRenderPasses/Public/MoviePipelineWaveOutput.h
--- a/Engine/Plugins/MovieScene/MovieRenderPipeline/Source/MovieRenderPipelineRenderPasses/Public/MoviePipelineWaveOutput.h
+++ b/Engine/Plugins/MovieScene/MovieRenderPipeline/Source/MovieRenderPipelineRenderPasses/Public/MoviePipelineWaveOutput.h
@@ -24,7 +24,8 @@ protected:
 	virtual void BeginFinalizeImpl() override;
 	virtual bool HasFinishedProcessingImpl() override;
 	virtual void ValidateStateImpl() override;
-	virtual void BuildNewProcessCommandLineImpl(FString& InOutUnrealURLParams, FString& InOutCommandLineArgs) const override;
+	virtual void BuildNewProcessCommandLineArgsImpl(TArray<FString>& InOutUnrealURLParams, TArray<FString>& InOutCommandLineArgs, TArray<FString>& InOutDeviceProfileCvars, TArray<FString>& InOutExecCmds) const override;
+
 public:
 	/* File name format string override. If specified it will override the FileNameFormat from the Output setting. */
 	UPROPERTY(EditAnywhere, BlueprintReadWrite, Category = "File Output")
