MoviePipeline: The changelist string specified in the burn-in is now a modular feature so plugins can provide different/more specific information without needing to change the default burn-in.

diff --git a/Engine/Plugins/MovieScene/MovieRenderPipeline/Source/MovieRenderPipelineCore/Private/MoviePipelineBlueprintLibrary.cpp b/Engine/Plugins/MovieScene/MovieRenderPipeline/Source/MovieRenderPipelineCore/Private/MoviePipelineBlueprintLibrary.cpp
--- a/Engine/Plugins/MovieScene/MovieRenderPipeline/Source/MovieRenderPipelineCore/Private/MoviePipelineBlueprintLibrary.cpp
+++ b/Engine/Plugins/MovieScene/MovieRenderPipeline/Source/MovieRenderPipelineCore/Private/MoviePipelineBlueprintLibrary.cpp
@@ -29,7 +29,9 @@
 #include "MovieSceneSequenceID.h"
 #include "MoviePipelineUtils.h"
 #include "UObject/UObjectHash.h"
-
+#include "MovieRenderPipelineCoreModule.h"
+#include "Features/IModularFeatures.h"
+#include "Misc/EngineVersion.h"
 
 // For camera settings
 #include "GameFramework/PlayerController.h"
@@ -961,4 +963,31 @@ UMoviePipelineExecutorShot* UMoviePipelineBlueprintLibrary::GetCurrentExecutorSh
 	}
 
 	return nullptr;
+}
+
+FText UMoviePipelineBlueprintLibrary::GetMoviePipelineEngineChangelistLabel(const UMoviePipeline*)
+{
+	// This is a modular feature so that you can write a plugin that provides the string. If you don't write a plugin
+	// then it falls back to the existing logic which pulls the engine build version.
+	if(IModularFeatures::Get().IsModularFeatureAvailable(IMoviePipelineBurnInExtension::ModularFeatureName))
+	{
+		IMoviePipelineBurnInExtension& Extension = IModularFeatures::Get().GetModularFeature<IMoviePipelineBurnInExtension>(IMoviePipelineBurnInExtension::ModularFeatureName);
+		const FText Label = Extension.GetEngineChangelistLabel();
+		if (!Label.IsEmpty())
+		{
+			return Label;
+		}
+	}
+
+	FString EngineVersion = FEngineVersion::Current().ToString();
+	int32 ChangeListIndexStart = EngineVersion.Find(TEXT("-"));
+	int32 ChangeListIndexEnd = EngineVersion.Find(TEXT("+"));
+
+	if (ChangeListIndexStart > 0 && ChangeListIndexEnd > 0)
+	{
+		FString Substr = EngineVersion.Mid(ChangeListIndexStart + 1, (ChangeListIndexEnd - ChangeListIndexStart) - 1);
+		return FText::FromString(Substr);
+	}
+
+	return FText();
 }
\ No newline at end of file
diff --git a/Engine/Plugins/MovieScene/MovieRenderPipeline/Source/MovieRenderPipelineCore/Private/MovieRenderPipelineCoreModule.cpp b/Engine/Plugins/MovieScene/MovieRenderPipeline/Source/MovieRenderPipelineCore/Private/MovieRenderPipelineCoreModule.cpp
--- a/Engine/Plugins/MovieScene/MovieRenderPipeline/Source/MovieRenderPipelineCore/Private/MovieRenderPipelineCoreModule.cpp
+++ b/Engine/Plugins/MovieScene/MovieRenderPipeline/Source/MovieRenderPipelineCore/Private/MovieRenderPipelineCoreModule.cpp
@@ -9,6 +9,8 @@
 #include "TimerManager.h"
 #include "MoviePipeline.h"
 
+FName IMoviePipelineBurnInExtension::ModularFeatureName = "ModularFeature_MoviePipelineBurnInExt";
+
 void FMovieRenderPipelineCoreModule::StartupModule()
 {
 #if WITH_EDITOR
diff --git a/Engine/Plugins/MovieScene/MovieRenderPipeline/Source/MovieRenderPipelineCore/Public/MoviePipelineBlueprintLibrary.h b/Engine/Plugins/MovieScene/MovieRenderPipeline/Source/MovieRenderPipelineCore/Public/MoviePipelineBlueprintLibrary.h
--- a/Engine/Plugins/MovieScene/MovieRenderPipeline/Source/MovieRenderPipelineCore/Public/MoviePipelineBlueprintLibrary.h
+++ b/Engine/Plugins/MovieScene/MovieRenderPipeline/Source/MovieRenderPipelineCore/Public/MoviePipelineBlueprintLibrary.h
@@ -150,4 +150,9 @@ public:
 
 	UFUNCTION(BlueprintPure, Category = "Movie Render Pipeline")
 	static UMoviePipelineExecutorShot* GetCurrentExecutorShot(const UMoviePipeline* InMoviePipeline);
+
+	/** Get a string to represent the Changelist Number for the burn in. This can be driven by a Modular Feature if you want to permanently replace it with different information. */
+	UFUNCTION(BlueprintPure, Category = "Movie Render Pipeline")
+	static FText GetMoviePipelineEngineChangelistLabel(const UMoviePipeline* InMoviePipeline);
+
 };
\ No newline at end of file
diff --git a/Engine/Plugins/MovieScene/MovieRenderPipeline/Source/MovieRenderPipelineCore/Public/MovieRenderPipelineCoreModule.h b/Engine/Plugins/MovieScene/MovieRenderPipeline/Source/MovieRenderPipelineCore/Public/MovieRenderPipelineCoreModule.h
--- a/Engine/Plugins/MovieScene/MovieRenderPipeline/Source/MovieRenderPipelineCore/Public/MovieRenderPipelineCoreModule.h
+++ b/Engine/Plugins/MovieScene/MovieRenderPipeline/Source/MovieRenderPipelineCore/Public/MovieRenderPipelineCoreModule.h
@@ -5,6 +5,7 @@
 #include "CoreMinimal.h"
 #include "Stats/Stats2.h"
 #include "Modules/ModuleInterface.h"
+#include "Features/IModularFeature.h"
 
 // Forward Declare
 class UMoviePipelineExecutorBase;
@@ -54,5 +55,15 @@ private:
 	FString SettingsAssetValue;
 };
 
+class MOVIERENDERPIPELINECORE_API IMoviePipelineBurnInExtension : public IModularFeature
+{
+public:
+	virtual ~IMoviePipelineBurnInExtension() {}
+	static FName ModularFeatureName;
+
+	virtual FText GetEngineChangelistLabel() const { return FText(); }
+};
+
+
 MOVIERENDERPIPELINECORE_API DECLARE_LOG_CATEGORY_EXTERN(LogMovieRenderPipeline, Log, All);
 MOVIERENDERPIPELINECORE_API DECLARE_LOG_CATEGORY_EXTERN(LogMovieRenderPipelineIO, Log, All);
\ No newline at end of file
