MoviePipeline: Added DWAA and DWAB compression methods for multi-layer-enabled EXRs (even if they only have one layer). This only affects MRQ's ability to write exrs and does not change general exr support in the engine.

diff --git a/Engine/Plugins/MovieScene/MovieRenderPipeline/Source/MovieRenderPipelineRenderPasses/Private/MoviePipelineEXROutput.cpp b/Engine/Plugins/MovieScene/MovieRenderPipeline/Source/MovieRenderPipelineRenderPasses/Private/MoviePipelineEXROutput.cpp
--- a/Engine/Plugins/MovieScene/MovieRenderPipeline/Source/MovieRenderPipelineRenderPasses/Private/MoviePipelineEXROutput.cpp
+++ b/Engine/Plugins/MovieScene/MovieRenderPipeline/Source/MovieRenderPipelineRenderPasses/Private/MoviePipelineEXROutput.cpp
@@ -122,6 +122,10 @@ bool FEXRImageWriteTask::WriteToDisk()
 				FileCompression = Imf::Compression::ZIP_COMPRESSION; break;
 			case EEXRCompressionFormat::PIZ:
 				FileCompression = Imf::Compression::PIZ_COMPRESSION; break;
+			case EEXRCompressionFormat::DWAA:
+				FileCompression = Imf::Compression::DWAA_COMPRESSION; break;
+			case EEXRCompressionFormat::DWAB:
+				FileCompression = Imf::Compression::DWAB_COMPRESSION; break;
 			default: 
 				checkNoEntry();
 		}
@@ -144,6 +148,13 @@ bool FEXRImageWriteTask::WriteToDisk()
 		
 		Imf::Header Header(DisplayWindow, DataWindow, 1, Imath::V2f(0, 0), 1, Imf::LineOrder::INCREASING_Y, FileCompression);
 		
+		// If using lossy compression, specify the compression level in the header per exr spec.
+		if (FileCompression == Imf::Compression::DWAA_COMPRESSION ||
+			FileCompression == Imf::Compression::DWAB_COMPRESSION)
+		{
+			FileMetadata.Add("dwaCompressionLevel", CompressionLevel);
+		}
+
 		// Insert our key-value pair metadata (if any, can be an arbitrary set of key/value pairs)
 		AddFileMetadata(Header);
 
@@ -362,12 +373,7 @@ void FEXRImageWriteTask::AddFileMetadata(Imf::Header& InHeader)
 	IOpenExrRTTIModule* OpenExrModule = FModuleManager::LoadModulePtr<IOpenExrRTTIModule>(RTTIExtensionModuleName);
 	if (OpenExrModule)
 	{
-		TMap<FString, FStringFormatArg> NewMap;
-		for (const TPair<FString, FString>& Metadata : FileMetadata)
-		{
-			NewMap.Add(Metadata.Key, Metadata.Value);
-		}
-		OpenExrModule->AddFileMetadata(NewMap, InHeader);
+		OpenExrModule->AddFileMetadata(FileMetadata, InHeader);
 	}
 }
 
@@ -464,7 +470,18 @@ void UMoviePipelineImageSequenceOutput_EXR::OnReceiveImageDataImpl(FMoviePipelin
 		TUniquePtr<FEXRImageWriteTask> MultiLayerImageTask = MakeUnique<FEXRImageWriteTask>();
 		MultiLayerImageTask->Filename = FinalFilePath;
 		MultiLayerImageTask->Compression = Compression;
-		MultiLayerImageTask->FileMetadata = FinalFormatArgs.FileMetadata; // This is already merged by ResolveFilenameFormatArgs with the FrameOutputState.
+		// MultiLayerImageTask->CompressionLevel is intentionally skipped because it doesn't seem to make any practical difference
+		// so we don't expose it to the user because that will just cause confusion where the setting doesn't seem to do anything.
+
+		// FinalFormatArgs.FileMetadata has been merged by ResolveFilenameFormatArgs with the FrameOutputState,
+		// but we need to convert from FString, FString (needed for BP/Python purposes) to a FStringFormatArg as
+		// we need to preserve numeric metadata types later in the image writing process (for compression level)
+		TMap<FString, FStringFormatArg> NewFileMetdataMap;
+		for (const TPair<FString, FString>& Metadata : FinalFormatArgs.FileMetadata)
+		{
+			NewFileMetdataMap.Add(Metadata.Key, Metadata.Value);
+		}
+		MultiLayerImageTask->FileMetadata = NewFileMetdataMap;
 
 		int32 LayerIndex = 0;
 		bool bRequiresTransparentOutput = false;
diff --git a/Engine/Plugins/MovieScene/MovieRenderPipeline/Source/MovieRenderPipelineRenderPasses/Private/MoviePipelineEXROutput.h b/Engine/Plugins/MovieScene/MovieRenderPipeline/Source/MovieRenderPipelineRenderPasses/Private/MoviePipelineEXROutput.h
--- a/Engine/Plugins/MovieScene/MovieRenderPipeline/Source/MovieRenderPipelineRenderPasses/Private/MoviePipelineEXROutput.h
+++ b/Engine/Plugins/MovieScene/MovieRenderPipeline/Source/MovieRenderPipelineRenderPasses/Private/MoviePipelineEXROutput.h
@@ -31,7 +31,11 @@ enum class EEXRCompressionFormat : uint8
 	/** Good compression quality for grainy images. Lossless.*/
 	PIZ,
 	/** Good compression quality for images with low amounts of noise. Lossless. */
-	ZIP
+	ZIP,
+	/** Lossy DCT-based compression for RGB channels. Alpha and other channels are uncompressed. More efficient than DWAB for partial buffer access on read in 3rd party tools. */
+	DWAA,
+	/** Similar to DWAA but goes in blocks of 256 scanlines instead of 32. More efficient disk space and faster to decode than DWAA. */
+	DWAB
 };
 
 #if WITH_UNREALEXR
@@ -48,6 +52,9 @@ public:
 	/** Compression method used for the resulting EXR files. */
 	EEXRCompressionFormat Compression;
 
+	/** When using a lossy compression format, what is the base-error (CompressionLevel/100000) */
+	int32 CompressionLevel;
+
 	/** A function to invoke on the game thread when the task has completed */
 	TFunction<void(bool)> OnCompleted;
 	
@@ -57,7 +64,7 @@ public:
 	int32 Height;
 
 	/** A set of key/value pairs to write into the exr file as metadata. */
-	TMap<FString, FString> FileMetadata;
+	TMap<FString, FStringFormatArg> FileMetadata;
 
 	/** The image data to write. Supports multiple layers of different bitdepths. */
 	TArray<TUniquePtr<FImagePixelData>> Layers;
@@ -71,6 +78,7 @@ public:
 	FEXRImageWriteTask()
 		: bOverwriteFile(true)
 		, Compression(EEXRCompressionFormat::PIZ)
+		, CompressionLevel(45)
 		, OverscanPercentage(0.0f)
 	{}
 
