Movie Pipeline: Change CastChecked to Cast so that the error message can proceed

diff --git a/Engine/Plugins/MovieScene/MovieRenderPipeline/Source/MovieRenderPipelineCore/Private/MovieRenderPipelineCommandLine.cpp b/Engine/Plugins/MovieScene/MovieRenderPipeline/Source/MovieRenderPipelineCore/Private/MovieRenderPipelineCommandLine.cpp
--- a/Engine/Plugins/MovieScene/MovieRenderPipeline/Source/MovieRenderPipelineCore/Private/MovieRenderPipelineCommandLine.cpp
+++ b/Engine/Plugins/MovieScene/MovieRenderPipeline/Source/MovieRenderPipelineCore/Private/MovieRenderPipelineCommandLine.cpp
@@ -239,7 +239,7 @@ uint8 FMovieRenderPipelineCoreModule::ParseMovieRenderData(const FString& InSequ
 	{
 		// Convert it to a soft object path and use that load to ensure it follows redirectors, etc.
 		FSoftObjectPath AssetPath = FSoftObjectPath(InSequenceAssetPath);
-		TargetSequence = CastChecked<ULevelSequence>(AssetPath.TryLoad());
+		TargetSequence = Cast<ULevelSequence>(AssetPath.TryLoad());
 
 		if (!TargetSequence)
 		{
