MoviePipeline: Reset the render target extents used by the gbuffer texture pool after renders. If you had a 4k render but a 1080p viewport historically we would continue to allocate memory for the full 4k resolution, now it should reset back down to 1080p.

diff --git a/Engine/Plugins/MovieScene/MovieRenderPipeline/Source/MovieRenderPipelineCore/Private/MoviePipeline.cpp b/Engine/Plugins/MovieScene/MovieRenderPipeline/Source/MovieRenderPipelineCore/Private/MoviePipeline.cpp
--- a/Engine/Plugins/MovieScene/MovieRenderPipeline/Source/MovieRenderPipelineCore/Private/MoviePipeline.cpp
+++ b/Engine/Plugins/MovieScene/MovieRenderPipeline/Source/MovieRenderPipelineCore/Private/MoviePipeline.cpp
@@ -40,6 +40,7 @@
 #include "HAL/FileManager.h"
 #include "Misc/CoreDelegates.h"
 #include "MoviePipelineCommandLineEncoder.h"
+#include "Kismet/KismetSystemLibrary.h"
 
 #if WITH_EDITOR
 #include "MovieSceneExportMetadata.h"
@@ -518,6 +519,12 @@ void UMoviePipeline::TransitionToState(const EMovieRenderPipelineState InNewStat
 				Viewport->bDisableWorldRendering = false;
 			}
 
+			// Because the render target pool is shared, if you had a high-resolution render in editor the entire gbuffer
+			// has been resized up to match the new maximum extent. This console command will reset the size of the pool
+			// and cause it to re-allocate at the currrent size on the next render request, which is likely to be the size
+			// of the PIE window (720p) or the Viewport itself.
+			UKismetSystemLibrary::ExecuteConsoleCommand(GetWorld(), TEXT("r.ResetRenderTargetsExtent"), nullptr);
+
 			GAreScreenMessagesEnabled = bPrevGScreenMessagesEnabled;
 
 			UE_LOG(LogMovieRenderPipeline, Log, TEXT("Movie Pipeline completed. Duration: %s"), *(FDateTime::UtcNow() - InitializationTime).ToString());
