World Partition: add wp.Runtime.IgnoreStreamingPerformance to ignore streaming performance checks and make sure to set it when rendering movies.

diff --git a/Engine/Plugins/MovieScene/MovieRenderPipeline/Source/MovieRenderPipelineCore/Private/MoviePipelineGameOverrideSetting.cpp b/Engine/Plugins/MovieScene/MovieRenderPipeline/Source/MovieRenderPipelineCore/Private/MoviePipelineGameOverrideSetting.cpp
--- a/Engine/Plugins/MovieScene/MovieRenderPipeline/Source/MovieRenderPipelineCore/Private/MoviePipelineGameOverrideSetting.cpp
+++ b/Engine/Plugins/MovieScene/MovieRenderPipeline/Source/MovieRenderPipelineCore/Private/MoviePipelineGameOverrideSetting.cpp
@@ -106,6 +106,8 @@ void UMoviePipelineGameOverrideSetting::ApplyCVarSettings(const bool bOverrideVa
 	// To make sure that the skylight is always valid and consistent accross capture sessions, we enforce a full capture each frame, accepting a small GPU cost.
 	MOVIEPIPELINE_STORE_AND_OVERRIDE_CVAR_INT(PreviousVolumetricRenderTarget, TEXT("r.VolumetricRenderTarget"), 0, bOverrideValues);
 
+	// To make sure that the world partition streaming doesn't end up in critical streaming performances and stops streaming low priority cells.
+	MOVIEPIPELINE_STORE_AND_OVERRIDE_CVAR_INT(PreviousIgnoreStreamingPerformance, TEXT("wp.Runtime.IgnoreStreamingPerformance"), 1, bOverrideValues);
 }
 
 void UMoviePipelineGameOverrideSetting::BuildNewProcessCommandLineArgsImpl(TArray<FString>& InOutUnrealURLParams, TArray<FString>& InOutCommandLineArgs, TArray<FString>& InOutDeviceProfileCvars, TArray<FString>& InOutExecCmds) const
diff --git a/Engine/Plugins/MovieScene/MovieRenderPipeline/Source/MovieRenderPipelineCore/Public/MoviePipelineGameOverrideSetting.h b/Engine/Plugins/MovieScene/MovieRenderPipeline/Source/MovieRenderPipelineCore/Public/MoviePipelineGameOverrideSetting.h
--- a/Engine/Plugins/MovieScene/MovieRenderPipeline/Source/MovieRenderPipelineCore/Public/MoviePipelineGameOverrideSetting.h
+++ b/Engine/Plugins/MovieScene/MovieRenderPipeline/Source/MovieRenderPipelineCore/Public/MoviePipelineGameOverrideSetting.h
@@ -123,4 +123,5 @@ private:
 	int32 PreviousNeverMuteNonRealtimeAudio;
 	int32 PreviousSkyLightRealTimeReflectionCaptureTimeSlice;
 	int32 PreviousVolumetricRenderTarget;
+	int32 PreviousIgnoreStreamingPerformance;
 };
\ No newline at end of file
diff --git a/Engine/Source/Runtime/Engine/Private/WorldPartition/WorldPartitionRuntimeHash.cpp b/Engine/Source/Runtime/Engine/Private/WorldPartition/WorldPartitionRuntimeHash.cpp
--- a/Engine/Source/Runtime/Engine/Private/WorldPartition/WorldPartitionRuntimeHash.cpp
+++ b/Engine/Source/Runtime/Engine/Private/WorldPartition/WorldPartitionRuntimeHash.cpp
@@ -17,6 +17,14 @@
 
 #define LOCTEXT_NAMESPACE "WorldPartition"
 
+static int32 GIgnoreStreamingPerformance = 0;
+static FAutoConsoleVariableRef CVarGIgnoreStreamingPerformance(
+	TEXT("wp.Runtime.IgnoreStreamingPerformance"),
+	GIgnoreStreamingPerformance,
+	TEXT("Should we ignore streaming performance checks? Enabling this will not block on critical streaming situations and may break gameplay. Useful for scenarios where we flush streaming every frame, such as movie rendering."),
+	ECVF_Default
+);
+
 UWorldPartitionRuntimeHash::UWorldPartitionRuntimeHash(const FObjectInitializer& ObjectInitializer)
 	: Super(ObjectInitializer)
 {}
@@ -285,7 +293,7 @@ EWorldPartitionStreamingPerformance UWorldPartitionRuntimeHash::GetStreamingPerf
 	bool bForceGoodStreamingPerformance = !GetWorld()->bMatchStarted;
 
 #if WITH_EDITOR
-	bForceGoodStreamingPerformance |= FApp::UseFixedTimeStep();
+	bForceGoodStreamingPerformance |= GIgnoreStreamingPerformance > 0;
 #endif
 
 	EWorldPartitionStreamingPerformance StreamingPerformance = EWorldPartitionStreamingPerformance::Good;
