MoviePipeline: Fixed the case where the shot frame numbers shown in the burn in would be incorrect if the Master sequence and the Shot sub-sequence had different internal tick resolutions.

diff --git a/Engine/Plugins/MovieScene/MovieRenderPipeline/Source/MovieRenderPipelineCore/Private/MoviePipeline.cpp b/Engine/Plugins/MovieScene/MovieRenderPipeline/Source/MovieRenderPipelineCore/Private/MoviePipeline.cpp
--- a/Engine/Plugins/MovieScene/MovieRenderPipeline/Source/MovieRenderPipelineCore/Private/MoviePipeline.cpp
+++ b/Engine/Plugins/MovieScene/MovieRenderPipeline/Source/MovieRenderPipelineCore/Private/MoviePipeline.cpp
@@ -839,6 +839,11 @@ void UMoviePipeline::BuildShotListFromSequence()
 		Shot->ShotInfo.NumSpatialSamples = AntiAliasingSettings->SpatialSampleCount;
 		Shot->ShotInfo.CachedFrameRate = GetPipelineMasterConfig()->GetEffectiveFrameRate(TargetSequence);
 		Shot->ShotInfo.CachedTickResolution = TargetSequence->GetMovieScene()->GetTickResolution();
+		Shot->ShotInfo.CachedShotTickResolution = Shot->ShotInfo.CachedTickResolution;
+		if (Shot->ShotInfo.SubSectionHierarchy.IsValid() && Shot->ShotInfo.SubSectionHierarchy->MovieScene.IsValid())
+		{
+			Shot->ShotInfo.CachedShotTickResolution = Shot->ShotInfo.SubSectionHierarchy->MovieScene->GetTickResolution();
+		}
 		Shot->ShotInfo.NumTiles = FIntPoint(HighResSettings->TileCount, HighResSettings->TileCount);
 
 		// Expand the shot (but don't actually modify the sections)
@@ -1253,6 +1258,7 @@ MoviePipeline::FFrameConstantMetrics UMoviePipeline::CalculateShotFrameMetrics(c
 {
 	MoviePipeline::FFrameConstantMetrics Output;
 	Output.TickResolution = TargetSequence->GetMovieScene()->GetTickResolution();
+	Output.ShotTickResolution = InShot->ShotInfo.CachedShotTickResolution;
 	Output.FrameRate = GetPipelineMasterConfig()->GetEffectiveFrameRate(TargetSequence);
 	Output.TicksPerOutputFrame = FFrameRate::TransformTime(FFrameTime(FFrameNumber(1)), Output.FrameRate, Output.TickResolution);
 
diff --git a/Engine/Plugins/MovieScene/MovieRenderPipeline/Source/MovieRenderPipelineCore/Private/MoviePipelineTiming.cpp b/Engine/Plugins/MovieScene/MovieRenderPipeline/Source/MovieRenderPipelineCore/Private/MoviePipelineTiming.cpp
--- a/Engine/Plugins/MovieScene/MovieRenderPipeline/Source/MovieRenderPipelineCore/Private/MoviePipelineTiming.cpp
+++ b/Engine/Plugins/MovieScene/MovieRenderPipeline/Source/MovieRenderPipelineCore/Private/MoviePipelineTiming.cpp
@@ -708,7 +708,7 @@ void UMoviePipeline::CalculateFrameNumbersForOutputState(const MoviePipeline::FF
 		// Convert from master space back into shot space - based on the inner most detected shot.
 		CenteredTick = (CenteredTick * InCameraCut->ShotInfo.OuterToInnerTransform).FloorToFrame();
 
-		InOutOutputState.CurrentShotSourceFrameNumber = FFrameRate::TransformTime(CenteredTick, InFrameMetrics.TickResolution, SourceDisplayRate).RoundToFrame().Value;
+		InOutOutputState.CurrentShotSourceFrameNumber = FFrameRate::TransformTime(CenteredTick, InFrameMetrics.ShotTickResolution, SourceDisplayRate).RoundToFrame().Value;
 		InOutOutputState.CurrentShotSourceTimeCode = FTimecode::FromFrameNumber(InOutOutputState.CurrentShotSourceFrameNumber, InFrameMetrics.FrameRate, false);
 	}
 
diff --git a/Engine/Plugins/MovieScene/MovieRenderPipeline/Source/MovieRenderPipelineCore/Public/MovieRenderPipelineDataTypes.h b/Engine/Plugins/MovieScene/MovieRenderPipeline/Source/MovieRenderPipelineCore/Public/MovieRenderPipelineDataTypes.h
--- a/Engine/Plugins/MovieScene/MovieRenderPipeline/Source/MovieRenderPipelineCore/Public/MovieRenderPipelineDataTypes.h
+++ b/Engine/Plugins/MovieScene/MovieRenderPipeline/Source/MovieRenderPipelineCore/Public/MovieRenderPipelineDataTypes.h
@@ -244,8 +244,10 @@ namespace MoviePipeline
 
 	struct FFrameConstantMetrics
 	{
-		/** What is the tick resolution fo the sequence */
+		/** What is the tick resolution of the master sequence */
 		FFrameRate TickResolution;
+		/** What is the tick resolution of the current shot */
+		FFrameRate ShotTickResolution;
 		/** What is the effective frame rate of the output */
 		FFrameRate FrameRate;
 		/** How many ticks per output frame. */
@@ -503,6 +505,8 @@ public:
 	/** Cached Tick Resolution our numbers are in. Simplifies some APIs. */
 	FFrameRate CachedTickResolution;
 	
+	/** Cached Tick Resolution the movie scene this range was generated for is in. Can be different than the master due to mixed tick resolutions. */
+	FFrameRate CachedShotTickResolution;
 public:
 	/** The current state of processing this Shot is in. Not all states will be passed through. */
 	EMovieRenderShotState State;
