MoviePipeline: Fixed a case where extra audio samples could be generated before the MoviePipeline started which would be included in the resulting wav file which would cause the audio to become offset.

diff --git a/Engine/Plugins/MovieScene/MovieRenderPipeline/Source/MovieRenderPipelineCore/Private/MoviePipelineAudioRendering.cpp b/Engine/Plugins/MovieScene/MovieRenderPipeline/Source/MovieRenderPipelineCore/Private/MoviePipelineAudioRendering.cpp
--- a/Engine/Plugins/MovieScene/MovieRenderPipeline/Source/MovieRenderPipelineCore/Private/MoviePipelineAudioRendering.cpp
+++ b/Engine/Plugins/MovieScene/MovieRenderPipeline/Source/MovieRenderPipelineCore/Private/MoviePipelineAudioRendering.cpp
@@ -17,7 +17,6 @@
 #include "AudioMixerPlatformNonRealtime.h"
 #include "MoviePipelineQueue.h"
 
-
 static FAudioDevice* GetAudioDeviceFromWorldContext(const UObject* WorldContextObject)
 {
 	UWorld* ThisWorld = GEngine->GetWorldFromContextObject(WorldContextObject, EGetWorldErrorMode::LogAndReturnNull);
@@ -128,12 +127,37 @@ void UMoviePipeline::ProcessAudioTick()
 		return;
 	}
 
+	Audio::IAudioMixerPlatformInterface* AudioMixerPlatform = MixerDevice->GetAudioMixerPlatform();
+	Audio::FMixerPlatformNonRealtime* NRTPlatform = nullptr;
+	if (AudioMixerPlatform && AudioMixerPlatform->IsNonRealtime())
+	{
+		// There is only one non-realtime audio platform at this time, so we can safely static_cast this due to the IsNonRealtime check.
+		NRTPlatform = static_cast<Audio::FMixerPlatformNonRealtime*>(MixerDevice->GetAudioMixerPlatform());
+	}
+
+	if (!NRTPlatform)
+	{
+		return;
+	}
+
+
 	UMoviePipelineExecutorShot* CurrentShot = ActiveShotList[CurrentShotIndex];
 
 	// Start capturing any produced samples on the same frame we start submitting samples that will make it to disk.
 	// This comes before we process samples for this frame (below).
 	if (CurrentShot->ShotInfo.State == EMovieRenderShotState::Rendering && !AudioState.bIsRecordingAudio)
 	{
+		// There can be time in the Unreal world that passes when MRQ is not yet active. When MRQ isn't active,
+		// then the NRT audio platform doesn't process anything and it builds up for when we do start to actually
+		// process it for a shot. This causes our actual audio to become offset from the sequence once lined up
+		// afterwards. So for now we'll just flush out anything that needed to be processed before we start capturing
+			
+		MixerDevice->Update(true);
+		
+		// We don't have a great way of determining how much outstanding audio there is, so we're just going to 
+		// process an arbitrary amount of time.
+		NRTPlatform->RenderAudio(30.f);
+		
 		StartAudioRecording();
 	}
 
@@ -158,17 +182,9 @@ void UMoviePipeline::ProcessAudioTick()
 		MixerDevice->Update(true);
 	}
 
-	Audio::IAudioMixerPlatformInterface* AudioMixerPlatform = MixerDevice->GetAudioMixerPlatform();
-	if (bCanRenderAudio && AudioMixerPlatform && AudioMixerPlatform->IsNonRealtime())
 	{
-		// There is only one non-realtime audio platform at this time, so we can safely static_cast this due to the IsNonRealtime check.
-		Audio::FMixerPlatformNonRealtime* Platform = static_cast<Audio::FMixerPlatformNonRealtime*>(MixerDevice->GetAudioMixerPlatform());
-			
-		if (Platform)
-		{
-			// Process work that has been submitted from the game thread to the audio thread over the temporal samples of this frame.
-			FFrameRate OutputFrameRate = GetPipelineMasterConfig()->GetEffectiveFrameRate(TargetSequence);
-			Platform->RenderAudio(OutputFrameRate.AsInterval());
-		}
+		// Process work that has been submitted from the game thread to the audio thread over the temporal samples of this frame.
+		FFrameRate OutputFrameRate = GetPipelineMasterConfig()->GetEffectiveFrameRate(TargetSequence);
+		NRTPlatform->RenderAudio(OutputFrameRate.AsInterval());
 	}
 }
