MoviePipeline: Fixed a crash rendering when rendering and there were no files being output for a given render job.

diff --git a/Engine/Plugins/MovieScene/MovieRenderPipeline/Source/MovieRenderPipelineCore/Private/MoviePipeline.cpp b/Engine/Plugins/MovieScene/MovieRenderPipeline/Source/MovieRenderPipelineCore/Private/MoviePipeline.cpp
--- a/Engine/Plugins/MovieScene/MovieRenderPipeline/Source/MovieRenderPipelineCore/Private/MoviePipeline.cpp
+++ b/Engine/Plugins/MovieScene/MovieRenderPipeline/Source/MovieRenderPipelineCore/Private/MoviePipeline.cpp
@@ -936,10 +936,13 @@ void UMoviePipeline::TeardownShot(UMoviePipelineExecutorShot* InShot)
 		ProcessOutstandingFutures();
 
 		TArray<FMoviePipelineShotOutputData> LatestShotData;
-		LatestShotData.Add(GeneratedShotOutputData.Last());
+		if (GeneratedShotOutputData.Num() > 0)
+		{
+			LatestShotData.Add(GeneratedShotOutputData.Last());
 
-		// Temporarily remove it from the global array, as the encode may modify it.
-		GeneratedShotOutputData.RemoveAt(GeneratedShotOutputData.Num() - 1);
+			// Temporarily remove it from the global array, as the encode may modify it.
+			GeneratedShotOutputData.RemoveAt(GeneratedShotOutputData.Num() - 1);
+		}
 
 		// We call the command line encoder as a special case here because it may want to modify the file list
 		// ie: if it deletes the file after use we probably don't want scripting looking for those files.
