MoviePipeline: Color setting can now be disabled via the toggle button in the UI once added (like other settings).

diff --git a/Engine/Plugins/MovieScene/MovieRenderPipeline/Source/MovieRenderPipelineCore/Public/MoviePipelineColorSetting.h b/Engine/Plugins/MovieScene/MovieRenderPipeline/Source/MovieRenderPipelineCore/Public/MoviePipelineColorSetting.h
--- a/Engine/Plugins/MovieScene/MovieRenderPipeline/Source/MovieRenderPipelineCore/Public/MoviePipelineColorSetting.h
+++ b/Engine/Plugins/MovieScene/MovieRenderPipeline/Source/MovieRenderPipelineCore/Public/MoviePipelineColorSetting.h
@@ -17,7 +17,6 @@ public:
 public:
 #if WITH_EDITOR
 	virtual FText GetDisplayText() const override { return NSLOCTEXT("MovieRenderPipeline", "ColorSetting", "Color Output"); }
-	virtual bool CanBeDisabled() const override { return false; }
 	virtual bool CanEditChange(const FProperty* InProperty) const override;
 #endif
 	virtual bool IsValidOnShots() const override { return false; }
