MoviePipeline: Added support for shot-only sequences and no-shot/no-camera sequences again (fixing a regression in detected ranges to render).

diff --git a/Engine/Plugins/MovieScene/MovieRenderPipeline/Source/MovieRenderPipelineCore/Private/MoviePipelineBlueprintLibrary.cpp b/Engine/Plugins/MovieScene/MovieRenderPipeline/Source/MovieRenderPipelineCore/Private/MoviePipelineBlueprintLibrary.cpp
--- a/Engine/Plugins/MovieScene/MovieRenderPipeline/Source/MovieRenderPipelineCore/Private/MoviePipelineBlueprintLibrary.cpp
+++ b/Engine/Plugins/MovieScene/MovieRenderPipeline/Source/MovieRenderPipelineCore/Private/MoviePipelineBlueprintLibrary.cpp
@@ -352,50 +352,70 @@ void UMoviePipelineBlueprintLibrary::UpdateJobShotListFromSequence(ULevelSequenc
 		return;
 	}
 
-	using FCameraCutInfo = TTuple<UMovieSceneCameraCutSection*, FMovieSceneSequenceID, int16>;
-	struct FCameraCutVisitor : UE::MovieScene::ISequenceVisitor
+	using FSequenceRangeInfo = TTuple<UMovieSceneSection*, FMovieSceneSequenceID, int16>;
+	struct FSequenceRangeVisitor : UE::MovieScene::ISequenceVisitor
 	{
-		virtual void VisitTrack(UMovieSceneTrack* InTrack, const FGuid&, const UE::MovieScene::FSubSequenceSpace& LocalSpace)
+		void GatherFromSection(UMovieSceneTrack* InTrack, const UE::MovieScene::FSubSequenceSpace& LocalSpace, TMovieSceneEvaluationTree<FSequenceRangeInfo>* InTree)
 		{
-			if (UMovieSceneCameraCutTrack* CameraCutTrack = Cast<UMovieSceneCameraCutTrack>(InTrack))
+			TArray<FSequenceRangeInfo> OutRanges;
+
+			for (UMovieSceneSection* Section : InTrack->GetAllSections())
 			{
-				for (UMovieSceneSection* Section : CameraCutTrack->GetAllSections())
+				if (!Section->IsActive())
 				{
-					if (!Section->IsActive())
-					{
-						continue;
-					}
+					continue;
+				}
 
-					UMovieSceneCameraCutSection* CameraCutSection = CastChecked<UMovieSceneCameraCutSection>(Section);
-					
-					// Intersect it with our local range so that any sections that fall outside our playback bounds gets discarded
-					UMovieScene* OwningScene = CameraCutSection->GetTypedOuter<UMovieScene>();
-					TRange<FFrameNumber> LocalCameraRange = TRange<FFrameNumber>::Intersection(OwningScene->GetPlaybackRange(), CameraCutSection->GetRange());
-					if (LocalCameraRange.IsEmpty())
+
+				// Intersect it with our local range so that any sections that fall outside our playback bounds gets discarded
+				UMovieScene* OwningScene = Section->GetTypedOuter<UMovieScene>();
+				TRange<FFrameNumber> LocalCameraRange = TRange<FFrameNumber>::Intersection(OwningScene->GetPlaybackRange(), Section->GetRange());
+				if (LocalCameraRange.IsEmpty())
+				{
+					continue;
+				}
+
+				// Intersect it with the root range so that if the parent has trimmed down the sub-section we don't render outside that.
+				TRange<FFrameNumber> RootCameraRange = TRange<FFrameNumber>::Intersection(LocalSpace.RootClampRange, LocalCameraRange * LocalSpace.RootToSequenceTransform.InverseLinearOnly());
+				if (!RootCameraRange.IsEmpty())
+				{
+					if (UMovieSceneCinematicShotSection* ShotSection = Cast<UMovieSceneCinematicShotSection>(Section))
 					{
-						continue;
+						// We add the child sequence here so that the hierarchy tree gets built correctly.
+						InTree->Add(RootCameraRange, MakeTuple(Section, ShotSection->GetSequenceID(), LocalSpace.HierarchicalBias));
 					}
-
-					// Intersect it with the root range so that if the parent has trimmed down the sub-section we don't render outside that.
-					TRange<FFrameNumber> RootCameraRange = TRange<FFrameNumber>::Intersection(LocalSpace.RootClampRange, LocalCameraRange * LocalSpace.RootToSequenceTransform.InverseLinearOnly());
-					if (!RootCameraRange.IsEmpty())
+					else
 					{
-						CameraCutTree.Add(RootCameraRange, MakeTuple(CameraCutSection, LocalSpace.SequenceID, LocalSpace.HierarchicalBias));
+						InTree->Add(RootCameraRange, MakeTuple(Section, LocalSpace.SequenceID, LocalSpace.HierarchicalBias));
 					}
 				}
 			}
 		}
-		TMovieSceneEvaluationTree<FCameraCutInfo> CameraCutTree;
+
+		virtual void VisitTrack(UMovieSceneTrack* InTrack, const FGuid&, const UE::MovieScene::FSubSequenceSpace& LocalSpace)
+		{
+			if (UMovieSceneCameraCutTrack* CameraCutTrack = Cast<UMovieSceneCameraCutTrack>(InTrack))
+			{
+				GatherFromSection(CameraCutTrack, LocalSpace, &CameraCutTree);
+			}
+
+			if (UMovieSceneCinematicShotTrack* ShotTrack = Cast<UMovieSceneCinematicShotTrack>(InTrack))
+			{
+				GatherFromSection(ShotTrack, LocalSpace, &ShotSectionTree);
+			}
+		}
+		TMovieSceneEvaluationTree<FSequenceRangeInfo> CameraCutTree;
+		TMovieSceneEvaluationTree<FSequenceRangeInfo> ShotSectionTree;
 	};
 
 	{
 		UE::MovieScene::FSequenceVisitParams Params;
 		Params.bVisitMasterTracks = true;
 		Params.bVisitSubSequences = true;
-		FCameraCutVisitor CameraCutVisitor;
+		FSequenceRangeVisitor RangeVisitor;
 
 		// Visit all camera cuts
-		VisitSequence(InSequence, Params, CameraCutVisitor);
+		VisitSequence(InSequence, Params, RangeVisitor);
 
 		struct FLinearizedEntity
 		{
@@ -412,19 +432,22 @@ void UMoviePipelineBlueprintLibrary::UpdateJobShotListFromSequence(ULevelSequenc
 
 		TArray<FLinearizedEntity> Entities;
 
+		// Give preference to the camera cut tree if we have any entries, otherwise we fall back to just shot sections
+		TMovieSceneEvaluationTree<FSequenceRangeInfo>* RangeToWalk = RangeVisitor.CameraCutTree.IsEmpty() ? &RangeVisitor.ShotSectionTree : &RangeVisitor.CameraCutTree;
+		
 		// Walk the camera cut tree, linearizing the results
-		FMovieSceneEvaluationTreeRangeIterator RangeIt(CameraCutVisitor.CameraCutTree);
+		FMovieSceneEvaluationTreeRangeIterator RangeIt(*RangeToWalk);
 		for (; RangeIt; ++RangeIt)
 		{
 			// Look up the camera cuts for this range
-			TMovieSceneEvaluationTreeDataIterator<FCameraCutInfo> DataIt = CameraCutVisitor.CameraCutTree.GetAllData(RangeIt.Node());
+			TMovieSceneEvaluationTreeDataIterator<FSequenceRangeInfo> DataIt = RangeToWalk->GetAllData(RangeIt.Node());
 			if (!DataIt)
 			{
 				continue;
 			}
 
 			// Find the best Camera Cut based on HBias
-			FCameraCutInfo BestInfo = *DataIt;
+			FSequenceRangeInfo BestInfo = *DataIt;
 			++DataIt;
 			for (; DataIt; ++DataIt)
 			{
@@ -445,7 +468,7 @@ void UMoviePipelineBlueprintLibrary::UpdateJobShotListFromSequence(ULevelSequenc
 			}
 
 			// Use our intersected range to trim by the outmost Playback Range.
-			FLinearizedEntity NewEntry{ IntersectedWithPlaybackBounds, BestInfo.Get<0>(), BestInfo.Get<1>() };
+			FLinearizedEntity NewEntry{ IntersectedWithPlaybackBounds, Cast<UMovieSceneCameraCutSection>(BestInfo.Get<0>()), BestInfo.Get<1>() };
 			
 			bool bMerged = false;
 			if (Entities.Num())
@@ -453,7 +476,7 @@ void UMoviePipelineBlueprintLibrary::UpdateJobShotListFromSequence(ULevelSequenc
 				FLinearizedEntity& LastEntity = Entities.Last();
 				// Check whether we should combine with the last one
 				const bool bShotMatches = LastEntity.SequenceID == NewEntry.SequenceID;
-				const bool bCameraCutMatches = LastEntity.CameraCut == NewEntry.CameraCut;
+				const bool bCameraCutMatches = (LastEntity.CameraCut == NewEntry.CameraCut) && LastEntity.CameraCut != nullptr;
 				if (LastEntity.Range.Adjoins(NewEntry.Range) && bShotMatches && bCameraCutMatches)
 				{
 					LastEntity.Range.SetUpperBound(NewEntry.Range.GetUpperBound());
@@ -476,6 +499,7 @@ void UMoviePipelineBlueprintLibrary::UpdateJobShotListFromSequence(ULevelSequenc
 			// Let's build a sub-section hierarchy for this node
 			TSharedPtr<MoviePipeline::FCameraCutSubSectionHierarchyNode> LeafNode = MakeShared<MoviePipeline::FCameraCutSubSectionHierarchyNode>();
 			
+			// This can be null in the fallback case where no camera cuts were detected
 			LeafNode->CameraCutSection = Entity.CameraCut;
 			MoviePipeline::BuildSectionHierarchyRecursive(SequenceHierarchyCache, InSequence, Entity.SequenceID, MovieSceneSequenceID::Invalid, LeafNode);
 
@@ -491,6 +515,25 @@ void UMoviePipelineBlueprintLibrary::UpdateJobShotListFromSequence(ULevelSequenc
 			Entity.InnerToOuterTransform = InnerToOuterTransform;
 		}
 
+		// Fallback case for no shot sections detected and no camera cut sections detected.
+		if (Entities.Num() == 0)
+		{
+			FLinearizedEntity& Entity = Entities.AddDefaulted_GetRef();
+
+			Entity.Range = InSequence->GetMovieScene()->GetPlaybackRange();
+			Entity.CameraCut = nullptr;
+			Entity.SequenceID = MovieSceneSequenceID::Root;
+
+			TSharedPtr<MoviePipeline::FCameraCutSubSectionHierarchyNode> LeafNode = MakeShared<MoviePipeline::FCameraCutSubSectionHierarchyNode>();
+			LeafNode->CameraCutSection = nullptr;
+			MoviePipeline::BuildSectionHierarchyRecursive(SequenceHierarchyCache, InSequence, MovieSceneSequenceID::Root, MovieSceneSequenceID::Invalid, LeafNode);
+
+			Entity.CameraCutWarmUpRange = TRange<FFrameNumber>::Empty();
+			Entity.LeafNode = LeafNode;
+			Entity.Name = MoviePipeline::GetNameForShot(SequenceHierarchyCache, InSequence, LeafNode);
+			Entity.InnerToOuterTransform = FMovieSceneTimeTransform();
+		}
+
 		// We need to generate all of the linearized segments first so that we have all of the names available.
 		// We need all of the names available because we need to ensure unique shot names (for XMLs), and they 
 		// expect to be consistent (ie: if duplicates are found, the first one is retroactively changed to (1)).
diff --git a/Engine/Plugins/MovieScene/MovieRenderPipeline/Source/MovieRenderPipelineCore/Private/MoviePipelineUtils.cpp b/Engine/Plugins/MovieScene/MovieRenderPipeline/Source/MovieRenderPipelineCore/Private/MoviePipelineUtils.cpp
--- a/Engine/Plugins/MovieScene/MovieRenderPipeline/Source/MovieRenderPipelineCore/Private/MoviePipelineUtils.cpp
+++ b/Engine/Plugins/MovieScene/MovieRenderPipeline/Source/MovieRenderPipelineCore/Private/MoviePipelineUtils.cpp
@@ -587,17 +587,17 @@ namespace MoviePipeline
 			}
 		}
 
-		TSharedPtr<FCameraCutSubSectionHierarchyNode> ParentNode;
-		if (!ParentNode)
-		{
-			ParentNode = MakeShared<FCameraCutSubSectionHierarchyNode>();
-		}
-
-		OutSubsectionHierarchy->SetParent(ParentNode);
 		OutSubsectionHierarchy->MovieScene = TWeakObjectPtr<UMovieScene>(Sequence->GetMovieScene());
 		OutSubsectionHierarchy->NodeID = InSequenceId;
 
-		BuildSectionHierarchyRecursive(InHierarchy, InRootSequence, SequenceNode->ParentID, InSequenceId, ParentNode);
+		// Only try assigning the parent and diving in if this node isn't already the root, roots have no parents.
+		if (InSequenceId != MovieSceneSequenceID::Root)
+		{
+			TSharedPtr<FCameraCutSubSectionHierarchyNode> ParentNode = MakeShared<FCameraCutSubSectionHierarchyNode>();
+			OutSubsectionHierarchy->SetParent(ParentNode);
+
+			BuildSectionHierarchyRecursive(InHierarchy, InRootSequence, SequenceNode->ParentID, InSequenceId, ParentNode);
+		}
 	}
 
 
