MoviePipeline: Added standalone asset editors for Movie Pipeline Master/Shot configs. Changed the behavior for MoviePipelineQueues to open in the Movie Pipeline Queue UI, but direct editing of the asset is not supported.

diff --git a/Engine/Plugins/MovieScene/MovieRenderPipeline/Source/MovieRenderPipelineEditor/MovieRenderPipelineEditor.Build.cs b/Engine/Plugins/MovieScene/MovieRenderPipeline/Source/MovieRenderPipelineEditor/MovieRenderPipelineEditor.Build.cs
--- a/Engine/Plugins/MovieScene/MovieRenderPipeline/Source/MovieRenderPipelineEditor/MovieRenderPipelineEditor.Build.cs
+++ b/Engine/Plugins/MovieScene/MovieRenderPipeline/Source/MovieRenderPipelineEditor/MovieRenderPipelineEditor.Build.cs
@@ -57,6 +57,7 @@ public class MovieRenderPipelineEditor : ModuleRules
 				"DeveloperSettings",
 				"MessageLog",
 				"Sequencer",
+				"ToolMenus"
 			}
         );
 
diff --git a/Engine/Plugins/MovieScene/MovieRenderPipeline/Source/MovieRenderPipelineEditor/Private/AssetTypeActions_PipelineConfigs.cpp b/Engine/Plugins/MovieScene/MovieRenderPipeline/Source/MovieRenderPipelineEditor/Private/AssetTypeActions_PipelineConfigs.cpp
new file mode 100644
--- /dev/null
+++ b/Engine/Plugins/MovieScene/MovieRenderPipeline/Source/MovieRenderPipelineEditor/Private/AssetTypeActions_PipelineConfigs.cpp
@@ -0,0 +1,66 @@
+// Copyright Epic Games, Inc. All Rights Reserved.
+
+#include "AssetTypeActions_PipelineConfigs.h"
+#include "MoviePipelineAssetEditor.h"
+#include "Framework/Notifications/NotificationManager.h"
+#include "MoviePipelineQueueSubsystem.h"
+#include "IMovieRenderPipelineEditorModule.h"
+#include "Widgets/Notifications/SNotificationList.h"
+
+void FAssetTypeActions_PipelineMasterConfig::OpenAssetEditor( const TArray<UObject*>& InObjects, TSharedPtr<IToolkitHost> EditWithinLevelEditor )
+{
+	UAssetEditorSubsystem* AssetEditorSubsystem = GEditor->GetEditorSubsystem<UAssetEditorSubsystem>();
+	if (AssetEditorSubsystem)
+	{
+		UMoviePipelineAssetEditor* AssetEditor = NewObject<UMoviePipelineAssetEditor>(AssetEditorSubsystem, NAME_None, RF_Transient);
+		if (AssetEditor)
+		{
+			if (InObjects.Num() > 0)
+			{
+				AssetEditor->SetObjectToEdit(InObjects[0]);
+				AssetEditor->Initialize();
+			}
+		}
+	}
+}
+
+void FAssetTypeActions_PipelineShotConfig::OpenAssetEditor( const TArray<UObject*>& InObjects, TSharedPtr<IToolkitHost> EditWithinLevelEditor )
+{
+	UAssetEditorSubsystem* AssetEditorSubsystem = GEditor->GetEditorSubsystem<UAssetEditorSubsystem>();
+	if (AssetEditorSubsystem)
+	{
+		UMoviePipelineAssetEditor* AssetEditor = NewObject<UMoviePipelineAssetEditor>(AssetEditorSubsystem, NAME_None, RF_Transient);
+		if (AssetEditor)
+		{
+			if (InObjects.Num() > 0)
+			{
+				AssetEditor->SetObjectToEdit(InObjects[0]);
+				AssetEditor->Initialize();
+			}
+		}
+	}
+}
+
+void FAssetTypeActions_PipelineQueue::OpenAssetEditor( const TArray<UObject*>& InObjects, TSharedPtr<IToolkitHost> EditWithinLevelEditor )
+{
+	if(InObjects.Num() > 0)
+	{
+		UMoviePipelineQueue* QueueAsset = Cast<UMoviePipelineQueue>(InObjects[0]);
+		if (QueueAsset)
+		{
+			// We don't allow direct editing of Queue assets, the best we can do is import a copy into our queue ui and focus it.
+			UMoviePipelineQueueSubsystem* PipelineSubsystem = GEditor->GetEditorSubsystem<UMoviePipelineQueueSubsystem>();
+			if(!PipelineSubsystem->IsRendering())
+			{
+				UMoviePipelineQueue* EditorQueue = PipelineSubsystem->GetQueue();
+				EditorQueue->CopyFrom(QueueAsset);
+				
+				FNotificationInfo Info(NSLOCTEXT("MoviePipeline", "QueueAssetActions_Imported", "Imported queue asset into UI. Save queue asset again from UI if you want to update the asset."));
+				Info.ExpireDuration = 3.0f;
+				FSlateNotificationManager::Get().AddNotification(Info);
+
+				FGlobalTabmanager::Get()->TryInvokeTab(IMovieRenderPipelineEditorModule::MoviePipelineQueueTabName);
+			}
+		}
+	}
+}
\ No newline at end of file
diff --git a/Engine/Plugins/MovieScene/MovieRenderPipeline/Source/MovieRenderPipelineEditor/Private/AssetTypeActions_PipelineConfigs.h b/Engine/Plugins/MovieScene/MovieRenderPipeline/Source/MovieRenderPipelineEditor/Private/AssetTypeActions_PipelineConfigs.h
new file mode 100644
--- /dev/null
+++ b/Engine/Plugins/MovieScene/MovieRenderPipeline/Source/MovieRenderPipelineEditor/Private/AssetTypeActions_PipelineConfigs.h
@@ -0,0 +1,47 @@
+// Copyright Epic Games, Inc. All Rights Reserved.
+
+#pragma once
+
+#include "CoreMinimal.h"
+#include "Toolkits/IToolkitHost.h"
+#include "AssetTypeActions_Base.h"
+#include "MoviePipelineMasterConfig.h"
+#include "MoviePipelineShotConfig.h"
+#include "MoviePipelineQueue.h"
+
+class FAssetTypeActions_PipelineMasterConfig : public FAssetTypeActions_Base
+{
+public:
+	// IAssetTypeActions Implementation
+	virtual FText GetName() const override { return NSLOCTEXT("AssetTypeActions", "AssetTypeActions_PipelineMasterConfig", "Movie Pipeline Master Config"); }
+	virtual FColor GetTypeColor() const override { return FColor(78, 40, 165); }
+	virtual UClass* GetSupportedClass() const override { return UMoviePipelineMasterConfig::StaticClass(); }
+	virtual void OpenAssetEditor( const TArray<UObject*>& InObjects, TSharedPtr<class IToolkitHost> EditWithinLevelEditor = TSharedPtr<IToolkitHost>() ) override;
+	virtual bool CanFilter() override { return false; }
+	virtual uint32 GetCategories() override { return EAssetTypeCategories::Animation; }
+};
+
+class FAssetTypeActions_PipelineShotConfig : public FAssetTypeActions_Base
+{
+public:
+	// IAssetTypeActions Implementation
+	virtual FText GetName() const override { return NSLOCTEXT("AssetTypeActions", "AssetTypeActions_PipelineShotConfig", "Movie Pipeline Shot Config"); }
+	virtual FColor GetTypeColor() const override { return FColor(78, 40, 165); }
+	virtual UClass* GetSupportedClass() const override { return UMoviePipelineShotConfig::StaticClass(); }
+	virtual void OpenAssetEditor( const TArray<UObject*>& InObjects, TSharedPtr<class IToolkitHost> EditWithinLevelEditor = TSharedPtr<IToolkitHost>() ) override;
+	virtual bool CanFilter() override { return false; }
+	virtual uint32 GetCategories() override { return EAssetTypeCategories::Animation; }
+};
+
+
+class FAssetTypeActions_PipelineQueue : public FAssetTypeActions_Base
+{
+public:
+	// IAssetTypeActions Implementation
+	virtual FText GetName() const override { return NSLOCTEXT("AssetTypeActions", "AssetTypeActions_PipelineQueue", "Movie Pipeline Queue"); }
+	virtual FColor GetTypeColor() const override { return FColor(78, 40, 165); }
+	virtual UClass* GetSupportedClass() const override { return UMoviePipelineQueue::StaticClass(); }
+	virtual void OpenAssetEditor( const TArray<UObject*>& InObjects, TSharedPtr<class IToolkitHost> EditWithinLevelEditor = TSharedPtr<IToolkitHost>() ) override;
+	virtual bool CanFilter() override { return false; }
+	virtual uint32 GetCategories() override { return EAssetTypeCategories::Animation; }
+};
\ No newline at end of file
diff --git a/Engine/Plugins/MovieScene/MovieRenderPipeline/Source/MovieRenderPipelineEditor/Private/MoviePipelineAssetEditor.cpp b/Engine/Plugins/MovieScene/MovieRenderPipeline/Source/MovieRenderPipelineEditor/Private/MoviePipelineAssetEditor.cpp
new file mode 100644
--- /dev/null
+++ b/Engine/Plugins/MovieScene/MovieRenderPipeline/Source/MovieRenderPipelineEditor/Private/MoviePipelineAssetEditor.cpp
@@ -0,0 +1,19 @@
+// Copyright Epic Games, Inc. All Rights Reserved.
+
+#include "MoviePipelineAssetEditor.h"
+#include "MoviePipelineConfigAssetEditor.h"
+
+TSharedPtr<FBaseAssetToolkit> UMoviePipelineAssetEditor::CreateToolkit()
+{
+	return MakeShared<FMoviePipelineConfigAssetEditor>(this);
+}
+
+void UMoviePipelineAssetEditor::GetObjectsToEdit(TArray<UObject*>& InObjectsToEdit)
+{
+	InObjectsToEdit.Add(ObjectToEdit);
+}
+
+void UMoviePipelineAssetEditor::SetObjectToEdit(UObject* InObject)
+{
+	ObjectToEdit = InObject;
+}
\ No newline at end of file
diff --git a/Engine/Plugins/MovieScene/MovieRenderPipeline/Source/MovieRenderPipelineEditor/Private/MoviePipelineAssetEditor.h b/Engine/Plugins/MovieScene/MovieRenderPipeline/Source/MovieRenderPipelineEditor/Private/MoviePipelineAssetEditor.h
new file mode 100644
--- /dev/null
+++ b/Engine/Plugins/MovieScene/MovieRenderPipeline/Source/MovieRenderPipelineEditor/Private/MoviePipelineAssetEditor.h
@@ -0,0 +1,23 @@
+// Copyright Epic Games, Inc. All Rights Reserved.
+
+#pragma once
+
+#include "Tools/UAssetEditor.h"
+
+#include "MoviePipelineAssetEditor.generated.h"
+
+
+UCLASS()
+class UMoviePipelineAssetEditor : public UAssetEditor
+{
+	GENERATED_BODY()
+
+public:
+	virtual TSharedPtr<FBaseAssetToolkit> CreateToolkit() override;
+	virtual void GetObjectsToEdit(TArray<UObject*>& InObjectsToEdit) override;
+	
+	void SetObjectToEdit(UObject* InObject);
+private:
+	UPROPERTY(Transient)
+	UObject* ObjectToEdit;
+};
\ No newline at end of file
diff --git a/Engine/Plugins/MovieScene/MovieRenderPipeline/Source/MovieRenderPipelineEditor/Private/MoviePipelineConfigAssetEditor.cpp b/Engine/Plugins/MovieScene/MovieRenderPipeline/Source/MovieRenderPipelineEditor/Private/MoviePipelineConfigAssetEditor.cpp
new file mode 100644
--- /dev/null
+++ b/Engine/Plugins/MovieScene/MovieRenderPipeline/Source/MovieRenderPipelineEditor/Private/MoviePipelineConfigAssetEditor.cpp
@@ -0,0 +1,153 @@
+// Copyright Epic Games, Inc. All Rights Reserved.
+
+#include "MoviePipelineConfigAssetEditor.h"
+#include "MovieRenderPipelineEditorModule.h"
+#include "MoviePipelineConfigBase.h"
+#include "Widgets/SMoviePipelineConfigPanel.h"
+#include "MoviePipelineMasterConfig.h"
+#include "Widgets/Docking/SDockTab.h"
+#include "Styling/CoreStyle.h"
+#include "Widgets/SOverlay.h"
+#include "Widgets/Images/SImage.h"
+#include "Framework/MultiBox/MultiBoxBuilder.h"
+#include "ToolMenus.h"
+#include "EditorFontGlyphs.h"
+
+#define LOCTEXT_NAMESPACE "MoviePipelineConfigAssetEditor"
+
+const FName FMoviePipelineConfigAssetEditor::ContentTabId(TEXT("PipelineAssetEditor_Config"));
+
+FMoviePipelineConfigAssetEditor::FMoviePipelineConfigAssetEditor(UAssetEditor* InOwningAssetEditor)
+	: FBaseAssetToolkit(InOwningAssetEditor)
+{
+	LayoutAppendix = TEXT("MoviePipelineConfigAssetEditor");
+	FString LayoutString = TEXT("Standalone_Test_Layout_") + LayoutAppendix;
+
+	// Override the Default Layout provided by FBaseAssetToolkit to hide the viewport and details panel tabs.
+	StandaloneDefaultLayout = FTabManager::NewLayout(FName(LayoutString))
+		->AddArea
+		(
+			FTabManager::NewPrimaryArea()->SetOrientation(Orient_Vertical)
+			->Split
+			(
+				FTabManager::NewStack()
+				->SetSizeCoefficient(0.1f)
+				->AddTab(GetToolbarTabId(), ETabState::OpenedTab)
+				->SetHideTabWell(true)
+			)
+			->Split
+			(
+				FTabManager::NewStack()
+				->SetSizeCoefficient(1.f)
+				->SetHideTabWell(true)
+				->AddTab(ContentTabId, ETabState::OpenedTab)
+			)
+		);
+}
+
+void FMoviePipelineConfigAssetEditor::RegisterTabSpawners(const TSharedRef<class FTabManager>& InTabManager)
+{
+	WorkspaceMenuCategory = InTabManager->AddLocalWorkspaceMenuCategory(LOCTEXT("WorkspaceMenu_PipelineAssetEditor", "Movie Pipeline Asset Editor"));
+
+	FAssetEditorToolkit::RegisterTabSpawners(InTabManager);
+
+	InTabManager->RegisterTabSpawner(ContentTabId, FOnSpawnTab::CreateSP(this, &FMoviePipelineConfigAssetEditor::SpawnTab_ConfigEditor))
+		.SetDisplayName(LOCTEXT("ConfigTab", "Config"))
+		.SetGroup(WorkspaceMenuCategory.ToSharedRef());
+}
+
+void FMoviePipelineConfigAssetEditor::UnregisterTabSpawners(const TSharedRef<class FTabManager>& InTabManager)
+{
+	InTabManager->UnregisterTabSpawner(ContentTabId);
+}
+
+TSharedRef<SDockTab> FMoviePipelineConfigAssetEditor::SpawnTab_ConfigEditor(const FSpawnTabArgs& Args)
+{
+	UMoviePipelineConfigBase* Config = Cast<UMoviePipelineConfigBase>(GetEditingObject());
+
+	ConfigEditorPanel =
+		SNew(SMoviePipelineConfigPanel, Config->GetClass())
+		.Job(nullptr)
+		.Shot(nullptr)
+		.BasePreset(nullptr)
+		.BaseConfig(nullptr)
+		.AssetToEdit(Config);
+
+	// Extend the tool menu specific to this tool
+	UToolMenus* ToolMenus = UToolMenus::Get();
+	FName ParentName;
+	FName ToolMenuToolbarName = GetToolMenuToolbarName(ParentName);
+
+	{
+		UToolMenu* ToolbarBuilder = ToolMenus->RegisterMenu(ToolMenuToolbarName, ParentName, EMultiBoxType::ToolBar);
+		{
+			FToolMenuInsert InsertAfterAssetSection("Asset", EToolMenuInsertType::After);
+			{
+				FToolMenuSection& Section = ToolbarBuilder->AddSection("Settings", TAttribute<FText>(), InsertAfterAssetSection);
+				Section.AddEntry(FToolMenuEntry::InitWidget(FName("Settings"), ConfigEditorPanel->MakeSettingsWidget(), LOCTEXT("AddSettings_Label", "Settings")));
+			}
+		}
+		RegenerateMenusAndToolbars();
+	}
+
+	// Tabs and Windows have different backgrounds, so we copy the window style to
+	// make the 'void' area look similar.
+	const FSlateBrush* BackgroundBrush = &FCoreStyle::Get().GetWidgetStyle<FWindowStyle>("Window").ChildBackgroundBrush;
+
+	TSharedRef<SDockTab> NewDockTab = SNew(SDockTab)
+		.Icon(FEditorStyle::GetBrush("CurveAssetEditor.Tabs.Properties"))
+		.TabColorScale(GetTabColorScale())
+		[
+			SNew(SOverlay)
+			+ SOverlay::Slot()
+			[
+				SNew(SImage)
+				.Image(BackgroundBrush)
+			]
+			+ SOverlay::Slot()
+			[
+				SNew(SVerticalBox)
+				+ SVerticalBox::Slot()
+				[
+					ConfigEditorPanel.ToSharedRef()
+				]
+				+SVerticalBox::Slot()
+				.AutoHeight()
+				[
+					SNew(SBorder)
+					.Padding(4.f)
+					.BorderImage(FEditorStyle::GetBrush("DetailsView.CategoryTop"))
+					.BorderBackgroundColor(FLinearColor(.6, .6, .6, 1.0f))
+					.Content()
+					[
+						SNew(SHorizontalBox)
+
+						+SHorizontalBox::Slot()
+						.AutoWidth()
+						.VAlign(VAlign_Center)
+						.Padding(0, 0, 4, 0)
+						[
+							SNew(STextBlock)
+							.Font(FEditorStyle::Get().GetFontStyle("FontAwesome.9"))
+							.Text(FEditorFontGlyphs::Exclamation_Triangle)
+							.ColorAndOpacity(FLinearColor::Yellow)
+						]
+						+SHorizontalBox::Slot()
+						[
+							SNew(STextBlock)
+							.Text(LOCTEXT("PipelineLimitedEditing", "Editing capabilities are limited in the asset editor. Edit this preset through the Movie Render Queue UI for full editing capabilities (validation, format texts, etc.)"))
+						]
+					]
+				]
+			]
+		];
+
+	return NewDockTab;
+}
+
+FName FMoviePipelineConfigAssetEditor::GetToolkitFName() const
+{
+	return FName("MoviePipelineConfigAssetEditor");
+}
+
+#undef LOCTEXT_NAMESPACE
diff --git a/Engine/Plugins/MovieScene/MovieRenderPipeline/Source/MovieRenderPipelineEditor/Private/MoviePipelineConfigAssetEditor.h b/Engine/Plugins/MovieScene/MovieRenderPipeline/Source/MovieRenderPipelineEditor/Private/MoviePipelineConfigAssetEditor.h
new file mode 100644
--- /dev/null
+++ b/Engine/Plugins/MovieScene/MovieRenderPipeline/Source/MovieRenderPipelineEditor/Private/MoviePipelineConfigAssetEditor.h
@@ -0,0 +1,29 @@
+// Copyright Epic Games, Inc. All Rights Reserved.
+
+#pragma once
+
+#include "CoreMinimal.h"
+#include "Widgets/SWidget.h"
+#include "Tools/BaseAssetToolkit.h"
+
+class UMoviePipelineConfigBase;
+
+class FMoviePipelineConfigAssetEditor :  public FBaseAssetToolkit
+{
+public:
+
+	FMoviePipelineConfigAssetEditor(UAssetEditor* InOwningAssetEditor);
+	virtual ~FMoviePipelineConfigAssetEditor() {}
+
+	/** FBaseAssetToolkit interface */
+	virtual void RegisterTabSpawners(const TSharedRef<class FTabManager>& TabManager) override;
+	virtual void UnregisterTabSpawners(const TSharedRef<class FTabManager>& TabManager) override;
+	virtual FName GetToolkitFName() const override;
+
+private:
+	TSharedRef<SDockTab> SpawnTab_ConfigEditor(const FSpawnTabArgs& Args);
+
+private:
+	static const FName ContentTabId;
+	TSharedPtr<class SMoviePipelineConfigPanel> ConfigEditorPanel;
+};
diff --git a/Engine/Plugins/MovieScene/MovieRenderPipeline/Source/MovieRenderPipelineEditor/Private/MovieRenderPipelineEditorModule.cpp b/Engine/Plugins/MovieScene/MovieRenderPipeline/Source/MovieRenderPipelineEditor/Private/MovieRenderPipelineEditorModule.cpp
--- a/Engine/Plugins/MovieScene/MovieRenderPipeline/Source/MovieRenderPipelineEditor/Private/MovieRenderPipelineEditorModule.cpp
+++ b/Engine/Plugins/MovieScene/MovieRenderPipeline/Source/MovieRenderPipelineEditor/Private/MovieRenderPipelineEditorModule.cpp
@@ -29,6 +29,8 @@
 #include "MoviePipelineEditorBlueprintLibrary.h"
 #include "MovieSceneSection.h"
 #include "Sections/MovieSceneCinematicShotSection.h"
+#include "AssetToolsModule.h"
+#include "AssetTypeActions_PipelineConfigs.h"
 
 #define LOCTEXT_NAMESPACE "FMovieRenderPipelineEditorModule"
 
@@ -194,10 +196,30 @@ void FMovieRenderPipelineEditorModule::StartupModule()
 	RegisterTabImpl();
 	RegisterSettings();
 	RegisterMovieRenderer();
+
+	IAssetTools& AssetTools = FModuleManager::LoadModuleChecked<FAssetToolsModule>("AssetTools").Get();
+	{
+		MasterConfigAssetActions = MakeShared<FAssetTypeActions_PipelineMasterConfig>();
+		AssetTools.RegisterAssetTypeActions(MasterConfigAssetActions.ToSharedRef());
+
+		ShotConfigAssetActions = MakeShared<FAssetTypeActions_PipelineShotConfig>();
+		AssetTools.RegisterAssetTypeActions(ShotConfigAssetActions.ToSharedRef());
+
+		QueueAssetActions = MakeShared<FAssetTypeActions_PipelineQueue>();
+		AssetTools.RegisterAssetTypeActions(QueueAssetActions.ToSharedRef());
+	}
 }
 
 void FMovieRenderPipelineEditorModule::ShutdownModule()
 {
+	FAssetToolsModule* AssetToolsModule = FModuleManager::GetModulePtr<FAssetToolsModule>("AssetTools");
+	if (AssetToolsModule)
+	{
+		AssetToolsModule->Get().UnregisterAssetTypeActions(MasterConfigAssetActions.ToSharedRef());
+		AssetToolsModule->Get().UnregisterAssetTypeActions(ShotConfigAssetActions.ToSharedRef());
+		AssetToolsModule->Get().UnregisterAssetTypeActions(QueueAssetActions.ToSharedRef());
+	}
+
 	UnregisterMovieRenderer();
 	UnregisterSettings();
 	FMoviePipelineCommands::Unregister();
diff --git a/Engine/Plugins/MovieScene/MovieRenderPipeline/Source/MovieRenderPipelineEditor/Private/MovieRenderPipelineEditorModule.h b/Engine/Plugins/MovieScene/MovieRenderPipeline/Source/MovieRenderPipelineEditor/Private/MovieRenderPipelineEditorModule.h
--- a/Engine/Plugins/MovieScene/MovieRenderPipeline/Source/MovieRenderPipelineEditor/Private/MovieRenderPipelineEditorModule.h
+++ b/Engine/Plugins/MovieScene/MovieRenderPipeline/Source/MovieRenderPipelineEditor/Private/MovieRenderPipelineEditorModule.h
@@ -24,5 +24,8 @@ private:
 	void RegisterMovieRenderer();
 	void UnregisterMovieRenderer();
 
+	TSharedPtr<class FAssetTypeActions_PipelineMasterConfig> MasterConfigAssetActions;
+	TSharedPtr<class FAssetTypeActions_PipelineShotConfig> ShotConfigAssetActions;
+	TSharedPtr<class FAssetTypeActions_PipelineQueue> QueueAssetActions;
 	FDelegateHandle MovieRendererDelegate;
 };
diff --git a/Engine/Plugins/MovieScene/MovieRenderPipeline/Source/MovieRenderPipelineEditor/Private/Widgets/SMoviePipelineConfigPanel.cpp b/Engine/Plugins/MovieScene/MovieRenderPipeline/Source/MovieRenderPipelineEditor/Private/Widgets/SMoviePipelineConfigPanel.cpp
--- a/Engine/Plugins/MovieScene/MovieRenderPipeline/Source/MovieRenderPipelineEditor/Private/Widgets/SMoviePipelineConfigPanel.cpp
+++ b/Engine/Plugins/MovieScene/MovieRenderPipeline/Source/MovieRenderPipelineEditor/Private/Widgets/SMoviePipelineConfigPanel.cpp
@@ -50,18 +50,26 @@ void SMoviePipelineConfigPanel::Construct(const FArguments& InArgs, TSubclassOf<
 {
 	ConfigAssetType = InConfigType;
 
-	// Allocate a transient preset automatically so they can start editing without having to create an asset.
-	TransientPreset = AllocateTransientPreset();
-
-	// Copy the base preset into the transient preset if it was provided.
-	if (InArgs._BasePreset)
+	if (!InArgs._AssetToEdit)
 	{
-		TransientPreset->CopyFrom(InArgs._BasePreset);
-		PresetUsedIfNotModified = InArgs._BasePreset;
+		// Allocate a transient preset automatically so they can start editing without having to create an asset.
+		TransientPreset = AllocateTransientPreset();
+
+		// Copy the base preset into the transient preset if it was provided.
+		if (InArgs._BasePreset)
+		{
+			TransientPreset->CopyFrom(InArgs._BasePreset);
+			PresetUsedIfNotModified = InArgs._BasePreset;
+		}
+		else if (InArgs._BaseConfig)
+		{
+			TransientPreset->CopyFrom(InArgs._BaseConfig);
+		}
 	}
-	else if (InArgs._BaseConfig)
+	else
 	{
-		TransientPreset->CopyFrom(InArgs._BaseConfig);
+		// If they want to edit an asset directly, we don't copy from other sources
+		TransientPreset = InArgs._AssetToEdit;
 	}
 
 	// Listen for modificatiosn to our edited configuration so we can determine if they've modified the preset
@@ -78,12 +86,17 @@ void SMoviePipelineConfigPanel::Construct(const FArguments& InArgs, TSubclassOf<
 	.PipelineConfig(this, &SMoviePipelineConfigPanel::GetPipelineConfig)
 	.OwningJob(this, &SMoviePipelineConfigPanel::GetOwningJob);
 
+	TSharedRef<SVerticalBox> VerticalBox = SNew(SVerticalBox);
 	ChildSlot
 	[
-		SNew(SVerticalBox)
+		VerticalBox
+	];
 
+	// Direct editing of assets doesn't support 'transient' preset management toolbars.
+	if (!InArgs._AssetToEdit)
+	{
 		// Create the toolbar for adding new stuff, choosing/saving a preset and resetting to the preset defaults.
-		+ SVerticalBox::Slot()
+		VerticalBox->AddSlot()
 		.Padding(FMargin(0.f, 1.0f))
 		.AutoHeight()
 		[
@@ -145,32 +158,37 @@ void SMoviePipelineConfigPanel::Construct(const FArguments& InArgs, TSubclassOf<
 					]
 				]
 			]
-		]
+		];
+	}
 
-		// Large Label Telling You What You're Editing
-		+SVerticalBox::Slot()
-		.AutoHeight()
-		.HAlign(HAlign_Right)
-		[
-			SNew(STextBlock)
-			.TextStyle(FMovieRenderPipelineStyle::Get(), "MovieRenderPipeline.Config.TypeLabel")
-			.Text(this, &SMoviePipelineConfigPanel::GetConfigTypeLabel)
-			.Visibility(EVisibility::Collapsed)
-		]
+	// Large Label Telling You What You're Editing
+	VerticalBox->AddSlot()
+	.AutoHeight()
+	.HAlign(HAlign_Right)
+	[
+		SNew(STextBlock)
+		.TextStyle(FMovieRenderPipelineStyle::Get(), "MovieRenderPipeline.Config.TypeLabel")
+		.Text(this, &SMoviePipelineConfigPanel::GetConfigTypeLabel)
+		.Visibility(EVisibility::Collapsed)
+	];
 
-		// Main Editor 
-		+ SVerticalBox::Slot()
-		.FillHeight(1.0f)
+		
+	// Main Editor 
+	VerticalBox->AddSlot()
+	.FillHeight(1.0f)
+	[
+		SNew(SHorizontalBox)
+		+ SHorizontalBox::Slot()
 		[
-			SNew(SHorizontalBox)
-			+ SHorizontalBox::Slot()
-			[
-				MoviePipelineEditorWidget.ToSharedRef()
-			]
+			MoviePipelineEditorWidget.ToSharedRef()
 		]
+	];
 
+	// Direct editing of assets doesn't support 'transient' preset management footers.
+	if (!InArgs._AssetToEdit)
+	{
 		// Footer Bar
-		+ SVerticalBox::Slot()
+		VerticalBox->AddSlot()
 		.AutoHeight()
 		[
 			SNew(SBorder)
@@ -227,11 +245,10 @@ void SMoviePipelineConfigPanel::Construct(const FArguments& InArgs, TSubclassOf<
 						.ToolTipText(LOCTEXT("ConfirmChangesButton_Tooltip", "Accepts the changes made and applies them to the particular job instance."))
 						.Margin(FMargin(4, 0, 4, 0))
 					]
-
 				]
 			]
-		]
-	];
+		];
+	}
 }
 
 PRAGMA_ENABLE_OPTIMIZATION
@@ -548,8 +565,6 @@ void SMoviePipelineConfigPanel::OnImportPreset(const FAssetData& InPresetAsset)
 void SMoviePipelineConfigPanel::AddReferencedObjects(FReferenceCollector& Collector)
 {
 	Collector.AddReferencedObject(TransientPreset);
-	// Collector.AddReferencedObject(SuppliedLevelSequence);
-	// Collector.AddReferencedObject(RecordingLevelSequence);
 }
 
 UMoviePipelineConfigBase* SMoviePipelineConfigPanel::GetPipelineConfig() const
@@ -567,4 +582,8 @@ UMoviePipelineExecutorShot* SMoviePipelineConfigPanel::GetOwningShot() const
 	return WeakShot.Get();
 }
 
+TSharedRef<SWidget> SMoviePipelineConfigPanel::MakeSettingsWidget()
+{
+	return MoviePipelineEditorWidget->MakeAddSettingButton();
+}
 #undef LOCTEXT_NAMESPACE // SMoviePipelinePanel
\ No newline at end of file
diff --git a/Engine/Plugins/MovieScene/MovieRenderPipeline/Source/MovieRenderPipelineEditor/Private/Widgets/SMoviePipelineConfigPanel.h b/Engine/Plugins/MovieScene/MovieRenderPipeline/Source/MovieRenderPipelineEditor/Private/Widgets/SMoviePipelineConfigPanel.h
--- a/Engine/Plugins/MovieScene/MovieRenderPipeline/Source/MovieRenderPipelineEditor/Private/Widgets/SMoviePipelineConfigPanel.h
+++ b/Engine/Plugins/MovieScene/MovieRenderPipeline/Source/MovieRenderPipelineEditor/Private/Widgets/SMoviePipelineConfigPanel.h
@@ -27,6 +27,7 @@ public:
 	SLATE_BEGIN_ARGS(SMoviePipelineConfigPanel)
 		: _BasePreset(nullptr)
 		,  _BaseConfig(nullptr)
+		, _AssetToEdit(nullptr)
 
 		{}
 		SLATE_ARGUMENT(TWeakObjectPtr<UMoviePipelineExecutorJob>, Job)
@@ -42,6 +43,10 @@ public:
 
 		/** An existing configuration to copy into the transient UI object. This will not get modified */
 		SLATE_ARGUMENT(UMoviePipelineConfigBase*, BaseConfig)
+
+		/** An existing asset to edit directly with no copy. Hides some of the UI. */
+		SLATE_ARGUMENT(UMoviePipelineConfigBase*, AssetToEdit)
+
 		/*-------------------------------------------------*/
 
 	SLATE_END_ARGS()
@@ -86,6 +91,9 @@ private:
 
 	FText GetValidationWarningText() const;
 
+public:
+	TSharedRef<SWidget> MakeSettingsWidget();
+
 private:
 	/** The transient preset that we use - kept alive by AddReferencedObjects */
 	UMoviePipelineConfigBase* TransientPreset;
