CIS: Fixed deprecated function warning usage.

diff --git a/Engine/Plugins/MovieScene/MovieRenderPipeline/Source/MovieRenderPipelineEditor/Private/MoviePipelineConfigAssetEditor.cpp b/Engine/Plugins/MovieScene/MovieRenderPipeline/Source/MovieRenderPipelineEditor/Private/MoviePipelineConfigAssetEditor.cpp
--- a/Engine/Plugins/MovieScene/MovieRenderPipeline/Source/MovieRenderPipelineEditor/Private/MoviePipelineConfigAssetEditor.cpp
+++ b/Engine/Plugins/MovieScene/MovieRenderPipeline/Source/MovieRenderPipelineEditor/Private/MoviePipelineConfigAssetEditor.cpp
@@ -95,7 +95,6 @@ TSharedRef<SDockTab> FMoviePipelineConfigAssetEditor::SpawnTab_ConfigEditor(cons
 	const FSlateBrush* BackgroundBrush = &FCoreStyle::Get().GetWidgetStyle<FWindowStyle>("Window").ChildBackgroundBrush;
 
 	TSharedRef<SDockTab> NewDockTab = SNew(SDockTab)
-		.Icon(FEditorStyle::GetBrush("CurveAssetEditor.Tabs.Properties"))
 		.TabColorScale(GetTabColorScale())
 		[
 			SNew(SOverlay)
