Movie Render Queue: Ensure new job has default settings

diff --git a/Engine/Plugins/MovieScene/MovieRenderPipeline/Source/MovieRenderPipelineEditor/Private/MoviePipelineEditorBlueprintLibrary.cpp b/Engine/Plugins/MovieScene/MovieRenderPipeline/Source/MovieRenderPipelineEditor/Private/MoviePipelineEditorBlueprintLibrary.cpp
--- a/Engine/Plugins/MovieScene/MovieRenderPipeline/Source/MovieRenderPipelineEditor/Private/MoviePipelineEditorBlueprintLibrary.cpp
+++ b/Engine/Plugins/MovieScene/MovieRenderPipeline/Source/MovieRenderPipelineEditor/Private/MoviePipelineEditorBlueprintLibrary.cpp
@@ -180,5 +180,27 @@ UMoviePipelineExecutorJob* UMoviePipelineEditorBlueprintLibrary::CreateJobFromSe
 	return NewJob;
 }
 
+void UMoviePipelineEditorBlueprintLibrary::EnsureJobHasDefaultSettings(UMoviePipelineExecutorJob* NewJob)
+{
+	const UMovieRenderPipelineProjectSettings* ProjectSettings = GetDefault<UMovieRenderPipelineProjectSettings>();
+	for (TSubclassOf<UMoviePipelineSetting> SettingClass : ProjectSettings->DefaultClasses)
+	{
+		if (!SettingClass)
+		{
+			continue;
+		}
+
+		if (SettingClass->HasAnyClassFlags(CLASS_Abstract))
+		{
+			continue;
+		}
+
+		UMoviePipelineSetting* ExistingSetting = NewJob->GetConfiguration()->FindSettingByClass(SettingClass);
+		if (!ExistingSetting)
+		{
+			NewJob->GetConfiguration()->FindOrAddSettingByClass(SettingClass);
+		}
+	}
+}
 
 #undef LOCTEXT_NAMESPACE // "MoviePipelineEditorBlueprintLibrary"
\ No newline at end of file
diff --git a/Engine/Plugins/MovieScene/MovieRenderPipeline/Source/MovieRenderPipelineEditor/Private/MovieRenderPipelineEditorModule.cpp b/Engine/Plugins/MovieScene/MovieRenderPipeline/Source/MovieRenderPipelineEditor/Private/MovieRenderPipelineEditorModule.cpp
--- a/Engine/Plugins/MovieScene/MovieRenderPipeline/Source/MovieRenderPipelineEditor/Private/MovieRenderPipelineEditorModule.cpp
+++ b/Engine/Plugins/MovieScene/MovieRenderPipeline/Source/MovieRenderPipelineEditor/Private/MovieRenderPipelineEditorModule.cpp
@@ -138,6 +138,8 @@ class FMovieRenderPipelineRenderer : public IMovieRendererInterface
 			}
 		}
 
+		UMoviePipelineEditorBlueprintLibrary::EnsureJobHasDefaultSettings(ActiveJob);
+
 		TArray<FString> ShotNames;
 		for (UMovieSceneCinematicShotSection* ShotSection : InSections)
 		{
diff --git a/Engine/Plugins/MovieScene/MovieRenderPipeline/Source/MovieRenderPipelineEditor/Private/Widgets/SMoviePipelineQueueEditor.cpp b/Engine/Plugins/MovieScene/MovieRenderPipeline/Source/MovieRenderPipelineEditor/Private/Widgets/SMoviePipelineQueueEditor.cpp
--- a/Engine/Plugins/MovieScene/MovieRenderPipeline/Source/MovieRenderPipelineEditor/Private/Widgets/SMoviePipelineQueueEditor.cpp
+++ b/Engine/Plugins/MovieScene/MovieRenderPipeline/Source/MovieRenderPipelineEditor/Private/Widgets/SMoviePipelineQueueEditor.cpp
@@ -56,10 +56,6 @@ struct FMoviePipelineMapTreeItem;
 struct FMoviePipelineShotItem;
 class SQueueJobListRow;
 
-// Forward Declare
-static void EnsureJobHasDefaultSettings(UMoviePipelineExecutorJob* NewJob);
-
-
 struct IMoviePipelineQueueTreeItem : TSharedFromThis<IMoviePipelineQueueTreeItem>
 {
 	virtual ~IMoviePipelineQueueTreeItem() {}
@@ -202,7 +198,7 @@ public:
 		{
 			// Copy from the CDO's version of the job to pick up the right name.
 			Job->SetConfiguration(GetMutableDefault<UMoviePipelineExecutorJob>()->GetConfiguration());
-			EnsureJobHasDefaultSettings(Job);
+			UMoviePipelineEditorBlueprintLibrary::EnsureJobHasDefaultSettings(Job);
 		}
 
 		OnChosePresetCallback.ExecuteIfBound(WeakJob, nullptr);
@@ -1022,29 +1018,6 @@ TSharedRef<SWidget> SMoviePipelineQueueEditor::OnGenerateNewJobFromAssetMenu()
 
 PRAGMA_ENABLE_OPTIMIZATION
 
-void EnsureJobHasDefaultSettings(UMoviePipelineExecutorJob* NewJob)
-{
-	const UMovieRenderPipelineProjectSettings* ProjectSettings = GetDefault<UMovieRenderPipelineProjectSettings>();
-	for (TSubclassOf<UMoviePipelineSetting> SettingClass : ProjectSettings->DefaultClasses)
-	{
-		if (!SettingClass)
-		{
-			continue;
-		}
-
-		if (SettingClass->HasAnyClassFlags(CLASS_Abstract))
-		{
-			continue;
-		}
-
-		UMoviePipelineSetting* ExistingSetting = NewJob->GetConfiguration()->FindSettingByClass(SettingClass);
-		if (!ExistingSetting)
-		{
-			NewJob->GetConfiguration()->FindOrAddSettingByClass(SettingClass);
-		}
-	}
-}
-
 void SMoviePipelineQueueEditor::OnCreateJobFromAsset(const FAssetData& InAsset)
 {
 	// Close the dropdown menu that showed them the assets to pick from.
@@ -1078,8 +1051,8 @@ void SMoviePipelineQueueEditor::OnCreateJobFromAsset(const FAssetData& InAsset)
 		}
 
 		// Ensure the job has the settings specified by the project settings added. If they're already added
-		// we don't modify the object so that we don't make it confused about whehter or not you've modified the preset.
-		EnsureJobHasDefaultSettings(NewJob);
+		// we don't modify the object so that we don't make it confused about whether or not you've modified the preset.
+		UMoviePipelineEditorBlueprintLibrary::EnsureJobHasDefaultSettings(NewJob);
 	}
 }
 
diff --git a/Engine/Plugins/MovieScene/MovieRenderPipeline/Source/MovieRenderPipelineEditor/Public/MoviePipelineEditorBlueprintLibrary.h b/Engine/Plugins/MovieScene/MovieRenderPipeline/Source/MovieRenderPipelineEditor/Public/MoviePipelineEditorBlueprintLibrary.h
--- a/Engine/Plugins/MovieScene/MovieRenderPipeline/Source/MovieRenderPipelineEditor/Public/MoviePipelineEditorBlueprintLibrary.h
+++ b/Engine/Plugins/MovieScene/MovieRenderPipeline/Source/MovieRenderPipelineEditor/Public/MoviePipelineEditorBlueprintLibrary.h
@@ -44,4 +44,8 @@ public:
 	/** Create a job from a level sequence. Sets the map as the currently editor world, the author, the sequence and the job name as the sequence name on the new job. Returns the newly created job. */
 	UFUNCTION(BlueprintCallable, Category = "Movie Render Pipeline")
 	static UMoviePipelineExecutorJob* CreateJobFromSequence(UMoviePipelineQueue* InPipelineQueue, const ULevelSequence* InSequence);
+
+	/** Ensure the job has the settings specified by the project settings added. If they're already added we don't modify the object so that we don't make it confused about whether or not you've modified the preset. */
+	UFUNCTION(BlueprintCallable, Category = "Movie Render Pipeline")
+	static void EnsureJobHasDefaultSettings(UMoviePipelineExecutorJob* InJob);
 };
\ No newline at end of file
