Fix for compilation. A 0 got inserted before I checked in 16347066 for some reason.

diff --git a/Engine/Plugins/MovieScene/MovieRenderPipeline/Source/MovieRenderPipelineCore/Private/MoviePipelineGameOverrideSetting.cpp b/Engine/Plugins/MovieScene/MovieRenderPipeline/Source/MovieRenderPipelineCore/Private/MoviePipelineGameOverrideSetting.cpp
--- a/Engine/Plugins/MovieScene/MovieRenderPipeline/Source/MovieRenderPipelineCore/Private/MoviePipelineGameOverrideSetting.cpp
+++ b/Engine/Plugins/MovieScene/MovieRenderPipeline/Source/MovieRenderPipelineCore/Private/MoviePipelineGameOverrideSetting.cpp
@@ -104,7 +104,7 @@ void UMoviePipelineGameOverrideSetting::ApplyCVarSettings(const bool bOverrideVa
 	MOVIEPIPELINE_STORE_AND_OVERRIDE_CVAR_INT(PreviousSkyLightRealTimeReflectionCaptureTimeSlice, TEXT("r.SkyLight.RealTimeReflectionCapture.TimeSlice"), 0, bOverrideValues);
 
 	// To make sure that the skylight is always valid and consistent accross capture sessions, we enforce a full capture each frame, accepting a small GPU cost.
-	MOVIEPIPELINE_STORE_AND_OVERRIDE_CVAR_INT(PreviousVolumetricRenderTarget 0, TEXT("r.VolumetricRenderTarget"), 0, bOverrideValues);
+	MOVIEPIPELINE_STORE_AND_OVERRIDE_CVAR_INT(PreviousVolumetricRenderTarget, TEXT("r.VolumetricRenderTarget"), 0, bOverrideValues);
 
 }
 
