MoviePipeline: Renders launched via Render (Remote) now add -MRQInstance to the command line. Nothing reads this at this time but it allows code to check if it is a MRQ instance if needed. Does not affect Editor sessions.

diff --git a/Engine/Plugins/MovieScene/MovieRenderPipeline/Source/MovieRenderPipelineCore/Public/MoviePipelineInProcessExecutorSettings.h b/Engine/Plugins/MovieScene/MovieRenderPipeline/Source/MovieRenderPipelineCore/Public/MoviePipelineInProcessExecutorSettings.h
--- a/Engine/Plugins/MovieScene/MovieRenderPipeline/Source/MovieRenderPipelineCore/Public/MoviePipelineInProcessExecutorSettings.h
+++ b/Engine/Plugins/MovieScene/MovieRenderPipeline/Source/MovieRenderPipelineCore/Public/MoviePipelineInProcessExecutorSettings.h
@@ -19,7 +19,7 @@ public:
 		: Super()
 	{
 		bCloseEditor = false;
-		AdditionalCommandLineArguments = TEXT("-NoLoadingScreen -FixedSeed -log -Unattended");
+		AdditionalCommandLineArguments = TEXT("-NoLoadingScreen -FixedSeed -log -Unattended -MRQInstance");
 		InitialDelayFrameCount = 0;
 
 		// Find all arguments from the command line and set them as the InheritedCommandLineArguments.
