MoviePipeline: Added the Camera overscan_percentage as a format tag and to exr metadata.

diff --git a/Engine/Plugins/MovieScene/MovieRenderPipeline/Source/MovieRenderPipelineCore/Public/MoviePipelineCameraSetting.h b/Engine/Plugins/MovieScene/MovieRenderPipeline/Source/MovieRenderPipelineCore/Public/MoviePipelineCameraSetting.h
--- a/Engine/Plugins/MovieScene/MovieRenderPipeline/Source/MovieRenderPipelineCore/Public/MoviePipelineCameraSetting.h
+++ b/Engine/Plugins/MovieScene/MovieRenderPipeline/Source/MovieRenderPipelineCore/Public/MoviePipelineCameraSetting.h
@@ -28,7 +28,9 @@ protected:
 		Super::GetFormatArguments(InOutFormatArgs);
 
 		InOutFormatArgs.FilenameArguments.Add(TEXT("shutter_timing"), StaticEnum<EMoviePipelineShutterTiming>()->GetNameStringByValue((int64)ShutterTiming));
+		InOutFormatArgs.FilenameArguments.Add(TEXT("overscan_percentage"), FString::SanitizeFloat(OverscanPercentage));
 		InOutFormatArgs.FileMetadata.Add(TEXT("unreal/camera/shutterTiming"), StaticEnum<EMoviePipelineShutterTiming>()->GetNameStringByValue((int64)ShutterTiming));
+		InOutFormatArgs.FileMetadata.Add(TEXT("unreal/camera/overscanPercentage"), FString::SanitizeFloat(OverscanPercentage));
 	}
 public:	
 	/**
