MoviePipeline: Changed FCPXML exporter setting to have the filename format be an *override* of the Output's Filename Format. This means you can skip specifying it on the FCPXML exporter if you want it to be the same as the main output string (though {frame_number} and friends are stripped, just like video outputs)

diff --git a/Engine/Plugins/MovieScene/MovieRenderPipeline/Source/MovieRenderPipelineCore/Private/MoviePipelineFCPXMLExporterSetting.cpp b/Engine/Plugins/MovieScene/MovieRenderPipeline/Source/MovieRenderPipelineCore/Private/MoviePipelineFCPXMLExporterSetting.cpp
--- a/Engine/Plugins/MovieScene/MovieRenderPipeline/Source/MovieRenderPipelineCore/Private/MoviePipelineFCPXMLExporterSetting.cpp
+++ b/Engine/Plugins/MovieScene/MovieRenderPipeline/Source/MovieRenderPipelineCore/Private/MoviePipelineFCPXMLExporterSetting.cpp
@@ -12,7 +12,9 @@
 #include "MovieSceneExportMetadata.h"
 #include "MovieSceneToolHelpers.h"
 #endif
+
 #include "LevelSequence.h"
+#include "MoviePipelineUtils.h"
 
 // For logs
 #include "MovieRenderPipelineCoreModule.h"
@@ -25,7 +27,17 @@ void UMoviePipelineFCPXMLExporter::BeginExportImpl()
 	UMoviePipelineOutputSetting* OutputSetting = GetPipeline()->GetPipelineMasterConfig()->FindSetting<UMoviePipelineOutputSetting>();
 
 	// Use our file name format on the end of the shared common directory.
+	FString FileNameFormat = FileNameFormatOverride.Len() > 0 ? FileNameFormatOverride : OutputSetting->FileNameFormat;
+
 	FString FileNameFormatString = OutputSetting->OutputDirectory.Path / FileNameFormat;
+
+	const bool bIncludeRenderPass = false;
+	const bool bTestFrameNumber = false;
+
+	UE::MoviePipeline::ValidateOutputFormatString(FileNameFormatString, bIncludeRenderPass, bTestFrameNumber);
+
+	// Strip any frame number tags.
+	UE::MoviePipeline::RemoveFrameNumberFormatStrings(FileNameFormatString, true);
 	
 	TMap<FString, FString> FormatOverrides;
 	FormatOverrides.Add(TEXT("ext"), TEXT("xml"));
@@ -48,6 +60,7 @@ void UMoviePipelineFCPXMLExporter::BeginExportImpl()
 		UMovieSceneCinematicShotTrack* ShotTrack = MovieScene->FindMasterTrack<UMovieSceneCinematicShotTrack>();
 		if (!ShotTrack)
 		{
+			UE_LOG(LogMovieRenderPipeline, Error, TEXT("FCPXML Export only works with a Cinematic Shot track. No FCPXML file will be written."));
 			return;
 		}
 		
diff --git a/Engine/Plugins/MovieScene/MovieRenderPipeline/Source/MovieRenderPipelineCore/Public/MoviePipelineFCPXMLExporterSetting.h b/Engine/Plugins/MovieScene/MovieRenderPipeline/Source/MovieRenderPipelineCore/Public/MoviePipelineFCPXMLExporterSetting.h
--- a/Engine/Plugins/MovieScene/MovieRenderPipeline/Source/MovieRenderPipelineCore/Public/MoviePipelineFCPXMLExporterSetting.h
+++ b/Engine/Plugins/MovieScene/MovieRenderPipeline/Source/MovieRenderPipelineCore/Public/MoviePipelineFCPXMLExporterSetting.h
@@ -18,7 +18,7 @@ class MOVIERENDERPIPELINECORE_API UMoviePipelineFCPXMLExporter : public UMoviePi
 	GENERATED_BODY()
 public:
 	UMoviePipelineFCPXMLExporter()
-		: FileNameFormat(TEXT("{sequence_name}"))
+		: DataSource(FCPXMLExportDataSource::SequenceData)
 		, bHasFinishedExporting(false)
 	{}
 
@@ -33,9 +33,9 @@ protected:
 	bool EnsureWritableFile();
 	bool bOverwriteFile;
 public:
-	/** What format string should the final files use? Can include folder prefixes, and format string ({sequence_name}, etc.) */
+	/** File name format string override. If specified it will override the FileNameFormat from the Output setting. Can include folder prefixes, and format string tags ({sequence_name}, etc.) */
 	UPROPERTY(EditAnywhere, BlueprintReadWrite, Category = "File Output")
-	FString FileNameFormat;
+	FString FileNameFormatOverride;
 	
 	/** Whether to build the FCPXML from sequence data directly (for reimporting) or from actual frame output data (for post processing) */
 	UPROPERTY(EditAnywhere, BlueprintReadWrite, Category = "File Output")
