MoviePipeline: Made the config editor window more consistent with the Queue window by not force-showing plugin content on the content browser. Sorts settings by display text now to increase consistency with already added settings.

diff --git a/Engine/Plugins/MovieScene/MovieRenderPipeline/Source/MovieRenderPipelineEditor/Private/Widgets/SMoviePipelineConfigEditor.cpp b/Engine/Plugins/MovieScene/MovieRenderPipeline/Source/MovieRenderPipelineEditor/Private/Widgets/SMoviePipelineConfigEditor.cpp
--- a/Engine/Plugins/MovieScene/MovieRenderPipeline/Source/MovieRenderPipelineEditor/Private/Widgets/SMoviePipelineConfigEditor.cpp
+++ b/Engine/Plugins/MovieScene/MovieRenderPipeline/Source/MovieRenderPipelineEditor/Private/Widgets/SMoviePipelineConfigEditor.cpp
@@ -228,17 +228,6 @@ void SMoviePipelineConfigEditor::Tick(const FGeometry& AllottedGeometry, const d
 TSharedRef<SWidget> SMoviePipelineConfigEditor::OnGenerateSettingsMenu()
 {
 	TSharedRef<FExtender> Extender = MakeShared<FExtender>();
-	{
-		// ULevelSequence*       LevelSequence = LevelSequenceAttribute.Get();
-		// UTakeRecorderSources* Sources = LevelSequence ? LevelSequence->FindOrAddMetaData<UTakeRecorderSources>() : nullptr;
-		// 
-		// if (Sources)
-		// {
-		// 	FTakeRecorderModule& TakeRecorderModule = FModuleManager::GetModuleChecked<FTakeRecorderModule>("TakeRecorder");
-		// 	TakeRecorderModule.PopulateSourcesMenu(Extender, Sources);
-		// }
-	}
-
 	FMenuBuilder MenuBuilder(true, nullptr, Extender);
 
 	// Put the different categories into different sections
@@ -256,7 +245,14 @@ TSharedRef<SWidget> SMoviePipelineConfigEditor::OnGenerateSettingsMenu()
 	{
 		MenuBuilder.BeginSection(NAME_None, FText::FromString(KVP.Key));
 
-		Algo::SortBy(KVP.Value, &UClass::GetDisplayNameText, FText::FSortPredicate());
+		// Sort the classes by their CDO's GetDisplayText and not the classes GetDisplayName.
+		Algo::Sort(KVP.Value, [](UClass* A, UClass* B)
+			{
+				const UMoviePipelineSetting* DefaultA = GetDefault<UMoviePipelineSetting>(A);
+				const UMoviePipelineSetting* DefaultB = GetDefault<UMoviePipelineSetting>(B);
+
+				return DefaultA->GetDisplayText().CompareTo(DefaultB->GetDisplayText()) < 0;
+			});
 
 		for (UClass* Class : KVP.Value)
 		{
diff --git a/Engine/Plugins/MovieScene/MovieRenderPipeline/Source/MovieRenderPipelineEditor/Private/Widgets/SMoviePipelineConfigPanel.cpp b/Engine/Plugins/MovieScene/MovieRenderPipeline/Source/MovieRenderPipelineEditor/Private/Widgets/SMoviePipelineConfigPanel.cpp
--- a/Engine/Plugins/MovieScene/MovieRenderPipeline/Source/MovieRenderPipelineEditor/Private/Widgets/SMoviePipelineConfigPanel.cpp
+++ b/Engine/Plugins/MovieScene/MovieRenderPipeline/Source/MovieRenderPipelineEditor/Private/Widgets/SMoviePipelineConfigPanel.cpp
@@ -377,8 +377,8 @@ TSharedRef<SWidget> SMoviePipelineConfigPanel::OnGeneratePresetsMenu()
 		AssetPickerConfig.bShowPathInColumnView = true;
 		AssetPickerConfig.bShowTypeInColumnView = false;
 		AssetPickerConfig.bSortByPathInColumnView = false;
-		AssetPickerConfig.bForceShowEngineContent = true;
-		AssetPickerConfig.bForceShowPluginContent = true;
+		AssetPickerConfig.bForceShowEngineContent = false;
+		AssetPickerConfig.bForceShowPluginContent = false;
 
 		AssetPickerConfig.AssetShowWarningText = LOCTEXT("NoPresets_Warning", "No Presets Found");
 		AssetPickerConfig.Filter.ClassNames.Add(ConfigAssetType->GetFName());
