Movie Pipeline: FCameraCutSubSectionHierarchyNode stores sequence hierarchy info for restoring after setting flags and ranges for rendering. The initial root node's movie scene wasn't being set, so the sequence was getting marked dirty and not restoring when done rendering. This looks like a regression from when this code was refactored.

diff --git a/Engine/Plugins/MovieScene/MovieRenderPipeline/Source/MovieRenderPipelineCore/Private/MoviePipelineUtils.cpp b/Engine/Plugins/MovieScene/MovieRenderPipeline/Source/MovieRenderPipelineCore/Private/MoviePipelineUtils.cpp
--- a/Engine/Plugins/MovieScene/MovieRenderPipeline/Source/MovieRenderPipelineCore/Private/MoviePipelineUtils.cpp
+++ b/Engine/Plugins/MovieScene/MovieRenderPipeline/Source/MovieRenderPipelineCore/Private/MoviePipelineUtils.cpp
@@ -147,6 +147,8 @@ namespace MoviePipeline
 
 	void BuildCompleteSequenceHierarchyRecursive(UMovieSceneSequence* InSequence, TSharedPtr<FCameraCutSubSectionHierarchyNode> InNode)
 	{
+		InNode->MovieScene = InSequence->GetMovieScene();
+
 		if (UMovieSceneCameraCutTrack* CameraCutTrack = Cast<UMovieSceneCameraCutTrack>(InSequence->GetMovieScene()->GetCameraCutTrack()))
 		{
 			// We create leaf nodes for each section. This kind of makes duplicate leafs but since this is separate from the evaluation tree
