MoviePipeline: Fixed a mismatch when using time dilation where the UWorld's version of time had time dilation applied twice, while the render world only had it applied (correctly) once.

diff --git a/Engine/Plugins/MovieScene/MovieRenderPipeline/Source/MovieRenderPipelineCore/Private/MoviePipeline.cpp b/Engine/Plugins/MovieScene/MovieRenderPipeline/Source/MovieRenderPipelineCore/Private/MoviePipeline.cpp
--- a/Engine/Plugins/MovieScene/MovieRenderPipeline/Source/MovieRenderPipelineCore/Private/MoviePipeline.cpp
+++ b/Engine/Plugins/MovieScene/MovieRenderPipeline/Source/MovieRenderPipelineCore/Private/MoviePipeline.cpp
@@ -751,10 +751,10 @@ void UMoviePipeline::TickPostFinalizeExport(const bool bInForceFinish)
 
 bool UMoviePipelineCustomTimeStep::UpdateTimeStep(UEngine* /*InEngine*/)
 {
-	if (ensureMsgf(!FMath::IsNearlyZero(TimeCache.DeltaTime), TEXT("An incorrect or uninitialized time step was used! Delta Time of 0 isn't allowed.")))
+	if (ensureMsgf(!FMath::IsNearlyZero(TimeCache.UndilatedDeltaTime), TEXT("An incorrect or uninitialized time step was used! Delta Time of 0 isn't allowed.")))
 	{
 		FApp::UpdateLastTime();
-		FApp::SetDeltaTime(TimeCache.DeltaTime);
+		FApp::SetDeltaTime(TimeCache.UndilatedDeltaTime);
 		FApp::SetCurrentTime(FApp::GetCurrentTime() + FApp::GetDeltaTime());
 	}
 
@@ -768,7 +768,7 @@ bool UMoviePipelineCustomTimeStep::UpdateTimeStep(UEngine* /*InEngine*/)
 
 void UMoviePipelineCustomTimeStep::SetCachedFrameTiming(const MoviePipeline::FFrameTimeStepCache& InTimeCache)
 { 
-	if (ensureMsgf(!FMath::IsNearlyZero(InTimeCache.DeltaTime), TEXT("An incorrect or uninitialized time step was used! Delta Time of 0 isn't allowed.")))
+	if (ensureMsgf(!FMath::IsNearlyZero(InTimeCache.UndilatedDeltaTime), TEXT("An incorrect or uninitialized time step was used! Delta Time of 0 isn't allowed.")))
 	{
 		TimeCache = InTimeCache;
 	}
diff --git a/Engine/Plugins/MovieScene/MovieRenderPipeline/Source/MovieRenderPipelineCore/Private/MoviePipelineTiming.cpp b/Engine/Plugins/MovieScene/MovieRenderPipeline/Source/MovieRenderPipelineCore/Private/MoviePipelineTiming.cpp
--- a/Engine/Plugins/MovieScene/MovieRenderPipeline/Source/MovieRenderPipelineCore/Private/MoviePipelineTiming.cpp
+++ b/Engine/Plugins/MovieScene/MovieRenderPipeline/Source/MovieRenderPipelineCore/Private/MoviePipelineTiming.cpp
@@ -535,6 +535,7 @@ void UMoviePipeline::TickProducingFrames()
 		// Now that we've calculated our delta ticks, we need to multiply it by
 		// time dilation so that we advance through the sequence as slow as we
 		// advance through the world.
+		FFrameTime UndilatedDeltaFrameTime = DeltaFrameTime;
 		if (!FMath::IsNearlyEqual(WorldTimeDilation, 1.f))
 		{
 			UE_LOG(LogMovieRenderPipeline, VeryVerbose, TEXT("[%d] Modified FrameDeltaTime by a factor of %f to account for World Time Dilation."), GFrameCounter, WorldTimeDilation);
@@ -674,8 +675,10 @@ void UMoviePipeline::TickProducingFrames()
 		
 
 		
-		// Set our time step for the next frame
-		CustomTimeStep->SetCachedFrameTiming(MoviePipeline::FFrameTimeStepCache(FrameDeltaTime));
+		// Set our time step for the next frame. We use the undilated delta time for the Custom Timestep as the engine will
+		// apply the time dilation to the world tick for us, so we don't want to double up time dilation.
+		double UndilatedDeltaTime = FrameMetrics.TickResolution.AsSeconds(FFrameTime(UndilatedDeltaFrameTime.GetFrame()));
+		CustomTimeStep->SetCachedFrameTiming(MoviePipeline::FFrameTimeStepCache(UndilatedDeltaTime));
 		CustomSequenceTimeController->SetCachedFrameTiming(FQualifiedFrameTime(FinalEvalTime, FrameMetrics.TickResolution));
 		return;
 	}
diff --git a/Engine/Plugins/MovieScene/MovieRenderPipeline/Source/MovieRenderPipelineCore/Public/MovieRenderPipelineDataTypes.h b/Engine/Plugins/MovieScene/MovieRenderPipeline/Source/MovieRenderPipelineCore/Public/MovieRenderPipelineDataTypes.h
--- a/Engine/Plugins/MovieScene/MovieRenderPipeline/Source/MovieRenderPipelineCore/Public/MovieRenderPipelineDataTypes.h
+++ b/Engine/Plugins/MovieScene/MovieRenderPipeline/Source/MovieRenderPipelineCore/Public/MovieRenderPipelineDataTypes.h
@@ -156,14 +156,14 @@ namespace MoviePipeline
 	struct FFrameTimeStepCache
 	{
 		FFrameTimeStepCache()
-			: DeltaTime(0.0)
+			: UndilatedDeltaTime(0.0)
 		{}
 
 		FFrameTimeStepCache(double InDeltaTime)
-			: DeltaTime(InDeltaTime)
+			: UndilatedDeltaTime(InDeltaTime)
 		{}
 
-		double DeltaTime;
+		double UndilatedDeltaTime;
 	};
 
 	struct FOutputFrameData
