MoviePipeline: Fixed an issue where sub-frame accumulation could never roll over into whole frames when not using temporal sub-samples. Time Dilation could put you into a situation where you need sub-frame accumulation even without temporal sub-samples.

diff --git a/Engine/Plugins/MovieScene/MovieRenderPipeline/Source/MovieRenderPipelineCore/Private/MoviePipelineTiming.cpp b/Engine/Plugins/MovieScene/MovieRenderPipeline/Source/MovieRenderPipelineCore/Private/MoviePipelineTiming.cpp
--- a/Engine/Plugins/MovieScene/MovieRenderPipeline/Source/MovieRenderPipelineCore/Private/MoviePipelineTiming.cpp
+++ b/Engine/Plugins/MovieScene/MovieRenderPipeline/Source/MovieRenderPipelineCore/Private/MoviePipelineTiming.cpp
@@ -365,7 +365,8 @@ void UMoviePipeline::TickProducingFrames()
 		}
 
 		CachedOutputState.TemporalSampleIndex++;
-
+		const bool bFirstFrame = CachedOutputState.TemporalSampleIndex == 0;
+		
 		if (AntiAliasingSettings->TemporalSampleCount == 1)
 		{
 			// It should always be zero when using no sub-sampling.
@@ -486,7 +487,7 @@ void UMoviePipeline::TickProducingFrames()
 			check(AntiAliasingSettings);
 			check(AntiAliasingSettings->TemporalSampleCount > 1);
 				
-			const bool bFirstFrame = CachedOutputState.TemporalSampleIndex == 0;
+			
 			if (bFirstFrame)
 			{
 				/* The first sub-frame has its timing calculated differently than the rest, because we're trying to skip
@@ -515,13 +516,7 @@ void UMoviePipeline::TickProducingFrames()
 				// Finally, we take the time it is closed plus the sample duration to put us into the new sample.
 				DeltaFrameTime = FrameMetrics.TicksWhileShutterClosed + FrameMetrics.TicksPerSample;
 
-				while (AccumulatedTickSubFrameDeltas >= 1.f)
-				{
-					// We add the deltas on the start of the first sample, since it's the one where 
-					// we jump by a non-consistent amount of time anyways.
-					DeltaFrameTime += FFrameTime(FFrameNumber(1));
-					AccumulatedTickSubFrameDeltas -= 1.f;
-				}
+
 			}
 			else
 			{
@@ -529,6 +524,17 @@ void UMoviePipeline::TickProducingFrames()
 				DeltaFrameTime = FrameMetrics.TicksPerSample;
 			}
 		}
+		
+		if(bFirstFrame)
+		{
+			while (AccumulatedTickSubFrameDeltas >= 1.f)
+			{
+				// We add the deltas on the start of the first sample, since it's the one where 
+				// we jump by a non-consistent amount of time anyways.
+				DeltaFrameTime += FFrameTime(FFrameNumber(1));
+				AccumulatedTickSubFrameDeltas -= 1.f;
+			}
+		}
 
 		
 
