MoviePipeline: Don't allow changing the configuration while a job is running, as it's actively frame-by-frame fetching values from it.

diff --git a/Engine/Plugins/MovieScene/MovieRenderPipeline/Source/MovieRenderPipelineEditor/Private/Widgets/SMoviePipelineQueueEditor.cpp b/Engine/Plugins/MovieScene/MovieRenderPipeline/Source/MovieRenderPipelineEditor/Private/Widgets/SMoviePipelineQueueEditor.cpp
--- a/Engine/Plugins/MovieScene/MovieRenderPipeline/Source/MovieRenderPipelineEditor/Private/Widgets/SMoviePipelineQueueEditor.cpp
+++ b/Engine/Plugins/MovieScene/MovieRenderPipeline/Source/MovieRenderPipelineEditor/Private/Widgets/SMoviePipelineQueueEditor.cpp
@@ -311,6 +311,16 @@ public:
 		return false;
 	}
 
+	bool IsConfigEditingEnabled() const
+	{
+		// Don't allow editing the UI while a job is running as it will change job parameters mid-job!
+		UMoviePipelineQueueSubsystem* Subsystem = GEditor->GetEditorSubsystem<UMoviePipelineQueueSubsystem>();
+		check(Subsystem);
+		const bool bNotRendering = !Subsystem->IsRendering();
+
+		return IsEnabled() && bNotRendering;
+	}
+
 	TOptional<float> GetProgressPercent() const
 	{
 		UMoviePipelineExecutorJob* Job = WeakJob.Get();
@@ -422,7 +432,7 @@ TSharedRef<SWidget> SQueueJobListRow::GenerateWidgetForColumn(const FName& Colum
 	else if (ColumnName == NAME_Settings)
 	{
 		return SNew(SHorizontalBox)
-		.IsEnabled(Item.Get(), &FMoviePipelineQueueJobTreeItem::IsEnabled)
+		.IsEnabled(Item.Get(), &FMoviePipelineQueueJobTreeItem::IsConfigEditingEnabled)
 
 		// Preset Label
 		+ SHorizontalBox::Slot()
@@ -640,6 +650,16 @@ struct FMoviePipelineShotItem : IMoviePipelineQueueTreeItem
 		return false;
 	}
 
+	bool IsConfigEditingEnabled() const
+	{
+		// Don't allow editing the UI while a job is running as it will change job parameters mid-job!
+		UMoviePipelineQueueSubsystem* Subsystem = GEditor->GetEditorSubsystem<UMoviePipelineQueueSubsystem>();
+		check(Subsystem);
+		const bool bNotRendering = !Subsystem->IsRendering();
+
+		return IsEnabled() && bNotRendering;
+	}
+
 	TOptional<float> GetProgressPercent() const
 	{
 		UMoviePipelineExecutorShot* Shot = WeakShot.Get();
@@ -769,7 +789,7 @@ TSharedRef<SWidget> SQueueShotListRow::GenerateWidgetForColumn(const FName& Colu
 	else if (ColumnName == SQueueJobListRow::NAME_Settings)
 	{
 		return SNew(SHorizontalBox)
-		.IsEnabled(Item.Get(), &FMoviePipelineShotItem::IsEnabled)
+		.IsEnabled(Item.Get(), &FMoviePipelineShotItem::IsConfigEditingEnabled)
 
 		// Preset Label
 		+ SHorizontalBox::Slot()
diff --git a/Engine/Plugins/MovieScene/MovieRenderPipeline/Source/MovieRenderPipelineEditor/Private/Widgets/SMoviePipelineQueuePanel.cpp b/Engine/Plugins/MovieScene/MovieRenderPipeline/Source/MovieRenderPipelineEditor/Private/Widgets/SMoviePipelineQueuePanel.cpp
--- a/Engine/Plugins/MovieScene/MovieRenderPipeline/Source/MovieRenderPipelineEditor/Private/Widgets/SMoviePipelineQueuePanel.cpp
+++ b/Engine/Plugins/MovieScene/MovieRenderPipeline/Source/MovieRenderPipelineEditor/Private/Widgets/SMoviePipelineQueuePanel.cpp
@@ -279,6 +279,7 @@ bool SMoviePipelineQueuePanel::IsRenderLocalEnabled() const
 	const UMovieRenderPipelineProjectSettings* ProjectSettings = GetDefault<UMovieRenderPipelineProjectSettings>();
 	const bool bHasExecutor = ProjectSettings->DefaultLocalExecutor != nullptr;
 	const bool bNotRendering = !Subsystem->IsRendering();
+	const bool bConfigWindowIsOpen = WeakEditorWindow.IsValid();
 
 	bool bAtLeastOneJobAvailable = false;
 	for (UMoviePipelineExecutorJob* Job : Subsystem->GetQueue()->GetJobs())
@@ -291,7 +292,7 @@ bool SMoviePipelineQueuePanel::IsRenderLocalEnabled() const
 	}
 
 	const bool bWorldIsActive = GEditor->IsPlaySessionInProgress();
-	return bHasExecutor && bNotRendering && bAtLeastOneJobAvailable && !bWorldIsActive;
+	return bHasExecutor && bNotRendering && bAtLeastOneJobAvailable && !bWorldIsActive && !bConfigWindowIsOpen;
 }
 
 FReply SMoviePipelineQueuePanel::OnRenderRemoteRequested()
@@ -314,6 +315,7 @@ bool SMoviePipelineQueuePanel::IsRenderRemoteEnabled() const
 	const UMovieRenderPipelineProjectSettings* ProjectSettings = GetDefault<UMovieRenderPipelineProjectSettings>();
 	const bool bHasExecutor = ProjectSettings->DefaultRemoteExecutor != nullptr;
 	const bool bNotRendering = !Subsystem->IsRendering();
+	const bool bConfigWindowIsOpen = WeakEditorWindow.IsValid();
 
 	bool bAtLeastOneJobAvailable = false;
 	for (UMoviePipelineExecutorJob* Job : Subsystem->GetQueue()->GetJobs())
@@ -325,7 +327,7 @@ bool SMoviePipelineQueuePanel::IsRenderRemoteEnabled() const
 		}
 	}
 
-	return bHasExecutor && bNotRendering && bAtLeastOneJobAvailable;
+	return bHasExecutor && bNotRendering && bAtLeastOneJobAvailable && !bConfigWindowIsOpen;
 }
 
 void SMoviePipelineQueuePanel::OnJobPresetChosen(TWeakObjectPtr<UMoviePipelineExecutorJob> InJob, TWeakObjectPtr<UMoviePipelineExecutorShot> InShot)
