MoviePipeline: Having {shot_name} or {camera_name} in the format string resolving into an incorrect folder name when clicking the hyperlink in the Output column.

diff --git a/Engine/Plugins/MovieScene/MovieRenderPipeline/Source/MovieRenderPipelineEditor/Private/Widgets/SMoviePipelineQueueEditor.cpp b/Engine/Plugins/MovieScene/MovieRenderPipeline/Source/MovieRenderPipelineEditor/Private/Widgets/SMoviePipelineQueueEditor.cpp
--- a/Engine/Plugins/MovieScene/MovieRenderPipeline/Source/MovieRenderPipelineEditor/Private/Widgets/SMoviePipelineQueueEditor.cpp
+++ b/Engine/Plugins/MovieScene/MovieRenderPipeline/Source/MovieRenderPipelineEditor/Private/Widgets/SMoviePipelineQueueEditor.cpp
@@ -253,8 +253,13 @@ public:
 				FormatString.LeftInline(LastSlashIndex + 1);
 			}
 
+			// By having it swap {camera_name} and {shot_name} with an unresolvable tag, it will
+			// stay in the resolved path and can be removed using the code below.
+			static const FString DummyTag = TEXT("{dontresolvethis}");
 			FMoviePipelineFilenameResolveParams Params;
 			Params.Job = Job;
+			Params.ShotNameOverride = DummyTag;
+			Params.CameraNameOverride = DummyTag;
 
 			FString OutResolvedPath;
 			FMoviePipelineFormatArgs Dummy;
@@ -275,9 +280,11 @@ public:
 			{
 				// Just as a last bit of saftey, we'll trim anything between the { and the preceeding /. This is
 				// in case they did something like Render_{Date}, we wouldn't want to make a folder named Render_.
-				if (OutResolvedPath.FindLastChar(TEXT('/'), LastSlashIndex))
+				// We search backwards from where we found the first { brace, so that will get us the last usable slash.
+				LastSlashIndex = OutResolvedPath.Find(TEXT("/"), ESearchCase::IgnoreCase, ESearchDir::FromEnd, FormatStringToken);
+				if (LastSlashIndex != INDEX_NONE)
 				{
-					OutResolvedPath.LeftInline(LastSlashIndex);
+					OutResolvedPath.LeftInline(LastSlashIndex + 1);
 				}
 			}
 
@@ -1387,4 +1394,4 @@ void SMoviePipelineQueueEditor::OnJobSelectionChanged_Impl(TSharedPtr<IMoviePipe
 	OnJobSelectionChanged.ExecuteIfBound(SelectedJobs);
 }
 
-#undef LOCTEXT_NAMESPACE
\ No newline at end of file
+#undef LOCTEXT_NAMESPACE
