Movie Pipeline: Use GetWorld() instead of InShot->GetWorld()

diff --git a/Engine/Plugins/MovieScene/MovieRenderPipeline/Source/MovieRenderPipelineCore/Private/MoviePipelineRendering.cpp b/Engine/Plugins/MovieScene/MovieRenderPipeline/Source/MovieRenderPipelineCore/Private/MoviePipelineRendering.cpp
--- a/Engine/Plugins/MovieScene/MovieRenderPipeline/Source/MovieRenderPipelineCore/Private/MoviePipelineRendering.cpp
+++ b/Engine/Plugins/MovieScene/MovieRenderPipeline/Source/MovieRenderPipelineCore/Private/MoviePipelineRendering.cpp
@@ -85,7 +85,7 @@ void UMoviePipeline::SetupRenderingPipelineForShot(UMoviePipelineExecutorShot* I
 	// Note how many tiles we wish to render with.
 	BackbufferTileCount = FIntPoint(HighResSettings->TileCount, HighResSettings->TileCount);
 
-	const ERHIFeatureLevel::Type FeatureLevel = InShot->GetWorld()->FeatureLevel;
+	const ERHIFeatureLevel::Type FeatureLevel = GetWorld()->FeatureLevel;
 
 	// Initialize our render pass. This is a copy of the settings to make this less coupled to the Settings UI.
 	const MoviePipeline::FMoviePipelineRenderPassInitSettings RenderPassInitSettings(FeatureLevel, BackbufferResolution, BackbufferTileCount);
