MoviePipeline: Updated example python scripts to invoke the executor through the UMoviePipelineQueueSubsystem so that the automatic suppression of Sequencer's 'Auto-bind to PIE' kicks in.

diff --git a/Engine/Plugins/MovieScene/MovieRenderPipeline/Source/MovieRenderPipelineEditor/Private/MoviePipelineQueueSubsystem.cpp b/Engine/Plugins/MovieScene/MovieRenderPipeline/Source/MovieRenderPipelineEditor/Private/MoviePipelineQueueSubsystem.cpp
--- a/Engine/Plugins/MovieScene/MovieRenderPipeline/Source/MovieRenderPipelineEditor/Private/MoviePipelineQueueSubsystem.cpp
+++ b/Engine/Plugins/MovieScene/MovieRenderPipeline/Source/MovieRenderPipelineEditor/Private/MoviePipelineQueueSubsystem.cpp
@@ -11,16 +11,34 @@ UMoviePipelineExecutorBase* UMoviePipelineQueueSubsystem::RenderQueueWithExecuto
 		return nullptr;
 	}
 
+	ActiveExecutor = NewObject<UMoviePipelineExecutorBase>(this, InExecutorType);
+	RenderQueueWithExecutorInstance(ActiveExecutor);
+	return ActiveExecutor;
+}
+
+void UMoviePipelineQueueSubsystem::RenderQueueWithExecutorInstance(UMoviePipelineExecutorBase* InExecutor)
+{
+	if (!ensureMsgf(!IsRendering(), TEXT("RenderQueueWithExecutor cannot be called while already rendering!")))
+	{
+		FFrame::KismetExecutionMessage(TEXT("Render already in progress."), ELogVerbosity::Error);
+		return;
+	}
+
+	if (!InExecutor)
+	{
+		FFrame::KismetExecutionMessage(TEXT("Invalid executor supplied."), ELogVerbosity::Error);
+		return;
+	}
+
 	ILevelSequenceEditorModule* LevelSequenceEditorModule = FModuleManager::GetModulePtr<ILevelSequenceEditorModule>("LevelSequenceEditor");
 	if (LevelSequenceEditorModule)
 	{
 		LevelSequenceEditorModule->OnComputePlaybackContext().AddUObject(this, &UMoviePipelineQueueSubsystem::OnSequencerContextBinding);
 	}
-		
-	ActiveExecutor = NewObject<UMoviePipelineExecutorBase>(this, InExecutorType);
+
+	ActiveExecutor = InExecutor;
 	ActiveExecutor->OnExecutorFinished().AddUObject(this, &UMoviePipelineQueueSubsystem::OnExecutorFinished);
 	ActiveExecutor->Execute(GetQueue());
-	return ActiveExecutor;
 }
 
 void UMoviePipelineQueueSubsystem::OnSequencerContextBinding(bool& bAllowBinding)
diff --git a/Engine/Plugins/MovieScene/MovieRenderPipeline/Source/MovieRenderPipelineEditor/Public/MoviePipelineQueueSubsystem.h b/Engine/Plugins/MovieScene/MovieRenderPipeline/Source/MovieRenderPipelineEditor/Public/MoviePipelineQueueSubsystem.h
--- a/Engine/Plugins/MovieScene/MovieRenderPipeline/Source/MovieRenderPipelineEditor/Public/MoviePipelineQueueSubsystem.h
+++ b/Engine/Plugins/MovieScene/MovieRenderPipeline/Source/MovieRenderPipelineEditor/Public/MoviePipelineQueueSubsystem.h
@@ -44,6 +44,16 @@ public:
 	UFUNCTION(BlueprintCallable, meta = (DeterminesOutputType = "InExecutorType"), Category = "Movie Render Pipeline|Rendering")
 	UMoviePipelineExecutorBase* RenderQueueWithExecutor(TSubclassOf<UMoviePipelineExecutorBase> InExecutorType);
 
+	/**
+	* Starts processing the current queue with the supplied executor. This starts an async process which
+	* may or may not run in a separate process (or on separate machines), determined by the executor implementation.
+	* The executor should report progress for jobs depending on the implementation.
+	*
+	* @param InExecutor	Instance of a subclass of UMoviePipelineExecutorBase.
+	*/
+	UFUNCTION(BlueprintCallable, Category = "Movie Render Pipeline|Rendering")
+	void RenderQueueWithExecutorInstance(UMoviePipelineExecutorBase* InExecutor);
+
 	/**
 	* Returns true if there is an active executor working on producing a movie. This could be actively rendering frames,
 	* or working on post processing (finalizing file writes, etc.). Use GetActiveExecutor() and query it directly for
