MoviePipeline: Basic setting classes can now be made in Blueprints, which allows creating custom properties visible in the job configuration, and adding custom metadata (for exrs).

diff --git a/Engine/Plugins/MovieScene/MovieRenderPipeline/Source/MovieRenderPipelineCore/Private/MoviePipelineConfigBase.cpp b/Engine/Plugins/MovieScene/MovieRenderPipeline/Source/MovieRenderPipelineCore/Private/MoviePipelineConfigBase.cpp
--- a/Engine/Plugins/MovieScene/MovieRenderPipeline/Source/MovieRenderPipelineCore/Private/MoviePipelineConfigBase.cpp
+++ b/Engine/Plugins/MovieScene/MovieRenderPipeline/Source/MovieRenderPipelineCore/Private/MoviePipelineConfigBase.cpp
@@ -2,6 +2,7 @@
 
 #include "MoviePipelineConfigBase.h"
 #include "MoviePipelineSetting.h"
+#include "MovieRenderPipelineCoreModule.h"
 
 #define LOCTEXT_NAMESPACE "MoviePipelineConfigBase"
 
@@ -29,6 +30,12 @@ void UMoviePipelineConfigBase::CopyFrom(UMoviePipelineConfigBase* InConfig)
 	// Only access the direct Settings array
 	for (UMoviePipelineSetting* Setting : InConfig->Settings)
 	{
+		if (!Setting)
+		{
+			UE_LOG(LogMovieRenderPipeline, Error, TEXT("Null setting found in config: %s - Did you disable a plugin that contained this setting?"), *GetNameSafe(InConfig));
+			continue;
+		}
+
 		UMoviePipelineSetting* Duplicate = Cast<UMoviePipelineSetting>(StaticDuplicateObject(Setting, this, Setting->GetFName()));
 		Duplicate->ValidateState();
 
diff --git a/Engine/Plugins/MovieScene/MovieRenderPipeline/Source/MovieRenderPipelineCore/Private/MoviePipelineMasterConfig.cpp b/Engine/Plugins/MovieScene/MovieRenderPipeline/Source/MovieRenderPipelineCore/Private/MoviePipelineMasterConfig.cpp
--- a/Engine/Plugins/MovieScene/MovieRenderPipeline/Source/MovieRenderPipelineCore/Private/MoviePipelineMasterConfig.cpp
+++ b/Engine/Plugins/MovieScene/MovieRenderPipeline/Source/MovieRenderPipelineCore/Private/MoviePipelineMasterConfig.cpp
@@ -9,6 +9,8 @@
 #include "MoviePipelineQueue.h"
 #include "MoviePipelineUtils.h"
 #include "MoviePipelineSetting.h"
+#include "MovieRenderPipelineCoreModule.h"
+#include "GenericPlatform/GenericPlatformProcess.h"
 
 #define LOCTEXT_NAMESPACE "MoviePipelineMasterConfig"
 
@@ -38,6 +40,12 @@ void UMoviePipelineMasterConfig::InitializeTransientSettings()
 	// Now remove any from the list that we already have a user setting for
 	for (const UMoviePipelineSetting* ExistingSetting : GetUserSettings())
 	{
+		if (!ExistingSetting)
+		{
+			UE_LOG(LogMovieRenderPipeline, Error, TEXT("Null setting found in config: %s - Did you disable a plugin that contained this setting?"), *GetNameSafe(this));
+			continue;
+		}
+
 		AllSettingClasses.RemoveSwap(ExistingSetting->GetClass());
 	}
 
@@ -117,7 +125,6 @@ void UMoviePipelineMasterConfig::GetFormatArguments(FMoviePipelineFormatArgs& In
 		InOutFormatArgs.FileMetadata.Add(TEXT("unreal/jobName"), JobName);
 		InOutFormatArgs.FileMetadata.Add(TEXT("unreal/frameRate"), FString::SanitizeFloat(FrameRate));
 
-		
 		// Normally these are filled when resolving the file name by the job (so that the time is shared), but stub them in here so
 		// they show up in the UI with a value.
 		FDateTime CurrentTime = FDateTime::Now();
@@ -126,6 +133,7 @@ void UMoviePipelineMasterConfig::GetFormatArguments(FMoviePipelineFormatArgs& In
 		InOutFormatArgs.FilenameArguments.Add(TEXT("month"), CurrentTime.ToString(TEXT("%m")));
 		InOutFormatArgs.FilenameArguments.Add(TEXT("day"), CurrentTime.ToString(TEXT("%d")));
 		InOutFormatArgs.FilenameArguments.Add(TEXT("time"), CurrentTime.ToString(TEXT("%H.%M.%S")));
+		InOutFormatArgs.FilenameArguments.Add(TEXT("job_author"), FPlatformProcess::UserName(false));
 		
 		// Let the output state fill out some too, since its the keeper of the information.
 		UMoviePipelineOutputSetting* OutputSettings = FindSetting<UMoviePipelineOutputSetting>();
@@ -195,71 +203,6 @@ TRange<FFrameNumber> UMoviePipelineMasterConfig::GetEffectivePlaybackRange(const
 	return InSequence->GetMovieScene()->GetPlaybackRange();
 }
 
-
-/*bool UMoviePipelineShotConfig::ValidateConfig(TArray<FText>& OutValidationErrors) const{
-	bool bValidPipeline = true;
-	OutValidationErrors.Reset();
-	
-	// Check that we have a sequence to render
-	if(!Sequence.IsValid())
-	{
-		OutValidationErrors.Add(LOCTEXT("InvalidInputSequence", "No sequence to render has been specified."));
-		bValidPipeline = false;
-	}
-	
-	// Check that we have a valid output resolution
-	if(OutputResolution.X == 0 || OutputResolution.Y == 0)
-	{
-		OutValidationErrors.Add(LOCTEXT("InvalidOutputResolution", "Output Resolution must be greater than zero."));
-		bValidPipeline = false;
-	}
-	
-	// Check for a non-zero framerate
-	if(OutputFrameRate.AsDecimal() == 0.0)
-	{
-		OutValidationErrors.Add(LOCTEXT("InvalidOutputFrameRate", "Output FrameRate must be greater than zero."));
-		bValidPipeline = false;
-	}
-	
-	// Output Directory can probably be ignored. The folder path you output to may not exist yet anyways.
-	
-	// Give each setting a chance to validate the pipeline
-	for(const UMoviePipelineSetting* Setting : Settings)
-	{
-		bValidPipeline &= Setting->ValidatePipeline(OutValidationErrors);
-	}
-	
-	// ToDo-MoviePipeline: Reflect over all settings classes and make sure all required classes have been added.
-	// ToDo-MoviePipeline: Scan for any shots/sections that don't start on whole frames issue warning.
-	
-	// Give each Input Buffer a chance to validate the pipeline
-	for(const UMoviePipelineRenderPass* InputBuffer : InputBuffers)
-	{
-		bValidPipeline &= InputBuffer->ValidatePipeline(OutValidationErrors);
-	}
-	
-	// Give each output container a chance to validate the pipeline
-	for(const UMoviePipelineOutput* OutputContainer : OutputContainers)
-	{
-		bValidPipeline &= OutputContainer->ValidatePipeline(OutValidationErrors);
-	}
-
-	// Ensure we have at least one input buffer and one output buffer.
-	if (InputBuffers.Num() == 0)
-	{
-		OutValidationErrors.Add(LOCTEXT("NoInputBuffer", "Must specify at least one Input Buffer to capture."));
-		bValidPipeline = false;
-	}
-
-	if (OutputContainers.Num() == 0)
-	{
-		OutValidationErrors.Add(LOCTEXT("NoOutputContainer", "Must specify at least one Output Container to write to."));
-		bValidPipeline = false;
-	}
-	
-	return bValidPipeline;
-}*/
-
 UMoviePipelineMasterConfig::UMoviePipelineMasterConfig()
 {
 	// Always add at least the output settings block since having a framerate/directory/etc. is critical.
@@ -281,6 +224,12 @@ TArray<UMoviePipelineSetting*> UMoviePipelineMasterConfig::GetAllSettings(const
 	// User settings were explicitly added by the user via the UI or scripting.
 	for (UMoviePipelineSetting* Setting : GetUserSettings())
 	{
+		if (!Setting)
+		{
+			UE_LOG(LogMovieRenderPipeline, Error, TEXT("Null setting found in config: %s - Did you disable a plugin that contained this setting?"), *GetNameSafe(this));
+			continue;
+		}
+
 		if (bIncludeDisabledSettings || Setting->IsEnabled())
 		{
 			OutSettings.Add(Setting);
diff --git a/Engine/Plugins/MovieScene/MovieRenderPipeline/Source/MovieRenderPipelineCore/Private/MoviePipelineUtils.cpp b/Engine/Plugins/MovieScene/MovieRenderPipeline/Source/MovieRenderPipelineCore/Private/MoviePipelineUtils.cpp
--- a/Engine/Plugins/MovieScene/MovieRenderPipeline/Source/MovieRenderPipelineCore/Private/MoviePipelineUtils.cpp
+++ b/Engine/Plugins/MovieScene/MovieRenderPipeline/Source/MovieRenderPipelineCore/Private/MoviePipelineUtils.cpp
@@ -32,6 +32,7 @@ namespace UE
 
 			FARFilter Filter;
 			Filter.ClassNames.Add(UMoviePipelineSetting::StaticClass()->GetFName());
+			Filter.ClassNames.Add(UBlueprint::StaticClass()->GetFName());
 
 			// Include any Blueprint based objects as well, this includes things like Blutilities, UMG, and GameplayAbility objects
 			Filter.bRecursiveClasses = true;
@@ -41,18 +42,47 @@ namespace UE
 
 			for (const FAssetData& Data : ClassList)
 			{
-				UClass* Class = Data.GetClass();
-				if (Class)
+				// We don't want to load every blueprint class to find out if it's an MRQ related one, so we'll look up their native class
+				// based on the asset registry metadata before deciding if it needs to be loaded.
+				UClass* ParentClass = nullptr;
+				FString ParentClassName;
+
+				if (!Data.GetTagValue(FBlueprintTags::NativeParentClassPath, ParentClassName))
+				{
+					Data.GetTagValue(FBlueprintTags::ParentClassPath, ParentClassName);
+				}
+
+				if (!ParentClassName.IsEmpty())
 				{
-					Classes.Add(Class);
+					UObject* Outer = nullptr;
+					ResolveName(Outer, ParentClassName, false, false);
+
+					ParentClass = FindObject<UClass>(ANY_PACKAGE, *ParentClassName);
+
+					if (!ParentClass->IsChildOf(UMoviePipelineSetting::StaticClass()))
+					{
+						continue;
+					}
 				}
+
+				// Now attempt to load it as it's an MRQ related setting. This will cause it to show up in the class iterator
+				// below.
+				UBlueprint* LoadedBlueprint = Cast<UBlueprint>(Data.ToSoftObjectPath().TryLoad());
 			}
 
+			// Now iterate through the loaded classes.
 			for (TObjectIterator<UClass> ClassIterator; ClassIterator; ++ClassIterator)
 			{
 				if (ClassIterator->IsChildOf(UMoviePipelineSetting::StaticClass()) && !ClassIterator->HasAnyClassFlags(CLASS_Abstract | CLASS_Deprecated | CLASS_NewerVersionExists))
 				{
-					Classes.Add(*ClassIterator);
+					// While a blueprint is being compiled, there can be some transient instances of the class. We don't want these to show up in the list
+					// so just manually skip over them as they will go away next time GC is run anyways.
+					if (ClassIterator->GetName().StartsWith(TEXT("SKEL_")) || ClassIterator->GetName().StartsWith(TEXT("REINST_")))
+					{
+						continue;
+					}
+
+					Classes.AddUnique(*ClassIterator);
 				}
 			}
 
diff --git a/Engine/Plugins/MovieScene/MovieRenderPipeline/Source/MovieRenderPipelineCore/Public/MoviePipelineConfigBase.h b/Engine/Plugins/MovieScene/MovieRenderPipeline/Source/MovieRenderPipelineCore/Public/MoviePipelineConfigBase.h
--- a/Engine/Plugins/MovieScene/MovieRenderPipeline/Source/MovieRenderPipelineCore/Public/MoviePipelineConfigBase.h
+++ b/Engine/Plugins/MovieScene/MovieRenderPipeline/Source/MovieRenderPipelineCore/Public/MoviePipelineConfigBase.h
@@ -55,7 +55,7 @@ public:
 
 		for (UMoviePipelineSetting* Setting : AllSettings)
 		{
-			if (Setting->IsA<SettingType>() && (Setting->IsEnabled() || bIncludeDisabledSettings))
+			if (Setting && Setting->IsA<SettingType>() && (Setting->IsEnabled() || bIncludeDisabledSettings))
 			{
 				FoundSettings.Add(Cast<SettingType>(Setting));
 			}
diff --git a/Engine/Plugins/MovieScene/MovieRenderPipeline/Source/MovieRenderPipelineCore/Public/MoviePipelineSetting.h b/Engine/Plugins/MovieScene/MovieRenderPipeline/Source/MovieRenderPipelineCore/Public/MoviePipelineSetting.h
--- a/Engine/Plugins/MovieScene/MovieRenderPipeline/Source/MovieRenderPipelineCore/Public/MoviePipelineSetting.h
+++ b/Engine/Plugins/MovieScene/MovieRenderPipeline/Source/MovieRenderPipelineCore/Public/MoviePipelineSetting.h
@@ -20,7 +20,7 @@ enum class EMoviePipelineValidationState : uint8
 /**
 * A base class for all Movie Render Pipeline settings.
 */
-UCLASS(Blueprintable, Abstract)
+UCLASS(BlueprintType, Abstract)
 class MOVIERENDERPIPELINECORE_API UMoviePipelineSetting : public UObject
 {
 	GENERATED_BODY()
diff --git a/Engine/Plugins/MovieScene/MovieRenderPipeline/Source/MovieRenderPipelineCore/Public/MoviePipelineSettingBlueprintBase.h b/Engine/Plugins/MovieScene/MovieRenderPipeline/Source/MovieRenderPipelineCore/Public/MoviePipelineSettingBlueprintBase.h
new file mode 100644
--- /dev/null
+++ b/Engine/Plugins/MovieScene/MovieRenderPipeline/Source/MovieRenderPipelineCore/Public/MoviePipelineSettingBlueprintBase.h
@@ -0,0 +1,84 @@
+// Copyright Epic Games, Inc. All Rights Reserved.
+#pragma once
+
+#include "MoviePipelineSetting.h"
+#include "MovieRenderPipelineDataTypes.h"
+#include "MoviePipelineSettingBlueprintBase.generated.h"
+
+
+/**
+* A base class for all Movie Render Pipeline settings which can be implemented in Blueprints. This features
+* a slightly different API than the regular UMoviePipelineSetting to make the Blueprint integration nicer
+* without breaking the C++ API backwards compatibility.
+*/
+UCLASS(Blueprintable, Abstract)
+class MOVIERENDERPIPELINECORE_API UMoviePipelineSetting_BlueprintBase : public UMoviePipelineSetting
+{
+	GENERATED_BODY()
+		
+public:
+	UMoviePipelineSetting_BlueprintBase()
+	{
+		CategoryText = NSLOCTEXT("MovieRenderPipelineBP", "DefaultCategoryName_Text", "Custom Settings");
+		bIsValidOnMaster = true;
+		bIsValidOnShots = true;
+		bCanBeDisabled = true;
+	}
+	
+	// Setup
+	virtual void SetupForPipelineImpl(UMoviePipeline* InPipeline) override 
+	{
+		ReceiveSetupForPipelineImpl(InPipeline);
+		FCoreDelegates::OnBeginFrame.AddUObject(this, &UMoviePipelineSetting_BlueprintBase::OnEngineTickBeginFrame);
+	}
+	
+	UFUNCTION(BlueprintImplementableEvent, meta=(DisplayName = "SetupForPipeline"))
+	void ReceiveSetupForPipelineImpl(UMoviePipeline* InPipeline);
+	
+	// Teardown
+	virtual void TeardownForPipelineImpl(UMoviePipeline* InPipeline) override 
+	{ 
+		FCoreDelegates::OnBeginFrame.RemoveAll(this);
+		ReceiveTeardownForPipelineImpl(InPipeline); 
+	}
+	
+	UFUNCTION(BlueprintImplementableEvent, meta=(DisplayName="TeardownForPipeline"))
+	void ReceiveTeardownForPipelineImpl(UMoviePipeline* InPipeline);
+	
+	// Format Arguments & EXR Metadata
+	virtual void GetFormatArguments(FMoviePipelineFormatArgs& InOutFormatArgs) const override { InOutFormatArgs = ReceiveGetFormatArguments(InOutFormatArgs); }
+	UFUNCTION(BlueprintImplementableEvent, meta=(DisplayName="GetFormatArgumentsAndMetadata"))
+	FMoviePipelineFormatArgs ReceiveGetFormatArguments(UPARAM(ref) FMoviePipelineFormatArgs& InOutFormatArgs) const;
+	
+	// Tick
+	UFUNCTION(BlueprintImplementableEvent)
+	void OnEngineTickBeginFrame();
+	
+public:
+#if WITH_EDITOR
+	/** Warning: This gets called on the CDO of the object */
+	virtual FText GetCategoryText() const override { return CategoryText; }
+	
+	// UI Footer Text
+	virtual FText GetFooterText(UMoviePipelineExecutorJob* InJob) const override { return ReceiveGetFooterText(InJob); } 
+	UFUNCTION(BlueprintImplementableEvent, meta=(DisplayName="GetFooterText"))
+	FText ReceiveGetFooterText(UMoviePipelineExecutorJob* InJob) const;
+		
+	virtual bool CanBeDisabled() const override { return bCanBeDisabled; }
+#endif
+	virtual bool IsValidOnMaster() const override { return bIsValidOnMaster; }
+	virtual bool IsValidOnShots() const override { return bIsValidOnShots; }
+
+protected:
+	UPROPERTY(EditDefaultsOnly, Category = "Movie Pipeline")
+	FText CategoryText;
+	
+	UPROPERTY(EditDefaultsOnly, Category = "Movie Pipeline")
+	bool bIsValidOnMaster;
+	
+	UPROPERTY(EditDefaultsOnly, Category = "Movie Pipeline")
+	bool bIsValidOnShots;
+	
+	UPROPERTY(EditDefaultsOnly, Category = "Movie Pipeline")
+	bool bCanBeDisabled;
+};
\ No newline at end of file
diff --git a/Engine/Plugins/MovieScene/MovieRenderPipeline/Source/MovieRenderPipelineEditor/Private/Widgets/SMoviePipelineConfigEditor.cpp b/Engine/Plugins/MovieScene/MovieRenderPipeline/Source/MovieRenderPipelineEditor/Private/Widgets/SMoviePipelineConfigEditor.cpp
--- a/Engine/Plugins/MovieScene/MovieRenderPipeline/Source/MovieRenderPipelineEditor/Private/Widgets/SMoviePipelineConfigEditor.cpp
+++ b/Engine/Plugins/MovieScene/MovieRenderPipeline/Source/MovieRenderPipelineEditor/Private/Widgets/SMoviePipelineConfigEditor.cpp
@@ -43,6 +43,7 @@
 #include "IDetailsView.h"
 #include "MovieRenderPipelineStyle.h"
 #include "FrameNumberDetailsCustomization.h"
+#include "Editor.h"
 
 #define LOCTEXT_NAMESPACE "SMoviePipelineEditor"
 
@@ -145,6 +146,9 @@ void SMoviePipelineConfigEditor::Construct(const FArguments& InArgs)
 			]
 		]
 	];
+
+	// Register a callback so that we can refresh the details panel after it is reinstanced.
+	GEditor->OnBlueprintReinstanced().AddSP(this, &SMoviePipelineConfigEditor::OnBlueprintReinstanced);
 }
 
 TSharedRef<SWidget> SMoviePipelineConfigEditor::MakeAddSettingButton()
@@ -196,6 +200,12 @@ TSharedRef<SWidget> SMoviePipelineConfigEditor::MakeAddSettingButton()
 }
 PRAGMA_ENABLE_OPTIMIZATION
 
+void SMoviePipelineConfigEditor::OnBlueprintReinstanced()
+{
+	SettingsWidget->InvalidateCachedSettingsSerialNumber();
+	bRequestDetailsRefresh = true;
+}
+
 void SMoviePipelineConfigEditor::CheckForNewSettingsObject()
 {
 	UMoviePipelineConfigBase* NewMoviePipeline = PipelineConfigAttribute.Get();
diff --git a/Engine/Plugins/MovieScene/MovieRenderPipeline/Source/MovieRenderPipelineEditor/Private/Widgets/SMoviePipelineConfigEditor.h b/Engine/Plugins/MovieScene/MovieRenderPipeline/Source/MovieRenderPipelineEditor/Private/Widgets/SMoviePipelineConfigEditor.h
--- a/Engine/Plugins/MovieScene/MovieRenderPipeline/Source/MovieRenderPipelineEditor/Private/Widgets/SMoviePipelineConfigEditor.h
+++ b/Engine/Plugins/MovieScene/MovieRenderPipeline/Source/MovieRenderPipelineEditor/Private/Widgets/SMoviePipelineConfigEditor.h
@@ -72,6 +72,7 @@ private:
 
 	EVisibility IsValidationWarningVisible() const;
 	FText GetValidationWarningText() const;
+	void OnBlueprintReinstanced();
 private:
 
 	bool bRequestDetailsRefresh;
diff --git a/Engine/Plugins/MovieScene/MovieRenderPipeline/Source/MovieRenderPipelineEditor/Private/Widgets/SMoviePipelineConfigSettings.h b/Engine/Plugins/MovieScene/MovieRenderPipeline/Source/MovieRenderPipelineEditor/Private/Widgets/SMoviePipelineConfigSettings.h
--- a/Engine/Plugins/MovieScene/MovieRenderPipeline/Source/MovieRenderPipelineEditor/Private/Widgets/SMoviePipelineConfigSettings.h
+++ b/Engine/Plugins/MovieScene/MovieRenderPipeline/Source/MovieRenderPipelineEditor/Private/Widgets/SMoviePipelineConfigSettings.h
@@ -40,7 +40,7 @@ public:
 
 	void GetSelectedSettings(TArray<UMoviePipelineSetting*>& OutSettings) const;
 	void SetSelectedSettings(const TArray<UMoviePipelineSetting*>& Settings);
-
+	void InvalidateCachedSettingsSerialNumber() { CachedSettingsSerialNumber = -1; }
 private:
 
 	virtual FReply OnKeyDown(const FGeometry& MyGeometry, const FKeyEvent& InKeyEvent) override;
