MoviePipeline: Switched FRand() for FRandomStream in png/jpeg image quantization. FRand() is [0, 1] inclusive which occasionally caused alpha pixels to go from 0/255 to 1/255 causing a speckling. FRandomStream() is much better as GetFraction() returns [0, 1) exclusive, so the cast to uint8 will always floor.

diff --git a/Engine/Plugins/MovieScene/MovieRenderPipeline/Source/MovieRenderPipelineCore/Private/MoviePipelineImageQuantization.cpp b/Engine/Plugins/MovieScene/MovieRenderPipeline/Source/MovieRenderPipelineCore/Private/MoviePipelineImageQuantization.cpp
--- a/Engine/Plugins/MovieScene/MovieRenderPipeline/Source/MovieRenderPipelineCore/Private/MoviePipelineImageQuantization.cpp
+++ b/Engine/Plugins/MovieScene/MovieRenderPipeline/Source/MovieRenderPipelineCore/Private/MoviePipelineImageQuantization.cpp
@@ -160,9 +160,10 @@ static TArray<FColor> ConvertLinearTosRGB8bppViaLookupTable(FFloat16Color* InCol
 {
 	TArray<float> sRGBTable = GenerateSRGBTableFloat16toFloat();
 
+	static FRandomStream RandStream(GFrameCounter);
+
 	// This quantization uses a table of float values from 0.0-255.0, and does the proper rounding for each pixel.
 	// The proper rounding is done by adding a random value from 0.0-1.0 instead of 0.5 to each value returned from the table.
-
 	// Convert all of our pixels.
 	TArray<FColor> OutsRGBData;
 	OutsRGBData.SetNumUninitialized(InCount);
@@ -175,12 +176,12 @@ static TArray<FColor> ConvertLinearTosRGB8bppViaLookupTable(FFloat16Color* InCol
 	{
 		// Avoid the bounds checking of TArray[]
 		FColor* OutColor = &OutData[PixelIndex];
-		OutColor->R = (uint8)FMath::FloorToInt(sRGBTableData[InColor[PixelIndex].R.Encoded] + FMath::FRand());
-		OutColor->G = (uint8)FMath::FloorToInt(sRGBTableData[InColor[PixelIndex].G.Encoded] + FMath::FRand());
-		OutColor->B = (uint8)FMath::FloorToInt(sRGBTableData[InColor[PixelIndex].B.Encoded] + FMath::FRand());
+		OutColor->R = (uint8)FMath::FloorToInt(sRGBTableData[InColor[PixelIndex].R.Encoded] + RandStream.GetFraction());
+		OutColor->G = (uint8)FMath::FloorToInt(sRGBTableData[InColor[PixelIndex].G.Encoded] + RandStream.GetFraction());
+		OutColor->B = (uint8)FMath::FloorToInt(sRGBTableData[InColor[PixelIndex].B.Encoded] + RandStream.GetFraction());
 
 		// Alpha doesn't get sRGB conversion, it just gets linearly converted to 8 bit. Flooring avoids an extra branch for Round.
-		OutColor->A = (uint8)FMath::Clamp(FMath::FloorToInt((InColor[PixelIndex].A * 255.f) + FMath::FRand()), 0, 255);
+		OutColor->A = (uint8)FMath::Clamp(FMath::FloorToInt((InColor[PixelIndex].A * 255.f) + RandStream.GetFraction()), 0, 255);
 	}
 
 	return OutsRGBData;
