MoviePipeline: Changed the EXR/Metadata gathering function from a BlueprintImplementableEvent to BlueprintNativeEvent so that we could provide a default implementation that passed the structure back.

Without this, if you created a BP setting and did not override ReceiveGetFormatArguments then the rest of the metadata/filename format arguments would disappear without cause which was very confusing.

diff --git a/Engine/Plugins/MovieScene/MovieRenderPipeline/Source/MovieRenderPipelineCore/Public/MoviePipelineSettingBlueprintBase.h b/Engine/Plugins/MovieScene/MovieRenderPipeline/Source/MovieRenderPipelineCore/Public/MoviePipelineSettingBlueprintBase.h
--- a/Engine/Plugins/MovieScene/MovieRenderPipeline/Source/MovieRenderPipelineCore/Public/MoviePipelineSettingBlueprintBase.h
+++ b/Engine/Plugins/MovieScene/MovieRenderPipeline/Source/MovieRenderPipelineCore/Public/MoviePipelineSettingBlueprintBase.h
@@ -47,8 +47,10 @@ public:
 	
 	// Format Arguments & EXR Metadata
 	virtual void GetFormatArguments(FMoviePipelineFormatArgs& InOutFormatArgs) const override { InOutFormatArgs = ReceiveGetFormatArguments(InOutFormatArgs); }
-	UFUNCTION(BlueprintImplementableEvent, meta=(DisplayName="GetFormatArgumentsAndMetadata"))
+	UFUNCTION(BlueprintNativeEvent, meta=(DisplayName="GetFormatArgumentsAndMetadata"))
 	FMoviePipelineFormatArgs ReceiveGetFormatArguments(UPARAM(ref) FMoviePipelineFormatArgs& InOutFormatArgs) const;
+	// Native Implementation so that if a blueprint class doesn't override this function we still pass on the previous arguments to the rest of the chain.
+	FMoviePipelineFormatArgs ReceiveGetFormatArguments_Implementation(FMoviePipelineFormatArgs& InOutFormatArgs) const { return InOutFormatArgs; }
 	
 	// Tick
 	UFUNCTION(BlueprintImplementableEvent)
