Movie Render Queue: Drag drop jobs to reorder

diff --git a/Engine/Plugins/MovieScene/MovieRenderPipeline/Source/MovieRenderPipelineCore/Private/MoviePipelineQueue.cpp b/Engine/Plugins/MovieScene/MovieRenderPipeline/Source/MovieRenderPipelineCore/Private/MoviePipelineQueue.cpp
--- a/Engine/Plugins/MovieScene/MovieRenderPipeline/Source/MovieRenderPipelineCore/Private/MoviePipelineQueue.cpp
+++ b/Engine/Plugins/MovieScene/MovieRenderPipeline/Source/MovieRenderPipelineCore/Private/MoviePipelineQueue.cpp
@@ -93,6 +93,34 @@ void UMoviePipelineQueue::CopyFrom(UMoviePipelineQueue* InQueue)
 	QueueSerialNumber++;
 }
 
+void UMoviePipelineQueue::SetJobIndex(UMoviePipelineExecutorJob* InJob, int32 Index)
+{
+#if WITH_EDITOR
+	Modify();
+#endif
+
+	int32 CurrentIndex = INDEX_NONE;
+	Jobs.Find(InJob, CurrentIndex);
+
+	if (CurrentIndex == INDEX_NONE)
+	{
+		FFrame::KismetExecutionMessage(*FString::Printf(TEXT("Cannot find Job %s in queue"), *InJob->JobName), ELogVerbosity::Error);
+		return;
+	}
+
+	if (Index > CurrentIndex)
+	{
+		--Index;
+	}
+
+	Jobs.Remove(InJob);
+	Jobs.Insert(InJob, Index);
+
+	// Ensure the serial number gets bumped at least once so the UI refreshes in case
+	// the queue we are copying from was empty.
+	QueueSerialNumber++;
+}
+
 #if WITH_EDITOR
 void UMoviePipelineExecutorJob::PostEditChangeProperty(FPropertyChangedEvent& PropertyChangedEvent)
 {
diff --git a/Engine/Plugins/MovieScene/MovieRenderPipeline/Source/MovieRenderPipelineCore/Public/MoviePipelineQueue.h b/Engine/Plugins/MovieScene/MovieRenderPipeline/Source/MovieRenderPipelineCore/Public/MoviePipelineQueue.h
--- a/Engine/Plugins/MovieScene/MovieRenderPipeline/Source/MovieRenderPipelineCore/Public/MoviePipelineQueue.h
+++ b/Engine/Plugins/MovieScene/MovieRenderPipeline/Source/MovieRenderPipelineCore/Public/MoviePipelineQueue.h
@@ -429,6 +429,10 @@ public:
 	UFUNCTION(BlueprintCallable, Category = "Movie Render Pipeline")
 	void CopyFrom(UMoviePipelineQueue* InQueue);
 	
+	/* Set the index of the given job */
+	UFUNCTION(BlueprintCallable, Category = "Movie Render Pipeline")
+	void SetJobIndex(UMoviePipelineExecutorJob* InJob, int32 Index);
+	
 	/**
 	 * Retrieve the serial number that is incremented when a job is added or removed from this list.
 	 * @note: This field is not serialized, and not copied along with UObject duplication.
diff --git a/Engine/Plugins/MovieScene/MovieRenderPipeline/Source/MovieRenderPipelineEditor/Private/Widgets/SMoviePipelineQueueEditor.cpp b/Engine/Plugins/MovieScene/MovieRenderPipeline/Source/MovieRenderPipelineEditor/Private/Widgets/SMoviePipelineQueueEditor.cpp
--- a/Engine/Plugins/MovieScene/MovieRenderPipeline/Source/MovieRenderPipelineEditor/Private/Widgets/SMoviePipelineQueueEditor.cpp
+++ b/Engine/Plugins/MovieScene/MovieRenderPipeline/Source/MovieRenderPipelineEditor/Private/Widgets/SMoviePipelineQueueEditor.cpp
@@ -87,6 +87,11 @@ public:
 	void Construct(const FArguments& InArgs, const TSharedRef<STableViewBase>& OwnerTable);
 	virtual TSharedRef<SWidget> GenerateWidgetForColumn(const FName& ColumnName) override;
 
+protected:
+	FReply OnDragDetected(const FGeometry& InGeometry, const FPointerEvent& InPointerEvent);
+	TOptional<EItemDropZone> OnCanAcceptDrop(const FDragDropEvent& DragDropEvent, EItemDropZone InItemDropZone, TSharedPtr<IMoviePipelineQueueTreeItem> InItem);
+	FReply OnAcceptDrop(const FDragDropEvent& DragDropEvent, EItemDropZone InItemDropZone, TSharedPtr<IMoviePipelineQueueTreeItem> InItem);
+
 private:
 	FOnMoviePipelineEditConfig OnEditConfigRequested;
 };
@@ -409,7 +414,11 @@ void SQueueJobListRow::Construct(const FArguments& InArgs, const TSharedRef<STab
 	OnEditConfigRequested = InArgs._OnEditConfigRequested;
 
 	FSuperRowType::FArguments SuperArgs = FSuperRowType::FArguments();
-	FSuperRowType::Construct(SuperArgs, OwnerTable);
+	FSuperRowType::Construct(SuperArgs 
+		.OnDragDetected(this, &SQueueJobListRow::OnDragDetected)
+		.OnCanAcceptDrop(this, &SQueueJobListRow::OnCanAcceptDrop)
+		.OnAcceptDrop(this, &SQueueJobListRow::OnAcceptDrop),		
+		OwnerTable);
 }
 
 TSharedRef<SWidget> SQueueJobListRow::GenerateWidgetForColumn(const FName& ColumnName)
@@ -537,6 +546,102 @@ const FName SQueueJobListRow::NAME_Settings = FName(TEXT("Settings"));
 const FName SQueueJobListRow::NAME_Output = FName(TEXT("Output"));
 const FName SQueueJobListRow::NAME_Status = FName(TEXT("Status"));
 
+/**
+ * This drag drop operation allows us to move around queues in the widget tree
+ */
+class FQueueJobListDragDropOp : public FDecoratedDragDropOp
+{
+public:
+	DRAG_DROP_OPERATOR_TYPE(FQueueJobListDragDropOp, FDecoratedDragDropOp)
+
+	/** The template to create an instance */
+	TSharedPtr<IMoviePipelineQueueTreeItem> ListItem;
+
+	/** Constructs the drag drop operation */
+	static TSharedRef<FQueueJobListDragDropOp> New(const TSharedPtr<IMoviePipelineQueueTreeItem>& InListItem, FText InDragText)
+	{
+		TSharedRef<FQueueJobListDragDropOp> Operation = MakeShared<FQueueJobListDragDropOp>();
+		Operation->ListItem = InListItem;
+		Operation->DefaultHoverText = InDragText;
+		Operation->CurrentHoverText = InDragText;
+		Operation->Construct();
+
+		return Operation;
+	}
+};
+
+/** Called whenever a drag is detected by the tree view. */
+FReply SQueueJobListRow::OnDragDetected(const FGeometry& InGeometry, const FPointerEvent& InPointerEvent)
+{
+	if (Item.IsValid())
+	{
+		FText DefaultText = LOCTEXT("DefaultDragDropFormat", "Move 1 item(s)");
+		return FReply::Handled().BeginDragDrop(FQueueJobListDragDropOp::New(Item, DefaultText));
+	}
+	return FReply::Unhandled();
+}
+
+/** Called to determine whether a current drag operation is valid for this row. */
+TOptional<EItemDropZone> SQueueJobListRow::OnCanAcceptDrop(const FDragDropEvent& DragDropEvent, EItemDropZone InItemDropZone, TSharedPtr<IMoviePipelineQueueTreeItem> InItem)
+{
+	TSharedPtr<FQueueJobListDragDropOp> DragDropOp = DragDropEvent.GetOperationAs<FQueueJobListDragDropOp>();
+	if (DragDropOp.IsValid())
+	{
+		if (InItemDropZone == EItemDropZone::OntoItem)
+		{
+			DragDropOp->CurrentIconBrush = FEditorStyle::GetBrush(TEXT("Graph.ConnectorFeedback.Error"));
+		}
+		else
+		{
+			DragDropOp->CurrentIconBrush = FEditorStyle::GetBrush(TEXT("Graph.ConnectorFeedback.Ok"));
+		}
+		return InItemDropZone;
+	}
+	return TOptional<EItemDropZone>();
+}
+
+/** Called to complete a drag and drop onto this drop. */
+FReply SQueueJobListRow::OnAcceptDrop(const FDragDropEvent& DragDropEvent, EItemDropZone InItemDropZone, TSharedPtr<IMoviePipelineQueueTreeItem> InItem)
+{
+	TSharedPtr<FQueueJobListDragDropOp> DragDropOp = DragDropEvent.GetOperationAs<FQueueJobListDragDropOp>();
+	if (!DragDropOp.IsValid() || !DragDropOp->ListItem.IsValid())
+	{
+		return FReply::Unhandled();
+	}
+
+	TSharedPtr<FMoviePipelineQueueJobTreeItem> DragDropJobItem = StaticCastSharedPtr<FMoviePipelineQueueJobTreeItem>(DragDropOp->ListItem);
+	UMoviePipelineExecutorJob* DragDropJob = DragDropJobItem->GetOwningJob();
+
+	TSharedPtr<FMoviePipelineQueueJobTreeItem> JobItem = StaticCastSharedPtr<FMoviePipelineQueueJobTreeItem>(InItem);
+	UMoviePipelineExecutorJob* Job = JobItem->GetOwningJob();
+
+	UMoviePipelineQueue* ActiveQueue = GEditor->GetEditorSubsystem<UMoviePipelineQueueSubsystem>()->GetQueue();
+	if (!ActiveQueue || !DragDropJob || !Job)
+	{
+		return FReply::Unhandled();
+	}
+	
+	int32 Index = 0;
+	ActiveQueue->GetJobs().Find(Job, Index);
+
+	if (InItemDropZone == EItemDropZone::BelowItem)
+	{
+		++Index;
+		if (Index > ActiveQueue->GetJobs().Num())
+		{
+			Index = ActiveQueue->GetJobs().Num()-1;
+		}
+	}
+
+	FScopedTransaction Transaction(LOCTEXT("ReorderJob_Transaction", "Reorder Job"));
+
+	ActiveQueue->Modify();
+
+	ActiveQueue->SetJobIndex(DragDropJob, Index);
+
+	return FReply::Handled();
+}
+
 class SQueueShotListRow : public SMultiColumnTableRow<TSharedPtr<IMoviePipelineQueueTreeItem>>
 {
 public:
