MoviePipeline: Fixed a scenario where a master sequence could specify a sub-sequence shift that pushed the evaluated range entirely outside of the inner PlaybackRange. This technically evaluates correctly so MRQ was in the wrong by not recognizing it as a correct field.

This change removes some code (that was written before the code that stayed), as the remaining code actually handles both cases in one go without the need for two separate checks.

diff --git a/Engine/Plugins/MovieScene/MovieRenderPipeline/Source/MovieRenderPipelineCore/Private/MoviePipelineBlueprintLibrary.cpp b/Engine/Plugins/MovieScene/MovieRenderPipeline/Source/MovieRenderPipelineCore/Private/MoviePipelineBlueprintLibrary.cpp
--- a/Engine/Plugins/MovieScene/MovieRenderPipeline/Source/MovieRenderPipelineCore/Private/MoviePipelineBlueprintLibrary.cpp
+++ b/Engine/Plugins/MovieScene/MovieRenderPipeline/Source/MovieRenderPipelineCore/Private/MoviePipelineBlueprintLibrary.cpp
@@ -406,20 +406,23 @@ void UMoviePipelineBlueprintLibrary::UpdateJobShotListFromSequence(ULevelSequenc
 
 			for (UMovieSceneSection* Section : InTrack->GetAllSections())
 			{
-				if (!Section->IsActive())
+				if(!Section)
 				{
 					continue;
 				}
-
-
-				// Intersect it with our local range so that any sections that fall outside our playback bounds gets discarded
-				UMovieScene* OwningScene = Section->GetTypedOuter<UMovieScene>();
-				TRange<FFrameNumber> LocalCameraRange = TRange<FFrameNumber>::Intersection(OwningScene->GetPlaybackRange(), Section->GetRange());
-				if (LocalCameraRange.IsEmpty())
+				
+				if (!Section->IsActive())
 				{
 					continue;
 				}
 
+				// When a sequence is contained by a Shot Section, it can offset/clip the playback range. So we want to scan for any sections that
+				// fall entirely out of the actually evaluated range, and discard them. To do this, we transform from our parent (RootToSequenceTransform)
+				// to local space, and then compare the local camera range to that and discard it. When there is no parent, the RootToSequence transform
+				// is identity, and thus it just clips it against the normal Playback Range.
+				UMovieScene* OwningScene = Section->GetTypedOuter<UMovieScene>();
+				TRange<FFrameNumber> LocalCameraRange = Section->GetRange();
+
 				// Intersect it with the root range so that if the parent has trimmed down the sub-section we don't render outside that.
 				TRange<FFrameNumber> RootCameraRange = TRange<FFrameNumber>::Intersection(LocalSpace.RootClampRange, LocalCameraRange * LocalSpace.RootToSequenceTransform.InverseLinearOnly());
 				if (!RootCameraRange.IsEmpty())
@@ -1006,4 +1009,4 @@ FText UMoviePipelineBlueprintLibrary::GetMoviePipelineEngineChangelistLabel(cons
 	}
 
 	return FText();
-}
\ No newline at end of file
+}
