Movie Pipeline: Fix queue icon

diff --git a/Engine/Plugins/MovieScene/MovieRenderPipeline/Source/MovieRenderPipelineEditor/Private/Widgets/SMoviePipelineQueuePanel.cpp b/Engine/Plugins/MovieScene/MovieRenderPipeline/Source/MovieRenderPipelineEditor/Private/Widgets/SMoviePipelineQueuePanel.cpp
--- a/Engine/Plugins/MovieScene/MovieRenderPipeline/Source/MovieRenderPipelineEditor/Private/Widgets/SMoviePipelineQueuePanel.cpp
+++ b/Engine/Plugins/MovieScene/MovieRenderPipeline/Source/MovieRenderPipelineEditor/Private/Widgets/SMoviePipelineQueuePanel.cpp
@@ -134,7 +134,7 @@ void SMoviePipelineQueuePanel::Construct(const FArguments& InArgs)
 						.AutoWidth()
 						[
 							SNew(SImage)
-							.Image(FEditorStyle::Get().GetBrush("AssetEditor.SaveAsset.Greyscale"))
+							.Image(FEditorStyle::Get().GetBrush("AssetEditor.SaveAsset"))
 						]
 
 						+ SHorizontalBox::Slot()
