MoviePipeline: UMoviePipelineSettings can now be notified of a shot being started/ended which can help in certain situations where you need your setting to react accordingly.

diff --git a/Engine/Plugins/MovieScene/MovieRenderPipeline/Source/MovieRenderPipelineCore/Private/MoviePipeline.cpp b/Engine/Plugins/MovieScene/MovieRenderPipeline/Source/MovieRenderPipelineCore/Private/MoviePipeline.cpp
--- a/Engine/Plugins/MovieScene/MovieRenderPipeline/Source/MovieRenderPipelineCore/Private/MoviePipeline.cpp
+++ b/Engine/Plugins/MovieScene/MovieRenderPipeline/Source/MovieRenderPipelineCore/Private/MoviePipeline.cpp
@@ -900,6 +900,12 @@ void UMoviePipeline::InitializeShot(UMoviePipelineExecutorShot* InShot)
 	// Set the new shot as the active shot. This enables the specified shot section and disables all other shot sections.
 	SetSoloShot(InShot);
 
+	// Loop through just our master settings and let them know which shot we're about to start.
+	for (UMoviePipelineSetting* Setting : GetPipelineMasterConfig()->GetAllSettings())
+	{
+		Setting->OnSetupForShot(InShot);
+	}
+
 	if (InShot->GetShotOverrideConfiguration() != nullptr)
 	{
 		// Any shot-specific overrides haven't had first time initialization. So we'll do that now.
@@ -944,6 +950,12 @@ void UMoviePipeline::TeardownShot(UMoviePipelineExecutorShot* InShot)
 		}
 	}
 
+	// Loop through just our master settings and let them know which shot we're about to end.
+	for (UMoviePipelineSetting* Setting : GetPipelineMasterConfig()->GetAllSettings())
+	{
+		Setting->OnTeardownForShot(InShot);
+	}
+
 	// Restore the sequence to the original state. We made changes to this when we solo'd it, so we want to unsolo now.
 	for (UMoviePipelineExecutorShot* Shot : GetCurrentJob()->ShotInfo)
 	{
diff --git a/Engine/Plugins/MovieScene/MovieRenderPipeline/Source/MovieRenderPipelineCore/Public/MoviePipelineBlueprintLibrary.h b/Engine/Plugins/MovieScene/MovieRenderPipeline/Source/MovieRenderPipelineCore/Public/MoviePipelineBlueprintLibrary.h
--- a/Engine/Plugins/MovieScene/MovieRenderPipeline/Source/MovieRenderPipelineCore/Public/MoviePipelineBlueprintLibrary.h
+++ b/Engine/Plugins/MovieScene/MovieRenderPipeline/Source/MovieRenderPipelineCore/Public/MoviePipelineBlueprintLibrary.h
@@ -142,7 +142,7 @@ public:
 	static FIntPoint GetEffectiveOutputResolution(UMoviePipelineMasterConfig* InMasterConfig, UMoviePipelineExecutorShot* InPipelineExecutorShot);
 
 	/** Allows access to a setting of provided type for specific shot. */
-	UFUNCTION(BlueprintCallable, Category = "Movie Render Pipeline")
+	UFUNCTION(BlueprintCallable, Category = "Movie Render Pipeline", meta = (DeterminesOutputType = "InSettingType"))
 	static UMoviePipelineSetting* FindOrGetDefaultSettingForShot(TSubclassOf<UMoviePipelineSetting> InSettingType, const UMoviePipelineMasterConfig* InMasterConfig, const UMoviePipelineExecutorShot* InShot);
 
 	UFUNCTION(BlueprintPure, Category = "Movie Render Pipeline")
diff --git a/Engine/Plugins/MovieScene/MovieRenderPipeline/Source/MovieRenderPipelineCore/Public/MoviePipelineConfigBase.h b/Engine/Plugins/MovieScene/MovieRenderPipeline/Source/MovieRenderPipelineCore/Public/MoviePipelineConfigBase.h
--- a/Engine/Plugins/MovieScene/MovieRenderPipeline/Source/MovieRenderPipelineCore/Public/MoviePipelineConfigBase.h
+++ b/Engine/Plugins/MovieScene/MovieRenderPipeline/Source/MovieRenderPipelineCore/Public/MoviePipelineConfigBase.h
@@ -135,7 +135,7 @@ public:
 			}
 			else
 			{
-				FFrame::KismetExecutionMessage(*FString::Printf(TEXT("Setting %d is not compatible with this Config Type and was not added."), *InClass->GetName()), ELogVerbosity::Error);
+				FFrame::KismetExecutionMessage(*FString::Printf(TEXT("Setting %s is not compatible with this Config Type and was not added."), *InClass->GetName()), ELogVerbosity::Error);
 				return nullptr;
 			}
 		}
diff --git a/Engine/Plugins/MovieScene/MovieRenderPipeline/Source/MovieRenderPipelineCore/Public/MoviePipelineSetting.h b/Engine/Plugins/MovieScene/MovieRenderPipeline/Source/MovieRenderPipelineCore/Public/MoviePipelineSetting.h
--- a/Engine/Plugins/MovieScene/MovieRenderPipeline/Source/MovieRenderPipelineCore/Public/MoviePipelineSetting.h
+++ b/Engine/Plugins/MovieScene/MovieRenderPipeline/Source/MovieRenderPipelineCore/Public/MoviePipelineSetting.h
@@ -9,6 +9,7 @@ class UMoviePipeline;
 struct FSlateBrush;
 struct FMoviePipelineFormatArgs;
 class UMoviePipelineExecutorJob;
+class UMoviePipelineExecutorShot;
 
 enum class EMoviePipelineValidationState : uint8
 {
@@ -24,7 +25,7 @@ UCLASS(BlueprintType, Abstract)
 class MOVIERENDERPIPELINECORE_API UMoviePipelineSetting : public UObject
 {
 	GENERATED_BODY()
-		
+
 public:
 	UMoviePipelineSetting();
 
@@ -41,6 +42,18 @@ public:
 	*/
 	void OnMoviePipelineShutdown(UMoviePipeline* InPipeline) { TeardownForPipelineImpl(InPipeline); }
 
+	/**
+	* Called only on settings that have been added to the Master Configuration to let you know that
+	* a shot is about to be rendered. Useful if your setting needs to know something about the shot
+	* to do something correctly, without using a per-shot override.
+	*/
+	void OnSetupForShot(UMoviePipelineExecutorShot* InShot) { OnSetupForShotImpl(InShot); }
+	/**
+	* Called only on settings that have been added to the Master Configuration to let you know that
+	* a shot has finished rendering and is being torn down.
+	*/
+	void OnTeardownForShot(UMoviePipelineExecutorShot* InShot) { OnTeardownForShotImpl(InShot); }
+
 	/**
 	* When rendering in a new process some settings may need to provide command line arguments
 	* to affect engine settings that need to be set before most of the engine boots up. This function
@@ -78,6 +91,8 @@ protected:
 
 	virtual void SetupForPipelineImpl(UMoviePipeline* InPipeline) {}
 	virtual void TeardownForPipelineImpl(UMoviePipeline* InPipeline) {}
+	virtual	void OnSetupForShotImpl(UMoviePipelineExecutorShot* InShot) {}
+	virtual	void OnTeardownForShotImpl(UMoviePipelineExecutorShot* InShot) {}
 	
 public:
 #if WITH_EDITOR
diff --git a/Engine/Plugins/MovieScene/MovieRenderPipeline/Source/MovieRenderPipelineSettings/Public/MoviePipelineWidgetRenderSetting.h b/Engine/Plugins/MovieScene/MovieRenderPipeline/Source/MovieRenderPipelineSettings/Public/MoviePipelineWidgetRenderSetting.h
--- a/Engine/Plugins/MovieScene/MovieRenderPipeline/Source/MovieRenderPipelineSettings/Public/MoviePipelineWidgetRenderSetting.h
+++ b/Engine/Plugins/MovieScene/MovieRenderPipeline/Source/MovieRenderPipelineSettings/Public/MoviePipelineWidgetRenderSetting.h
@@ -40,7 +40,7 @@ public:
 	virtual FText GetDisplayText() const override { return NSLOCTEXT("MovieRenderPipeline", "WidgetRendererSettingDisplayName", "UI Renderer"); }
 	virtual FText GetCategoryText() const { return NSLOCTEXT("MovieRenderPipeline", "WidgetRendererSettingCategoryName", "Rendering"); }
 #endif
-	virtual bool IsValidOnShots() const override { return false; }
+	virtual bool IsValidOnShots() const override { return true; }
 	virtual bool IsValidOnMaster() const override { return true; }
 private:
 	TSharedPtr<FWidgetRenderer> WidgetRenderer;
