MoviePipeline: Added some classes to perform functional testing of the Movie Render Queue ecosystem. This allows creating functional tests that render pre-made queues and then compare their output to known good outputs.

diff --git a/Engine/Plugins/MovieScene/MovieRenderPipeline/Source/MovieRenderPipelineEditor/MovieRenderPipelineEditor.Build.cs b/Engine/Plugins/MovieScene/MovieRenderPipeline/Source/MovieRenderPipelineEditor/MovieRenderPipelineEditor.Build.cs
--- a/Engine/Plugins/MovieScene/MovieRenderPipeline/Source/MovieRenderPipelineEditor/MovieRenderPipelineEditor.Build.cs
+++ b/Engine/Plugins/MovieScene/MovieRenderPipeline/Source/MovieRenderPipelineEditor/MovieRenderPipelineEditor.Build.cs
@@ -34,6 +34,7 @@ public class MovieRenderPipelineEditor : ModuleRules
 				"PropertyEditor",
 				"EditorWidgets",
 				"EditorSubsystem",
+				"FunctionalTesting",
             }
         );
 
@@ -57,7 +58,11 @@ public class MovieRenderPipelineEditor : ModuleRules
 				"DeveloperSettings",
 				"MessageLog",
 				"Sequencer",
-				"ToolMenus"
+				"ToolMenus",
+				"Json",
+				"JsonUtilities",
+				"ScreenShotComparisonTools",
+				"AutomationMessages"
 			}
         );
 
diff --git a/Engine/Plugins/MovieScene/MovieRenderPipeline/Source/MovieRenderPipelineEditor/Private/MoviePipelineEditorBlueprintLibrary.cpp b/Engine/Plugins/MovieScene/MovieRenderPipeline/Source/MovieRenderPipelineEditor/Private/MoviePipelineEditorBlueprintLibrary.cpp
--- a/Engine/Plugins/MovieScene/MovieRenderPipeline/Source/MovieRenderPipelineEditor/Private/MoviePipelineEditorBlueprintLibrary.cpp
+++ b/Engine/Plugins/MovieScene/MovieRenderPipeline/Source/MovieRenderPipelineEditor/Private/MoviePipelineEditorBlueprintLibrary.cpp
@@ -18,6 +18,8 @@
 #include "Editor.h"
 #include "GenericPlatform/GenericPlatformProcess.h"
 #include "SequencerUtilities.h"
+#include "MoviePipelineBlueprintLibrary.h"
+#include "MoviePipelineOutputSetting.h"
 
 #define LOCTEXT_NAMESPACE "MoviePipelineEditorBlueprintLibrary"
 
@@ -208,4 +210,59 @@ void UMoviePipelineEditorBlueprintLibrary::EnsureJobHasDefaultSettings(UMoviePip
 	}
 }
 
+FString UMoviePipelineEditorBlueprintLibrary::ResolveOutputDirectoryFromJob(UMoviePipelineExecutorJob* InJob)
+{
+	UMoviePipelineOutputSetting* OutputSetting = InJob->GetConfiguration()->FindSetting<UMoviePipelineOutputSetting>();
+	check(OutputSetting);
+
+	// Set up as many parameters as we can to try and resolve most of the string.
+	FString FormatString = OutputSetting->OutputDirectory.Path / OutputSetting->FileNameFormat;
+
+	// If they've set up any folders within the filename portion of it, let's be nice and resolve that.
+	FPaths::NormalizeFilename(FormatString);
+	int32 LastSlashIndex;
+	if (FormatString.FindLastChar(TEXT('/'), LastSlashIndex))
+	{
+		FormatString.LeftInline(LastSlashIndex + 1);
+	}
+
+	// By having it swap {camera_name} and {shot_name} with an unresolvable tag, it will
+	// stay in the resolved path and can be removed using the code below.
+	static const FString DummyTag = TEXT("{dontresolvethis}");
+	FMoviePipelineFilenameResolveParams Params;
+	Params.Job = InJob;
+	Params.ShotNameOverride = DummyTag;
+	Params.CameraNameOverride = DummyTag;
+
+	FString OutResolvedPath;
+	FMoviePipelineFormatArgs Dummy;
+	UMoviePipelineBlueprintLibrary::ResolveFilenameFormatArguments(FormatString, Params, OutResolvedPath, Dummy);
+
+	// Drop the .{ext} resolving always puts on.
+	OutResolvedPath.LeftChopInline(6);
+
+	if (FPaths::IsRelative(OutResolvedPath))
+	{
+		OutResolvedPath = FPaths::ConvertRelativePathToFull(OutResolvedPath);
+	}
+
+	// In the event that they used a {format_string} we couldn't resolve (such as shot name), then
+	// we'll trim off anything after the format string.
+	int32 FormatStringToken;
+	if (OutResolvedPath.FindChar(TEXT('{'), FormatStringToken))
+	{
+		// Just as a last bit of saftey, we'll trim anything between the { and the preceeding /. This is
+		// in case they did something like Render_{Date}, we wouldn't want to make a folder named Render_.
+		// We search backwards from where we found the first { brace, so that will get us the last usable slash.
+		LastSlashIndex = OutResolvedPath.Find(TEXT("/"), ESearchCase::IgnoreCase, ESearchDir::FromEnd, FormatStringToken);
+		if (LastSlashIndex != INDEX_NONE)
+		{
+			OutResolvedPath.LeftInline(LastSlashIndex + 1);
+		}
+	}
+
+	return OutResolvedPath;
+}
+
+
 #undef LOCTEXT_NAMESPACE // "MoviePipelineEditorBlueprintLibrary"
\ No newline at end of file
diff --git a/Engine/Plugins/MovieScene/MovieRenderPipeline/Source/MovieRenderPipelineEditor/Private/MoviePipelineFunctionalTestBase.cpp b/Engine/Plugins/MovieScene/MovieRenderPipeline/Source/MovieRenderPipelineEditor/Private/MoviePipelineFunctionalTestBase.cpp
new file mode 100644
--- /dev/null
+++ b/Engine/Plugins/MovieScene/MovieRenderPipeline/Source/MovieRenderPipelineEditor/Private/MoviePipelineFunctionalTestBase.cpp
@@ -0,0 +1,273 @@
+// Copyright Epic Games, Inc. All Rights Reserved.
+
+#include "MoviePipelineFunctionalTestBase.h"
+#include "MoviePipeline.h"
+#include "MoviePipelineQueue.h"
+#include "Tests/AutomationCommon.h"
+#include "IImageWrapper.h"
+#include "IImageWrapperModule.h"
+#include "Misc/FileHelper.h"
+#include "Modules/ModuleManager.h"
+#include "Misc/Paths.h"
+#include "JsonObjectConverter.h"
+#include "Misc/FileHelper.h"
+#include "MoviePipelineOutputSetting.h"
+#include "MoviePipelineEditorBlueprintLibrary.h"
+#include "ImageComparer.h"
+#include "Misc/DataDrivenPlatformInfoRegistry.h"
+#include "Interfaces/IScreenShotManager.h"
+#include "Interfaces/IScreenShotToolsModule.h"
+#include "AutomationBlueprintFunctionLibrary.h"
+#include "AutomationWorkerMessages.h"
+
+AMoviePipelineFunctionalTestBase::AMoviePipelineFunctionalTestBase()
+{
+	
+}
+
+void AMoviePipelineFunctionalTestBase::PrepareTest() 
+{
+	if (!QueuePreset)
+	{
+		FinishTest(EFunctionalTestResult::Failed, TEXT("No Queue Preset asset specified, nothing to test!"));
+	}
+
+	if (QueuePreset->GetJobs().Num() == 0)
+	{
+		FinishTest(EFunctionalTestResult::Failed, TEXT("Queue Preset has no jobs, nothing to test!"));
+	}
+
+	if (QueuePreset->GetJobs().Num() > 1)
+	{
+		FinishTest(EFunctionalTestResult::Failed, TEXT("Only one job per queue currently supported!"));
+	}
+}
+
+bool AMoviePipelineFunctionalTestBase::IsReady_Implementation()
+{ 
+	return Super::IsReady_Implementation();
+}
+
+void AMoviePipelineFunctionalTestBase::StartTest()
+{
+
+	// PIE will already be running at this point so we want to instantiate an instance of
+	// the Movie Pipeline in the current world and just run it. This doesn't test the UI/PIE 
+	// portion of the system but that is more stable than the actual featureset.
+	ActiveMoviePipeline = NewObject<UMoviePipeline>(GetWorld());
+	ActiveMoviePipeline->OnMoviePipelineWorkFinished().AddUObject(this, &AMoviePipelineFunctionalTestBase::OnMoviePipelineFinished);
+	ActiveMoviePipeline->OnMoviePipelineShotWorkFinished().AddUObject(this, &AMoviePipelineFunctionalTestBase::OnJobShotFinished);
+
+	ActiveQueue = NewObject<UMoviePipelineQueue>(GetWorld());
+	ActiveQueue->CopyFrom(QueuePreset);
+
+	ActiveMoviePipeline->Initialize(ActiveQueue->GetJobs()[0]); // Prepare Test will ensure we have one job
+}
+
+void AMoviePipelineFunctionalTestBase::OnJobShotFinished(FMoviePipelineOutputData InOutputData)
+{
+	if (ActiveMoviePipeline)
+	{
+		ActiveMoviePipeline->OnMoviePipelineShotWorkFinished().RemoveAll(this);
+	}
+}
+
+void AMoviePipelineFunctionalTestBase::OnMoviePipelineFinished(FMoviePipelineOutputData InOutputData)
+{
+	if (ActiveMoviePipeline)
+	{
+		ActiveMoviePipeline->OnMoviePipelineWorkFinished().RemoveAll(this);
+	}
+
+	if (!InOutputData.bSuccess)
+	{
+		FinishTest(EFunctionalTestResult::Failed, TEXT("MoviePipeline encountered an internal error. See log for details."));
+	}
+	else
+	{
+		CompareRenderOutputToGroundTruth(InOutputData);
+	}
+}
+
+bool SaveOutputToGroundTruth(const FString& GroundTruthDirectory, FMoviePipelineOutputData OutputData)
+{
+	FString GroundTruthFilepath = GroundTruthDirectory / "GroundTruth.json";
+
+	// We need to rewrite the Output Data to be relative to our new directory, and then copy all
+	// of the files from the old location to the new location. We want to keep these relative to
+	// the original output directory from MRQ, so that we can make tests that ensure sub-folder
+	// structures get generated correctly.
+	UMoviePipelineOutputSetting* OutputSetting = Cast<UMoviePipelineOutputSetting>(OutputData.Job->GetConfiguration()->FindOrAddSettingByClass(UMoviePipelineOutputSetting::StaticClass()));
+	FString OriginalRootOutputDirectory = UMoviePipelineEditorBlueprintLibrary::ResolveOutputDirectoryFromJob(OutputData.Job);
+
+	for (FMoviePipelineShotOutputData& Shot : OutputData.ShotData)
+	{
+		for (TPair<FMoviePipelinePassIdentifier, FMoviePipelineRenderPassOutputData>& Pair : Shot.RenderPassData)
+		{
+			for (FString& FilePath : Pair.Value.FilePaths)
+			{
+				FString RelativePath = FilePath;
+				if (!FPaths::MakePathRelativeTo(/*Out*/ RelativePath, *OriginalRootOutputDirectory))
+				{
+					UE_LOG(LogTemp, Error, TEXT("Could not generate ground truth, unable to generate relative paths."));
+					return false;
+				}
+
+				FString NewPath = GroundTruthDirectory / RelativePath;
+				FString AbsoluteNewPath = FPaths::ConvertRelativePathToFull(NewPath);
+				if (IFileManager::Get().Copy(*AbsoluteNewPath, *FilePath) != ECopyResult::COPY_OK)
+				{
+					UE_LOG(LogTemp, Error, TEXT("Could not generate ground truth, unable to copy existing image to Automation directory."));
+					return false;
+				}
+
+				// Now rewrite the FilePath in the struct to the new location so that when we save the json below it has the right path.
+				FString RelativeToAutomationDir = NewPath;
+				if (!FPaths::MakePathRelativeTo(/*InOut*/ RelativeToAutomationDir, *GroundTruthDirectory))
+				{
+					UE_LOG(LogTemp, Error, TEXT("Could not generate ground truth, unable to generate relative paths."));
+					return false;
+				}
+
+				FilePath = RelativeToAutomationDir;
+			}
+		}
+	}
+
+	// Now that we've copied all of the images across we can serialize the struct to a json string
+	FString SerializedJson;
+	FJsonObjectConverter::UStructToJsonObjectString(OutputData, SerializedJson);
+
+	// And finally write it to disk.
+	FFileHelper::SaveStringToFile(SerializedJson, *GroundTruthFilepath);
+
+	return true;
+}
+
+void AMoviePipelineFunctionalTestBase::CompareRenderOutputToGroundTruth(FMoviePipelineOutputData InOutputData)
+{
+	// Grab the (expected) resolution from out rendered data
+	UMoviePipelineOutputSetting* OutputSetting = Cast< UMoviePipelineOutputSetting>(InOutputData.Job->GetConfiguration()->FindOrAddSettingByClass(UMoviePipelineOutputSetting::StaticClass()));
+	check(OutputSetting);
+
+	// Build screenshot data for this test. This contains a lot of metadata about RHI, Platform, etc that we'll use to generate the output folder name.
+	FAutomationScreenshotData Data = UAutomationBlueprintFunctionLibrary::BuildScreenshotData(GetWorld()->GetName(), TestLabel, OutputSetting->OutputResolution.X, OutputSetting->OutputResolution.Y);
+
+	// Convert the Screenshot Data into Metadata
+	FAutomationScreenshotMetadata MetaData(Data);
+
+	IScreenShotToolsModule& ScreenShotModule = FModuleManager::LoadModuleChecked<IScreenShotToolsModule>("ScreenShotComparisonTools");
+	IScreenShotManagerPtr ScreenshotManager = ScreenShotModule.GetScreenShotManager();
+
+	// Now we know where to look for our ground truth data.
+	FString ReportDirectory = ScreenshotManager->GetIdealApprovedFolderForImage(MetaData);
+	FString GroundTruthFilename = ReportDirectory / "GroundTruth.json";
+
+	if (FPaths::FileExists(GroundTruthFilename))
+	{
+		UE_LOG(LogTemp, Log, TEXT("GroundTruth file located at %s"), *GroundTruthFilename);
+	}
+	else
+	{
+		UE_LOG(LogTemp, Error, TEXT("Failed to find a GroundTruth file at %s, generating one now. Rerun the test after verifying them!"), *GroundTruthFilename);
+		SaveOutputToGroundTruth(ReportDirectory, InOutputData);
+		FinishTest(EFunctionalTestResult::Failed, TEXT("Generated ground truth, run test again after verifying the ground truth is correct."));
+		return;
+	}
+
+	// The groundtruth file exists, so we can load it and turn it back into a FMoviePipelineOutputData struct.
+	FString LoadedGroundTruthJsonStr;
+	FFileHelper::LoadFileToString(LoadedGroundTruthJsonStr, *GroundTruthFilename);
+
+	FMoviePipelineOutputData GroundTruthData;
+	FJsonObjectConverter::JsonObjectStringToUStruct<FMoviePipelineOutputData>(LoadedGroundTruthJsonStr, &GroundTruthData);
+
+	// Do some basic checks on our new data to ensure it output the expected number of files/shots/etc, before doing
+	// the computationally expensive image comparisons.
+	int32 NumShotsNew = InOutputData.ShotData.Num();
+	int32 NumShotsGT = GroundTruthData.ShotData.Num();
+	if (NumShotsNew != NumShotsGT)
+	{
+		FinishTest(EFunctionalTestResult::Failed, TEXT("Mismatched number of shots between GT and New."));
+		return;
+	}
+
+	TMap<FString, FString> OldToNewImagesToCompare;
+
+	for (const FMoviePipelineShotOutputData& GroundTruthShot : GroundTruthData.ShotData)
+	{
+		// Look up the matching shot in the new data.
+		const FMoviePipelineShotOutputData* FoundInNewData = InOutputData.ShotData.FindByPredicate([GroundTruthShot](const FMoviePipelineShotOutputData& InOther)
+			{
+				return GroundTruthShot.Shot->InnerName == InOther.Shot->InnerName &&
+					GroundTruthShot.Shot->OuterName == InOther.Shot->OuterName;
+			});
+		if (!FoundInNewData)
+		{
+			FinishTest(EFunctionalTestResult::Failed, TEXT("Did not render a shot that was in the Ground Truth."));
+			return;
+		}
+
+		// And then repeat this for each render pass in the shot
+		int32 NumRenderPassesNew = FoundInNewData->RenderPassData.Num();
+		int32 NumRenderPassesGT = GroundTruthShot.RenderPassData.Num();
+		if (NumRenderPassesNew != NumRenderPassesGT)
+		{
+			FinishTest(EFunctionalTestResult::Failed, TEXT("Mismatched number of Render Passes between GT and New."));
+			return;
+		}
+
+		for (const TPair<FMoviePipelinePassIdentifier, FMoviePipelineRenderPassOutputData>& GroundTruthPair : GroundTruthShot.RenderPassData)
+		{
+			const FMoviePipelineRenderPassOutputData* NewDataOutputData = FoundInNewData->RenderPassData.Find(GroundTruthPair.Key);
+			if (!NewDataOutputData)
+			{
+				FinishTest(EFunctionalTestResult::Failed, TEXT("Did not render a render pass that was in the Ground Truth."));
+				return;
+			}
+
+			// Now we can check that they output the same number of files
+			int32 NumOutputFilesNew = NewDataOutputData->FilePaths.Num();
+			int32 NumOutputFilesGT = GroundTruthPair.Value.FilePaths.Num();
+			if (NumOutputFilesNew != NumOutputFilesGT)
+			{
+				FinishTest(EFunctionalTestResult::Failed, TEXT("Mismatched number of files on disk between GT and New."));
+				return;
+			}
+
+			// Now we'll just loop through them in lockstep (as MRQ should put them in order anyways) and just add
+			// them to an array to be tested later. We don't test that the sub-file names are exact matches right now.
+			for (int32 Index = 0; Index < NumOutputFilesGT; Index++)
+			{
+				FString AbsoluteGTPath = FPaths::ConvertRelativePathToFull(FPaths::CreateStandardFilename(ReportDirectory / GroundTruthPair.Value.FilePaths[Index]));
+				FString AbsoluteNewPath = FPaths::ConvertRelativePathToFull(FPaths::CreateStandardFilename(NewDataOutputData->FilePaths[Index]));
+				OldToNewImagesToCompare.Add(AbsoluteGTPath, AbsoluteNewPath);
+			}
+		}
+	}
+
+	// Time for the computationally expensive part, doing image comparisons!
+	for (const TPair<FString, FString>& Pair : OldToNewImagesToCompare)
+	{
+		const FString& OldImagePath = Pair.Key;
+		const FString& NewImagePath = Pair.Value;
+
+		// Calculate a path for the delta image to be saved.
+		FString DeltaPath = FPaths::ChangeExtension(NewImagePath, TEXT(""));
+		FString OldExtension = FPaths::GetExtension(NewImagePath, true);
+		DeltaPath += TEXT("_Delta") + OldExtension;
+
+		// Alright we have both images in memory now, now use a FImageComparer for less strict comparison.
+		FImageTolerance Tolerance = FImageTolerance::DefaultIgnoreLess;
+		FImageComparer Comparer;
+		FImageComparisonResult ComparisonResults = Comparer.Compare(OldImagePath, NewImagePath, Tolerance, DeltaPath);
+
+		if (!ComparisonResults.AreSimilar())
+		{
+			FinishTest(EFunctionalTestResult::Failed, TEXT("Frames failed comparison tolerance!"));
+			return;
+		}
+	}
+
+	FinishTest(EFunctionalTestResult::Succeeded, TEXT(""));
+}
diff --git a/Engine/Plugins/MovieScene/MovieRenderPipeline/Source/MovieRenderPipelineEditor/Private/MoviePipelineFunctionalTestBase.h b/Engine/Plugins/MovieScene/MovieRenderPipeline/Source/MovieRenderPipelineEditor/Private/MoviePipelineFunctionalTestBase.h
new file mode 100644
--- /dev/null
+++ b/Engine/Plugins/MovieScene/MovieRenderPipeline/Source/MovieRenderPipelineEditor/Private/MoviePipelineFunctionalTestBase.h
@@ -0,0 +1,50 @@
+// Copyright Epic Games, Inc. All Rights Reserved.
+
+#pragma once
+
+#include "CoreMinimal.h"
+#include "UObject/ObjectMacros.h"
+#include "FunctionalTest.h"
+#include "Misc/AutomationTest.h"
+#include "MovieRenderPipelineDataTypes.h"
+
+#include "MoviePipelineFunctionalTestBase.generated.h"
+
+class UMoviePipelineQueue;
+class UMoviePipeline;
+
+/**
+* Base class for Movie Pipeline functional tests which render pre-made queues
+* and compare their output against pre-existing render outputs.
+*/
+UCLASS(Blueprintable)
+class MOVIERENDERPIPELINEEDITOR_API AMoviePipelineFunctionalTestBase : public AFunctionalTest
+{
+	GENERATED_BODY()
+
+public:
+	AMoviePipelineFunctionalTestBase();
+
+protected:
+	// AFunctionalTest
+	virtual void PrepareTest() override;
+	virtual bool IsReady_Implementation() override;
+	virtual void StartTest() override;
+	// ~AFunctionalTest
+
+	void OnJobShotFinished(FMoviePipelineOutputData InOutputData);
+	void OnMoviePipelineFinished(FMoviePipelineOutputData InOutputData);
+	void CompareRenderOutputToGroundTruth(FMoviePipelineOutputData InOutputData);
+
+protected:
+
+	UPROPERTY(EditAnywhere, BlueprintReadOnly, Category = "Movie Pipeline")
+	UMoviePipelineQueue* QueuePreset;
+
+protected:
+	UPROPERTY(Transient)
+	UMoviePipeline* ActiveMoviePipeline;
+
+	UPROPERTY(Transient)
+	UMoviePipelineQueue* ActiveQueue;
+};
diff --git a/Engine/Plugins/MovieScene/MovieRenderPipeline/Source/MovieRenderPipelineEditor/Private/Widgets/SMoviePipelineQueueEditor.cpp b/Engine/Plugins/MovieScene/MovieRenderPipeline/Source/MovieRenderPipelineEditor/Private/Widgets/SMoviePipelineQueueEditor.cpp
--- a/Engine/Plugins/MovieScene/MovieRenderPipeline/Source/MovieRenderPipelineEditor/Private/Widgets/SMoviePipelineQueueEditor.cpp
+++ b/Engine/Plugins/MovieScene/MovieRenderPipeline/Source/MovieRenderPipelineEditor/Private/Widgets/SMoviePipelineQueueEditor.cpp
@@ -244,60 +244,13 @@ public:
 		UMoviePipelineExecutorJob* Job = WeakJob.Get();
 		if (Job && Job->GetConfiguration())
 		{
-			UMoviePipelineOutputSetting* OutputSetting = Job->GetConfiguration()->FindSetting<UMoviePipelineOutputSetting>();
-			check(OutputSetting);
-
-			// Set up as many parameters as we can to try and resolve most of the string.
-			FString FormatString = OutputSetting->OutputDirectory.Path / OutputSetting->FileNameFormat;
-
-			// If they've set up any folders within the filename portion of it, let's be nice and resolve that.
-			FPaths::NormalizeFilename(FormatString);
-			int32 LastSlashIndex;
-			if(FormatString.FindLastChar(TEXT('/'), LastSlashIndex))
-			{
-				FormatString.LeftInline(LastSlashIndex + 1);
-			}
-
-			// By having it swap {camera_name} and {shot_name} with an unresolvable tag, it will
-			// stay in the resolved path and can be removed using the code below.
-			static const FString DummyTag = TEXT("{dontresolvethis}");
-			FMoviePipelineFilenameResolveParams Params;
-			Params.Job = Job;
-			Params.ShotNameOverride = DummyTag;
-			Params.CameraNameOverride = DummyTag;
-
-			FString OutResolvedPath;
-			FMoviePipelineFormatArgs Dummy;
-			UMoviePipelineBlueprintLibrary::ResolveFilenameFormatArguments(FormatString, Params, OutResolvedPath, Dummy);
-
-			// Drop the .{ext} resolving always puts on.
-			OutResolvedPath.LeftChopInline(6);
-
-			if (FPaths::IsRelative(OutResolvedPath))
-			{
-				OutResolvedPath = FPaths::ConvertRelativePathToFull(OutResolvedPath);
-			}
-			
-			// In the event that they used a {format_string} we couldn't resolve (such as shot name), then
-			// we'll trim off anything after the format string.
-			int32 FormatStringToken;
-			if(OutResolvedPath.FindChar(TEXT('{'), FormatStringToken))
-			{
-				// Just as a last bit of saftey, we'll trim anything between the { and the preceeding /. This is
-				// in case they did something like Render_{Date}, we wouldn't want to make a folder named Render_.
-				// We search backwards from where we found the first { brace, so that will get us the last usable slash.
-				LastSlashIndex = OutResolvedPath.Find(TEXT("/"), ESearchCase::IgnoreCase, ESearchDir::FromEnd, FormatStringToken);
-				if (LastSlashIndex != INDEX_NONE)
-				{
-					OutResolvedPath.LeftInline(LastSlashIndex + 1);
-				}
-			}
+			FString ResolvedOutputDir = UMoviePipelineEditorBlueprintLibrary::ResolveOutputDirectoryFromJob(Job);
 
 			// Attempt to make the directory. The user can see the output folder before they render so the folder
 			// may not have been created yet and the ExploreFolder call will fail.
-			IFileManager::Get().MakeDirectory(*OutResolvedPath, true);
+			IFileManager::Get().MakeDirectory(*ResolvedOutputDir, true);
 
-			FPlatformProcess::ExploreFolder(*OutResolvedPath);
+			FPlatformProcess::ExploreFolder(*ResolvedOutputDir);
 		}
 	}
 
diff --git a/Engine/Plugins/MovieScene/MovieRenderPipeline/Source/MovieRenderPipelineEditor/Public/MoviePipelineEditorBlueprintLibrary.h b/Engine/Plugins/MovieScene/MovieRenderPipeline/Source/MovieRenderPipelineEditor/Public/MoviePipelineEditorBlueprintLibrary.h
--- a/Engine/Plugins/MovieScene/MovieRenderPipeline/Source/MovieRenderPipelineEditor/Public/MoviePipelineEditorBlueprintLibrary.h
+++ b/Engine/Plugins/MovieScene/MovieRenderPipeline/Source/MovieRenderPipelineEditor/Public/MoviePipelineEditorBlueprintLibrary.h
@@ -48,4 +48,8 @@ public:
 	/** Ensure the job has the settings specified by the project settings added. If they're already added we don't modify the object so that we don't make it confused about whether or not you've modified the preset. */
 	UFUNCTION(BlueprintCallable, Category = "Movie Render Pipeline")
 	static void EnsureJobHasDefaultSettings(UMoviePipelineExecutorJob* InJob);
+
+	/** Resolves as much of the output directory for this job into a usable directory path as possible. Cannot resolve anything that relies on shot name, frame numbers, etc. */
+	UFUNCTION(BlueprintCallable, Category = "Movie Render Pipeline")
+	static FString ResolveOutputDirectoryFromJob(UMoviePipelineExecutorJob* InJob);
 };
\ No newline at end of file
