Movie Render Queue:
* Added GeometryCache streamer cvars to the MRQ Game Override that are automatically applied for optimal rendering of scenes with streaming geometry caches so that the MRQ waits for the streaming to finish before rendering the frame.

diff --git a/Engine/Plugins/MovieScene/MovieRenderPipeline/Source/MovieRenderPipelineCore/Private/MoviePipelineGameOverrideSetting.cpp b/Engine/Plugins/MovieScene/MovieRenderPipeline/Source/MovieRenderPipelineCore/Private/MoviePipelineGameOverrideSetting.cpp
--- a/Engine/Plugins/MovieScene/MovieRenderPipeline/Source/MovieRenderPipelineCore/Private/MoviePipelineGameOverrideSetting.cpp
+++ b/Engine/Plugins/MovieScene/MovieRenderPipeline/Source/MovieRenderPipelineCore/Private/MoviePipelineGameOverrideSetting.cpp
@@ -93,6 +93,10 @@ void UMoviePipelineGameOverrideSetting::ApplyCVarSettings(const bool bOverrideVa
 		MOVIEPIPELINE_STORE_AND_OVERRIDE_CVAR_INT_IF_EXIST(PreviousGPUTimeout, TEXT("r.D3D12.GPUTimeout"), 0, bOverrideValues);
 	}
 	
+	// To make sure the GeometryCache streamer doesn't skip frames and doesn't pop up notification during rendering
+	MOVIEPIPELINE_STORE_AND_OVERRIDE_CVAR_INT(PreviousGeoCacheStreamerBlockTillFinish, TEXT("GeometryCache.Streamer.BlockTillFinishStreaming"), 1, bOverrideValues);
+	MOVIEPIPELINE_STORE_AND_OVERRIDE_CVAR_INT(PreviousGeoCacheStreamerShowNotification, TEXT("GeometryCache.Streamer.ShowNotification"), 0, bOverrideValues);
+
 	{
 		// Disable systems that try to preserve performance in runtime games.
 		MOVIEPIPELINE_STORE_AND_OVERRIDE_CVAR_INT(PreviousAnimationUROEnabled, TEXT("a.URO.Enable"), 0, bOverrideValues);
@@ -179,6 +183,11 @@ void UMoviePipelineGameOverrideSetting::BuildNewProcessCommandLineArgsImpl(TArra
 		InOutDeviceProfileCvars.Add(TEXT("r.D3D12.GPUTimeout=0"));
 	}
 	
+	{
+		InOutDeviceProfileCvars.Add(TEXT("GeometryCache.Streamer.BlockTillFinishStreaming=1"));
+		InOutDeviceProfileCvars.Add(TEXT("GeometryCache.Streamer.ShowNotification=0"));
+	}
+
 	{
 		InOutDeviceProfileCvars.Add(FString::Printf(TEXT("a.URO.Enable=%d"), 0));
 	}
diff --git a/Engine/Plugins/MovieScene/MovieRenderPipeline/Source/MovieRenderPipelineCore/Public/MoviePipelineGameOverrideSetting.h b/Engine/Plugins/MovieScene/MovieRenderPipeline/Source/MovieRenderPipelineCore/Public/MoviePipelineGameOverrideSetting.h
--- a/Engine/Plugins/MovieScene/MovieRenderPipeline/Source/MovieRenderPipelineCore/Public/MoviePipelineGameOverrideSetting.h
+++ b/Engine/Plugins/MovieScene/MovieRenderPipeline/Source/MovieRenderPipelineCore/Public/MoviePipelineGameOverrideSetting.h
@@ -124,4 +124,6 @@ private:
 	int32 PreviousSkyLightRealTimeReflectionCaptureTimeSlice;
 	int32 PreviousVolumetricRenderTarget;
 	int32 PreviousIgnoreStreamingPerformance;
+	int32 PreviousGeoCacheStreamerShowNotification;
+	int32 PreviousGeoCacheStreamerBlockTillFinish;
 };
\ No newline at end of file
