* Fixed GeometryCache.Streamer.BlockTillFinishStreaming cvar not found in package. The GeometryCacheStreamer module is editor only.

diff --git a/Engine/Plugins/MovieScene/MovieRenderPipeline/Source/MovieRenderPipelineCore/Private/MoviePipelineGameOverrideSetting.cpp b/Engine/Plugins/MovieScene/MovieRenderPipeline/Source/MovieRenderPipelineCore/Private/MoviePipelineGameOverrideSetting.cpp
--- a/Engine/Plugins/MovieScene/MovieRenderPipeline/Source/MovieRenderPipelineCore/Private/MoviePipelineGameOverrideSetting.cpp
+++ b/Engine/Plugins/MovieScene/MovieRenderPipeline/Source/MovieRenderPipelineCore/Private/MoviePipelineGameOverrideSetting.cpp
@@ -93,9 +93,11 @@ void UMoviePipelineGameOverrideSetting::ApplyCVarSettings(const bool bOverrideVa
 		MOVIEPIPELINE_STORE_AND_OVERRIDE_CVAR_INT_IF_EXIST(PreviousGPUTimeout, TEXT("r.D3D12.GPUTimeout"), 0, bOverrideValues);
 	}
 	
+#if WITH_EDITOR
 	// To make sure the GeometryCache streamer doesn't skip frames and doesn't pop up notification during rendering
 	MOVIEPIPELINE_STORE_AND_OVERRIDE_CVAR_INT(PreviousGeoCacheStreamerBlockTillFinish, TEXT("GeometryCache.Streamer.BlockTillFinishStreaming"), 1, bOverrideValues);
 	MOVIEPIPELINE_STORE_AND_OVERRIDE_CVAR_INT(PreviousGeoCacheStreamerShowNotification, TEXT("GeometryCache.Streamer.ShowNotification"), 0, bOverrideValues);
+#endif
 
 	{
 		// Disable systems that try to preserve performance in runtime games.
@@ -183,10 +185,12 @@ void UMoviePipelineGameOverrideSetting::BuildNewProcessCommandLineArgsImpl(TArra
 		InOutDeviceProfileCvars.Add(TEXT("r.D3D12.GPUTimeout=0"));
 	}
 	
+#if WITH_EDITOR
 	{
 		InOutDeviceProfileCvars.Add(TEXT("GeometryCache.Streamer.BlockTillFinishStreaming=1"));
 		InOutDeviceProfileCvars.Add(TEXT("GeometryCache.Streamer.ShowNotification=0"));
 	}
+#endif
 
 	{
 		InOutDeviceProfileCvars.Add(FString::Printf(TEXT("a.URO.Enable=%d"), 0));
diff --git a/Engine/Plugins/MovieScene/MovieRenderPipeline/Source/MovieRenderPipelineCore/Public/MoviePipelineGameOverrideSetting.h b/Engine/Plugins/MovieScene/MovieRenderPipeline/Source/MovieRenderPipelineCore/Public/MoviePipelineGameOverrideSetting.h
--- a/Engine/Plugins/MovieScene/MovieRenderPipeline/Source/MovieRenderPipelineCore/Public/MoviePipelineGameOverrideSetting.h
+++ b/Engine/Plugins/MovieScene/MovieRenderPipeline/Source/MovieRenderPipelineCore/Public/MoviePipelineGameOverrideSetting.h
@@ -124,6 +124,8 @@ private:
 	int32 PreviousSkyLightRealTimeReflectionCaptureTimeSlice;
 	int32 PreviousVolumetricRenderTarget;
 	int32 PreviousIgnoreStreamingPerformance;
+#if WITH_EDITOR
 	int32 PreviousGeoCacheStreamerShowNotification;
 	int32 PreviousGeoCacheStreamerBlockTillFinish;
+#endif
 };
\ No newline at end of file
