MoviePipeline: Fixed a GC leak with UMoviePipelinePIEExecutor which was an issue when PendingKill was disabled.

diff --git a/Engine/Plugins/MovieScene/MovieRenderPipeline/Source/MovieRenderPipelineEditor/Private/MoviePipelinePIEExecutor.cpp b/Engine/Plugins/MovieScene/MovieRenderPipeline/Source/MovieRenderPipelineEditor/Private/MoviePipelinePIEExecutor.cpp
--- a/Engine/Plugins/MovieScene/MovieRenderPipeline/Source/MovieRenderPipelineEditor/Private/MoviePipelinePIEExecutor.cpp
+++ b/Engine/Plugins/MovieScene/MovieRenderPipeline/Source/MovieRenderPipelineEditor/Private/MoviePipelinePIEExecutor.cpp
@@ -289,6 +289,10 @@ void UMoviePipelinePIEExecutor::OnPIEEnded(bool)
 		// Cache this off so we can use it in DelayedFinishNotification, at which point ActiveMoviePipeline is null.
 		CachedOutputDataParams = ActiveMoviePipeline->GetOutputDataParams();
 	}
+
+	// We need to null out this reference this frame, otherwise we try to hold onto a PIE
+	// object after PIE finishes which causes a GC leak.
+	ActiveMoviePipeline = nullptr;
 	// ToDo: bAnyJobHadFatalError
 
 	// Delay for one frame so that PIE can finish shut down. It's not a huge fan of us starting up on the same frame.
@@ -307,16 +311,9 @@ void UMoviePipelinePIEExecutor::DelayedFinishNotification()
 {
 	// Get the params for the job (including output info) from the cache. ActiveMoviePipeline is null already.
 	OnIndividualJobFinishedImpl(CachedOutputDataParams);
-
-	// Now that PIE has finished
-	UMoviePipeline* MoviePipeline = ActiveMoviePipeline;
-	
-	// Null these out now since OnIndividualPipelineFinished might invoke something that causes a GC
-	// and we want them to go away with the GC.
-	ActiveMoviePipeline = nullptr;
 	
 	// Now that another frame has passed and we should be OK to start another PIE session, notify our owner.
-	OnIndividualPipelineFinished(MoviePipeline);
+	OnIndividualPipelineFinished(nullptr);
 }
 
 void UMoviePipelinePIEExecutor::OnIndividualJobFinishedImpl(FMoviePipelineOutputData InOutputData)
