MoviePipeline: Fixed a regression in jpeg/png file writing times that significantly slowed down large resolution renders.

diff --git a/Engine/Plugins/MovieScene/MovieRenderPipeline/Source/MovieRenderPipelineCore/Private/MoviePipelineImageQuantization.cpp b/Engine/Plugins/MovieScene/MovieRenderPipeline/Source/MovieRenderPipelineCore/Private/MoviePipelineImageQuantization.cpp
--- a/Engine/Plugins/MovieScene/MovieRenderPipeline/Source/MovieRenderPipelineCore/Private/MoviePipelineImageQuantization.cpp
+++ b/Engine/Plugins/MovieScene/MovieRenderPipeline/Source/MovieRenderPipelineCore/Private/MoviePipelineImageQuantization.cpp
@@ -5,7 +5,7 @@
 #include "Stats/Stats2.h"
 #include "MovieRenderPipelineCoreModule.h"
 #include "Math/NumericLimits.h"
-
+#include "Async/ParallelFor.h"
 
 namespace UE
 {
@@ -160,10 +160,21 @@ static TArray<FColor> ConvertLinearTosRGB8bppViaLookupTable(FFloat16Color* InCol
 {
 	TArray<float> sRGBTable = GenerateSRGBTableFloat16toFloat();
 
-	static FRandomStream RandStream(GFrameCounter);
-
 	// This quantization uses a table of float values from 0.0-255.0, and does the proper rounding for each pixel.
 	// The proper rounding is done by adding a random value from 0.0-1.0 instead of 0.5 to each value returned from the table.
+	static FRandomStream RandStream(GFrameCounter);
+
+	// FRandomStream is too slow to call 30,000,000 times (4k image) so instead
+	// we pre-calculate a reasonable number of offsets into a table, and then 
+	// each pixel in the output image takes a offset. Some care needs to be taken
+	// that the table doesn't introduce an obvious repeating pattern in dithering.
+	TArray<float> RandStreamTable;
+	int32 TableSize = (1024 * 1024) + 7; // Arbitrarily picked for now
+	RandStreamTable.SetNumUninitialized(TableSize);
+	for (int32 Index = 0; Index < TableSize; Index++)
+	{
+		RandStreamTable[Index] = RandStream.GetFraction();
+	}
 	// Convert all of our pixels.
 	TArray<FColor> OutsRGBData;
 	OutsRGBData.SetNumUninitialized(InCount);
@@ -172,17 +183,28 @@ static TArray<FColor> ConvertLinearTosRGB8bppViaLookupTable(FFloat16Color* InCol
 	FColor* OutData = static_cast<FColor*>(OutsRGBData.GetData());
 
 	QUICK_SCOPE_CYCLE_COUNTER(STAT_ImageQuant_ApplysRGB);
-	for (int32 PixelIndex = 0; PixelIndex < InCount; PixelIndex++)
-	{
-		// Avoid the bounds checking of TArray[]
-		FColor* OutColor = &OutData[PixelIndex];
-		OutColor->R = (uint8)FMath::FloorToInt(sRGBTableData[InColor[PixelIndex].R.Encoded] + RandStream.GetFraction());
-		OutColor->G = (uint8)FMath::FloorToInt(sRGBTableData[InColor[PixelIndex].G.Encoded] + RandStream.GetFraction());
-		OutColor->B = (uint8)FMath::FloorToInt(sRGBTableData[InColor[PixelIndex].B.Encoded] + RandStream.GetFraction());
 
-		// Alpha doesn't get sRGB conversion, it just gets linearly converted to 8 bit. Flooring avoids an extra branch for Round.
-		OutColor->A = (uint8)FMath::Clamp(FMath::FloorToInt((InColor[PixelIndex].A * 255.f) + RandStream.GetFraction()), 0, 255);
-	}
+	int32 BatchSize = 64 * 1024 * 8; // Arbitrarily picked for now.
+	const uint32 Loops = FMath::CeilToInt((float)InCount / BatchSize);
+	ParallelFor(Loops,
+		[&](int32 LoopIndex)
+		{
+			const int32 Start = LoopIndex * BatchSize;
+			const int32 End = FMath::Min((LoopIndex + 1) * BatchSize, InCount);
+
+
+			for (int32 PixelIndex = Start; PixelIndex < End; PixelIndex++)
+			{
+				// Avoid the bounds checking of TArray[]
+				FColor* OutColor = &OutData[PixelIndex];
+				OutColor->R = (uint8)FMath::FloorToInt(sRGBTableData[InColor[PixelIndex].R.Encoded] + RandStreamTable[(PixelIndex + 0) % TableSize]);
+				OutColor->G = (uint8)FMath::FloorToInt(sRGBTableData[InColor[PixelIndex].G.Encoded] + RandStreamTable[(PixelIndex + 1) % TableSize]);
+				OutColor->B = (uint8)FMath::FloorToInt(sRGBTableData[InColor[PixelIndex].B.Encoded] + RandStreamTable[(PixelIndex + 2) % TableSize]);
+	
+				// Alpha doesn't get sRGB conversion, it just gets linearly converted to 8 bit. Flooring avoids an extra branch for Round.
+				OutColor->A = (uint8)FMath::Clamp(FMath::FloorToInt((InColor[PixelIndex].A * 255.f) + RandStreamTable[(PixelIndex + 3) % TableSize]), 0, 255);
+			}
+		});
 
 	return OutsRGBData;
 }
diff --git a/Engine/Plugins/MovieScene/MovieRenderPipeline/Source/MovieRenderPipelineRenderPasses/Private/MoviePipelineImageSequenceOutput.cpp b/Engine/Plugins/MovieScene/MovieRenderPipeline/Source/MovieRenderPipelineRenderPasses/Private/MoviePipelineImageSequenceOutput.cpp
--- a/Engine/Plugins/MovieScene/MovieRenderPipeline/Source/MovieRenderPipelineRenderPasses/Private/MoviePipelineImageSequenceOutput.cpp
+++ b/Engine/Plugins/MovieScene/MovieRenderPipeline/Source/MovieRenderPipelineRenderPasses/Private/MoviePipelineImageSequenceOutput.cpp
@@ -25,6 +25,24 @@
 
 DECLARE_CYCLE_STAT(TEXT("ImgSeqOutput_RecieveImageData"), STAT_ImgSeqRecieveImageData, STATGROUP_MoviePipeline);
 
+struct FAsyncImageQuantization
+{
+	FAsyncImageQuantization(FImageWriteTask* InWriteTask, const bool bInApplysRGB)
+		: ParentWriteTask(InWriteTask)
+		, bApplysRGB(bInApplysRGB)
+	{}
+
+	void operator()(FImagePixelData* PixelData)
+	{
+		// Convert the incoming data to 8-bit, potentially with sRGB applied.
+		TUniquePtr<FImagePixelData> QuantizedPixelData = UE::MoviePipeline::QuantizeImagePixelDataToBitDepth(PixelData, 8, nullptr,  bApplysRGB);
+		ParentWriteTask->PixelData = MoveTemp(QuantizedPixelData);
+	}
+
+	FImageWriteTask* ParentWriteTask;
+	bool bApplysRGB;
+};
+
 UMoviePipelineImageSequenceOutputBase::UMoviePipelineImageSequenceOutputBase()
 {
 	if (!HasAnyFlags(RF_ArchetypeObject))
@@ -118,27 +136,6 @@ void UMoviePipelineImageSequenceOutputBase::OnReceiveImageDataImpl(FMoviePipelin
 		case EImageFormat::EXR: Extension = TEXT("exr"); break;
 		}
 
-		TUniquePtr<FImagePixelData> QuantizedPixelData = nullptr;
-		
-
-		switch (PreferredOutputFormat)
-		{
-		case EImageFormat::PNG:
-		case EImageFormat::JPEG:
-		case EImageFormat::BMP:
-		{
-			// All three of these formats only support 8 bit data, so we need to take the incoming buffer type,
-			// copy it into a new 8-bit array and optionally apply a little noise to the data to help hide gradient banding.
-			QuantizedPixelData = UE::MoviePipeline::QuantizeImagePixelDataToBitDepth(RenderPassData.Value.Get(), 8, nullptr, !(ColorSetting && ColorSetting->OCIOConfiguration.bIsEnabled));
-			break;
-		}
-		case EImageFormat::EXR:
-			// No quantization required, just copy the data as we will move it into the image write task.
-			QuantizedPixelData = RenderPassData.Value->CopyImageData();
-			break;
-		default:
-			check(false);
-		}
 
 		// We need to resolve the filename format string. We combine the folder and file name into one long string first
 		MoviePipeline::FMoviePipelineOutputFutureData OutputData;
@@ -198,6 +195,32 @@ void UMoviePipelineImageSequenceOutputBase::OnReceiveImageDataImpl(FMoviePipelin
 		TileImageTask->CompressionQuality = 100;
 		TileImageTask->Filename = OutputData.FilePath;
 
+		TUniquePtr<FImagePixelData> QuantizedPixelData = RenderPassData.Value->CopyImageData();
+		EImagePixelType QuantizedPixelType = QuantizedPixelData->GetType();
+
+		switch (PreferredOutputFormat)
+		{
+		case EImageFormat::PNG:
+		case EImageFormat::JPEG:
+		case EImageFormat::BMP:
+		{
+			// All three of these formats only support 8 bit data, so we need to take the incoming buffer type,
+			// copy it into a new 8-bit array and apply a little noise to the data to help hide gradient banding.
+			const bool bApplysRGB = !(ColorSetting && ColorSetting->OCIOConfiguration.bIsEnabled);
+			TileImageTask->PixelPreProcessors.Add(FAsyncImageQuantization(TileImageTask.Get(), bApplysRGB));
+
+			// The pixel type will get changed by this pre-processor so future calculations below need to know the correct type they'll be editing.
+			QuantizedPixelType = EImagePixelType::Color; 
+			break;
+		}
+		case EImageFormat::EXR:
+			// No quantization required, just copy the data as we will move it into the image write task.
+			break;
+		default:
+			check(false);
+		}
+
+
 		// We composite before flipping the alpha so that it is consistent for all formats.
 		if (RenderPassData.Key == FMoviePipelinePassIdentifier(TEXT("FinalImage")))
 		{
@@ -205,7 +228,7 @@ void UMoviePipelineImageSequenceOutputBase::OnReceiveImageDataImpl(FMoviePipelin
 			{
 				// We don't need to copy the data here (even though it's being passed to a async system) because we already made a unique copy of the
 				// burn in/widget data when we decided to composite it.
-				switch (QuantizedPixelData->GetType())
+				switch (QuantizedPixelType)
 				{
 				case EImagePixelType::Color:
 					TileImageTask->PixelPreProcessors.Add(TAsyncCompositeImage<FColor>(CompositePass.PixelData->MoveImageDataToNew()));
diff --git a/Engine/Source/Runtime/ImageWriteQueue/Private/ImageWriteTask.cpp b/Engine/Source/Runtime/ImageWriteQueue/Private/ImageWriteTask.cpp
--- a/Engine/Source/Runtime/ImageWriteQueue/Private/ImageWriteTask.cpp
+++ b/Engine/Source/Runtime/ImageWriteQueue/Private/ImageWriteTask.cpp
@@ -146,10 +146,11 @@ bool FImageWriteTask::WriteBitmap()
 
 void FImageWriteTask::PreProcess()
 {
-	FImagePixelData* Data = PixelData.Get();
 	for (const FPixelPreProcessor& PreProcessor : PixelPreProcessors)
 	{
-		// PreProcessors are assumed to be valid.
+		// PreProcessors are assumed to be valid. Fetch the Data pointer each time
+		// in case a pre-processor changes our pixel data.
+		FImagePixelData* Data = PixelData.Get();
 		PreProcessor(Data);
 	}
 }
