MoviePipeline: The project settings class can now properly store references to late-loaded class types such as those defined in Python or in a game's Plugins.

diff --git a/Engine/Plugins/MovieScene/MovieRenderPipeline/Source/MovieRenderPipelineCore/Private/MovieRenderPipelineCommandLine.cpp b/Engine/Plugins/MovieScene/MovieRenderPipeline/Source/MovieRenderPipelineCore/Private/MovieRenderPipelineCommandLine.cpp
--- a/Engine/Plugins/MovieScene/MovieRenderPipeline/Source/MovieRenderPipelineCore/Private/MovieRenderPipelineCommandLine.cpp
+++ b/Engine/Plugins/MovieScene/MovieRenderPipeline/Source/MovieRenderPipelineCore/Private/MovieRenderPipelineCommandLine.cpp
@@ -321,7 +321,7 @@ uint8 FMovieRenderPipelineCoreModule::ParseMovieRenderData(const FString& InSequ
 
 		if (!ExecutorClass)
 		{
-			UE_LOG(LogMovieRenderPipeline, Fatal, TEXT("No class set in Python Host Executor. This does nothing without setting a class via \"unreal.get_default_object(unreal.MoviePipelinePythonHostExecutor).executor_class = self.get_class(), or passing -ExecutorPythonClass=...\"."));
+			UE_LOG(LogMovieRenderPipeline, Fatal, TEXT("No class set in Python Host Executor. This does nothing without setting a class via \"unreal.get_default_object(unreal.MoviePipelinePythonHostExecutor).executor_class = self.get_class()\", or passing \"-ExecutorPythonClass=/Engine/PythonTypes.MoviePipelineExampleRuntimeExecutor\" on the command line!"));
 			return MoviePipelineErrorCodes::Critical;
 		}
 	}
diff --git a/Engine/Plugins/MovieScene/MovieRenderPipeline/Source/MovieRenderPipelineEditor/Private/MoviePipelineEditorBlueprintLibrary.cpp b/Engine/Plugins/MovieScene/MovieRenderPipeline/Source/MovieRenderPipelineEditor/Private/MoviePipelineEditorBlueprintLibrary.cpp
--- a/Engine/Plugins/MovieScene/MovieRenderPipeline/Source/MovieRenderPipelineEditor/Private/MoviePipelineEditorBlueprintLibrary.cpp
+++ b/Engine/Plugins/MovieScene/MovieRenderPipeline/Source/MovieRenderPipelineEditor/Private/MoviePipelineEditorBlueprintLibrary.cpp
@@ -148,7 +148,7 @@ UMoviePipelineExecutorJob* UMoviePipelineEditorBlueprintLibrary::CreateJobFromSe
 
 	InPipelineQueue->Modify();
 
-	UMoviePipelineExecutorJob* NewJob = InPipelineQueue->AllocateNewJob(ProjectSettings->DefaultExecutorJob);
+	UMoviePipelineExecutorJob* NewJob = InPipelineQueue->AllocateNewJob(ProjectSettings->DefaultExecutorJob.TryLoadClass<UMoviePipelineExecutorJob>());
 	if (!ensureAlwaysMsgf(NewJob, TEXT("Failed to allocate new job! Check the DefaultExecutorJob is not null in Project Settings!")))
 	{
 		return nullptr;
@@ -190,8 +190,9 @@ UMoviePipelineExecutorJob* UMoviePipelineEditorBlueprintLibrary::CreateJobFromSe
 void UMoviePipelineEditorBlueprintLibrary::EnsureJobHasDefaultSettings(UMoviePipelineExecutorJob* NewJob)
 {
 	const UMovieRenderPipelineProjectSettings* ProjectSettings = GetDefault<UMovieRenderPipelineProjectSettings>();
-	for (TSubclassOf<UMoviePipelineSetting> SettingClass : ProjectSettings->DefaultClasses)
+	for (FSoftClassPath SettingClassPath : ProjectSettings->DefaultClasses)
 	{
+		TSubclassOf<UMoviePipelineSetting> SettingClass = SettingClassPath.TryLoadClass<UMoviePipelineSetting>();
 		if (!SettingClass)
 		{
 			continue;
diff --git a/Engine/Plugins/MovieScene/MovieRenderPipeline/Source/MovieRenderPipelineEditor/Private/MoviePipelinePIEExecutor.cpp b/Engine/Plugins/MovieScene/MovieRenderPipeline/Source/MovieRenderPipelineEditor/Private/MoviePipelinePIEExecutor.cpp
--- a/Engine/Plugins/MovieScene/MovieRenderPipeline/Source/MovieRenderPipelineEditor/Private/MoviePipelinePIEExecutor.cpp
+++ b/Engine/Plugins/MovieScene/MovieRenderPipeline/Source/MovieRenderPipelineEditor/Private/MoviePipelinePIEExecutor.cpp
@@ -178,7 +178,13 @@ void UMoviePipelinePIEExecutor::OnPIEStartupFinished(bool)
 	// Allow the user to have overridden which Pipeline is actually run. This is an unlikely scenario but allows
 	// the user to create their own implementation while still re-using the rest of our UI and infrastructure.
 	const UMovieRenderPipelineProjectSettings* ProjectSettings = GetDefault<UMovieRenderPipelineProjectSettings>();
-	TSubclassOf<UMoviePipeline> PipelineClass = ProjectSettings->DefaultPipeline;
+	TSubclassOf<UMoviePipeline> PipelineClass = ProjectSettings->DefaultPipeline.TryLoadClass<UMoviePipeline>();
+	if (!PipelineClass)
+	{
+		UE_LOG(LogMovieRenderPipeline, Error, TEXT("Failed to load project specified MoviePipeline class type!"));
+		OnExecutorFinishedImpl();
+		return;
+	}
 
 	// This Pipeline belongs to the world being created so that they have context for things they execute.
 	ActiveMoviePipeline = NewObject<UMoviePipeline>(ExecutingWorld, PipelineClass);
diff --git a/Engine/Plugins/MovieScene/MovieRenderPipeline/Source/MovieRenderPipelineEditor/Private/Widgets/SMoviePipelineQueuePanel.cpp b/Engine/Plugins/MovieScene/MovieRenderPipeline/Source/MovieRenderPipelineEditor/Private/Widgets/SMoviePipelineQueuePanel.cpp
--- a/Engine/Plugins/MovieScene/MovieRenderPipeline/Source/MovieRenderPipelineEditor/Private/Widgets/SMoviePipelineQueuePanel.cpp
+++ b/Engine/Plugins/MovieScene/MovieRenderPipeline/Source/MovieRenderPipelineEditor/Private/Widgets/SMoviePipelineQueuePanel.cpp
@@ -265,9 +265,11 @@ FReply SMoviePipelineQueuePanel::OnRenderLocalRequested()
 	check(Subsystem);
 
 	const UMovieRenderPipelineProjectSettings* ProjectSettings = GetDefault<UMovieRenderPipelineProjectSettings>();
-	check(ProjectSettings->DefaultLocalExecutor != nullptr);
+	TSubclassOf<UMoviePipelineExecutorBase> ExecutorClass = ProjectSettings->DefaultLocalExecutor.TryLoadClass<UMoviePipelineExecutorBase>();
 
-	Subsystem->RenderQueueWithExecutor(ProjectSettings->DefaultLocalExecutor);
+	// OnRenderLocalRequested should only get called if IsRenderLocalEnabled() returns true, meaning there's a valid class.
+	check(ExecutorClass != nullptr);
+	Subsystem->RenderQueueWithExecutor(ExecutorClass);
 	return FReply::Handled();
 }
 
@@ -277,7 +279,7 @@ bool SMoviePipelineQueuePanel::IsRenderLocalEnabled() const
 	check(Subsystem);
 
 	const UMovieRenderPipelineProjectSettings* ProjectSettings = GetDefault<UMovieRenderPipelineProjectSettings>();
-	const bool bHasExecutor = ProjectSettings->DefaultLocalExecutor != nullptr;
+	const bool bHasExecutor = ProjectSettings->DefaultLocalExecutor.TryLoadClass<UMoviePipelineExecutorBase>() != nullptr;
 	const bool bNotRendering = !Subsystem->IsRendering();
 	const bool bConfigWindowIsOpen = WeakEditorWindow.IsValid();
 
@@ -301,9 +303,12 @@ FReply SMoviePipelineQueuePanel::OnRenderRemoteRequested()
 	check(Subsystem);
 
 	const UMovieRenderPipelineProjectSettings* ProjectSettings = GetDefault<UMovieRenderPipelineProjectSettings>();
-	check(ProjectSettings->DefaultRemoteExecutor != nullptr);
+	TSubclassOf<UMoviePipelineExecutorBase> ExecutorClass = ProjectSettings->DefaultRemoteExecutor.TryLoadClass<UMoviePipelineExecutorBase>();
 
-	Subsystem->RenderQueueWithExecutor(ProjectSettings->DefaultRemoteExecutor);
+	// OnRenderRemoteRequested should only get called if IsRenderRemoteEnabled() returns true, meaning there's a valid class.
+	check(ExecutorClass != nullptr);
+
+	Subsystem->RenderQueueWithExecutor(ExecutorClass);
 	return FReply::Handled();
 }
 
@@ -313,7 +318,7 @@ bool SMoviePipelineQueuePanel::IsRenderRemoteEnabled() const
 	check(Subsystem);
 
 	const UMovieRenderPipelineProjectSettings* ProjectSettings = GetDefault<UMovieRenderPipelineProjectSettings>();
-	const bool bHasExecutor = ProjectSettings->DefaultRemoteExecutor != nullptr;
+	const bool bHasExecutor = ProjectSettings->DefaultRemoteExecutor.TryLoadClass<UMoviePipelineExecutorBase>() != nullptr;
 	const bool bNotRendering = !Subsystem->IsRendering();
 	const bool bConfigWindowIsOpen = WeakEditorWindow.IsValid();
 
diff --git a/Engine/Plugins/MovieScene/MovieRenderPipeline/Source/MovieRenderPipelineEditor/Public/MovieRenderPipelineSettings.h b/Engine/Plugins/MovieScene/MovieRenderPipeline/Source/MovieRenderPipelineEditor/Public/MovieRenderPipelineSettings.h
--- a/Engine/Plugins/MovieScene/MovieRenderPipeline/Source/MovieRenderPipelineEditor/Public/MovieRenderPipelineSettings.h
+++ b/Engine/Plugins/MovieScene/MovieRenderPipeline/Source/MovieRenderPipelineEditor/Public/MovieRenderPipelineSettings.h
@@ -33,7 +33,7 @@ public:
 	/**
 	* What was the last configuration preset the user used? Can be null.
 	*/
-	UPROPERTY(Transient, EditAnywhere, BlueprintReadWrite, Category = "Movie Render Pipeline")
+	UPROPERTY(Transient, BlueprintReadWrite, Category = "Movie Render Pipeline")
 	TSoftObjectPtr<UMoviePipelineMasterConfig> LastPresetOrigin;
 	
 	/**
@@ -42,8 +42,8 @@ public:
 	* which does different logic. See UMoviePipelineExecutorBase for more information. This is used for
 	* the Render button on the UI.
 	*/
-	UPROPERTY(config, EditAnywhere, BlueprintReadWrite, Category="Movie Render Pipeline")
-	TSubclassOf<UMoviePipelineExecutorBase> DefaultLocalExecutor;
+	UPROPERTY(config, EditAnywhere, BlueprintReadWrite, meta = (MetaClass = "MoviePipelineExecutorBase"), Category="Movie Render Pipeline")
+	FSoftClassPath DefaultLocalExecutor;
 
 	/**
 	* When the user uses the UI to request we render a movie remotely, which implementation should we use
@@ -51,24 +51,24 @@ public:
 	* which does different logic. See UMoviePipelineExecutorBase for more information. This is used for
 	* the Render Remotely button on the UI.
 	*/
-	UPROPERTY(config, EditAnywhere, BlueprintReadWrite, Category = "Movie Render Pipeline")
-	TSubclassOf<UMoviePipelineExecutorBase> DefaultRemoteExecutor;
+	UPROPERTY(config, EditAnywhere, BlueprintReadWrite, meta = (MetaClass = "MoviePipelineExecutorBase"), Category = "Movie Render Pipeline")
+	FSoftClassPath DefaultRemoteExecutor;
 	
 	/**
 	* Which Job class should we create by default when adding a job? This allows you to make custom jobs
 	* that will have editable properties in the UI for special handling with your executor. This can be
 	* made dynamic if you add jobs to the queue programatically instead of through the UI.
 	*/
-	UPROPERTY(config, EditAnywhere, BlueprintReadWrite, Category="Movie Render Pipeline")
-	TSubclassOf<UMoviePipelineExecutorJob> DefaultExecutorJob;
+	UPROPERTY(config, EditAnywhere, BlueprintReadWrite, meta = (MetaClass = "MoviePipelineExecutorJob"), Category="Movie Render Pipeline")
+	FSoftClassPath DefaultExecutorJob;
 	
 	/**
 	* This allows you to implement your own Pipeline to handle timing and rendering of a movie. Changing
 	* this will allow you to re-use the existing UI/Executors while providing your own logic for producing
 	* a single render.
 	*/
-	UPROPERTY(config, EditAnywhere, BlueprintReadWrite, Category="Movie Render Pipeline")
-	TSubclassOf<UMoviePipeline> DefaultPipeline;
+	UPROPERTY(config, EditAnywhere, BlueprintReadWrite, meta = (MetaClass = "MoviePipeline"), Category="Movie Render Pipeline")
+	FSoftClassPath DefaultPipeline;
 
 	/**
 	* The settings specified here will automatically be added to a Movie Pipeline Master Configuration when using the UI. 
@@ -78,6 +78,6 @@ public:
 	*
 	* This only applies to jobs created via the UI. If you do not use the UI (ie: Scripting/Python) you will need to
 	* add settings by hand for each job you create. */
-	UPROPERTY(config, EditAnywhere, BlueprintReadWrite, Category = "Movie Render Pipeline", DisplayName = "Default Job Settings Classes")
-	TArray<TSubclassOf<UMoviePipelineSetting>> DefaultClasses;
+	UPROPERTY(config, EditAnywhere, BlueprintReadWrite, Category = "Movie Render Pipeline", meta = (MetaClass = "MoviePipelineSetting"), DisplayName = "Default Job Settings Classes")
+	TArray<FSoftClassPath> DefaultClasses;
 };
\ No newline at end of file
