Movie Pipeline: Fix missing preset icon

diff --git a/Engine/Plugins/MovieScene/MovieRenderPipeline/Source/MovieRenderPipelineEditor/Private/Widgets/SMoviePipelineConfigPanel.cpp b/Engine/Plugins/MovieScene/MovieRenderPipeline/Source/MovieRenderPipelineEditor/Private/Widgets/SMoviePipelineConfigPanel.cpp
--- a/Engine/Plugins/MovieScene/MovieRenderPipeline/Source/MovieRenderPipelineEditor/Private/Widgets/SMoviePipelineConfigPanel.cpp
+++ b/Engine/Plugins/MovieScene/MovieRenderPipeline/Source/MovieRenderPipelineEditor/Private/Widgets/SMoviePipelineConfigPanel.cpp
@@ -146,7 +146,7 @@ void SMoviePipelineConfigPanel::Construct(const FArguments& InArgs, TSubclassOf<
 						.AutoWidth()
 						[
 							SNew(SImage)
-							.Image(FEditorStyle::Get().GetBrush("AssetEditor.SaveAsset.Greyscale"))
+							.Image(FEditorStyle::Get().GetBrush("AssetEditor.SaveAsset"))
 						]
 
 						+ SHorizontalBox::Slot()
