MoviePipeline: Fixed multiple issues in the Command Line Encoder.
Allows adding the duration of each frame to the generated text file - this fixes a bug in third party software when using png as an image source.
Fixed file being written to a folder containing {project_dir} in the final name (wasn't being resolved), which generated empty folders on the users harddrive.
Added additional logging when the process isn't found/couldn't be started.
Fixed an issue where if the user copy/pasted.

diff --git a/Engine/Plugins/MovieScene/MovieRenderPipeline/Source/MovieRenderPipelineCore/Private/MoviePipelineCommandLineEncoder.cpp b/Engine/Plugins/MovieScene/MovieRenderPipeline/Source/MovieRenderPipelineCore/Private/MoviePipelineCommandLineEncoder.cpp
--- a/Engine/Plugins/MovieScene/MovieRenderPipeline/Source/MovieRenderPipelineCore/Private/MoviePipelineCommandLineEncoder.cpp
+++ b/Engine/Plugins/MovieScene/MovieRenderPipeline/Source/MovieRenderPipelineCore/Private/MoviePipelineCommandLineEncoder.cpp
@@ -13,6 +13,7 @@
 #include "MoviePipelineUtils.h"
 #include "MoviePipelineQueue.h"
 #include "MoviePipelineDebugSettings.h"
+#include "MoviePipelineBlueprintLibrary.h"
 
 // Forward Declare
 static TArray<FText> GetErrorTexts();
@@ -24,6 +25,7 @@ UMoviePipelineCommandLineEncoder::UMoviePipelineCommandLineEncoder()
 	Quality = EMoviePipelineEncodeQuality::Epic;
 	bDeleteSourceFiles = false;
 	bSkipEncodeOnRenderCanceled = true;
+	bWriteEachFrameDuration = true;
 }
 
 bool UMoviePipelineCommandLineEncoder::HasFinishedExportingImpl()
@@ -94,8 +96,13 @@ void UMoviePipelineCommandLineEncoder::StartEncodingProcess(TArray<FMoviePipelin
 	/** We produce one file per render pass we detect */
 	TMap<FMoviePipelinePassIdentifier, FEncoderParams> RenderPasses;
 	
+	// The path shouldn't have quotes on it as it's already kept as a separate argument right up until creating the process, at which point
+	// the platform puts quotes around the FString if needed.
+	FString ExecutablePathNoQuotes = EncoderSettings->ExecutablePath.Replace(TEXT("\""), TEXT(""));
+	FPaths::NormalizeFilename(ExecutablePathNoQuotes);
+	
 	FStringFormatNamedArguments SharedArguments;
-	SharedArguments.Add(TEXT("Executable"), EncoderSettings->ExecutablePath);
+	SharedArguments.Add(TEXT("Executable"), ExecutablePathNoQuotes);
 	SharedArguments.Add(TEXT("AudioCodec"), EncoderSettings->AudioCodec);
 	SharedArguments.Add(TEXT("VideoCodec"), EncoderSettings->VideoCodec);
 	FFrameRate RenderFrameRate = GetPipeline()->GetPipelineMasterConfig()->GetEffectiveFrameRate(GetPipeline()->GetTargetSequence());
@@ -196,34 +203,48 @@ void UMoviePipelineCommandLineEncoder::LaunchEncoder(const FEncoderParams& InPar
 	TStringBuilder<64> StringBuilder;
 	TArray<FString> VideoInputs;
 	TArray<FString> AudioInputs;
+
+	double InFrameRate = InParams.NamedArguments[TEXT("FrameRate")].DoubleValue;
+	double FrameRateAsDuration = 1.0 / InFrameRate;
+
 	for (const TTuple<FString, TArray<FString>>& Pair : InParams.FilesByExtensionType)
 	{
 		FGuid FileGuid = FGuid::NewGuid();
-		FString FilePath = OutputSetting->OutputDirectory.Path / FileGuid.ToString() + TEXT("_input.txt");
+		FString FilePath = OutputSetting->OutputDirectory.Path / FileGuid.ToString() + TEXT("_input");
 
-		if (FPaths::IsRelative(FilePath))
-		{
-			FilePath = FPaths::ConvertRelativePathToFull(FilePath);
-		}
+		FMoviePipelineFormatArgs FinalFormatArgs;
+
+		FString FinalFilePath;
+		TMap<FString, FString> FormatOverrides;
+		FormatOverrides.Add(TEXT("ext"), TEXT("txt"));
+
+		GetPipeline()->ResolveFilenameFormatArguments(FilePath, FormatOverrides, FinalFilePath, FinalFormatArgs);
+		
 
-		UE_LOG(LogMovieRenderPipelineIO, Log, TEXT("Generated Path '%s' for input data."), *FilePath);
+		UE_LOG(LogMovieRenderPipelineIO, Log, TEXT("Generated Path '%s' for input data."), *FinalFilePath);
 		StringBuilder.Reset();
 		for (const FString& Path : Pair.Value)
 		{
 			StringBuilder.Appendf(TEXT("file 'file:%s'%s"), *Path, LINE_TERMINATOR);
+
+			// Some encoders require the duration of each file to be listed after the file.
+			if (Pair.Key != TEXT("wav") && bWriteEachFrameDuration)
+			{
+				StringBuilder.Appendf(TEXT("duration %f%s"), FrameRateAsDuration, LINE_TERMINATOR);
+			}
 		}
 
 		// Save this to disk.
-		FFileHelper::SaveStringToFile(StringBuilder.ToString(), *FilePath);
+		FFileHelper::SaveStringToFile(StringBuilder.ToString(), *FinalFilePath);
 
 		// Not a great solution but best we've got right now
 		if (Pair.Key == TEXT("wav"))
 		{
-			AudioInputs.Add(FilePath);
+			AudioInputs.Add(FinalFilePath);
 		}
 		else
 		{
-			VideoInputs.Add(FilePath);
+			VideoInputs.Add(FinalFilePath);
 		}
 	}
 
@@ -290,8 +311,8 @@ void UMoviePipelineCommandLineEncoder::LaunchEncoder(const FEncoderParams& InPar
 			NewJob.FilesToDelete.Append(AudioInputs);
 		}
 
-		// And the user's input files (if requested)
-		if (bDeleteSourceFiles)
+		// And the user's input files (if requested), though we ignore this if you have the Debug Setting asking you to write all samples.
+		if (bDeleteSourceFiles && bDeleteInputTexts)
 		{
 			for (const TTuple<FString, TArray<FString>>& Pair : InParams.FilesByExtensionType)
 			{
@@ -299,6 +320,11 @@ void UMoviePipelineCommandLineEncoder::LaunchEncoder(const FEncoderParams& InPar
 			}
 		}
 	}
+	else
+	{
+		UE_LOG(LogMovieRenderPipeline, Error, TEXT("Failed to launch encoder process, see output log for more details."));
+		GetPipeline()->Shutdown(true);
+	}
 }
 
 void UMoviePipelineCommandLineEncoder::OnTick()
diff --git a/Engine/Plugins/MovieScene/MovieRenderPipeline/Source/MovieRenderPipelineCore/Public/MoviePipelineCommandLineEncoder.h b/Engine/Plugins/MovieScene/MovieRenderPipeline/Source/MovieRenderPipelineCore/Public/MoviePipelineCommandLineEncoder.h
--- a/Engine/Plugins/MovieScene/MovieRenderPipeline/Source/MovieRenderPipelineCore/Public/MoviePipelineCommandLineEncoder.h
+++ b/Engine/Plugins/MovieScene/MovieRenderPipeline/Source/MovieRenderPipelineCore/Public/MoviePipelineCommandLineEncoder.h
@@ -73,6 +73,10 @@ public:
 	/** If a render was canceled (via hitting escape mid render) should we skip trying to encode the files we did produce? */
 	UPROPERTY(EditAnywhere, BlueprintReadWrite, Category = "Command Line Encoder")
 	bool bSkipEncodeOnRenderCanceled;
+
+	/** Write the duration for each frame into the generated text file. Needed for some input types on some CLI encoding software. */
+	UPROPERTY(EditAnywhere, BlueprintReadWrite, Category = "Command Line Encoder")
+	bool bWriteEachFrameDuration;
 	
 private:
 	struct FActiveJob
