MoviePipeline: Format strings that come from shots now provide the correct values when used in filenames, instead of always using the Master Configuration's values.

diff --git a/Engine/Plugins/MovieScene/MovieRenderPipeline/Source/MovieRenderPipelineCore/Private/MoviePipelineBlueprintLibrary.cpp b/Engine/Plugins/MovieScene/MovieRenderPipeline/Source/MovieRenderPipelineCore/Private/MoviePipelineBlueprintLibrary.cpp
--- a/Engine/Plugins/MovieScene/MovieRenderPipeline/Source/MovieRenderPipelineCore/Private/MoviePipelineBlueprintLibrary.cpp
+++ b/Engine/Plugins/MovieScene/MovieRenderPipeline/Source/MovieRenderPipelineCore/Private/MoviePipelineBlueprintLibrary.cpp
@@ -900,12 +900,24 @@ void UMoviePipelineBlueprintLibrary::ResolveFilenameFormatArguments(const FStrin
 		OutMergedFormatArgs.FilenameArguments.Add(TEXT("file_dup"), FString());
 	}
 
-	// Apply any overrides from shots
 	if (InParams.ShotOverride && InParams.ShotOverride->GetShotOverrideConfiguration())
 	{
-		// ToDo
-		// InParams.ShotOverride->GetShotOverrideConfiguration()->GetFormatArguments(OutMergedFormatArgs, true);
+		UMoviePipelineShotConfig* ShotConfig = InParams.ShotOverride->GetShotOverrideConfiguration();
+		for (UMoviePipelineSetting* Setting : ShotConfig->GetUserSettings())
+		{
+			if (!Setting)
+			{
+				UE_LOG(LogMovieRenderPipeline, Error, TEXT("Null setting found in config: %s - Did you disable a plugin that contained this setting?"), *GetNameSafe(ShotConfig));
+				continue;
+			}
+
+			if (Setting->IsEnabled())
+			{
+				Setting->GetFormatArguments(OutMergedFormatArgs);
+			}
+		}
 	}
+
 	// Overwrite the variables with overrides if needed. This allows different requesters to share the same variables (ie: filename extension, render pass name)
 	for (const TPair<FString, FString>& KVP : InParams.FileNameFormatOverrides)
 	{
diff --git a/Engine/Plugins/MovieScene/MovieRenderPipeline/Source/MovieRenderPipelineCore/Private/MoviePipelineMasterConfig.cpp b/Engine/Plugins/MovieScene/MovieRenderPipeline/Source/MovieRenderPipelineCore/Private/MoviePipelineMasterConfig.cpp
--- a/Engine/Plugins/MovieScene/MovieRenderPipeline/Source/MovieRenderPipelineCore/Private/MoviePipelineMasterConfig.cpp
+++ b/Engine/Plugins/MovieScene/MovieRenderPipeline/Source/MovieRenderPipelineCore/Private/MoviePipelineMasterConfig.cpp
@@ -155,9 +155,6 @@ void UMoviePipelineMasterConfig::GetFormatArguments(FMoviePipelineFormatArgs& In
 		{
 			Setting->GetFormatArguments(InOutFormatArgs);
 		}
-
-		// ToDo: Should shots be able to provide arguments too? They're only overrides, and
-		// outside of that shot the format would fail and end up leaving the {text} anyways.
 	}
 }
 
diff --git a/Engine/Plugins/MovieScene/MovieRenderPipeline/Source/MovieRenderPipelineCore/Private/MoviePipelineSetting.cpp b/Engine/Plugins/MovieScene/MovieRenderPipeline/Source/MovieRenderPipelineCore/Private/MoviePipelineSetting.cpp
--- a/Engine/Plugins/MovieScene/MovieRenderPipeline/Source/MovieRenderPipelineCore/Private/MoviePipelineSetting.cpp
+++ b/Engine/Plugins/MovieScene/MovieRenderPipeline/Source/MovieRenderPipelineCore/Private/MoviePipelineSetting.cpp
@@ -15,11 +15,8 @@ void UMoviePipelineSetting::OnMoviePipelineInitialized(UMoviePipeline* InPipelin
 
 UMoviePipeline* UMoviePipelineSetting::GetPipeline() const
 {
-	// If this check trips then life cycles aren't as expected.
-	UMoviePipeline* OutPipeline = CachedPipeline.Get();
-	check(OutPipeline);
-
-	return OutPipeline;
+	// Can be null if we're not in the middle of a render (ie in the UI)
+	return CachedPipeline.Get();
 }
 
 UWorld* UMoviePipelineSetting::GetWorld() const
