MoviePipeline: The Path Tracer render pass will now produce validation warnings in the UI / automatically fail renders if you attempt to use it when the engine doesn't support the path tracer.

diff --git a/Engine/Plugins/MovieScene/MovieRenderPipeline/Source/MovieRenderPipelineRenderPasses/Private/MoviePipelineDeferredPasses.cpp b/Engine/Plugins/MovieScene/MovieRenderPipeline/Source/MovieRenderPipelineRenderPasses/Private/MoviePipelineDeferredPasses.cpp
--- a/Engine/Plugins/MovieScene/MovieRenderPipeline/Source/MovieRenderPipelineRenderPasses/Private/MoviePipelineDeferredPasses.cpp
+++ b/Engine/Plugins/MovieScene/MovieRenderPipeline/Source/MovieRenderPipelineRenderPasses/Private/MoviePipelineDeferredPasses.cpp
@@ -749,3 +749,47 @@ FText UMoviePipelineDeferredPass_PathTracer::GetFooterText(UMoviePipelineExecuto
 		"All other Path Tracer settings are taken from the Post Process settings as usual.");
 }
 #endif
+namespace UE
+{
+namespace MoviePipeline
+{
+	bool CheckIfPathTracerIsSupported()
+	{
+		bool bSupportsPathTracing = false;
+		if (IsRayTracingEnabled())
+		{
+			IConsoleVariable* PathTracingCVar = IConsoleManager::Get().FindConsoleVariable(TEXT("r.PathTracing"));
+			if (PathTracingCVar)
+			{
+				bSupportsPathTracing = PathTracingCVar->GetInt() != 0;
+			}
+		}
+		return bSupportsPathTracing;
+	}
+}
+}
+void UMoviePipelineDeferredPass_PathTracer::ValidateStateImpl()
+{
+	Super::ValidateStateImpl();
+
+	bool bSupportsPathTracing = UE::MoviePipeline::CheckIfPathTracerIsSupported();
+	
+	if (!bSupportsPathTracing)
+	{
+		const FText ValidationWarning = NSLOCTEXT("MovieRenderPipeline", "PathTracerValidation_Unsupported", "Path Tracing is currently not enabled for this project and this render pass will not work.");
+		ValidationResults.Add(ValidationWarning);
+		ValidationState = EMoviePipelineValidationState::Warnings;
+	}
+}
+
+void UMoviePipelineDeferredPass_PathTracer::SetupImpl(const MoviePipeline::FMoviePipelineRenderPassInitSettings& InPassInitSettings)
+{
+	if (!UE::MoviePipeline::CheckIfPathTracerIsSupported())
+	{
+		UE_LOG(LogMovieRenderPipeline, Error, TEXT("Cannot render a Path Tracer pass, Path Tracer is not enabled by this project."));
+		GetPipeline()->Shutdown(true);
+		return;
+	}
+
+	Super::SetupImpl(InPassInitSettings);
+}
diff --git a/Engine/Plugins/MovieScene/MovieRenderPipeline/Source/MovieRenderPipelineRenderPasses/Public/MoviePipelineDeferredPasses.h b/Engine/Plugins/MovieScene/MovieRenderPipeline/Source/MovieRenderPipelineRenderPasses/Public/MoviePipelineDeferredPasses.h
--- a/Engine/Plugins/MovieScene/MovieRenderPipeline/Source/MovieRenderPipelineRenderPasses/Public/MoviePipelineDeferredPasses.h
+++ b/Engine/Plugins/MovieScene/MovieRenderPipeline/Source/MovieRenderPipelineRenderPasses/Public/MoviePipelineDeferredPasses.h
@@ -263,4 +263,7 @@ public:
 	virtual int32 GetOutputFileSortingOrder() const override { return 2; }
 
 	virtual bool IsAntiAliasingSupported() const { return false; }
+	virtual void ValidateStateImpl() override;
+	virtual void SetupImpl(const MoviePipeline::FMoviePipelineRenderPassInitSettings& InPassInitSettings) override;
+
 };
