MoviePipeline: Fixed loading a saved Queue not showing the correct names for jobs that used Presets for their settings. Fixed jobs jobs with modified presets (a preset was assigned but then edited, which technically breaks the connection to the preset asset but preserves the preset name) having the wrong name after being duplicated in the queue.

diff --git a/Engine/Plugins/MovieScene/MovieRenderPipeline/Source/MovieRenderPipelineCore/Private/MoviePipelineQueue.cpp b/Engine/Plugins/MovieScene/MovieRenderPipeline/Source/MovieRenderPipelineCore/Private/MoviePipelineQueue.cpp
--- a/Engine/Plugins/MovieScene/MovieRenderPipeline/Source/MovieRenderPipelineCore/Private/MoviePipelineQueue.cpp
+++ b/Engine/Plugins/MovieScene/MovieRenderPipeline/Source/MovieRenderPipelineCore/Private/MoviePipelineQueue.cpp
@@ -63,6 +63,19 @@ UMoviePipelineExecutorJob* UMoviePipelineQueue::DuplicateJob(UMoviePipelineExecu
 #endif
 
 	UMoviePipelineExecutorJob* NewJob = CastChecked<UMoviePipelineExecutorJob>(StaticDuplicateObject(InJob, this));
+
+	// Duplication causes the DisplayNames to get renamed to the asset name (required to support both duplicating
+	// and renaming objects in the Content Browser). So after using the normal duplication, we run our custom
+	// CopyFrom code to transfer the configurations over (which do a display name fixup).
+	NewJob->GetConfiguration()->CopyFrom(InJob->GetConfiguration());
+	for (int32 Index = 0; Index < NewJob->ShotInfo.Num(); Index++)
+	{
+		if (InJob->ShotInfo[Index]->GetShotOverrideConfiguration())
+		{
+			NewJob->ShotInfo[Index]->GetShotOverrideConfiguration()->CopyFrom(InJob->ShotInfo[Index]->GetShotOverrideConfiguration());
+		}
+	}
+
 	NewJob->OnDuplicated();
 	Jobs.Add(NewJob);
 
diff --git a/Engine/Plugins/MovieScene/MovieRenderPipeline/Source/MovieRenderPipelineCore/Public/MoviePipelineQueue.h b/Engine/Plugins/MovieScene/MovieRenderPipeline/Source/MovieRenderPipelineCore/Public/MoviePipelineQueue.h
--- a/Engine/Plugins/MovieScene/MovieRenderPipeline/Source/MovieRenderPipelineCore/Public/MoviePipelineQueue.h
+++ b/Engine/Plugins/MovieScene/MovieRenderPipeline/Source/MovieRenderPipelineCore/Public/MoviePipelineQueue.h
@@ -278,9 +278,9 @@ public:
 	void SetPresetOrigin(UMoviePipelineMasterConfig* InPreset);
 
 	UFUNCTION(BlueprintPure, Category = "Movie Render Pipeline")
-	UMoviePipelineMasterConfig* GetPresetOrigin() const
+	UMoviePipelineMasterConfig* GetPresetOrigin() const 
 	{
-		return PresetOrigin.Get();
+		return PresetOrigin.LoadSynchronous();
 	}
 
 	UFUNCTION(BlueprintPure, Category = "Movie Render Pipeline")
