MoviePipeline: Fixed an issue where having online logins enabled can prevent renders from starting.

diff --git a/Engine/Plugins/MovieScene/MovieRenderPipeline/Source/MovieRenderPipelineEditor/Private/MoviePipelinePIEExecutor.cpp b/Engine/Plugins/MovieScene/MovieRenderPipeline/Source/MovieRenderPipelineEditor/Private/MoviePipelinePIEExecutor.cpp
--- a/Engine/Plugins/MovieScene/MovieRenderPipeline/Source/MovieRenderPipelineEditor/Private/MoviePipelinePIEExecutor.cpp
+++ b/Engine/Plugins/MovieScene/MovieRenderPipeline/Source/MovieRenderPipelineEditor/Private/MoviePipelinePIEExecutor.cpp
@@ -117,6 +117,8 @@ void UMoviePipelinePIEExecutor::Start(const UMoviePipelineExecutorJob* InJob)
 	Params.EditorPlaySettings = PlayInEditorSettings;
 	Params.CustomPIEWindow = CustomWindow;
 	Params.GlobalMapOverride = InJob->Map.GetAssetPathString();
+	// Don't allow the online subsystem to try and authenticate as it will delay PIE startup and no PIE world will exist when PostPIEStarted is called.
+	Params.bAllowOnlineSubsystem = false;
 
 	// Initialize the transient settings so that they will exist in time for the GameOverrides check.
 	InJob->GetConfiguration()->InitializeTransientSettings();
@@ -167,6 +169,7 @@ void UMoviePipelinePIEExecutor::OnPIEStartupFinished(bool)
 	if(!ExecutingWorld)
 	{
 		// This only happens if PIE startup fails and they've usually gotten a pop-up dialog already.
+		UE_LOG(LogMovieRenderPipeline, Error, TEXT("Failed to find a PIE UWorld after OnPIEStartupFinished!"));
 		OnExecutorFinishedImpl();
 		return;
 	}
diff --git a/Engine/Source/Editor/UnrealEd/Private/PlayLevel.cpp b/Engine/Source/Editor/UnrealEd/Private/PlayLevel.cpp
--- a/Engine/Source/Editor/UnrealEd/Private/PlayLevel.cpp
+++ b/Engine/Source/Editor/UnrealEd/Private/PlayLevel.cpp
@@ -2693,6 +2693,8 @@ void UEditorEngine::StartPlayInEditorSession(FRequestPlaySessionParams& InReques
 	// If they need to use the online services, validate that they have provided 
 	// enough pie credentials to launch the desired number of clients.
 	bool bUseOnlineSubsystemForLogin = false;
+	bool bRequestAllowsOnlineSubsystem = InRequestParams.bAllowOnlineSubsystem;
+	if(bRequestAllowsOnlineSubsystem)
 	{
 		int32 DesiredNumberOfClients;
 		EditorPlaySettings->GetPlayNumberOfClients(DesiredNumberOfClients);
@@ -2711,9 +2713,8 @@ void UEditorEngine::StartPlayInEditorSession(FRequestPlaySessionParams& InReques
 				FMessageLog(NAME_CategoryPIE).Warning(ErrorMsg);
 			}
 		}
-
-		UOnlineEngineInterface::Get()->SetShouldTryOnlinePIE(bUseOnlineSubsystemForLogin);
 	}
+	UOnlineEngineInterface::Get()->SetShouldTryOnlinePIE(bUseOnlineSubsystemForLogin);
 
 	PlayInEditorSessionInfo->bUsingOnlinePlatform = bUseOnlineSubsystemForLogin;
 	PlayInEditorSessionInfo->bAnyBlueprintErrors = ErroredBlueprints.Num() > 0;
diff --git a/Engine/Source/Editor/UnrealEd/Public/PlayInEditorDataTypes.h b/Engine/Source/Editor/UnrealEd/Public/PlayInEditorDataTypes.h
--- a/Engine/Source/Editor/UnrealEd/Public/PlayInEditorDataTypes.h
+++ b/Engine/Source/Editor/UnrealEd/Public/PlayInEditorDataTypes.h
@@ -155,6 +155,7 @@ public:
 		: SessionDestination(EPlaySessionDestinationType::InProcess)
 		, WorldType(EPlaySessionWorldType::PlayInEditor)
 		, EditorPlaySettings(nullptr)
+		, bAllowOnlineSubsystem(true)
 	{
 	}
 
@@ -217,6 +218,9 @@ public:
 
 	/** Override which map is loaded for the Play session. This overrides both offline & servers (server only can be overridden in ULevelEditorPlaySettings) */
 	FString GlobalMapOverride;
+	
+	/** If false, then PIE won't try to ask the Online Subsystem/"Use Online Logins" feature if it should authenticate the PIE sessions. */
+	bool bAllowOnlineSubsystem;
 
 	bool HasPlayWorldPlacement() const
 	{
