MoviePipeline: The debug setting now lets you automatically capture an Unreal Insights trace for the duration of the render. This makes it easy for artists to capture a trace file to send to programmers for further analysis when they see unexpectedly slow performance.

diff --git a/Engine/Plugins/MovieScene/MovieRenderPipeline/Source/MovieRenderPipelineCore/Private/MoviePipeline.cpp b/Engine/Plugins/MovieScene/MovieRenderPipeline/Source/MovieRenderPipelineCore/Private/MoviePipeline.cpp
--- a/Engine/Plugins/MovieScene/MovieRenderPipeline/Source/MovieRenderPipelineCore/Private/MoviePipeline.cpp
+++ b/Engine/Plugins/MovieScene/MovieRenderPipeline/Source/MovieRenderPipelineCore/Private/MoviePipeline.cpp
@@ -36,11 +36,13 @@
 #include "ImageWriteQueue.h"
 #include "MoviePipelineHighResSetting.h"
 #include "MoviePipelineCameraSetting.h"
+#include "MoviePipelineDebugSettings.h"
 #include "MoviePipelineQueue.h"
 #include "HAL/FileManager.h"
 #include "Misc/CoreDelegates.h"
 #include "MoviePipelineCommandLineEncoder.h"
 #include "Kismet/KismetSystemLibrary.h"
+#include "ProfilingDebugging/TraceAuxiliary.h"
 
 #if WITH_EDITOR
 #include "MovieSceneExportMetadata.h"
@@ -179,6 +181,15 @@ void UMoviePipeline::Initialize(UMoviePipelineExecutorJob* InJob)
 		return;
 	}
 
+	UMoviePipelineDebugSettings* DebugSetting = FindOrAddSettingForShot<UMoviePipelineDebugSettings>(nullptr);
+	if (DebugSetting)
+	{
+		if (DebugSetting->bCaptureUnrealInsightsTrace)
+		{
+			StartUnrealInsightsCapture();
+		}
+	}
+
 	TargetSequence = Cast<ULevelSequence>(GetCurrentJob()->Sequence.TryLoad());
 
 	CachedSequenceHierarchyRoot = MakeShared<MoviePipeline::FCameraCutSubSectionHierarchyNode>();
@@ -537,6 +548,15 @@ void UMoviePipeline::TransitionToState(const EMovieRenderPipelineState InNewStat
 
 			UE_LOG(LogMovieRenderPipeline, Log, TEXT("Movie Pipeline completed. Duration: %s"), *(FDateTime::UtcNow() - InitializationTime).ToString());
 
+			UMoviePipelineDebugSettings* DebugSetting = FindOrAddSettingForShot<UMoviePipelineDebugSettings>(nullptr);
+			if (DebugSetting)
+			{
+				if (DebugSetting->bCaptureUnrealInsightsTrace)
+				{
+					StopUnrealInsightsCapture();
+				}
+			}
+
 			OnMoviePipelineFinishedImpl();
 		}
 		break;
@@ -1547,4 +1567,40 @@ void UMoviePipeline::PrintVerboseLogForFiles(const TArray<FMoviePipelineShotOutp
 		}
 	}
 }
+
+void UMoviePipeline::StartUnrealInsightsCapture()
+{
+	// Generate a filepath to attempt to store the trace file in.
+	UMoviePipelineOutputSetting* OutputSetting = GetPipelineMasterConfig()->FindSetting<UMoviePipelineOutputSetting>();
+	FString FileName = OutputSetting->FileNameFormat + TEXT("_UnrealInsights");
+	FString FileNameFormatString = OutputSetting->OutputDirectory.Path / FileName;
+
+	// Generate a filename for this encoded file
+	TMap<FString, FString> FormatOverrides;
+	FormatOverrides.Add(TEXT("ext"), TEXT("utrace"));
+	FMoviePipelineFormatArgs FinalFormatArgs;
+
+	FString FinalFilePath;
+	ResolveFilenameFormatArguments(FileNameFormatString, FormatOverrides, FinalFilePath, FinalFormatArgs);
+
+	if (FPaths::IsRelative(FinalFilePath))
+	{
+		FinalFilePath = FPaths::ConvertRelativePathToFull(FinalFilePath);
+	}	
+
+	const bool bTraceStarted = FTraceAuxiliary::Start(FTraceAuxiliary::EConnectionType::File, *FinalFilePath);
+	if (bTraceStarted)
+	{
+		UE_LOG(LogMovieRenderPipeline, Log, TEXT("Started capturing UnrealInsights trace file to %s"), *FinalFilePath);
+	}
+	else
+	{
+		UE_LOG(LogMovieRenderPipeline, Warning, TEXT("Failed to start capturing UnrealInsights trace. Is there already a trace session in progress?"));
+	}
+}
+
+void UMoviePipeline::StopUnrealInsightsCapture()
+{
+	FTraceAuxiliary::Stop();
+}
 #undef LOCTEXT_NAMESPACE // "MoviePipeline"
diff --git a/Engine/Plugins/MovieScene/MovieRenderPipeline/Source/MovieRenderPipelineCore/Private/MoviePipelineBlueprintLibrary.cpp b/Engine/Plugins/MovieScene/MovieRenderPipeline/Source/MovieRenderPipelineCore/Private/MoviePipelineBlueprintLibrary.cpp
--- a/Engine/Plugins/MovieScene/MovieRenderPipeline/Source/MovieRenderPipelineCore/Private/MoviePipelineBlueprintLibrary.cpp
+++ b/Engine/Plugins/MovieScene/MovieRenderPipeline/Source/MovieRenderPipelineCore/Private/MoviePipelineBlueprintLibrary.cpp
@@ -775,7 +775,7 @@ FIntPoint UMoviePipelineBlueprintLibrary::GetEffectiveOutputResolution(UMoviePip
 UMoviePipelineSetting* UMoviePipelineBlueprintLibrary::FindOrGetDefaultSettingForShot(TSubclassOf<UMoviePipelineSetting> InSettingType, const UMoviePipelineMasterConfig* InMasterConfig, const UMoviePipelineExecutorShot* InShot)
 {
 	// Check to see if this setting is in the shot override, if it is we'll use the shot version of that.
-	if (InShot->GetShotOverrideConfiguration())
+	if (InShot && InShot->GetShotOverrideConfiguration())
 	{
 		UMoviePipelineSetting* Setting = InShot->GetShotOverrideConfiguration()->FindSettingByClass(InSettingType);
 		if (Setting)
diff --git a/Engine/Plugins/MovieScene/MovieRenderPipeline/Source/MovieRenderPipelineCore/Public/MoviePipeline.h b/Engine/Plugins/MovieScene/MovieRenderPipeline/Source/MovieRenderPipelineCore/Public/MoviePipeline.h
--- a/Engine/Plugins/MovieScene/MovieRenderPipeline/Source/MovieRenderPipelineCore/Public/MoviePipeline.h
+++ b/Engine/Plugins/MovieScene/MovieRenderPipeline/Source/MovieRenderPipelineCore/Public/MoviePipeline.h
@@ -268,6 +268,11 @@ private:
 	/** Called once when first moving to the Export state. */
 	void BeginExport();
 
+	/** Attempts to start an Unreal Insights capture to a file on disk adjacent to the movie output. */
+	void StartUnrealInsightsCapture();
+	/** Attempts to stop an already started Unreal Insights capture. */
+	void StopUnrealInsightsCapture();
+
 	/** 
 	* Runs the per-tick logic when doing the Finalize state.
 	*
diff --git a/Engine/Plugins/MovieScene/MovieRenderPipeline/Source/MovieRenderPipelineCore/Public/MoviePipelineDebugSettings.h b/Engine/Plugins/MovieScene/MovieRenderPipeline/Source/MovieRenderPipelineCore/Public/MoviePipelineDebugSettings.h
--- a/Engine/Plugins/MovieScene/MovieRenderPipeline/Source/MovieRenderPipelineCore/Public/MoviePipelineDebugSettings.h
+++ b/Engine/Plugins/MovieScene/MovieRenderPipeline/Source/MovieRenderPipelineCore/Public/MoviePipelineDebugSettings.h
@@ -39,6 +39,10 @@ public:
 	/** Used when capturing rendering information with RenderDoc. In Display Rate frames.*/
 	UPROPERTY(EditAnywhere, BlueprintReadWrite, Category = "RenderDoc")
 	int32 CaptureFrame;
+
+	/** If true, automatically capture an Unreal Insights trace file for the duration of the render. */
+	UPROPERTY(EditAnywhere, BlueprintReadWrite, Category = "Profiling")
+	bool bCaptureUnrealInsightsTrace;
 	
 private:
 	bool bIsRenderDebugCaptureAvailable;
