MoviePipeline: The {version} token can now come after {shot_name} in the format list and will be treated as a local version (per shot folder) instead of a global version (shared by all shots).

diff --git a/Engine/Plugins/MovieScene/MovieRenderPipeline/Source/MovieRenderPipelineCore/Private/MoviePipeline.cpp b/Engine/Plugins/MovieScene/MovieRenderPipeline/Source/MovieRenderPipelineCore/Private/MoviePipeline.cpp
--- a/Engine/Plugins/MovieScene/MovieRenderPipeline/Source/MovieRenderPipelineCore/Private/MoviePipeline.cpp
+++ b/Engine/Plugins/MovieScene/MovieRenderPipeline/Source/MovieRenderPipelineCore/Private/MoviePipeline.cpp
@@ -306,11 +306,19 @@ void UMoviePipeline::Initialize(UMoviePipelineExecutorJob* InJob)
 	// delta time not generated by us) so we'll wait until the next engine frame to start rendering.
 	InitializationTime = FDateTime::UtcNow();
 
-	// Can only resolve a version number if it comes before dynamic information 
+	// Resolve the version number for each shot. If the version number comes before the shot name in
+	// the format string, all shots will land on the same 'global' version number (as this is done
+	// before any file writing, so all shots within this job will calculate the same value). If the
+	// version number comes after theo shot name, then each shot may have its own 'local' version number
 	FMoviePipelineFilenameResolveParams Params;
 	Params.InitializationTime = InitializationTime;
 	Params.Job = GetCurrentJob();
-	InitializationVersion = UMoviePipelineBlueprintLibrary::ResolveVersionNumber(Params);
+
+	for (UMoviePipelineExecutorShot* Shot : ActiveShotList)
+	{
+		Params.ShotOverride = Shot;
+		Shot->ShotInfo.VersionNumber = UMoviePipelineBlueprintLibrary::ResolveVersionNumber(Params);
+	}
 
 	// If the shot mask entirely disabled everything we'll transition directly to finish as there is no work to do.
 	if (ActiveShotList.Num() == 0)
@@ -1480,6 +1488,8 @@ void UMoviePipeline::ResolveFilenameFormatArguments(const FString& InFormatStrin
 		Params.FrameNumberShotRel = InOutputState->ShotOutputFrameNumber;
 		Params.FileMetadata = InOutputState->FileMetadata;
 		Params.ShotOverride = ActiveShotList[InOutputState->ShotIndex];
+		Params.InitializationVersion = ActiveShotList[InOutputState->ShotIndex]->ShotInfo.VersionNumber;
+
 	}
 
 	UMoviePipelineOutputSetting* OutputSetting = GetPipelineMasterConfig()->FindSetting<UMoviePipelineOutputSetting>();
@@ -1498,7 +1508,6 @@ void UMoviePipeline::ResolveFilenameFormatArguments(const FString& InFormatStrin
 	Params.bForceRelativeFrameNumbers = bForceRelativeFrameNumbers;
 	Params.FileNameFormatOverrides = InFormatOverrides;
 	Params.InitializationTime = InitializationTime;
-	Params.InitializationVersion = InitializationVersion;
 	Params.Job = GetCurrentJob();
 	Params.AdditionalFrameNumberOffset = InFrameNumberOffset;
 
diff --git a/Engine/Plugins/MovieScene/MovieRenderPipeline/Source/MovieRenderPipelineCore/Public/MoviePipeline.h b/Engine/Plugins/MovieScene/MovieRenderPipeline/Source/MovieRenderPipelineCore/Public/MoviePipeline.h
--- a/Engine/Plugins/MovieScene/MovieRenderPipeline/Source/MovieRenderPipelineCore/Public/MoviePipeline.h
+++ b/Engine/Plugins/MovieScene/MovieRenderPipeline/Source/MovieRenderPipelineCore/Public/MoviePipeline.h
@@ -405,9 +405,6 @@ private:
 	/** The time (in UTC) that Initialize was called. Used to track elapsed time. */
 	FDateTime InitializationTime;
 
-	/** The version number for this render. Detected when job starts to increment to the next version so that image sequences don't think it's a new version for every frame. */
-	int32 InitializationVersion;
-
 	FMoviePipelineFrameOutputState CachedOutputState;
 
 	MoviePipeline::FAudioState AudioState;
diff --git a/Engine/Plugins/MovieScene/MovieRenderPipeline/Source/MovieRenderPipelineCore/Public/MovieRenderPipelineDataTypes.h b/Engine/Plugins/MovieScene/MovieRenderPipeline/Source/MovieRenderPipelineCore/Public/MovieRenderPipelineDataTypes.h
--- a/Engine/Plugins/MovieScene/MovieRenderPipeline/Source/MovieRenderPipelineCore/Public/MovieRenderPipelineDataTypes.h
+++ b/Engine/Plugins/MovieScene/MovieRenderPipeline/Source/MovieRenderPipelineCore/Public/MovieRenderPipelineDataTypes.h
@@ -469,6 +469,7 @@ public:
 		, State(EMovieRenderShotState::Uninitialized)
 		, bHasEvaluatedMotionBlurFrame(false)
 		, NumEngineWarmUpFramesRemaining(0)
+		, VersionNumber(0)
 	{
 	}
 
@@ -532,6 +533,9 @@ public:
 
 	/** How many engine warm up frames are left to process for this shot. May be zero. */
 	int32 NumEngineWarmUpFramesRemaining;
+
+	/** What version number should this shot use when resolving format arguments. */
+	int32 VersionNumber;
 };
 
 /**
