MoviePipeline: Add UnrealInsights trace bookmarks for each frame in MRQ so it's easier to understand the timing in context.

diff --git a/Engine/Plugins/MovieScene/MovieRenderPipeline/Source/MovieRenderPipelineCore/Private/MoviePipelineTiming.cpp b/Engine/Plugins/MovieScene/MovieRenderPipeline/Source/MovieRenderPipelineCore/Private/MoviePipelineTiming.cpp
--- a/Engine/Plugins/MovieScene/MovieRenderPipeline/Source/MovieRenderPipelineCore/Private/MoviePipelineTiming.cpp
+++ b/Engine/Plugins/MovieScene/MovieRenderPipeline/Source/MovieRenderPipelineCore/Private/MoviePipelineTiming.cpp
@@ -260,6 +260,7 @@ void UMoviePipeline::TickProducingFrames()
 
 		CachedOutputState.TimeData.FrameDeltaTime = FrameDeltaTime;
 		CachedOutputState.TimeData.WorldSeconds = CachedOutputState.TimeData.WorldSeconds + FrameDeltaTime;
+		TRACE_BOOKMARK(TEXT("MoviePipeline - WarmingUp %d"), CurrentCameraCut->ShotInfo.NumEngineWarmUpFramesRemaining);
 		// We don't want to execute the other possible states until at least the next frame.
 		// This ensures that we actually tick the world once for the WarmUp state.
 		return;
@@ -338,6 +339,8 @@ void UMoviePipeline::TickProducingFrames()
 		// with many renders for a frame to fill history buffers).
 		CurrentCameraCut->ShotInfo.bHasEvaluatedMotionBlurFrame = true;
 		UE_LOG(LogMovieRenderPipeline, Verbose, TEXT("[%d] Shot MotionBlur set engine DeltaTime to %f seconds."), GFrameCounter, FrameDeltaTime);
+		TRACE_BOOKMARK(TEXT("MoviePipeline - MotionBlurEmulation"));
+
 		return;
 	}
 
@@ -686,6 +689,8 @@ void UMoviePipeline::TickProducingFrames()
 		double UndilatedDeltaTime = FrameMetrics.TickResolution.AsSeconds(FFrameTime(UndilatedDeltaFrameTime.GetFrame()));
 		CustomTimeStep->SetCachedFrameTiming(MoviePipeline::FFrameTimeStepCache(UndilatedDeltaTime));
 		CustomSequenceTimeController->SetCachedFrameTiming(FQualifiedFrameTime(FinalEvalTime, FrameMetrics.TickResolution));
+		TRACE_BOOKMARK(TEXT("MoviePipeline - Rendering Frame %d"), CachedOutputState.EffectiveFrameNumber);
+
 		return;
 	}
 
