MoviePipeline: Fixed 'Write all Samples' not respecting the {project_dir} format string.

diff --git a/Engine/Plugins/MovieScene/MovieRenderPipeline/Source/MovieRenderPipelineCore/Private/MoviePipelineRendering.cpp b/Engine/Plugins/MovieScene/MovieRenderPipeline/Source/MovieRenderPipelineCore/Private/MoviePipelineRendering.cpp
--- a/Engine/Plugins/MovieScene/MovieRenderPipeline/Source/MovieRenderPipelineCore/Private/MoviePipelineRendering.cpp
+++ b/Engine/Plugins/MovieScene/MovieRenderPipeline/Source/MovieRenderPipelineCore/Private/MoviePipelineRendering.cpp
@@ -467,14 +467,33 @@ void UMoviePipeline::OnSampleRendered(TUniquePtr<FImagePixelData>&& OutputSample
 	TileImageTask->CompressionQuality = (int32)EImageCompressionQuality::Default;
 
 	FString OutputName = InFrameData->Debug_OverrideFilename.IsEmpty() ?
-		FString::Printf(TEXT("/%s_SS_%d_TS_%d_TileX_%d_TileY_%d.%d.exr"),
+		FString::Printf(TEXT("/%s_SS_%d_TS_%d_TileX_%d_TileY_%d.%d"),
 			*InFrameData->PassIdentifier.Name, InFrameData->SampleState.SpatialSampleIndex, InFrameData->SampleState.TemporalSampleIndex,
 			InFrameData->SampleState.TileIndexes.X, InFrameData->SampleState.TileIndexes.Y, InFrameData->SampleState.OutputState.OutputFrameNumber)
 		: InFrameData->Debug_OverrideFilename;
-
+	
 	FString OutputDirectory = OutputSettings->OutputDirectory.Path;
-	FString OutputPath = OutputDirectory + OutputName;
-	TileImageTask->Filename = OutputPath;
+	FString FileNameFormatString = OutputDirectory + OutputName;
+
+	TMap<FString, FString> FormatOverrides;
+	FormatOverrides.Add(TEXT("ext"), TEXT("exr"));
+	UMoviePipelineExecutorShot* Shot = nullptr;
+	if (InFrameData->SampleState.OutputState.ShotIndex >= 0 && InFrameData->SampleState.OutputState.ShotIndex < ActiveShotList.Num())
+	{
+		Shot = ActiveShotList[InFrameData->SampleState.OutputState.ShotIndex];
+	}
+
+	if (Shot)
+	{
+		FormatOverrides.Add(TEXT("shot_name"), Shot->OuterName);
+		FormatOverrides.Add(TEXT("camera_name"), Shot->InnerName);
+	}
+	FMoviePipelineFormatArgs FinalFormatArgs;
+
+	FString FinalFilePath;
+	ResolveFilenameFormatArguments(FileNameFormatString, FormatOverrides, FinalFilePath, FinalFormatArgs);
+
+	TileImageTask->Filename = FinalFilePath;
 
 	// Duplicate the data so that the Image Task can own it.
 	TileImageTask->PixelData = MoveTemp(OutputSample);
