Movie Pipeline: Drag and drop queue assets now appends the jobs in the queue to the active queue

diff --git a/Engine/Plugins/MovieScene/MovieRenderPipeline/Source/MovieRenderPipelineEditor/Private/Widgets/SMoviePipelineQueueEditor.cpp b/Engine/Plugins/MovieScene/MovieRenderPipeline/Source/MovieRenderPipelineEditor/Private/Widgets/SMoviePipelineQueueEditor.cpp
--- a/Engine/Plugins/MovieScene/MovieRenderPipeline/Source/MovieRenderPipelineEditor/Private/Widgets/SMoviePipelineQueueEditor.cpp
+++ b/Engine/Plugins/MovieScene/MovieRenderPipeline/Source/MovieRenderPipelineEditor/Private/Widgets/SMoviePipelineQueueEditor.cpp
@@ -1139,27 +1139,43 @@ void SMoviePipelineQueueEditor::OnCreateJobFromAsset(const FAssetData& InAsset)
 	// Close the dropdown menu that showed them the assets to pick from.
 	FSlateApplication::Get().DismissAllMenus();
 
-	// Only try to initialize level sequences, in the event they had more than a level sequence selected when drag/dropping.
-	ULevelSequence* LevelSequence = Cast<ULevelSequence>(InAsset.GetAsset());
-	if (LevelSequence)
-	{
-		FScopedTransaction Transaction(FText::Format(LOCTEXT("CreateJob_Transaction", "Add {0}|plural(one=Job, other=Jobs)"), 1));
+	UMoviePipelineQueue* ActiveQueue = GEditor->GetEditorSubsystem<UMoviePipelineQueueSubsystem>()->GetQueue();
+	check(ActiveQueue);
+	
+	FScopedTransaction Transaction(FText::Format(LOCTEXT("CreateJob_Transaction", "Add {0}|plural(one=Job, other=Jobs)"), 1));
 
-		UMoviePipelineQueue* ActiveQueue = GEditor->GetEditorSubsystem<UMoviePipelineQueueSubsystem>()->GetQueue();
-		check(ActiveQueue);
-		
+	TArray<UMoviePipelineExecutorJob*> NewJobs;
+	// Only try to initialize level sequences, in the event they had more than a level sequence selected when drag/dropping.
+	if (ULevelSequence* LevelSequence = Cast<ULevelSequence>(InAsset.GetAsset()))
+	{		
 		UMoviePipelineExecutorJob* NewJob = UMoviePipelineEditorBlueprintLibrary::CreateJobFromSequence(ActiveQueue, LevelSequence);
 		if (!NewJob)
 		{
 			return;
 		}
 
+		NewJobs.Add(NewJob);
+	}
+	else if (UMoviePipelineQueue* Queue = Cast<UMoviePipelineQueue>(InAsset.GetAsset()))
+	{
+		for (UMoviePipelineExecutorJob* Job : Queue->GetJobs())
+		{
+			UMoviePipelineExecutorJob* NewJob = ActiveQueue->DuplicateJob(Job);
+			if (NewJob)
+			{
+				NewJobs.Add(NewJob);
+			}
+		}
+	}
+
+	const UMovieRenderPipelineProjectSettings* ProjectSettings = GetDefault<UMovieRenderPipelineProjectSettings>();
+	for (UMoviePipelineExecutorJob* NewJob : NewJobs)
+	{
 		PendingJobsToSelect.Add(NewJob);
 		
 		{
 			// The job configuration is already set up with an empty configuration, but we'll try and use their last used preset 
 			// (or a engine supplied default) for better user experience. 
-			const UMovieRenderPipelineProjectSettings* ProjectSettings = GetDefault<UMovieRenderPipelineProjectSettings>();
 			if (ProjectSettings->LastPresetOrigin.IsValid())
 			{
 				NewJob->SetPresetOrigin(ProjectSettings->LastPresetOrigin.Get());
@@ -1314,6 +1330,13 @@ bool SMoviePipelineQueueEditor::CanDragDropTarget(TSharedPtr<FDragDropOperation>
 					bIsValid = true;
 					break;
 				}
+
+				UMoviePipelineQueue* QueueAsset = Cast<UMoviePipelineQueue>(Asset.GetAsset());
+				if (QueueAsset)
+				{
+					bIsValid = true;
+					break;
+				}
 			}
 		}
 	}
