Movie Pipeline: Add the ability to toggle jobs in MRQ

Added bEnabled to UMoviePipelineExecutorJob that saves the state of whether a job should be rendered or not. UMoviePipelineLinearExecutorBase::StartPipelineByIndex checks whether the job is bEnabled and renders it only if it is.

diff --git a/Engine/Plugins/MovieScene/MovieRenderPipeline/Source/MovieRenderPipelineCore/Private/MoviePipelineLinearExecutor.cpp b/Engine/Plugins/MovieScene/MovieRenderPipeline/Source/MovieRenderPipelineCore/Private/MoviePipelineLinearExecutor.cpp
--- a/Engine/Plugins/MovieScene/MovieRenderPipeline/Source/MovieRenderPipelineCore/Private/MoviePipelineLinearExecutor.cpp
+++ b/Engine/Plugins/MovieScene/MovieRenderPipeline/Source/MovieRenderPipelineCore/Private/MoviePipelineLinearExecutor.cpp
@@ -43,6 +43,12 @@ void UMoviePipelineLinearExecutorBase::StartPipelineByIndex(int32 InPipelineInde
 	
 	CurrentPipelineIndex = InPipelineIndex;
 
+	if (!Queue->GetJobs()[CurrentPipelineIndex]->bEnabled)
+	{
+		OnIndividualPipelineFinished(nullptr);
+		return;
+	}
+
 	if (Queue->GetJobs()[CurrentPipelineIndex]->IsConsumed())
 	{
 		// We skip working on consumed jobs. Jobs submitted to remote renders might already be consumed and
diff --git a/Engine/Plugins/MovieScene/MovieRenderPipeline/Source/MovieRenderPipelineCore/Public/MoviePipelineQueue.h b/Engine/Plugins/MovieScene/MovieRenderPipeline/Source/MovieRenderPipelineCore/Public/MoviePipelineQueue.h
--- a/Engine/Plugins/MovieScene/MovieRenderPipeline/Source/MovieRenderPipelineCore/Public/MoviePipelineQueue.h
+++ b/Engine/Plugins/MovieScene/MovieRenderPipeline/Source/MovieRenderPipelineCore/Public/MoviePipelineQueue.h
@@ -166,6 +166,7 @@ class MOVIERENDERPIPELINECORE_API UMoviePipelineExecutorJob : public UObject
 public:
 	UMoviePipelineExecutorJob()
 	{
+		bEnabled = true;
 		StatusProgress = 0.f;
 		bIsConsumed = false;
 		Configuration = CreateDefaultSubobject<UMoviePipelineMasterConfig>("DefaultConfig");
@@ -334,6 +335,10 @@ public:
 	UPROPERTY(BlueprintReadWrite, Instanced, Category = "Movie Render Pipeline")
 	TArray<UMoviePipelineExecutorShot*> ShotInfo;
 
+	/** Whether this job is enabled and should be rendered. */
+	UPROPERTY(BlueprintReadWrite, EditAnywhere, Category = "Movie Render Pipeline")
+	bool bEnabled;
+
 	/** 
 	* Arbitrary data that can be associated with the job. Not used by default implementations, nor read.
 	* This can be used to attach third party metadata such as job ids from remote farms. 
diff --git a/Engine/Plugins/MovieScene/MovieRenderPipeline/Source/MovieRenderPipelineEditor/Private/Widgets/SMoviePipelineQueueEditor.cpp b/Engine/Plugins/MovieScene/MovieRenderPipeline/Source/MovieRenderPipelineEditor/Private/Widgets/SMoviePipelineQueueEditor.cpp
--- a/Engine/Plugins/MovieScene/MovieRenderPipeline/Source/MovieRenderPipelineEditor/Private/Widgets/SMoviePipelineQueueEditor.cpp
+++ b/Engine/Plugins/MovieScene/MovieRenderPipeline/Source/MovieRenderPipelineEditor/Private/Widgets/SMoviePipelineQueueEditor.cpp
@@ -77,6 +77,7 @@ public:
 		SLATE_EVENT(FOnMoviePipelineEditConfig, OnEditConfigRequested)
 	SLATE_END_ARGS()
 
+	static const FName NAME_Enabled;
 	static const FName NAME_JobName;
 	static const FName NAME_Settings;
 	static const FName NAME_Output;
@@ -101,6 +102,8 @@ struct FMoviePipelineQueueJobTreeItem : IMoviePipelineQueueTreeItem
 	/** The job that this tree item represents */
 	TWeakObjectPtr<UMoviePipelineExecutorJob> WeakJob;
 
+	TWeakPtr<SMoviePipelineQueueEditor> WeakQueueEditor;
+
 	/** Sorted list of this category's children */
 	TArray<TSharedPtr<IMoviePipelineQueueTreeItem>> Children;
 
@@ -115,6 +118,8 @@ struct FMoviePipelineQueueJobTreeItem : IMoviePipelineQueueTreeItem
 
 	virtual TSharedRef<ITableRow> ConstructWidget(TWeakPtr<SMoviePipelineQueueEditor> InQueueWidget, const TSharedRef<STableViewBase>& OwnerTable) override
 	{
+		WeakQueueEditor = InQueueWidget;
+
 		return SNew(SQueueJobListRow, OwnerTable)
 			.Item(SharedThis(this));
 	}
@@ -148,6 +153,39 @@ struct FMoviePipelineQueueJobTreeItem : IMoviePipelineQueueTreeItem
 	}
 
 public:
+	ECheckBoxState GetCheckState() const
+	{
+		UMoviePipelineExecutorJob* Job = WeakJob.Get();
+		if (Job)
+		{
+			return Job->bEnabled ? ECheckBoxState::Checked : ECheckBoxState::Unchecked;
+		}
+
+		return ECheckBoxState::Unchecked;
+	}
+
+	void SetCheckState(const ECheckBoxState NewState) const
+	{
+		if (WeakQueueEditor.IsValid() && WeakQueueEditor.Pin()->GetSelectedItems().Contains(SharedThis(this)))
+		{
+			for (TSharedPtr<IMoviePipelineQueueTreeItem> Item : WeakQueueEditor.Pin()->GetSelectedItems())
+			{
+				TSharedPtr<FMoviePipelineQueueJobTreeItem> JobTreeItem = StaticCastSharedPtr<FMoviePipelineQueueJobTreeItem>(Item);
+				if (JobTreeItem.IsValid())
+				{
+					if (UMoviePipelineExecutorJob* Job = JobTreeItem->WeakJob.Get())
+					{
+						Job->bEnabled = NewState == ECheckBoxState::Checked;
+					}
+				}
+			}
+		}
+		else if (UMoviePipelineExecutorJob* Job = WeakJob.Get())
+		{
+			Job->bEnabled = NewState == ECheckBoxState::Checked;
+		}
+	}
+
 	FText GetJobName() const
 	{
 		UMoviePipelineExecutorJob* Job = WeakJob.Get();
@@ -271,7 +309,7 @@ public:
 		UMoviePipelineExecutorJob* Job = WeakJob.Get();
 		if (Job)
 		{
-			return !Job->IsConsumed();
+			return Job->bEnabled && !Job->IsConsumed();
 		}
 		return false;
 	}
@@ -376,7 +414,19 @@ void SQueueJobListRow::Construct(const FArguments& InArgs, const TSharedRef<STab
 
 TSharedRef<SWidget> SQueueJobListRow::GenerateWidgetForColumn(const FName& ColumnName)
 {
-	if (ColumnName == NAME_JobName)
+	if (ColumnName == NAME_Enabled)
+	{
+		return SNew(SBox)
+			.WidthOverride(16)
+		  	[
+				SNew(SCheckBox)
+				.Style(FMovieRenderPipelineStyle::Get(), "MovieRenderPipeline.Setting.Switch")
+				.IsFocusable(false)
+				.IsChecked(Item.Get(), &FMoviePipelineQueueJobTreeItem::GetCheckState)
+				.OnCheckStateChanged(Item.Get(), &FMoviePipelineQueueJobTreeItem::SetCheckState)
+			];
+	}
+	else if (ColumnName == NAME_JobName)
 	{
 		return SNew(SBox)
 			.Padding(2.0f)
@@ -494,6 +544,7 @@ TSharedRef<SWidget> SQueueJobListRow::GenerateWidgetForColumn(const FName& Colum
 	return SNullWidget::NullWidget;
 }
 
+const FName SQueueJobListRow::NAME_Enabled = FName(TEXT("Enabled"));
 const FName SQueueJobListRow::NAME_JobName = FName(TEXT("Job Name"));
 const FName SQueueJobListRow::NAME_Settings = FName(TEXT("Settings"));
 const FName SQueueJobListRow::NAME_Output = FName(TEXT("Output"));
@@ -710,7 +761,7 @@ struct FMoviePipelineShotItem : IMoviePipelineQueueTreeItem
 		UMoviePipelineExecutorJob* Job = WeakJob.Get();
 		if (Job)
 		{
-			return !Job->IsConsumed();
+			return Job->bEnabled && !Job->IsConsumed();
 		}
 		return false;
 	}
@@ -964,12 +1015,16 @@ void SMoviePipelineQueueEditor::Construct(const FArguments& InArgs)
 		(
 			SNew(SHeaderRow)
 
+			+ SHeaderRow::Column(SQueueJobListRow::NAME_Enabled)
+			.FillWidth(0.05f)
+			.DefaultLabel(FText::FromString(TEXT(" ")))
+
 			+ SHeaderRow::Column(SQueueJobListRow::NAME_JobName)
 			.FillWidth(0.25f)
 			.DefaultLabel(LOCTEXT("QueueHeaderJobName_Text", "Job"))
 
 			+ SHeaderRow::Column(SQueueJobListRow::NAME_Settings)
-			.FillWidth(0.25f)
+			.FillWidth(0.20f)
 			.DefaultLabel(LOCTEXT("QueueHeaderSettings_Text", "Settings"))
 
 			+ SHeaderRow::Column(SQueueJobListRow::NAME_Output)
