MoviePipeline: A few additional fixes to IsEnabled support.

diff --git a/Engine/Plugins/MovieScene/MovieRenderPipeline/Source/MovieRenderPipelineCore/Private/MoviePipelineLinearExecutor.cpp b/Engine/Plugins/MovieScene/MovieRenderPipeline/Source/MovieRenderPipelineCore/Private/MoviePipelineLinearExecutor.cpp
--- a/Engine/Plugins/MovieScene/MovieRenderPipeline/Source/MovieRenderPipelineCore/Private/MoviePipelineLinearExecutor.cpp
+++ b/Engine/Plugins/MovieScene/MovieRenderPipeline/Source/MovieRenderPipelineCore/Private/MoviePipelineLinearExecutor.cpp
@@ -43,7 +43,7 @@ void UMoviePipelineLinearExecutorBase::StartPipelineByIndex(int32 InPipelineInde
 	
 	CurrentPipelineIndex = InPipelineIndex;
 
-	if (!Queue->GetJobs()[CurrentPipelineIndex]->bEnabled)
+	if (!Queue->GetJobs()[CurrentPipelineIndex]->IsEnabled())
 	{
 		OnIndividualPipelineFinished(nullptr);
 		return;
diff --git a/Engine/Plugins/MovieScene/MovieRenderPipeline/Source/MovieRenderPipelineCore/Public/MoviePipelineQueue.h b/Engine/Plugins/MovieScene/MovieRenderPipeline/Source/MovieRenderPipelineCore/Public/MoviePipelineQueue.h
--- a/Engine/Plugins/MovieScene/MovieRenderPipeline/Source/MovieRenderPipelineCore/Public/MoviePipelineQueue.h
+++ b/Engine/Plugins/MovieScene/MovieRenderPipeline/Source/MovieRenderPipelineCore/Public/MoviePipelineQueue.h
@@ -263,6 +263,32 @@ public:
 	UFUNCTION(BlueprintPure, BlueprintNativeEvent, Category = "Movie Render Pipeline")
 	bool IsConsumed() const;
 
+	/**
+	* Set the job to be enabled/disabled. This is exposed to the user in the Queue UI
+	* so they can disable a job after loading a queue to skip trying to run it.
+	*
+	* For C++ implementations override `virtual void SetIsEnabled_Implementation() override`
+	* For Python/BP implementations override
+	*	@unreal.ufunction(override=True)
+	*	def set_is_enabled(self, isEnabled):
+	*
+	* @param bInEnabled	True if the job should be enabled and rendered.
+	*/
+	UFUNCTION(BlueprintCallable, BlueprintNativeEvent, Category = "Movie Render Pipeline")
+	void SetIsEnabled(const bool bInEnabled);
+
+	/**
+	* Gets whether or not the job has been marked as being enabled. 
+	*
+	* For C++ implementations override `virtual bool IsEnabled_Implementation() const override`
+	* For Python/BP implementations override
+	*	@unreal.ufunction(override=True)
+	*	def is_enabled(self):
+	*		return ?
+	*/
+	UFUNCTION(BlueprintPure, BlueprintNativeEvent, Category = "Movie Render Pipeline")
+	bool IsEnabled() const;
+
 	/**
 	* Should be called to clear status and user data after duplication so that jobs stay
 	* unique and don't pick up ids or other unwanted behavior from the pareant job.
@@ -312,6 +338,8 @@ protected:
 	virtual float GetStatusProgress_Implementation() const { return StatusProgress; }
 	virtual bool IsConsumed_Implementation() const { return bIsConsumed; }
 	virtual void OnDuplicated_Implementation();
+	virtual void SetIsEnabled_Implementation(const bool bInEnabled) { bEnabled = bInEnabled; }
+	virtual bool IsEnabled_Implementation() const { return bEnabled; }
 	// ~UMoviePipelineExecutorJob Interface
 
 public:
@@ -335,10 +363,6 @@ public:
 	UPROPERTY(BlueprintReadWrite, Instanced, Category = "Movie Render Pipeline")
 	TArray<UMoviePipelineExecutorShot*> ShotInfo;
 
-	/** Whether this job is enabled and should be rendered. */
-	UPROPERTY(BlueprintReadWrite, EditAnywhere, Category = "Movie Render Pipeline")
-	bool bEnabled;
-
 	/** 
 	* Arbitrary data that can be associated with the job. Not used by default implementations, nor read.
 	* This can be used to attach third party metadata such as job ids from remote farms. 
@@ -363,6 +387,10 @@ private:
 	*/
 	UPROPERTY()
 	TSoftObjectPtr<UMoviePipelineMasterConfig> PresetOrigin;
+
+	/** Whether this job is enabled and should be rendered. */
+	UPROPERTY()
+	bool bEnabled;
 };
 
 /**
diff --git a/Engine/Plugins/MovieScene/MovieRenderPipeline/Source/MovieRenderPipelineEditor/Private/Widgets/SMoviePipelineQueueEditor.cpp b/Engine/Plugins/MovieScene/MovieRenderPipeline/Source/MovieRenderPipelineEditor/Private/Widgets/SMoviePipelineQueueEditor.cpp
--- a/Engine/Plugins/MovieScene/MovieRenderPipeline/Source/MovieRenderPipelineEditor/Private/Widgets/SMoviePipelineQueueEditor.cpp
+++ b/Engine/Plugins/MovieScene/MovieRenderPipeline/Source/MovieRenderPipelineEditor/Private/Widgets/SMoviePipelineQueueEditor.cpp
@@ -158,7 +158,7 @@ public:
 		UMoviePipelineExecutorJob* Job = WeakJob.Get();
 		if (Job)
 		{
-			return Job->bEnabled ? ECheckBoxState::Checked : ECheckBoxState::Unchecked;
+			return Job->IsEnabled() ? ECheckBoxState::Checked : ECheckBoxState::Unchecked;
 		}
 
 		return ECheckBoxState::Unchecked;
@@ -175,14 +175,14 @@ public:
 				{
 					if (UMoviePipelineExecutorJob* Job = JobTreeItem->WeakJob.Get())
 					{
-						Job->bEnabled = NewState == ECheckBoxState::Checked;
+						Job->SetIsEnabled(NewState == ECheckBoxState::Checked);
 					}
 				}
 			}
 		}
 		else if (UMoviePipelineExecutorJob* Job = WeakJob.Get())
 		{
-			Job->bEnabled = NewState == ECheckBoxState::Checked;
+			Job->SetIsEnabled(NewState == ECheckBoxState::Checked);
 		}
 	}
 
@@ -309,7 +309,7 @@ public:
 		UMoviePipelineExecutorJob* Job = WeakJob.Get();
 		if (Job)
 		{
-			return Job->bEnabled && !Job->IsConsumed();
+			return Job->IsEnabled() && !Job->IsConsumed();
 		}
 		return false;
 	}
@@ -761,7 +761,7 @@ struct FMoviePipelineShotItem : IMoviePipelineQueueTreeItem
 		UMoviePipelineExecutorJob* Job = WeakJob.Get();
 		if (Job)
 		{
-			return Job->bEnabled && !Job->IsConsumed();
+			return Job->IsEnabled()  && !Job->IsConsumed();
 		}
 		return false;
 	}
@@ -878,6 +878,7 @@ TSharedRef<SWidget> SQueueShotListRow::GenerateWidgetForColumn(const FName& Colu
 		.Padding(2.0f)
 		[
 			SNew(SHorizontalBox)
+			.IsEnabled(Item.Get(), &FMoviePipelineShotItem::IsEnabled)
 
 			// Toggle Checkbox for deciding to render or not.
 			+ SHorizontalBox::Slot()
@@ -964,7 +965,8 @@ TSharedRef<SWidget> SQueueShotListRow::GenerateWidgetForColumn(const FName& Colu
 	{
 		return SNew(SWidgetSwitcher)
 			.WidgetIndex(Item.Get(), &FMoviePipelineShotItem::GetStatusIndex)
-			
+			.IsEnabled(Item.Get(), &FMoviePipelineShotItem::IsEnabled)
+
 			// Ready Label
 			+ SWidgetSwitcher::Slot()
 			.VAlign(VAlign_Center)
diff --git a/Engine/Plugins/MovieScene/MovieRenderPipeline/Source/MovieRenderPipelineEditor/Private/Widgets/SMoviePipelineQueuePanel.cpp b/Engine/Plugins/MovieScene/MovieRenderPipeline/Source/MovieRenderPipelineEditor/Private/Widgets/SMoviePipelineQueuePanel.cpp
--- a/Engine/Plugins/MovieScene/MovieRenderPipeline/Source/MovieRenderPipelineEditor/Private/Widgets/SMoviePipelineQueuePanel.cpp
+++ b/Engine/Plugins/MovieScene/MovieRenderPipeline/Source/MovieRenderPipelineEditor/Private/Widgets/SMoviePipelineQueuePanel.cpp
@@ -286,7 +286,7 @@ bool SMoviePipelineQueuePanel::IsRenderLocalEnabled() const
 	bool bAtLeastOneJobAvailable = false;
 	for (UMoviePipelineExecutorJob* Job : Subsystem->GetQueue()->GetJobs())
 	{
-		if (!Job->IsConsumed())
+		if (!Job->IsConsumed() && Job->IsEnabled())
 		{
 			bAtLeastOneJobAvailable = true;
 			break;
@@ -325,7 +325,7 @@ bool SMoviePipelineQueuePanel::IsRenderRemoteEnabled() const
 	bool bAtLeastOneJobAvailable = false;
 	for (UMoviePipelineExecutorJob* Job : Subsystem->GetQueue()->GetJobs())
 	{
-		if (!Job->IsConsumed())
+		if (!Job->IsConsumed() && Job->IsEnabled())
 		{
 			bAtLeastOneJobAvailable = true;
 			break;
