Static Analysis: Fixed a potential use-before-nullcheck in MoviePipeline widget code.

diff --git a/Engine/Plugins/MovieScene/MovieRenderPipeline/Source/MovieRenderPipelineEditor/Private/Widgets/SMoviePipelineConfigEditor.cpp b/Engine/Plugins/MovieScene/MovieRenderPipeline/Source/MovieRenderPipelineEditor/Private/Widgets/SMoviePipelineConfigEditor.cpp
--- a/Engine/Plugins/MovieScene/MovieRenderPipeline/Source/MovieRenderPipelineEditor/Private/Widgets/SMoviePipelineConfigEditor.cpp
+++ b/Engine/Plugins/MovieScene/MovieRenderPipeline/Source/MovieRenderPipelineEditor/Private/Widgets/SMoviePipelineConfigEditor.cpp
@@ -268,6 +268,10 @@ TSharedRef<SWidget> SMoviePipelineConfigEditor::OnGenerateSettingsMenu()
 		{
 			// Get a display name for the setting from the CDO.
 			const UMoviePipelineSetting* SettingDefaultObject = GetDefault<UMoviePipelineSetting>(Class);
+			if (!SettingDefaultObject)
+			{
+				continue;
+			}
 
 			// Depending on the type of config we're editing, some settings may not be eligible. If this is the case, we omit them from the list.
 			bool bCanSettingBeAdded = CachedPipelineConfig->CanSettingBeAdded(SettingDefaultObject);
@@ -292,11 +296,7 @@ TSharedRef<SWidget> SMoviePipelineConfigEditor::OnGenerateSettingsMenu()
 				continue;
 			}
 
-			FText SettingDisplayName;
-			if (SettingDefaultObject)
-			{
-				SettingDisplayName = SettingDefaultObject->GetDisplayText();
-			}
+			FText SettingDisplayName = SettingDefaultObject->GetDisplayText();
 
 			TSubclassOf<UMoviePipelineSetting> SubclassOf = Class;
 			MenuBuilder.AddMenuEntry(
