MoviePipeline: Improved behavior of cloth when using temporal sub-sampling by automatically increasing the sub-step count on the correct sub-frame that has large delta times compared to the rest.

diff --git a/Engine/Plugins/MovieScene/MovieRenderPipeline/Source/MovieRenderPipelineCore/MovieRenderPipelineCore.Build.cs b/Engine/Plugins/MovieScene/MovieRenderPipeline/Source/MovieRenderPipelineCore/MovieRenderPipelineCore.Build.cs
--- a/Engine/Plugins/MovieScene/MovieRenderPipeline/Source/MovieRenderPipelineCore/MovieRenderPipelineCore.Build.cs
+++ b/Engine/Plugins/MovieScene/MovieRenderPipeline/Source/MovieRenderPipelineCore/MovieRenderPipelineCore.Build.cs
@@ -21,7 +21,8 @@ public class MovieRenderPipelineCore : ModuleRules
 				"Networking",
 				"HTTP",
 				"DeveloperSettings",
-				"CinematicCamera"
+				"CinematicCamera",
+				"ClothingSystemRuntimeInterface"
 
 			}
 		);
diff --git a/Engine/Plugins/MovieScene/MovieRenderPipeline/Source/MovieRenderPipelineCore/Private/MoviePipeline.cpp b/Engine/Plugins/MovieScene/MovieRenderPipeline/Source/MovieRenderPipelineCore/Private/MoviePipeline.cpp
--- a/Engine/Plugins/MovieScene/MovieRenderPipeline/Source/MovieRenderPipelineCore/Private/MoviePipeline.cpp
+++ b/Engine/Plugins/MovieScene/MovieRenderPipeline/Source/MovieRenderPipelineCore/Private/MoviePipeline.cpp
@@ -43,6 +43,9 @@
 #include "MoviePipelineCommandLineEncoder.h"
 #include "Kismet/KismetSystemLibrary.h"
 #include "ProfilingDebugging/TraceAuxiliary.h"
+#include "EngineUtils.h"
+#include "ClothingSimulationInteractor.h"
+#include "ClothingSimulationInterface.h"
 
 #if WITH_EDITOR
 #include "MovieSceneExportMetadata.h"
@@ -64,6 +67,7 @@ static TAutoConsoleVariable<int32> CVarMovieRenderPipelineFrameStepper(
 	ECVF_Default);
 
 FString UMoviePipeline::DefaultDebugWidgetAsset = TEXT("/MovieRenderPipeline/Blueprints/UI_MovieRenderPipelineScreenOverlay.UI_MovieRenderPipelineScreenOverlay_C");
+DECLARE_CYCLE_STAT(TEXT("STAT_MoviePipeline_ClothAdjust"), STAT_ClothSubstepAdjust, STATGROUP_MoviePipeline);
 
 UMoviePipeline::UMoviePipeline()
 	: CustomTimeStep(nullptr)
@@ -992,6 +996,9 @@ void UMoviePipeline::TeardownShot(UMoviePipelineExecutorShot* InShot)
 		const bool bSaveSettings = false;
 		MoviePipeline::SaveOrRestoreSubSectionHierarchy(Node, bSaveSettings);
 	}
+
+	RestoreSkeletalMeshClothSubSteps();
+	ClothSimCache.Reset();
 	
 	if (IsFlushDiskWritesPerShot())
 	{
@@ -1612,4 +1619,52 @@ void UMoviePipeline::StopUnrealInsightsCapture()
 {
 	FTraceAuxiliary::Stop();
 }
+
+void UMoviePipeline::SetSkeletalMeshClothSubSteps(const int32 InSubdivisionCount)
+{
+	SCOPE_CYCLE_COUNTER(STAT_ClothSubstepAdjust);
+
+	for (TActorIterator<AActor> ActorIt(GetWorld()); ActorIt; ++ActorIt)
+	{
+		AActor* FoundActor = *ActorIt;
+		if (FoundActor)
+		{
+			TArray<USkeletalMeshComponent*> SkeletalMeshComponents;
+			FoundActor->GetComponents<USkeletalMeshComponent>(SkeletalMeshComponents);
+
+			for (USkeletalMeshComponent* Component : SkeletalMeshComponents)
+			{
+				UClothingSimulationInteractor* ClothInteractor = Component->GetClothingSimulationInteractor();
+				if (ClothInteractor)
+				{
+					TWeakObjectPtr<UClothingSimulationInteractor> WeakPtr = TWeakObjectPtr< UClothingSimulationInteractor>(ClothInteractor);
+		
+					FClothSimSettingsCache* ExistingCacheEntry = ClothSimCache.Find(WeakPtr);
+					if (!ExistingCacheEntry)
+					{
+						ClothSimCache.Add(WeakPtr);
+						ExistingCacheEntry = ClothSimCache.Find(WeakPtr);
+						const int32 NumSubsteps = Component->GetClothingSimulation() ? FMath::Max(Component->GetClothingSimulation()->GetNumSubsteps(), 1)
+							: 1; // If there's no clothing simulation component just fall back to assuming they only had 1.
+						ExistingCacheEntry->NumSubSteps = NumSubsteps;
+					}
+
+					ClothInteractor->SetNumSubsteps(ExistingCacheEntry->NumSubSteps * InSubdivisionCount);
+				}
+			}
+		}
+	}
+}
+
+void UMoviePipeline::RestoreSkeletalMeshClothSubSteps()
+{
+	for (const TPair<TWeakObjectPtr<UClothingSimulationInteractor>, FClothSimSettingsCache>& Pair : ClothSimCache)
+	{
+		if (Pair.Key.Get())
+		{
+			Pair.Key->SetNumSubsteps(Pair.Value.NumSubSteps);
+		}
+	}
+}
+
 #undef LOCTEXT_NAMESPACE // "MoviePipeline"
diff --git a/Engine/Plugins/MovieScene/MovieRenderPipeline/Source/MovieRenderPipelineCore/Private/MoviePipelineTiming.cpp b/Engine/Plugins/MovieScene/MovieRenderPipeline/Source/MovieRenderPipelineCore/Private/MoviePipelineTiming.cpp
--- a/Engine/Plugins/MovieScene/MovieRenderPipeline/Source/MovieRenderPipelineCore/Private/MoviePipelineTiming.cpp
+++ b/Engine/Plugins/MovieScene/MovieRenderPipeline/Source/MovieRenderPipelineCore/Private/MoviePipelineTiming.cpp
@@ -554,6 +554,21 @@ void UMoviePipeline::TickProducingFrames()
 			DeltaFrameTime = DeltaFrameTime * WorldTimeDilation;
 		}
 
+		// Cloth doesn't like temporal sub-sampling due to the wildly varying delta times (due to the 'shutter closed') period.
+		// To help account for this, we increase the sub-division count for each skeletal mesh inversely to the deltatime
+		// for this frame. That means if a 'short' frame is 2ms, but a long frame is 10ms, we make the long frame 5x as many
+		// sub-divisions as the short frame to fake the cloth simulation system into processing each step with the same dt.
+		{
+			if (AntiAliasingSettings->TemporalSampleCount > 1)
+			{
+				double Ratio = DeltaFrameTime.FrameNumber.Value / (double)FrameMetrics.TicksPerSample.FrameNumber.Value;
+				
+				// Slomo can end up trying to do less than one iteration, we don't want that.	
+				int32 DivisionMultiplier = FMath::Max(FMath::FloorToInt(Ratio), 1);
+				SetSkeletalMeshClothSubSteps(DivisionMultiplier);
+			}
+		}
+
 		/*
 		* Now that we know how much we want to advance the sequence by (and the world)
 		* We can calculate out the desired frame for the shot to evaluate. This is slightly
diff --git a/Engine/Plugins/MovieScene/MovieRenderPipeline/Source/MovieRenderPipelineCore/Public/MoviePipeline.h b/Engine/Plugins/MovieScene/MovieRenderPipeline/Source/MovieRenderPipelineCore/Public/MoviePipeline.h
--- a/Engine/Plugins/MovieScene/MovieRenderPipeline/Source/MovieRenderPipelineCore/Public/MoviePipeline.h
+++ b/Engine/Plugins/MovieScene/MovieRenderPipeline/Source/MovieRenderPipelineCore/Public/MoviePipeline.h
@@ -366,6 +366,10 @@ private:
 
 	/** Handles transitioning between states, preventing reentrancy. Normal state flow should be respected, does not handle arbitrary x to y transitions. */
 	void TransitionToState(const EMovieRenderPipelineState InNewState);
+
+	void SetSkeletalMeshClothSubSteps(const int32 InSubdivisionCount);
+	void RestoreSkeletalMeshClothSubSteps();
+
 private:
 	/** Custom TimeStep used to drive the engine while rendering. */
 	UPROPERTY(Transient, Instanced)
@@ -478,6 +482,13 @@ private:
 
 	TSharedPtr<MoviePipeline::FCameraCutSubSectionHierarchyNode> CachedSequenceHierarchyRoot;
 
+	struct FClothSimSettingsCache
+	{
+		int32 NumSubSteps;
+	};
+
+	TMap<TWeakObjectPtr<class UClothingSimulationInteractor>, FClothSimSettingsCache> ClothSimCache;
+
 public:
 	static FString DefaultDebugWidgetAsset;
 };
