Movie Pipeline: Fix non-unity

diff --git a/Engine/Plugins/MovieScene/MovieRenderPipeline/Source/MovieRenderPipelineCore/Private/MoviePipeline.cpp b/Engine/Plugins/MovieScene/MovieRenderPipeline/Source/MovieRenderPipelineCore/Private/MoviePipeline.cpp
--- a/Engine/Plugins/MovieScene/MovieRenderPipeline/Source/MovieRenderPipelineCore/Private/MoviePipeline.cpp
+++ b/Engine/Plugins/MovieScene/MovieRenderPipeline/Source/MovieRenderPipelineCore/Private/MoviePipeline.cpp
@@ -46,6 +46,7 @@
 #include "EngineUtils.h"
 #include "ClothingSimulationInteractor.h"
 #include "ClothingSimulationInterface.h"
+#include "Components/SkeletalMeshComponent.h"
 
 #if WITH_EDITOR
 #include "MovieSceneExportMetadata.h"
