Fixing hardcoded shader feature level in MRQ (which was causing fatal asserts).

diff --git a/Engine/Plugins/MovieScene/MovieRenderPipeline/Source/MovieRenderPipelineRenderPasses/Private/MoviePipelineDeferredPasses.cpp b/Engine/Plugins/MovieScene/MovieRenderPipeline/Source/MovieRenderPipelineRenderPasses/Private/MoviePipelineDeferredPasses.cpp
--- a/Engine/Plugins/MovieScene/MovieRenderPipeline/Source/MovieRenderPipelineRenderPasses/Private/MoviePipelineDeferredPasses.cpp
+++ b/Engine/Plugins/MovieScene/MovieRenderPipeline/Source/MovieRenderPipelineRenderPasses/Private/MoviePipelineDeferredPasses.cpp
@@ -391,7 +391,7 @@ void UMoviePipelineDeferredPassBase::RenderSample_GameThreadImpl(const FMoviePip
 
 		// Submit to be rendered. Main render pass always uses target 0.
 		FRenderTarget* RenderTarget = GetViewRenderTarget()->GameThread_GetRenderTargetResource();
-		FCanvas Canvas = FCanvas(RenderTarget, nullptr, GetPipeline()->GetWorld(), ERHIFeatureLevel::SM5, FCanvas::CDM_DeferDrawing, 1.0f);
+		FCanvas Canvas = FCanvas(RenderTarget, nullptr, GetPipeline()->GetWorld(), View->GetFeatureLevel(), FCanvas::CDM_DeferDrawing, 1.0f);
 		GetRendererModule().BeginRenderingViewFamily(&Canvas, ViewFamily.Get());
 
 		// Readback + Accumulate.
