Path Tracer: add a reference motion blur mode to MRQ

In this mode the post-processed motion blur is ignored and the path tracer accumulates frames across temporal samples instead.

This can be used to workaround limitations of the 2d effect and allows pushing the temporal sample count higher than otherwise possible.

diff --git a/Engine/Plugins/MovieScene/MovieRenderPipeline/Source/MovieRenderPipelineCore/Public/MoviePipelineHighResSetting.h b/Engine/Plugins/MovieScene/MovieRenderPipeline/Source/MovieRenderPipelineCore/Public/MoviePipelineHighResSetting.h
--- a/Engine/Plugins/MovieScene/MovieRenderPipeline/Source/MovieRenderPipelineCore/Public/MoviePipelineHighResSetting.h
+++ b/Engine/Plugins/MovieScene/MovieRenderPipeline/Source/MovieRenderPipelineCore/Public/MoviePipelineHighResSetting.h
@@ -43,7 +43,7 @@ public:
 
 		if (TileCount > 1)
 		{
-			ValidationResults.Add(NSLOCTEXT("MovieRenderPipeline", "HighRes_UnsupportedFeatures", "Tiling does not support all rendering features (bloom, some screen-space effects). Additionally, TAA and Auto Exposure are not supported and will be forced off when rendering. Use Spatial/Temporal sampling and Manual Camera Exposure /w Exposure Compensation instead."));
+			ValidationResults.Add(NSLOCTEXT("MovieRenderPipeline", "HighRes_UnsupportedFeatures", "Tiling does not support all rendering features (bloom, some screen-space effects). Additionally, TAA and Auto Exposure are not supported and will be forced off when rendering. Use Spatial/Temporal sampling and Manual Camera Exposure /w Exposure Compensation instead. Tiling does not support reference motion blur with the Path Tracer."));
 			if (FMath::IsNearlyEqual(OverlapRatio, 0.f))
 			{
 				ValidationResults.Add(NSLOCTEXT("MovieRenderPipeline", "HighRes_OverlapNeeded", "Increase the overlap amount to avoid seams in the final image. More overlap results in longer renders so start at 0.1 and increase as needed."));
diff --git a/Engine/Plugins/MovieScene/MovieRenderPipeline/Source/MovieRenderPipelineRenderPasses/Private/MoviePipelineImagePassBase.cpp b/Engine/Plugins/MovieScene/MovieRenderPipeline/Source/MovieRenderPipelineRenderPasses/Private/MoviePipelineImagePassBase.cpp
--- a/Engine/Plugins/MovieScene/MovieRenderPipeline/Source/MovieRenderPipelineRenderPasses/Private/MoviePipelineImagePassBase.cpp
+++ b/Engine/Plugins/MovieScene/MovieRenderPipeline/Source/MovieRenderPipelineRenderPasses/Private/MoviePipelineImagePassBase.cpp
@@ -180,20 +180,32 @@ TSharedPtr<FSceneViewFamilyContext> UMoviePipelineImagePassBase::CalculateViewFa
 	if (OutViewFamily->EngineShowFlags.PathTracing)
 	{
 		// override whatever settings came from PostProcessVolume or Camera
-		int32 SampleCount = InOutSampleState.TemporalSampleCount * InOutSampleState.SpatialSampleCount;
-		int32 SampleIndex = InOutSampleState.TemporalSampleIndex * InOutSampleState.SpatialSampleCount + InOutSampleState.SpatialSampleIndex;
+
+		// If motion blur is enabled:
+		//    blend all spatial samples together while leaving the handling of temporal samples up to MRQ
+		//    each temporal sample will include denoising and post-process effects
+		// If motion blur is NOT enabled:
+		//    blend all temporal+spatial samples within the path tracer and only apply denoising on the last temporal sample
+		//    this way we minimize denoising cost and also allow a much higher number of temporal samples to be used which
+		//    can help reduce strobing
+
+		// NOTE: Tiling is not compatible with the reference motion blur mode because it changes the order of the loops over the image.
+		const bool bAccumulateSpatialSamplesOnly = OutViewFamily->EngineShowFlags.MotionBlur || InOutSampleState.GetTileCount() > 1;
+
+		const int32 SampleCount = bAccumulateSpatialSamplesOnly ? InOutSampleState.SpatialSampleCount : InOutSampleState.TemporalSampleCount * InOutSampleState.SpatialSampleCount;
+		const int32 SampleIndex = bAccumulateSpatialSamplesOnly ? InOutSampleState.SpatialSampleIndex : InOutSampleState.TemporalSampleIndex * InOutSampleState.SpatialSampleCount + InOutSampleState.SpatialSampleIndex;
 
 		// TODO: pass along FrameIndex (which includes SampleIndex) to make sure sampling is fully deterministic
 
 		// Overwrite whatever sampling count came from the PostProcessVolume
 		View->FinalPostProcessSettings.bOverride_PathTracingSamplesPerPixel = true;
-		View->FinalPostProcessSettings.PathTracingSamplesPerPixel = InOutSampleState.SpatialSampleCount;
+		View->FinalPostProcessSettings.PathTracingSamplesPerPixel = SampleCount;
 
-		// reset path tracer's accumulation at the start of each spatial sample
-		View->bForcePathTracerReset = InOutSampleState.SpatialSampleIndex == 0;
+		// reset path tracer's accumulation at the start of each sample
+		View->bForcePathTracerReset = SampleIndex == 0;
 
-		// discard the result, unless its the last spatial sample
-		InOutSampleState.bDiscardResult |= !(InOutSampleState.SpatialSampleIndex == InOutSampleState.SpatialSampleCount - 1);
+		// discard the result, unless its the last sample
+		InOutSampleState.bDiscardResult |= !(SampleIndex == SampleCount - 1);
 	}
 
 	// Object Occlusion/Histories
diff --git a/Engine/Plugins/MovieScene/MovieRenderPipeline/Source/MovieRenderPipelineRenderPasses/Public/MoviePipelineDeferredPasses.h b/Engine/Plugins/MovieScene/MovieRenderPipeline/Source/MovieRenderPipelineRenderPasses/Public/MoviePipelineDeferredPasses.h
--- a/Engine/Plugins/MovieScene/MovieRenderPipeline/Source/MovieRenderPipelineRenderPasses/Public/MoviePipelineDeferredPasses.h
+++ b/Engine/Plugins/MovieScene/MovieRenderPipeline/Source/MovieRenderPipelineRenderPasses/Public/MoviePipelineDeferredPasses.h
@@ -258,6 +258,7 @@ public:
 	{
 		OutShowFlag = FEngineShowFlags(EShowFlagInitMode::ESFIM_Game);
 		OutShowFlag.SetPathTracing(true);
+		OutShowFlag.SetMotionBlur(!bReferenceMotionBlur);
 		OutViewModeIndex = EViewModeIndex::VMI_PathTracing;
 	}
 	virtual int32 GetOutputFileSortingOrder() const override { return 2; }
@@ -266,4 +267,10 @@ public:
 	virtual void ValidateStateImpl() override;
 	virtual void SetupImpl(const MoviePipeline::FMoviePipelineRenderPassInitSettings& InPassInitSettings) override;
 
+	/** When enabled, the path tracer will blend all spatial and temporal samples prior to the denoising and will disable post-processed motion blur.
+	 *  In this mode it is possible to use higher temporal sample counts to improve the motion blur quality.
+	 *  When this option is disabled, the path tracer will accumulate spatial samples, but denoise them prior to accumulation of temporal samples.
+	 */
+	UPROPERTY(EditAnywhere, BlueprintReadWrite, Category = "Reference Motion Blur")
+	bool bReferenceMotionBlur;
 };
