Take Recorder/Movie Pipeline: Use positive action button and primary buttons

diff --git a/Engine/Plugins/MovieScene/MovieRenderPipeline/Source/MovieRenderPipelineEditor/MovieRenderPipelineEditor.Build.cs b/Engine/Plugins/MovieScene/MovieRenderPipeline/Source/MovieRenderPipelineEditor/MovieRenderPipelineEditor.Build.cs
--- a/Engine/Plugins/MovieScene/MovieRenderPipeline/Source/MovieRenderPipelineEditor/MovieRenderPipelineEditor.Build.cs
+++ b/Engine/Plugins/MovieScene/MovieRenderPipeline/Source/MovieRenderPipelineEditor/MovieRenderPipelineEditor.Build.cs
@@ -62,7 +62,8 @@ public class MovieRenderPipelineEditor : ModuleRules
 				"Json",
 				"JsonUtilities",
 				"ScreenShotComparisonTools",
-				"AutomationMessages"
+				"AutomationMessages",
+				"ToolWidgets"
 			}
         );
 
diff --git a/Engine/Plugins/MovieScene/MovieRenderPipeline/Source/MovieRenderPipelineEditor/Private/Widgets/SMoviePipelineConfigEditor.cpp b/Engine/Plugins/MovieScene/MovieRenderPipeline/Source/MovieRenderPipelineEditor/Private/Widgets/SMoviePipelineConfigEditor.cpp
--- a/Engine/Plugins/MovieScene/MovieRenderPipeline/Source/MovieRenderPipelineEditor/Private/Widgets/SMoviePipelineConfigEditor.cpp
+++ b/Engine/Plugins/MovieScene/MovieRenderPipeline/Source/MovieRenderPipelineEditor/Private/Widgets/SMoviePipelineConfigEditor.cpp
@@ -34,6 +34,7 @@
 #include "Widgets/Layout/SBorder.h"
 #include "Framework/MultiBox/MultiBoxBuilder.h"
 #include "Styling/SlateIconFinder.h"
+#include "SPositiveActionButton.h"
 
 // EditorStyle includes
 #include "EditorStyleSet.h"
@@ -153,50 +154,10 @@ void SMoviePipelineConfigEditor::Construct(const FArguments& InArgs)
 
 TSharedRef<SWidget> SMoviePipelineConfigEditor::MakeAddSettingButton()
 {
-	return SNew(SComboButton)
-		.ContentPadding(MoviePipeline::ButtonPadding)
-		.ButtonStyle(FMovieRenderPipelineStyle::Get(), "FlatButton.Success")
+	return SNew(SPositiveActionButton)
 		.OnGetMenuContent(this, &SMoviePipelineConfigEditor::OnGenerateSettingsMenu)
-		.ForegroundColor(FSlateColor::UseForeground())
-		.HasDownArrow(false)
-		.ButtonContent()
-		[
-			SNew(SHorizontalBox)
-    
-			// Plus Icon
-			+ SHorizontalBox::Slot()
-			.VAlign(VAlign_Center)
-			.AutoWidth()
-			[
-				SNew(STextBlock)
-				.TextStyle(FEditorStyle::Get(), "NormalText.Important")
-				.Font(FEditorStyle::Get().GetFontStyle("FontAwesome.10"))
-				.Text(FEditorFontGlyphs::Plus)
-			]
-    
-			// "Setting" Text
-			+ SHorizontalBox::Slot()
-			.VAlign(VAlign_Center)
-			.AutoWidth()
-			.Padding(4, 0, 0, 0)
-			[
-				SNew(STextBlock)
-				.TextStyle(FEditorStyle::Get(), "NormalText.Important")
-				.Text(LOCTEXT("AddNewSetting_Text", "Setting"))
-			]
-    
-			// Non-Default Down Caret arrow.
-			+ SHorizontalBox::Slot()
-			.VAlign(VAlign_Center)
-			.AutoWidth()
-			.Padding(4, 0, 0, 0)
-			[
-				SNew(STextBlock)
-				.TextStyle(FEditorStyle::Get(), "NormalText.Important")
-				.Font(FEditorStyle::Get().GetFontStyle("FontAwesome.10"))
-				.Text(FEditorFontGlyphs::Caret_Down)
-			]
-		];
+		.Icon(FAppStyle::Get().GetBrush("Icons.Plus"))
+		.Text(LOCTEXT("AddNewSetting_Text", "Setting"));
 }
 PRAGMA_ENABLE_OPTIMIZATION
 
diff --git a/Engine/Plugins/MovieScene/MovieRenderPipeline/Source/MovieRenderPipelineEditor/Private/Widgets/SMoviePipelineQueueEditor.cpp b/Engine/Plugins/MovieScene/MovieRenderPipeline/Source/MovieRenderPipelineEditor/Private/Widgets/SMoviePipelineQueueEditor.cpp
--- a/Engine/Plugins/MovieScene/MovieRenderPipeline/Source/MovieRenderPipelineEditor/Private/Widgets/SMoviePipelineQueueEditor.cpp
+++ b/Engine/Plugins/MovieScene/MovieRenderPipeline/Source/MovieRenderPipelineEditor/Private/Widgets/SMoviePipelineQueueEditor.cpp
@@ -27,6 +27,7 @@
 #include "EditorStyleSet.h"
 #include "EditorFontGlyphs.h"
 #include "SDropTarget.h"
+#include "SPositiveActionButton.h"
 
 // Editor
 #include "Editor.h"
@@ -1086,50 +1087,10 @@ TSharedPtr<SWidget> SMoviePipelineQueueEditor::GetContextMenuContent()
 
 TSharedRef<SWidget> SMoviePipelineQueueEditor::MakeAddSequenceJobButton()
 {
-	return SNew(SComboButton)
-		.ContentPadding(MoviePipeline::ButtonPadding)
-		.ButtonStyle(FMovieRenderPipelineStyle::Get(), "FlatButton.Success")
+	return SNew(SPositiveActionButton)
 		.OnGetMenuContent(this, &SMoviePipelineQueueEditor::OnGenerateNewJobFromAssetMenu)
-		.ForegroundColor(FSlateColor::UseForeground())
-		.HasDownArrow(false)
-		.ButtonContent()
-		[
-			SNew(SHorizontalBox)
-
-			// Plus Icon
-			+ SHorizontalBox::Slot()
-			.VAlign(VAlign_Center)
-			.AutoWidth()
-			[
-				SNew(STextBlock)
-				.TextStyle(FEditorStyle::Get(), "NormalText.Important")
-				.Font(FEditorStyle::Get().GetFontStyle("FontAwesome.10"))
-				.Text(FEditorFontGlyphs::Plus)
-			]
-
-			// "Render" Text
-			+ SHorizontalBox::Slot()
-			.VAlign(VAlign_Center)
-			.AutoWidth()
-			.Padding(4, 0, 0, 0)
-			[
-				SNew(STextBlock)
-				.TextStyle(FEditorStyle::Get(), "NormalText.Important")
-				.Text(LOCTEXT("AddNewJob_Text", "Render"))
-			]
-
-			// Non-Default Down Caret arrow.
-			+ SHorizontalBox::Slot()
-			.VAlign(VAlign_Center)
-			.AutoWidth()
-			.Padding(4, 0, 0, 0)
-			[
-				SNew(STextBlock)
-				.TextStyle(FEditorStyle::Get(), "NormalText.Important")
-				.Font(FEditorStyle::Get().GetFontStyle("FontAwesome.10"))
-				.Text(FEditorFontGlyphs::Caret_Down)
-			]
-		];
+		.Icon(FAppStyle::Get().GetBrush("Icons.Plus"))
+		.Text(LOCTEXT("AddNewJob_Text", "Render"));
 }
 
 TSharedRef<SWidget> SMoviePipelineQueueEditor::RemoveSelectedJobButton()
diff --git a/Engine/Plugins/MovieScene/MovieRenderPipeline/Source/MovieRenderPipelineEditor/Private/Widgets/SMoviePipelineQueuePanel.cpp b/Engine/Plugins/MovieScene/MovieRenderPipeline/Source/MovieRenderPipelineEditor/Private/Widgets/SMoviePipelineQueuePanel.cpp
--- a/Engine/Plugins/MovieScene/MovieRenderPipeline/Source/MovieRenderPipelineEditor/Private/Widgets/SMoviePipelineQueuePanel.cpp
+++ b/Engine/Plugins/MovieScene/MovieRenderPipeline/Source/MovieRenderPipelineEditor/Private/Widgets/SMoviePipelineQueuePanel.cpp
@@ -28,6 +28,7 @@
 #include "Styling/SlateIconFinder.h"
 #include "Widgets/Images/SImage.h"
 #include "Framework/MultiBox/MultiBoxBuilder.h"
+#include "SPrimaryButton.h"
 
 // ContentBrowser Includes
 #include "IContentBrowserSingleton.h"
@@ -216,19 +217,11 @@ void SMoviePipelineQueuePanel::Construct(const FArguments& InArgs)
 				.HAlign(HAlign_Right)
 				.AutoWidth()
 				[
-					SNew(SButton)
-					.ContentPadding(MoviePipeline::ButtonPadding)
-					.ButtonStyle(FMovieRenderPipelineStyle::Get(), "FlatButton.Success")
+					SNew(SPrimaryButton)
+					.Text(LOCTEXT("RenderQueueLocal_Text", "Render (Local)"))
+					.ToolTipText(LOCTEXT("RenderQueueLocal_Tooltip", "Renders the current queue in the current process using Play in Editor."))
 					.IsEnabled(this, &SMoviePipelineQueuePanel::IsRenderLocalEnabled)
 					.OnClicked(this, &SMoviePipelineQueuePanel::OnRenderLocalRequested)
-					.Content()
-					[
-						SNew(STextBlock)
-						.TextStyle(FEditorStyle::Get(), "NormalText.Important")
-						.Text(LOCTEXT("RenderQueueLocal_Text", "Render (Local)"))
-						.ToolTipText(LOCTEXT("RenderQueueLocal_Tooltip", "Renders the current queue in the current process using Play in Editor."))
-						.Margin(FMargin(4, 0, 4, 0))
-					]
 				]
 
 				// Render Remotely (Separate Process or Farm)
@@ -238,19 +231,11 @@ void SMoviePipelineQueuePanel::Construct(const FArguments& InArgs)
 				.HAlign(HAlign_Right)
 				.AutoWidth()
 				[
-					SNew(SButton)
-					.ContentPadding(MoviePipeline::ButtonPadding)
-					.ButtonStyle(FMovieRenderPipelineStyle::Get(), "FlatButton.Success")
+					SNew(SPrimaryButton)
+					.Text(LOCTEXT("RenderQueueRemote_Text", "Render (Remote)"))
+					.ToolTipText(LOCTEXT("RenderQueueRemote_Tooltip", "Renders the current queue in a separate process."))
 					.IsEnabled(this, &SMoviePipelineQueuePanel::IsRenderRemoteEnabled)
 					.OnClicked(this, &SMoviePipelineQueuePanel::OnRenderRemoteRequested)
-					.Content()
-					[
-						SNew(STextBlock)
-						.TextStyle(FEditorStyle::Get(), "NormalText.Important")
-						.Text(LOCTEXT("RenderQueueRemote_Text", "Render (Remote)"))
-						.ToolTipText(LOCTEXT("RenderQueueRemote_Tooltip", "Renders the current queue in a separate process."))
-						.Margin(FMargin(4, 0, 4, 0))
-					]
 				]
 			]
 		]
diff --git a/Engine/Plugins/VirtualProduction/Takes/Source/TakeRecorder/Private/Widgets/SLevelSequenceTakeEditor.cpp b/Engine/Plugins/VirtualProduction/Takes/Source/TakeRecorder/Private/Widgets/SLevelSequenceTakeEditor.cpp
--- a/Engine/Plugins/VirtualProduction/Takes/Source/TakeRecorder/Private/Widgets/SLevelSequenceTakeEditor.cpp
+++ b/Engine/Plugins/VirtualProduction/Takes/Source/TakeRecorder/Private/Widgets/SLevelSequenceTakeEditor.cpp
@@ -42,6 +42,7 @@
 #include "Widgets/Layout/SBorder.h"
 #include "Framework/MultiBox/MultiBoxBuilder.h"
 #include "Styling/SlateIconFinder.h"
+#include "SPositiveActionButton.h"
 
 // EditorStyle includes
 #include "EditorStyleSet.h"
@@ -143,47 +144,10 @@ void SLevelSequenceTakeEditor::Construct(const FArguments& InArgs)
 
 TSharedRef<SWidget> SLevelSequenceTakeEditor::MakeAddSourceButton()
 {
-	return SNew(SComboButton)
-		.ContentPadding(TakeRecorder::ButtonPadding)
-		.ButtonStyle(FTakeRecorderStyle::Get(), "FlatButton.Success")
+	return SNew(SPositiveActionButton)
 		.OnGetMenuContent(this, &SLevelSequenceTakeEditor::OnGenerateSourcesMenu)
-		.ForegroundColor(FSlateColor::UseForeground())
-		.HasDownArrow(false)
-		.ButtonContent()
-		[
-			SNew(SHorizontalBox)
-
-			+ SHorizontalBox::Slot()
-			.VAlign(VAlign_Center)
-			.AutoWidth()
-			[
-				SNew(STextBlock)
-				.TextStyle(FEditorStyle::Get(), "NormalText.Important")
-				.Font(FEditorStyle::Get().GetFontStyle("FontAwesome.10"))
-				.Text(FEditorFontGlyphs::Plus)
-			]
-
-			+ SHorizontalBox::Slot()
-			.VAlign(VAlign_Center)
-			.AutoWidth()
-			.Padding(4, 0, 0, 0)
-			[
-				SNew(STextBlock)
-				.TextStyle(FEditorStyle::Get(), "NormalText.Important")
-				.Text(LOCTEXT("AddNewSource_Text", "Source"))
-			]
-
-			+ SHorizontalBox::Slot()
-			.VAlign(VAlign_Center)
-			.AutoWidth()
-			.Padding(4, 0, 0, 0)
-			[
-				SNew(STextBlock)
-				.TextStyle(FEditorStyle::Get(), "NormalText.Important")
-				.Font(FEditorStyle::Get().GetFontStyle("FontAwesome.10"))
-				.Text(FEditorFontGlyphs::Caret_Down)
-			]
-		];
+		.Icon(FAppStyle::Get().GetBrush("Icons.Plus"))
+		.Text(LOCTEXT("AddNewSource_Text", "Source"));
 }
 PRAGMA_ENABLE_OPTIMIZATION
 
