MoviePipeline: Further fixes for Cloth Simulation when using Temporal sub-sampling. Warmup Frames use the full delta time (with no sub-stepping) so when we transitioned from warmup to rendering there would be a significant difference in delta times (as far as the physics system was concerned), causing a jump. This changes Warmup Frames to always set the number of iterations to match what the iteration will be once we switch into the Rendering state.

diff --git a/Engine/Plugins/MovieScene/MovieRenderPipeline/Source/MovieRenderPipelineCore/Private/MoviePipelineTiming.cpp b/Engine/Plugins/MovieScene/MovieRenderPipeline/Source/MovieRenderPipelineCore/Private/MoviePipelineTiming.cpp
--- a/Engine/Plugins/MovieScene/MovieRenderPipeline/Source/MovieRenderPipelineCore/Private/MoviePipelineTiming.cpp
+++ b/Engine/Plugins/MovieScene/MovieRenderPipeline/Source/MovieRenderPipelineCore/Private/MoviePipelineTiming.cpp
@@ -258,6 +258,17 @@ void UMoviePipeline::TickProducingFrames()
 			CustomSequenceTimeController->SetCachedFrameTiming(FQualifiedFrameTime(FinalEvalTime, FrameMetrics.TickResolution));
 		}
 
+		// Apply the cloth fixups during warmup differently. If we don't use the increased number of sub-steps used on the large delta times here,
+		// then transitioning from warmup to rendering causes a jump in delta times (from large to small) which makes all the cloth ripple. 
+		if (AntiAliasingSettings->TemporalSampleCount > 1)
+		{
+			double Ratio = FrameMetrics.TicksWhileShutterOpen.FloorToFrame().Value / (double)FrameMetrics.TicksPerSample.FrameNumber.Value;
+
+			// Slomo can end up trying to do less than one iteration, we don't want that.	
+			int32 DivisionMultiplier = FMath::Max(FMath::FloorToInt(Ratio), 1);
+			SetSkeletalMeshClothSubSteps(DivisionMultiplier);
+		}
+
 		CachedOutputState.TimeData.FrameDeltaTime = FrameDeltaTime;
 		CachedOutputState.TimeData.WorldSeconds = CachedOutputState.TimeData.WorldSeconds + FrameDeltaTime;
 		TRACE_BOOKMARK(TEXT("MoviePipeline - WarmingUp %d"), CurrentCameraCut->ShotInfo.NumEngineWarmUpFramesRemaining);
