Take Recorder/Movie Pipeline: Use positive action button and primary buttons

diff --git a/Engine/Plugins/MovieScene/MovieRenderPipeline/Source/MovieRenderPipelineEditor/Private/Widgets/SMoviePipelineConfigPanel.cpp b/Engine/Plugins/MovieScene/MovieRenderPipeline/Source/MovieRenderPipelineEditor/Private/Widgets/SMoviePipelineConfigPanel.cpp
--- a/Engine/Plugins/MovieScene/MovieRenderPipeline/Source/MovieRenderPipelineEditor/Private/Widgets/SMoviePipelineConfigPanel.cpp
+++ b/Engine/Plugins/MovieScene/MovieRenderPipeline/Source/MovieRenderPipelineEditor/Private/Widgets/SMoviePipelineConfigPanel.cpp
@@ -25,6 +25,7 @@
 #include "EditorStyleSet.h"
 #include "EditorFontGlyphs.h"
 #include "MovieRenderPipelineStyle.h"
+#include "SPrimaryButton.h"
 
 // ContentBrowser Includes
 #include "IContentBrowserSingleton.h"
@@ -213,16 +214,9 @@ void SMoviePipelineConfigPanel::Construct(const FArguments& InArgs, TSubclassOf<
 				.AutoWidth()
 				[
 					SNew(SButton)
-					.ContentPadding(MoviePipeline::ButtonPadding)
+					.Text(LOCTEXT("CancelChangesButton_Text", "Cancel"))
+					.ToolTipText(LOCTEXT("CancelChangesButton_Tooltip", "Discards changes made and does not modify anything."))
 					.OnClicked(this, &SMoviePipelineConfigPanel::OnCancelChanges)
-					.Content()
-					[
-						SNew(STextBlock)
-						.TextStyle(FEditorStyle::Get(), "NormalText.Important")
-						.Text(LOCTEXT("CancelChangesButton_Text", "Cancel"))
-						.ToolTipText(LOCTEXT("CancelChangesButton_Tooltip", "Discards changes made and does not modify anything."))
-						.Margin(FMargin(4, 0, 4, 0))
-					]
 				]
 
 				// Accept Changes
@@ -232,19 +226,11 @@ void SMoviePipelineConfigPanel::Construct(const FArguments& InArgs, TSubclassOf<
 				.HAlign(HAlign_Right)
 				.AutoWidth()
 				[
-					SNew(SButton)
-					.ContentPadding(MoviePipeline::ButtonPadding)
-					.ButtonStyle(FMovieRenderPipelineStyle::Get(), "FlatButton.Success")
+					SNew(SPrimaryButton)
+					.Text(LOCTEXT("ConfirmChangesButton_Text", "Accept"))
+					.ToolTipText(LOCTEXT("ConfirmChangesButton_Tooltip", "Accepts the changes made and applies them to the particular job instance."))
 					.OnClicked(this, &SMoviePipelineConfigPanel::OnConfirmChanges)
 					.IsEnabled(this, &SMoviePipelineConfigPanel::CanAcceptChanges)
-					.Content()
-					[
-						SNew(STextBlock)
-						.TextStyle(FEditorStyle::Get(), "NormalText.Important")
-						.Text(LOCTEXT("ConfirmChangesButton_Text", "Accept"))
-						.ToolTipText(LOCTEXT("ConfirmChangesButton_Tooltip", "Accepts the changes made and applies them to the particular job instance."))
-						.Margin(FMargin(4, 0, 4, 0))
-					]
 				]
 			]
 		];
