MoviePipeline: Add the duration it took to render a frame to the exr metadata.

diff --git a/Engine/Plugins/MovieScene/MovieRenderPipeline/Source/MovieRenderPipelineCore/Private/MoviePipeline.cpp b/Engine/Plugins/MovieScene/MovieRenderPipeline/Source/MovieRenderPipelineCore/Private/MoviePipeline.cpp
--- a/Engine/Plugins/MovieScene/MovieRenderPipeline/Source/MovieRenderPipelineCore/Private/MoviePipeline.cpp
+++ b/Engine/Plugins/MovieScene/MovieRenderPipeline/Source/MovieRenderPipelineCore/Private/MoviePipeline.cpp
@@ -1497,7 +1497,6 @@ void UMoviePipeline::ResolveFilenameFormatArguments(const FString& InFormatStrin
 		Params.FileMetadata = InOutputState->FileMetadata;
 		Params.ShotOverride = ActiveShotList[InOutputState->ShotIndex];
 		Params.InitializationVersion = ActiveShotList[InOutputState->ShotIndex]->ShotInfo.VersionNumber;
-
 	}
 
 	UMoviePipelineOutputSetting* OutputSetting = GetPipelineMasterConfig()->FindSetting<UMoviePipelineOutputSetting>();
@@ -1520,6 +1519,21 @@ void UMoviePipeline::ResolveFilenameFormatArguments(const FString& InFormatStrin
 	Params.AdditionalFrameNumberOffset = InFrameNumberOffset;
 
 	UMoviePipelineBlueprintLibrary::ResolveFilenameFormatArguments(InFormatString, Params, OutFinalPath, OutFinalFormatArgs);
+
+	// This needs to come after ResolveFilenameFormatArguments as that resets the OutFinalFormatArgs.
+	if (InOutputState)
+	{
+		const FRenderTimeStatistics* TimeStats = RenderTimeFrameStatistics.Find(InOutputState->OutputFrameNumber);
+		if (TimeStats)
+		{
+			FString StartTimeStr = TimeStats->StartTime.ToString();
+			FString EndTimeStr = TimeStats->EndTime.ToString();
+			FString DurationTimeStr = (TimeStats->EndTime - TimeStats->StartTime).ToString();
+			OutFinalFormatArgs.FileMetadata.Add(TEXT("unreal/frameRenderStartTimeUTC"), StartTimeStr);
+			OutFinalFormatArgs.FileMetadata.Add(TEXT("unreal/frameRenderEndTimeUTC"), EndTimeStr);
+			OutFinalFormatArgs.FileMetadata.Add(TEXT("unreal/frameRenderDuration"), DurationTimeStr);
+		}
+	}
 }
 
 void UMoviePipeline::SetProgressWidgetVisible(bool bVisible)
diff --git a/Engine/Plugins/MovieScene/MovieRenderPipeline/Source/MovieRenderPipelineCore/Private/MoviePipelineRendering.cpp b/Engine/Plugins/MovieScene/MovieRenderPipeline/Source/MovieRenderPipelineCore/Private/MoviePipelineRendering.cpp
--- a/Engine/Plugins/MovieScene/MovieRenderPipeline/Source/MovieRenderPipelineCore/Private/MoviePipelineRendering.cpp
+++ b/Engine/Plugins/MovieScene/MovieRenderPipeline/Source/MovieRenderPipelineCore/Private/MoviePipelineRendering.cpp
@@ -252,6 +252,9 @@ void UMoviePipeline::RenderFrame()
 		{
 			RenderPass->GatherOutputPasses(OutputFrame.ExpectedRenderPasses);
 		}
+
+		FRenderTimeStatistics& TimeStats = RenderTimeFrameStatistics.FindOrAdd(CachedOutputState.OutputFrameNumber);
+		TimeStats.StartTime = FDateTime::UtcNow();
 	}
 
 #if WITH_EDITOR && !UE_BUILD_SHIPPING
@@ -446,6 +449,9 @@ void UMoviePipeline::ProcessOutstandingFinishedFrames()
 	{
 		FMoviePipelineMergerOutputFrame OutputFrame;
 		OutputBuilder->FinishedFrames.Dequeue(OutputFrame);
+
+		FRenderTimeStatistics& TimeStats = RenderTimeFrameStatistics.FindOrAdd(OutputFrame.FrameOutputState.OutputFrameNumber);
+		TimeStats.EndTime = FDateTime::UtcNow();
 	
 		for (UMoviePipelineOutputBase* OutputContainer : GetPipelineMasterConfig()->GetOutputContainers())
 		{
diff --git a/Engine/Plugins/MovieScene/MovieRenderPipeline/Source/MovieRenderPipelineCore/Public/MoviePipeline.h b/Engine/Plugins/MovieScene/MovieRenderPipeline/Source/MovieRenderPipelineCore/Public/MoviePipeline.h
--- a/Engine/Plugins/MovieScene/MovieRenderPipeline/Source/MovieRenderPipelineCore/Public/MoviePipeline.h
+++ b/Engine/Plugins/MovieScene/MovieRenderPipeline/Source/MovieRenderPipelineCore/Public/MoviePipeline.h
@@ -489,6 +489,14 @@ private:
 
 	TMap<TWeakObjectPtr<class UClothingSimulationInteractor>, FClothSimSettingsCache> ClothSimCache;
 
+	struct FRenderTimeStatistics
+	{
+		FDateTime StartTime;
+		FDateTime EndTime;
+	};
+
+	TMap<int32, FRenderTimeStatistics> RenderTimeFrameStatistics;
+
 public:
 	static FString DefaultDebugWidgetAsset;
 };
