MoviePipeline: Added a delegate to the PIE Executor for when a job is about to start, to match up with the delegate that gets fired after a job completes.

diff --git a/Engine/Plugins/MovieScene/MovieRenderPipeline/Source/MovieRenderPipelineEditor/Private/MoviePipelinePIEExecutor.cpp b/Engine/Plugins/MovieScene/MovieRenderPipeline/Source/MovieRenderPipelineEditor/Private/MoviePipelinePIEExecutor.cpp
--- a/Engine/Plugins/MovieScene/MovieRenderPipeline/Source/MovieRenderPipelineEditor/Private/MoviePipelinePIEExecutor.cpp
+++ b/Engine/Plugins/MovieScene/MovieRenderPipeline/Source/MovieRenderPipelineEditor/Private/MoviePipelinePIEExecutor.cpp
@@ -206,6 +206,7 @@ void UMoviePipelinePIEExecutor::OnPIEStartupFinished(bool)
 	
 	if (ExecutorSettings->InitialDelayFrameCount == 0)
 	{
+		OnIndividualJobStartedImpl(Queue->GetJobs()[CurrentPipelineIndex]);
 		ActiveMoviePipeline->Initialize(Queue->GetJobs()[CurrentPipelineIndex]);
 		RemainingInitializationFrames = -1;
 	}
@@ -229,6 +230,7 @@ void UMoviePipelinePIEExecutor::OnTick()
 				ActiveMoviePipeline->SetInitializationTime(CustomInitializationTime.GetValue());
 			}
 
+			OnIndividualJobStartedImpl(Queue->GetJobs()[CurrentPipelineIndex]);
 			ActiveMoviePipeline->Initialize(Queue->GetJobs()[CurrentPipelineIndex]);
 		}
 
@@ -324,6 +326,12 @@ void UMoviePipelinePIEExecutor::DelayedFinishNotification()
 	// Now that another frame has passed and we should be OK to start another PIE session, notify our owner.
 	OnIndividualPipelineFinished(nullptr);
 }
+void UMoviePipelinePIEExecutor::OnIndividualJobStartedImpl(UMoviePipelineExecutorJob* InJob)
+{
+	// Broadcast to both Native and Python/BP
+	OnIndividualJobStartedDelegateNative.Broadcast(InJob);
+	OnIndividualJobStartedDelegate.Broadcast(InJob);
+}
 
 void UMoviePipelinePIEExecutor::OnIndividualJobFinishedImpl(FMoviePipelineOutputData InOutputData)
 {
diff --git a/Engine/Plugins/MovieScene/MovieRenderPipeline/Source/MovieRenderPipelineEditor/Public/MoviePipelinePIEExecutor.h b/Engine/Plugins/MovieScene/MovieRenderPipeline/Source/MovieRenderPipelineEditor/Public/MoviePipelinePIEExecutor.h
--- a/Engine/Plugins/MovieScene/MovieRenderPipeline/Source/MovieRenderPipelineEditor/Public/MoviePipelinePIEExecutor.h
+++ b/Engine/Plugins/MovieScene/MovieRenderPipeline/Source/MovieRenderPipelineEditor/Public/MoviePipelinePIEExecutor.h
@@ -9,9 +9,14 @@
 
 class UMoviePipeline;
 
+// These are deprecated, see OnIndividualShotWorkFinishedDelegate which contains both the job and the data actually rendered.
 DECLARE_MULTICAST_DELEGATE_TwoParams(FOnMoviePipelineIndividualJobFinishedNative, UMoviePipelineExecutorJob* /*FinishedJob*/, bool /*bSuccess*/);
 DECLARE_DYNAMIC_MULTICAST_DELEGATE_TwoParams(FOnMoviePipelineIndividualJobFinished, UMoviePipelineExecutorJob*, FinishedJob, bool, bSuccess);
 
+// These are called right before UMoviePipeline::Initialize(). They only contain the job because no work has been done yet.
+DECLARE_MULTICAST_DELEGATE_OneParam(FOnMoviePipelineIndividualJobStartedNative, UMoviePipelineExecutorJob* /*JobToStart*/);
+DECLARE_DYNAMIC_MULTICAST_DELEGATE_OneParam(FOnMoviePipelineIndividualJobStarted, UMoviePipelineExecutorJob*, StartedJob);
+
 /**
 * This is the implementation responsible for executing the rendering of
 * multiple movie pipelines in the currently running Editor process. This
@@ -52,6 +57,12 @@ public:
 		return OnIndividualShotWorkFinishedDelegateNative;
 	}
 
+	/** Called right before the specified job is started. This is your last chance to modify the Job Properties before things are initialized from it. */
+	FOnMoviePipelineIndividualJobStartedNative& OnIndividualJobStarted()
+	{
+		return OnIndividualJobStartedDelegateNative;
+	}
+
 	UFUNCTION(BlueprintCallable, Category = "Movie Render Pipeline")
 	void SetInitializationTime(const FDateTime& InInitializationTime) { CustomInitializationTime = InInitializationTime; }
 
@@ -64,7 +75,10 @@ protected:
 	* make modifications to the editor world.
 	*/
 	void OnIndividualJobFinishedImpl(FMoviePipelineOutputData InOutputData);
-
+	/**
+	* This is called after PIE starts up but before UMoviePIpeline::Initialize() is called.
+	*/
+	void OnIndividualJobStartedImpl(UMoviePipelineExecutorJob* InJob);
 private:
 	/** Called when PIE finishes booting up and it is safe for us to spawn an object into that world. */
 	void OnPIEStartupFinished(bool);
@@ -108,8 +122,15 @@ private:
 	UPROPERTY(BlueprintAssignable, Category = "Movie Render Pipeline")
 	FMoviePipelineWorkFinished OnIndividualShotWorkFinishedDelegate;
 
+	/**
+	* Called right before this job is used to initialize a UMoviePipeline.
+	*/
+	UPROPERTY(BlueprintAssignable, Category = "Movie Render Pipeline")
+	FOnMoviePipelineIndividualJobStarted OnIndividualJobStartedDelegate;
+
 	FMoviePipelineWorkFinishedNative OnIndividualJobWorkFinishedDelegateNative;
 	FMoviePipelineWorkFinishedNative OnIndividualShotWorkFinishedDelegateNative;
+	FOnMoviePipelineIndividualJobStartedNative OnIndividualJobStartedDelegateNative;
 
 	class FValidationMessageGatherer : public FOutputDevice
 	{
