MoviePipeline: Added an Actor Factory to make using the classes in the UI easier outside of MRQ (this allows you to add the "Create..." type in dropdowns that let you pick from content browser assets).

This is NOT enabled the Movie Render Queue UI as it would create too much confusion, since the MRQ ui edits *copies* of the assets (and the only way to update them is to save over it). But it's okay in other contexts where you generate a new asset and open it outside of MRQ in the slimmed down asset editor.

diff --git a/Engine/Plugins/MovieScene/MovieRenderPipeline/Source/MovieRenderPipelineEditor/Private/MoviePipelineConfigFactory.cpp b/Engine/Plugins/MovieScene/MovieRenderPipeline/Source/MovieRenderPipelineEditor/Private/MoviePipelineConfigFactory.cpp
new file mode 100644
--- /dev/null
+++ b/Engine/Plugins/MovieScene/MovieRenderPipeline/Source/MovieRenderPipelineEditor/Private/MoviePipelineConfigFactory.cpp
@@ -0,0 +1,40 @@
+// Copyright Epic Games, Inc. All Rights Reserved.
+
+#include "MoviePipelineConfigFactory.h"
+#include "MoviePipelineMasterConfig.h"
+#include "MoviePipelineShotConfig.h"
+#include "AssetTypeCategories.h"
+
+UMoviePipelineMasterConfigFactory::UMoviePipelineMasterConfigFactory()
+{
+	bCreateNew = true;
+	bEditAfterNew = false;
+	SupportedClass = UMoviePipelineMasterConfig::StaticClass();
+}
+
+UObject* UMoviePipelineMasterConfigFactory::FactoryCreateNew(UClass* Class, UObject* InParent, FName Name, EObjectFlags Flags, UObject* Context, FFeedbackContext* Warn)
+{
+	return NewObject<UMoviePipelineMasterConfig>(InParent, Class, Name, Flags);
+}
+
+uint32 UMoviePipelineMasterConfigFactory::GetMenuCategories() const
+{
+	return EAssetTypeCategories::Type::None; 
+}
+
+UMoviePipelineShotConfigFactory::UMoviePipelineShotConfigFactory()
+{
+	bCreateNew = true;
+	bEditAfterNew = false;
+	SupportedClass = UMoviePipelineShotConfig::StaticClass();
+}
+
+UObject* UMoviePipelineShotConfigFactory::FactoryCreateNew(UClass* Class, UObject* InParent, FName Name, EObjectFlags Flags, UObject* Context, FFeedbackContext* Warn)
+{
+	return NewObject<UMoviePipelineShotConfig>(InParent, Class, Name, Flags);
+}
+
+uint32 UMoviePipelineShotConfigFactory::GetMenuCategories() const
+{
+	return EAssetTypeCategories::Type::None; 
+}
\ No newline at end of file
diff --git a/Engine/Plugins/MovieScene/MovieRenderPipeline/Source/MovieRenderPipelineEditor/Public/MoviePipelineConfigFactory.h b/Engine/Plugins/MovieScene/MovieRenderPipeline/Source/MovieRenderPipelineEditor/Public/MoviePipelineConfigFactory.h
new file mode 100644
--- /dev/null
+++ b/Engine/Plugins/MovieScene/MovieRenderPipeline/Source/MovieRenderPipelineEditor/Public/MoviePipelineConfigFactory.h
@@ -0,0 +1,32 @@
+// Copyright Epic Games, Inc. All Rights Reserved.
+
+#pragma once
+
+#include "Factories/Factory.h"
+#include "MoviePipelineConfigFactory.generated.h"
+
+// Forward Declare
+UCLASS(BlueprintType)
+class MOVIERENDERPIPELINEEDITOR_API UMoviePipelineMasterConfigFactory : public UFactory
+{
+    GENERATED_BODY()
+public:
+	UMoviePipelineMasterConfigFactory();
+	//~ Begin UFactory Interface
+	virtual UObject* FactoryCreateNew(UClass* Class, UObject* InParent, FName Name, EObjectFlags Flags, UObject* Context, FFeedbackContext* Warn) override;
+	virtual uint32 GetMenuCategories() const override;
+	//~ End UFactory Interface
+};
+
+// Forward Declare
+UCLASS(BlueprintType)
+class MOVIERENDERPIPELINEEDITOR_API UMoviePipelineShotConfigFactory : public UFactory
+{
+    GENERATED_BODY()
+public:
+	UMoviePipelineShotConfigFactory();
+	//~ Begin UFactory Interface
+	virtual UObject* FactoryCreateNew(UClass* Class, UObject* InParent, FName Name, EObjectFlags Flags, UObject* Context, FFeedbackContext* Warn) override;
+	virtual uint32 GetMenuCategories() const override;
+	//~ End UFactory Interface
+};
\ No newline at end of file
