MoviePipeline: Fixed DetailLighting render pass not applying the correct view modes causing reflections to not show up.

diff --git a/Engine/Plugins/MovieScene/MovieRenderPipeline/Source/MovieRenderPipelineRenderPasses/Private/MoviePipelineImagePassBase.cpp b/Engine/Plugins/MovieScene/MovieRenderPipeline/Source/MovieRenderPipelineRenderPasses/Private/MoviePipelineImagePassBase.cpp
--- a/Engine/Plugins/MovieScene/MovieRenderPipeline/Source/MovieRenderPipelineRenderPasses/Private/MoviePipelineImagePassBase.cpp
+++ b/Engine/Plugins/MovieScene/MovieRenderPipeline/Source/MovieRenderPipelineRenderPasses/Private/MoviePipelineImagePassBase.cpp
@@ -83,6 +83,10 @@ TSharedPtr<FSceneViewFamilyContext> UMoviePipelineImagePassBase::CalculateViewFa
 	OutViewFamily->SceneCaptureSource = InOutSampleState.SceneCaptureSource;
 	OutViewFamily->bWorldIsPaused = InOutSampleState.bWorldIsPaused;
 	OutViewFamily->ViewMode = ViewModeIndex;
+
+	const bool bIsPerspective = true;
+	ApplyViewMode(OutViewFamily->ViewMode, bIsPerspective, OutViewFamily->EngineShowFlags);
+
 	EngineShowFlagOverride(ESFIM_Game, OutViewFamily->ViewMode, OutViewFamily->EngineShowFlags, false);
 	
 	const UMoviePipelineExecutorShot* Shot = GetPipeline()->GetActiveShotList()[InOutSampleState.OutputState.ShotIndex];
