MoviePipeline: Fixed {date} and {time} not resolving in output paths when clicked on in the UI. {time} is probably not that useful (since once the job starts the actual initialization time will be taken) but {date} will generally be correct.

diff --git a/Engine/Plugins/MovieScene/MovieRenderPipeline/Source/MovieRenderPipelineEditor/Private/MoviePipelineEditorBlueprintLibrary.cpp b/Engine/Plugins/MovieScene/MovieRenderPipeline/Source/MovieRenderPipelineEditor/Private/MoviePipelineEditorBlueprintLibrary.cpp
--- a/Engine/Plugins/MovieScene/MovieRenderPipeline/Source/MovieRenderPipelineEditor/Private/MoviePipelineEditorBlueprintLibrary.cpp
+++ b/Engine/Plugins/MovieScene/MovieRenderPipeline/Source/MovieRenderPipelineEditor/Private/MoviePipelineEditorBlueprintLibrary.cpp
@@ -234,6 +234,7 @@ FString UMoviePipelineEditorBlueprintLibrary::ResolveOutputDirectoryFromJob(UMov
 	Params.Job = InJob;
 	Params.ShotNameOverride = DummyTag;
 	Params.CameraNameOverride = DummyTag;
+	Params.InitializationTime = FDateTime::UtcNow();
 
 	FString OutResolvedPath;
 	FMoviePipelineFormatArgs Dummy;
