MoviePipeline: Fixed the {year} {month} {day} tags not using the job initialization time which would cause files to go into separate folders when a job rolled over midnight in UTC time.

diff --git a/Engine/Plugins/MovieScene/MovieRenderPipeline/Source/MovieRenderPipelineCore/Private/MoviePipelineBlueprintLibrary.cpp b/Engine/Plugins/MovieScene/MovieRenderPipeline/Source/MovieRenderPipelineCore/Private/MoviePipelineBlueprintLibrary.cpp
--- a/Engine/Plugins/MovieScene/MovieRenderPipeline/Source/MovieRenderPipelineCore/Private/MoviePipelineBlueprintLibrary.cpp
+++ b/Engine/Plugins/MovieScene/MovieRenderPipeline/Source/MovieRenderPipelineCore/Private/MoviePipelineBlueprintLibrary.cpp
@@ -74,25 +74,11 @@ FText UMoviePipelineBlueprintLibrary::GetJobName(UMoviePipeline* InMoviePipeline
 	return FText();
 }
 
-namespace MoviePipeline
-{
-	FString GetJobAuthor(const UMoviePipelineExecutorJob* InJob)
-	{
-		if (InJob && InJob->Author.Len() > 0)
-		{
-			return InJob->Author;
-		}
-
-		// If they didn't specify an author in the job, default to the local username.
-		return FPlatformProcess::UserName(false);
-	}
-}
-
 FText UMoviePipelineBlueprintLibrary::GetJobAuthor(UMoviePipeline* InMoviePipeline)
 {
 	if (InMoviePipeline)
 	{
-		return FText::FromString(MoviePipeline::GetJobAuthor(InMoviePipeline->GetCurrentJob()));
+		return FText::FromString(UE::MoviePipeline::GetJobAuthor(InMoviePipeline->GetCurrentJob()));
 	}
 
 	return FText();
@@ -882,20 +868,9 @@ void UMoviePipelineBlueprintLibrary::ResolveFilenameFormatArguments(const FStrin
 
 	// And from ourself
 	{
-		OutMergedFormatArgs.FilenameArguments.Add(TEXT("date"), InParams.InitializationTime.ToString(TEXT("%Y.%m.%d")));
-		OutMergedFormatArgs.FilenameArguments.Add(TEXT("time"), InParams.InitializationTime.ToString(TEXT("%H.%M.%S")));
-		OutMergedFormatArgs.FilenameArguments.Add(TEXT("job_author"), MoviePipeline::GetJobAuthor(InParams.Job));
-
-		FString VersionText = FString::Printf(TEXT("v%0*d"), 3, InParams.InitializationVersion);
-
-		OutMergedFormatArgs.FilenameArguments.Add(TEXT("version"), VersionText);
-
-		OutMergedFormatArgs.FileMetadata.Add(TEXT("unreal/jobDate"), InParams.InitializationTime.ToString(TEXT("%Y.%m.%d")));
-		OutMergedFormatArgs.FileMetadata.Add(TEXT("unreal/jobTime"), InParams.InitializationTime.ToString(TEXT("%H.%M.%S")));
-		OutMergedFormatArgs.FileMetadata.Add(TEXT("unreal/jobVersion"), FString::FromInt(InParams.InitializationVersion));
-		OutMergedFormatArgs.FileMetadata.Add(TEXT("unreal/jobName"), InParams.Job->JobName);
-		OutMergedFormatArgs.FileMetadata.Add(TEXT("unreal/jobAuthor"), MoviePipeline::GetJobAuthor(InParams.Job));
-
+		// Use the shared function as UMoviePipelineMasterConfig::GetFormatArguments to ensure all the time variables get overwritten.
+		UE::MoviePipeline::GetSharedFormatArguments(OutMergedFormatArgs.FilenameArguments, OutMergedFormatArgs.FileMetadata, InParams.InitializationTime, InParams.InitializationVersion, InParams.Job);
+		
 		// By default, we don't want to show frame duplication numbers. If we need to start writing them,
 		// they need to come before the frame number (so that tools recognize them as a sequence).
 		OutMergedFormatArgs.FilenameArguments.Add(TEXT("file_dup"), FString());
diff --git a/Engine/Plugins/MovieScene/MovieRenderPipeline/Source/MovieRenderPipelineCore/Private/MoviePipelineMasterConfig.cpp b/Engine/Plugins/MovieScene/MovieRenderPipeline/Source/MovieRenderPipelineCore/Private/MoviePipelineMasterConfig.cpp
--- a/Engine/Plugins/MovieScene/MovieRenderPipeline/Source/MovieRenderPipelineCore/Private/MoviePipelineMasterConfig.cpp
+++ b/Engine/Plugins/MovieScene/MovieRenderPipeline/Source/MovieRenderPipelineCore/Private/MoviePipelineMasterConfig.cpp
@@ -115,6 +115,7 @@ void UMoviePipelineMasterConfig::GetFormatArguments(FMoviePipelineFormatArgs& In
 			FrameRate = GetEffectiveFrameRate(Cast<ULevelSequence>(InOutFormatArgs.InJob->Sequence.TryLoad())).AsDecimal();
 		}
 
+
 		InOutFormatArgs.FilenameArguments.Add(TEXT("level_name"), LevelName);
 		InOutFormatArgs.FilenameArguments.Add(TEXT("sequence_name"), SequenceName);
 		InOutFormatArgs.FilenameArguments.Add(TEXT("job_name"), JobName);
@@ -128,13 +129,9 @@ void UMoviePipelineMasterConfig::GetFormatArguments(FMoviePipelineFormatArgs& In
 		// Normally these are filled when resolving the file name by the job (so that the time is shared), but stub them in here so
 		// they show up in the UI with a value.
 		FDateTime CurrentTime = FDateTime::Now();
-		InOutFormatArgs.FilenameArguments.Add(TEXT("date"), CurrentTime.ToString(TEXT("%Y.%m.%d")));
-		InOutFormatArgs.FilenameArguments.Add(TEXT("year"), CurrentTime.ToString(TEXT("%Y")));
-		InOutFormatArgs.FilenameArguments.Add(TEXT("month"), CurrentTime.ToString(TEXT("%m")));
-		InOutFormatArgs.FilenameArguments.Add(TEXT("day"), CurrentTime.ToString(TEXT("%d")));
-		InOutFormatArgs.FilenameArguments.Add(TEXT("time"), CurrentTime.ToString(TEXT("%H.%M.%S")));
-		InOutFormatArgs.FilenameArguments.Add(TEXT("job_author"), FPlatformProcess::UserName(false));
-		
+		int32 DummyVersionNumber = 1;
+		UE::MoviePipeline::GetSharedFormatArguments(InOutFormatArgs.FilenameArguments, InOutFormatArgs.FileMetadata, CurrentTime, DummyVersionNumber, InOutFormatArgs.InJob);
+
 		// Let the output state fill out some too, since its the keeper of the information.
 		UMoviePipelineOutputSetting* OutputSettings = FindSetting<UMoviePipelineOutputSetting>();
 		check(OutputSettings);
diff --git a/Engine/Plugins/MovieScene/MovieRenderPipeline/Source/MovieRenderPipelineCore/Private/MoviePipelineUtils.cpp b/Engine/Plugins/MovieScene/MovieRenderPipeline/Source/MovieRenderPipelineCore/Private/MoviePipelineUtils.cpp
--- a/Engine/Plugins/MovieScene/MovieRenderPipeline/Source/MovieRenderPipelineCore/Private/MoviePipelineUtils.cpp
+++ b/Engine/Plugins/MovieScene/MovieRenderPipeline/Source/MovieRenderPipelineCore/Private/MoviePipelineUtils.cpp
@@ -736,5 +736,46 @@ namespace UE
 				InOutFilenameFormatString.ReplaceInline(TEXT("{frame_number_shot_rel}"), TEXT(""));
 			}
 		}
+
+		FString GetJobAuthor(const UMoviePipelineExecutorJob* InJob)
+		{
+			if (InJob && InJob->Author.Len() > 0)
+			{
+				return InJob->Author;
+			}
+
+			// If they didn't specify an author in the job, default to the local username.
+			return FPlatformProcess::UserName(false);
+		}
+
+		void GetSharedFormatArguments(TMap<FString, FString>& InFilenameArguments, TMap<FString, FString>& InFileMetadata, const FDateTime& InDateTime, const int32 InVersionNumber, const UMoviePipelineExecutorJob* InJob)
+		{
+			InFilenameArguments.Add(TEXT("date"), InDateTime.ToString(TEXT("%Y.%m.%d")));
+			InFilenameArguments.Add(TEXT("time"), InDateTime.ToString(TEXT("%H.%M.%S")));
+			InFilenameArguments.Add(TEXT("year"), InDateTime.ToString(TEXT("%Y")));
+			InFilenameArguments.Add(TEXT("month"), InDateTime.ToString(TEXT("%m")));
+			InFilenameArguments.Add(TEXT("day"), InDateTime.ToString(TEXT("%d")));
+
+
+			FString VersionText = FString::Printf(TEXT("v%0*d"), 3, InVersionNumber);
+
+			InFilenameArguments.Add(TEXT("version"), VersionText);
+
+			InFileMetadata.Add(TEXT("unreal/jobDate"), InDateTime.ToString(TEXT("%Y.%m.%d")));
+			InFileMetadata.Add(TEXT("unreal/jobTime"), InDateTime.ToString(TEXT("%H.%M.%S")));
+			InFileMetadata.Add(TEXT("unreal/jobYear"), InDateTime.ToString(TEXT("%Y")));
+			InFileMetadata.Add(TEXT("unreal/jobMonth"), InDateTime.ToString(TEXT("%m")));
+			InFileMetadata.Add(TEXT("unreal/jobDay"), InDateTime.ToString(TEXT("%d")));
+
+			InFileMetadata.Add(TEXT("unreal/jobVersion"), FString::FromInt(InVersionNumber));
+
+			if (InJob)
+			{
+				InFilenameArguments.Add(TEXT("job_author"), GetJobAuthor(InJob));
+				InFilenameArguments.Add(TEXT("job_name"), InJob->JobName);
+				InFileMetadata.Add(TEXT("unreal/jobName"), InJob->JobName);
+				InFileMetadata.Add(TEXT("unreal/jobAuthor"), GetJobAuthor(InJob));
+			}
+		}
 	}
 }
diff --git a/Engine/Plugins/MovieScene/MovieRenderPipeline/Source/MovieRenderPipelineCore/Public/MoviePipelineUtils.h b/Engine/Plugins/MovieScene/MovieRenderPipeline/Source/MovieRenderPipelineCore/Public/MoviePipelineUtils.h
--- a/Engine/Plugins/MovieScene/MovieRenderPipeline/Source/MovieRenderPipelineCore/Public/MoviePipelineUtils.h
+++ b/Engine/Plugins/MovieScene/MovieRenderPipeline/Source/MovieRenderPipelineCore/Public/MoviePipelineUtils.h
@@ -11,6 +11,7 @@
 #include "Misc/FrameRate.h"
 #include "Evaluation/MovieSceneSequenceHierarchy.h"
 #include "MovieSceneSequenceID.h"
+#include "MoviePipelineQueue.h"
 
 // Forward Declare
 class UClass;
@@ -128,6 +129,9 @@ namespace UE
 	{
 		MOVIERENDERPIPELINECORE_API void ValidateOutputFormatString(FString& InOutFilenameFormatString, const bool bTestRenderPass, const bool bTestFrameNumber);
 		MOVIERENDERPIPELINECORE_API void RemoveFrameNumberFormatStrings(FString& InOutFilenameFormatString, const bool bIncludeShots);
+
+		MOVIERENDERPIPELINECORE_API FString GetJobAuthor(const UMoviePipelineExecutorJob* InJob);
+		MOVIERENDERPIPELINECORE_API void GetSharedFormatArguments(TMap<FString, FString>& InFilenameArguments, TMap<FString, FString>& InFileMetadata, const FDateTime& InDateTime, const int32 InVersionNumber, const UMoviePipelineExecutorJob* InJob);
 	}
 }
 
