MoviePipeline: Added some system information to exr metadata such as total memory available, total memory used (by all processes), cpu vendor/name, etc.

diff --git a/Engine/Plugins/MovieScene/MovieRenderPipeline/Source/MovieRenderPipelineCore/Private/MoviePipelineRendering.cpp b/Engine/Plugins/MovieScene/MovieRenderPipeline/Source/MovieRenderPipelineCore/Private/MoviePipelineRendering.cpp
--- a/Engine/Plugins/MovieScene/MovieRenderPipeline/Source/MovieRenderPipelineCore/Private/MoviePipelineRendering.cpp
+++ b/Engine/Plugins/MovieScene/MovieRenderPipeline/Source/MovieRenderPipelineCore/Private/MoviePipelineRendering.cpp
@@ -193,6 +193,18 @@ void UMoviePipeline::RenderFrame()
 
 	// Add appropriate metadata here that is shared by all passes.
 	{
+		// Add hardware stats such as total memory, cpu vendor, etc.
+		FString ResolvedOutputDirectory;
+		TMap<FString, FString> FormatOverrides;
+		FMoviePipelineFormatArgs FinalFormatArgs;
+
+		// We really only need the output disk path for disk size info, but we'll try to resolve as much as possible anyways
+		ResolveFilenameFormatArguments(OutputSettings->OutputDirectory.Path, FormatOverrides, ResolvedOutputDirectory, FinalFormatArgs);
+		// Strip .{ext}
+		ResolvedOutputDirectory.LeftChopInline(6);
+
+		UE::MoviePipeline::GetHardwareUsageMetadata(CachedOutputState.FileMetadata, ResolvedOutputDirectory);
+
 		CachedOutputState.FileMetadata.Add(TEXT("unreal/camera/curPos/x"), FString::SanitizeFloat(FrameInfo.CurrViewLocation.X));
 		CachedOutputState.FileMetadata.Add(TEXT("unreal/camera/curPos/y"), FString::SanitizeFloat(FrameInfo.CurrViewLocation.Y));
 		CachedOutputState.FileMetadata.Add(TEXT("unreal/camera/curPos/z"), FString::SanitizeFloat(FrameInfo.CurrViewLocation.Z));
diff --git a/Engine/Plugins/MovieScene/MovieRenderPipeline/Source/MovieRenderPipelineCore/Private/MoviePipelineUtils.cpp b/Engine/Plugins/MovieScene/MovieRenderPipeline/Source/MovieRenderPipelineCore/Private/MoviePipelineUtils.cpp
--- a/Engine/Plugins/MovieScene/MovieRenderPipeline/Source/MovieRenderPipelineCore/Private/MoviePipelineUtils.cpp
+++ b/Engine/Plugins/MovieScene/MovieRenderPipeline/Source/MovieRenderPipelineCore/Private/MoviePipelineUtils.cpp
@@ -20,6 +20,11 @@
 #include "Misc/Paths.h"
 #include "Misc/FrameRate.h"
 #include "MoviePipelineQueue.h"
+#include "HAL/PlatformMisc.h"
+#include "HAL/PlatformMemory.h"
+#include "GenericPlatform/GenericPlatformDriver.h"
+#include "MovieRenderPipelineCoreModule.h"
+#include "RHI.h"
 
 namespace UE
 {
@@ -652,6 +657,9 @@ namespace UE
 {
 	namespace MoviePipeline
 	{
+		DECLARE_CYCLE_STAT(TEXT("STAT_MoviePipeline_HardwareMetadata"), STAT_HardwareMetadata, STATGROUP_MoviePipeline);
+
+
 		void ValidateOutputFormatString(FString& InOutFilenameFormatString, const bool bTestRenderPass, const bool bTestFrameNumber)
 		{
 			const FString FrameNumberIdentifiers[] = { TEXT("{frame_number}"), TEXT("{frame_number_shot}"), TEXT("{frame_number_rel}"), TEXT("{frame_number_shot_rel}") };
@@ -777,5 +785,51 @@ namespace UE
 				InFileMetadata.Add(TEXT("unreal/jobAuthor"), GetJobAuthor(InJob));
 			}
 		}
+
+		void GetHardwareUsageMetadata(TMap<FString, FString>& InFileMetadata, const FString& InOutputDir)
+		{
+			SCOPE_CYCLE_COUNTER(STAT_HardwareMetadata);
+			
+			// Stats that change per-frame go in /stats/, consistent things go in /system/ though logical grouping
+			// overrides some things such as total vs. current memory.
+			static const uint64 MBDivider = 1024 * 1024;
+
+			InFileMetadata.Add(TEXT("unreal/system/cpuVendor"), FPlatformMisc::GetCPUVendor());
+			InFileMetadata.Add(TEXT("unreal/system/cpuChipset"), FPlatformMisc::GetCPUChipset());
+			InFileMetadata.Add(TEXT("unreal/system/cpuBrand"), FPlatformMisc::GetCPUBrand());
+			
+			FString OSVersionLabel, OSSubVersionLabel;
+			FPlatformMisc::GetOSVersions(OSVersionLabel, OSSubVersionLabel);
+			InFileMetadata.Add(TEXT("unreal/system/osVersion"), OSVersionLabel);
+			InFileMetadata.Add(TEXT("unreal/system/osSubVersion"), OSSubVersionLabel);
+			InFileMetadata.Add(TEXT("unreal/system/deviceTemp"), LexToString(FPlatformMisc::GetDeviceTemperatureLevel()));
+			InFileMetadata.Add(TEXT("unreal/system/deviceMakeAndModel"), FPlatformMisc::GetDeviceMakeAndModel());
+			
+			if (InOutputDir.Len() > 0)
+			{
+				uint64 TotalNumBytes, NumFreeBytes;
+				FPlatformMisc::GetDiskTotalAndFreeSpace(InOutputDir, TotalNumBytes, NumFreeBytes);
+				InFileMetadata.Add(TEXT("unreal/stats/outputDirectoryTotalSizeMB"), LexToString(TotalNumBytes / MBDivider));
+				InFileMetadata.Add(TEXT("unreal/stats/outputDirectoryTotalFreeMB"), LexToString(NumFreeBytes / MBDivider));
+			}
+
+			FGPUDriverInfo DriverInfo = FPlatformMisc::GetGPUDriverInfo(GRHIAdapterName);
+			InFileMetadata.Add(TEXT("unreal/system/gpu/vendorId"), LexToString(DriverInfo.VendorId));
+			InFileMetadata.Add(TEXT("unreal/system/gpu/deviceDescription"), DriverInfo.DeviceDescription);
+			InFileMetadata.Add(TEXT("unreal/system/gpu/providerName"), DriverInfo.ProviderName);
+			InFileMetadata.Add(TEXT("unreal/system/gpu/internalDriverVersion"), DriverInfo.InternalDriverVersion);
+			InFileMetadata.Add(TEXT("unreal/system/gpu/userDriverVersion"), DriverInfo.UserDriverVersion);
+			InFileMetadata.Add(TEXT("unreal/system/gpu/driverDate"), DriverInfo.DriverDate);
+			InFileMetadata.Add(TEXT("unreal/system/gpu/rhiName"), DriverInfo.RHIName);
+
+			FPlatformMemoryStats MemoryStats = FPlatformMemory::GetStats();
+			InFileMetadata.Add(TEXT("unreal/stats/memory/availablePhysicalMB"), LexToString(MemoryStats.AvailablePhysical / MBDivider));
+			InFileMetadata.Add(TEXT("unreal/stats/memory/availableVirtualMB"), LexToString(MemoryStats.AvailableVirtual / MBDivider));
+			InFileMetadata.Add(TEXT("unreal/stats/memory/totalPhysicalMB"), LexToString(MemoryStats.TotalPhysical / MBDivider));
+			InFileMetadata.Add(TEXT("unreal/stats/memory/totalVirtualMB"), LexToString(MemoryStats.TotalVirtual / MBDivider));
+			InFileMetadata.Add(TEXT("unreal/stats/memory/peakUsedPhysicalMB"), LexToString(MemoryStats.PeakUsedPhysical / MBDivider));
+			InFileMetadata.Add(TEXT("unreal/stats/memory/peakUsedVirtualMB"), LexToString(MemoryStats.PeakUsedVirtual / MBDivider));
+
+		}
 	}
 }
diff --git a/Engine/Plugins/MovieScene/MovieRenderPipeline/Source/MovieRenderPipelineCore/Public/MoviePipelineUtils.h b/Engine/Plugins/MovieScene/MovieRenderPipeline/Source/MovieRenderPipelineCore/Public/MoviePipelineUtils.h
--- a/Engine/Plugins/MovieScene/MovieRenderPipeline/Source/MovieRenderPipelineCore/Public/MoviePipelineUtils.h
+++ b/Engine/Plugins/MovieScene/MovieRenderPipeline/Source/MovieRenderPipelineCore/Public/MoviePipelineUtils.h
@@ -132,6 +132,7 @@ namespace UE
 
 		MOVIERENDERPIPELINECORE_API FString GetJobAuthor(const UMoviePipelineExecutorJob* InJob);
 		MOVIERENDERPIPELINECORE_API void GetSharedFormatArguments(TMap<FString, FString>& InFilenameArguments, TMap<FString, FString>& InFileMetadata, const FDateTime& InDateTime, const int32 InVersionNumber, const UMoviePipelineExecutorJob* InJob);
+		MOVIERENDERPIPELINECORE_API void GetHardwareUsageMetadata(TMap<FString, FString>& InFileMetadata, const FString& InOutputDir);
 	}
 }
 
