MoviePipeline: Override the minimum timestep for Chaos Physics to zero during renders to avoid high temporal sample counts being artificially clamped. Also fixed some cvars that weren't being applied when launching in -game.

diff --git a/Engine/Plugins/MovieScene/MovieRenderPipeline/Source/MovieRenderPipelineCore/Private/MoviePipelineGameOverrideSetting.cpp b/Engine/Plugins/MovieScene/MovieRenderPipeline/Source/MovieRenderPipelineCore/Private/MoviePipelineGameOverrideSetting.cpp
--- a/Engine/Plugins/MovieScene/MovieRenderPipeline/Source/MovieRenderPipelineCore/Private/MoviePipelineGameOverrideSetting.cpp
+++ b/Engine/Plugins/MovieScene/MovieRenderPipeline/Source/MovieRenderPipelineCore/Private/MoviePipelineGameOverrideSetting.cpp
@@ -114,6 +114,9 @@ void UMoviePipelineGameOverrideSetting::ApplyCVarSettings(const bool bOverrideVa
 
 	// To make sure that the world partition streaming doesn't end up in critical streaming performances and stops streaming low priority cells.
 	MOVIEPIPELINE_STORE_AND_OVERRIDE_CVAR_INT(PreviousIgnoreStreamingPerformance, TEXT("wp.Runtime.BlockOnSlowStreaming"), 0, bOverrideValues);
+
+	// Remove any minimum delta time requirements from Chaos Physics to ensure accuracy at high Temporal Sample counts
+	MOVIEPIPELINE_STORE_AND_OVERRIDE_CVAR_FLOAT(PreviousChaosImmPhysicsMinStepTime, TEXT("p.Chaos.ImmPhys.MinStepTime"), 0, bOverrideValues);
 }
 
 void UMoviePipelineGameOverrideSetting::BuildNewProcessCommandLineArgsImpl(TArray<FString>& InOutUnrealURLParams, TArray<FString>& InOutCommandLineArgs, TArray<FString>& InOutDeviceProfileCvars, TArray<FString>& InOutExecCmds) const
@@ -196,5 +199,9 @@ void UMoviePipelineGameOverrideSetting::BuildNewProcessCommandLineArgsImpl(TArra
 		InOutDeviceProfileCvars.Add(FString::Printf(TEXT("a.URO.Enable=%d"), 0));
 	}
 
+	InOutDeviceProfileCvars.Add(FString::Printf(TEXT("au.NeverMuteNonRealtimeAudioDevices=%d"), 1));
 	InOutDeviceProfileCvars.Add(FString::Printf(TEXT("r.SkyLight.RealTimeReflectionCapture.TimeSlice=%d"), 0));
+	InOutDeviceProfileCvars.Add(FString::Printf(TEXT("r.VolumetricRenderTarget=%d"), 0));
+	InOutDeviceProfileCvars.Add(FString::Printf(TEXT("wp.Runtime.BlockOnSlowStreaming=%d"), 0));
+	InOutDeviceProfileCvars.Add(FString::Printf(TEXT("p.Chaos.ImmPhys.MinStepTime=%d"), 0));
 }
\ No newline at end of file
diff --git a/Engine/Plugins/MovieScene/MovieRenderPipeline/Source/MovieRenderPipelineCore/Public/MoviePipelineGameOverrideSetting.h b/Engine/Plugins/MovieScene/MovieRenderPipeline/Source/MovieRenderPipelineCore/Public/MoviePipelineGameOverrideSetting.h
--- a/Engine/Plugins/MovieScene/MovieRenderPipeline/Source/MovieRenderPipelineCore/Public/MoviePipelineGameOverrideSetting.h
+++ b/Engine/Plugins/MovieScene/MovieRenderPipeline/Source/MovieRenderPipelineCore/Public/MoviePipelineGameOverrideSetting.h
@@ -128,4 +128,5 @@ private:
 	int32 PreviousGeoCacheStreamerShowNotification;
 	int32 PreviousGeoCacheStreamerBlockTillFinish;
 #endif
+	float PreviousChaosImmPhysicsMinStepTime;
 };
\ No newline at end of file
