MRQ: Tweak help text to clarify how samples work with the path tracer

diff --git a/Engine/Plugins/MovieScene/MovieRenderPipeline/Source/MovieRenderPipelineRenderPasses/Private/MoviePipelineDeferredPasses.cpp b/Engine/Plugins/MovieScene/MovieRenderPipeline/Source/MovieRenderPipelineRenderPasses/Private/MoviePipelineDeferredPasses.cpp
--- a/Engine/Plugins/MovieScene/MovieRenderPipeline/Source/MovieRenderPipelineRenderPasses/Private/MoviePipelineDeferredPasses.cpp
+++ b/Engine/Plugins/MovieScene/MovieRenderPipeline/Source/MovieRenderPipelineRenderPasses/Private/MoviePipelineDeferredPasses.cpp
@@ -745,8 +745,8 @@ FText UMoviePipelineDeferredPass_PathTracer::GetFooterText(UMoviePipelineExecuto
 	return NSLOCTEXT(
 		"MovieRenderPipeline",
 		"DeferredBasePassSetting_FooterText_PathTracer",
-		"Samples per Pixel for the Path Tracer are controlled by the Spatial Sample Count from the Anti-Aliasing settings.\n"
-		"All other Path Tracer settings are taken from the Post Process settings as usual.");
+		"Sampling for the Path Tracer is controlled by the Anti-Aliasing settings and the Reference Motion Blur setting.\n"
+		"All other Path Tracer settings are taken from the Post Process settings.");
 }
 #endif
 namespace UE
