MoviePipeline: Use a texture copy rather than MSAA resolve to copy to readback texture. This fixes a crash with movie pipeline readbacks.

diff --git a/Engine/Plugins/MovieScene/MovieRenderPipeline/Source/MovieRenderPipelineCore/Private/MoviePipelineSurfaceReader.cpp b/Engine/Plugins/MovieScene/MovieRenderPipeline/Source/MovieRenderPipelineCore/Private/MoviePipelineSurfaceReader.cpp
--- a/Engine/Plugins/MovieScene/MovieRenderPipeline/Source/MovieRenderPipelineCore/Private/MoviePipelineSurfaceReader.cpp
+++ b/Engine/Plugins/MovieScene/MovieRenderPipeline/Source/MovieRenderPipelineCore/Private/MoviePipelineSurfaceReader.cpp
@@ -103,9 +103,11 @@ void FMoviePipelineSurfaceReader::ResolveSampleToReadbackTexture_RenderThread(co
 	GRenderTargetPool.FindFreeElement(RHICmdList, OutputDesc, ResampleTexturePooledRenderTarget, TEXT("ResampleTexture"));
 	check(ResampleTexturePooledRenderTarget);
 
+	RHICmdList.Transition(FRHITransitionInfo(ResampleTexturePooledRenderTarget->GetRHI(), ERHIAccess::Unknown, ERHIAccess::RTV));
+
 	// Enqueue a render pass which uses a simple shader and a fullscreen quad to copy the SourceSurfaceSample to this SurfaceReader's cpu-readback enabled texture.
 	// This RenderPass resolves to our ReadbackTexture 
-	FRHIRenderPassInfo RPInfo(ResampleTexturePooledRenderTarget->GetRHI(), ERenderTargetActions::Load_Store, ReadbackTexture);
+	FRHIRenderPassInfo RPInfo(ResampleTexturePooledRenderTarget->GetRHI(), ERenderTargetActions::Load_Store);
 	RHICmdList.BeginRenderPass(RPInfo, TEXT("MoviePipelineSurfaceResolveRenderTarget"));
 	{
 		RHICmdList.SetViewport(0, 0, 0.0f, TargetSize.X, TargetSize.Y, 1.0f);
@@ -159,6 +161,8 @@ void FMoviePipelineSurfaceReader::ResolveSampleToReadbackTexture_RenderThread(co
 			EDRF_Default);
 	}
 	RHICmdList.EndRenderPass();
+
+	CopyTextureWithTransitions(RHICmdList, ResampleTexturePooledRenderTarget->GetRHI(), ReadbackTexture, {});
 }
 
 void FMoviePipelineSurfaceReader::CopyReadbackTexture_RenderThread(TUniqueFunction<void(TUniquePtr<FImagePixelData>&&)>&& InFunctionCallback, TSharedPtr<FImagePixelDataPayload, ESPMode::ThreadSafe> InFramePayload)
