MoviePipeline: Calling DeleteAllJobs immediately after calling RenderWithExecutor will no longer crash and will instead report an executor failure due to there being no jobs in the queue to render.

diff --git a/Engine/Plugins/MovieScene/MovieRenderPipeline/Source/MovieRenderPipelineCore/Private/MoviePipelineInProcessExecutor.cpp b/Engine/Plugins/MovieScene/MovieRenderPipeline/Source/MovieRenderPipelineCore/Private/MoviePipelineInProcessExecutor.cpp
--- a/Engine/Plugins/MovieScene/MovieRenderPipeline/Source/MovieRenderPipelineCore/Private/MoviePipelineInProcessExecutor.cpp
+++ b/Engine/Plugins/MovieScene/MovieRenderPipeline/Source/MovieRenderPipelineCore/Private/MoviePipelineInProcessExecutor.cpp
@@ -88,6 +88,17 @@ void UMoviePipelineInProcessExecutor::OnMapLoadFinished(UWorld* NewWorld)
 	
 	// Stop listening for map load until we're done and know we want to start the next config.
 	FCoreUObjectDelegates::PostLoadMapWithWorld.RemoveAll(this);
+
+	if (CurrentPipelineIndex >= Queue->GetJobs().Num())
+	{
+		FCoreDelegates::OnBeginFrame.RemoveAll(this);
+		UE_LOG(LogMovieRenderPipeline, Error, TEXT("Out of bounds Job Index (%d) in Queue (Length: %d)! %d"), CurrentPipelineIndex, Queue->GetJobs().Num());
+
+		FText FailureReason = LOCTEXT("OutOfBoundsFailureDialog", "The job index is out of bounds, did you clear the Queue while rendering was happening? See log for details.");
+		OnPipelineErrored(nullptr, true, FailureReason);
+		OnExecutorFinishedImpl();
+		return;
+	}
 	
 	UMoviePipelineExecutorJob* CurrentJob = Queue->GetJobs()[CurrentPipelineIndex];
 
diff --git a/Engine/Plugins/MovieScene/MovieRenderPipeline/Source/MovieRenderPipelineCore/Private/MoviePipelineLinearExecutor.cpp b/Engine/Plugins/MovieScene/MovieRenderPipeline/Source/MovieRenderPipelineCore/Private/MoviePipelineLinearExecutor.cpp
--- a/Engine/Plugins/MovieScene/MovieRenderPipeline/Source/MovieRenderPipelineCore/Private/MoviePipelineLinearExecutor.cpp
+++ b/Engine/Plugins/MovieScene/MovieRenderPipeline/Source/MovieRenderPipelineCore/Private/MoviePipelineLinearExecutor.cpp
@@ -72,7 +72,8 @@ void UMoviePipelineLinearExecutorBase::StartPipelineByIndex(int32 InPipelineInde
 
 void UMoviePipelineLinearExecutorBase::OnIndividualPipelineFinished(UMoviePipeline* /* FinishedPipeline */)
 {
-	if (CurrentPipelineIndex == Queue->GetJobs().Num() - 1 || bIsCanceling)
+	const bool bNoMoreJobs = CurrentPipelineIndex >= Queue->GetJobs().Num() - 1;
+	if (bNoMoreJobs || bIsCanceling)
 	{
 		OnExecutorFinishedImpl();
 	}
