MoviePipeline: MoviePipelineInProcessExecutor no longer crashes if provided with a bad Sequence path.

diff --git a/Engine/Plugins/MovieScene/MovieRenderPipeline/Source/MovieRenderPipelineCore/Private/MoviePipelineInProcessExecutor.cpp b/Engine/Plugins/MovieScene/MovieRenderPipeline/Source/MovieRenderPipelineCore/Private/MoviePipelineInProcessExecutor.cpp
--- a/Engine/Plugins/MovieScene/MovieRenderPipeline/Source/MovieRenderPipelineCore/Private/MoviePipelineInProcessExecutor.cpp
+++ b/Engine/Plugins/MovieScene/MovieRenderPipeline/Source/MovieRenderPipelineCore/Private/MoviePipelineInProcessExecutor.cpp
@@ -16,6 +16,7 @@
 
 void UMoviePipelineInProcessExecutor::Start(const UMoviePipelineExecutorJob* InJob)
 {
+	Super::Start(InJob);
 	UWorld* World = MoviePipeline::FindCurrentWorld();
 
 	if (bUseCurrentLevel)
@@ -35,6 +36,18 @@ void UMoviePipelineInProcessExecutor::Start(const UMoviePipelineExecutorJob* InJ
 		UE_LOG(LogMovieRenderPipeline, Log, TEXT("Starting %s"), *GetNameSafe(World));
 	}
 
+	// Check for null sequence
+	ULevelSequence* LevelSequence = Cast<ULevelSequence>(InJob->Sequence.TryLoad());
+	if (!LevelSequence)
+	{
+		UE_LOG(LogMovieRenderPipeline, Error, TEXT("Failed to load Sequence specified by job: %s"), *InJob->Sequence.ToString());
+
+		FText FailureReason = LOCTEXT("InvalidSequenceFailureDialog", "One or more jobs in the queue has an invalid/null sequence. See log for details.");
+		OnPipelineErrored(nullptr, true, FailureReason);
+		OnExecutorFinishedImpl();
+		return;
+	}
+
 	BackupState();
 	
 	// Initialize the transient settings so that they will exist in time for the GameOverrides check.
