MRQ: Post-processed output is not accumulated properly when using the path tracer

diff --git a/Engine/Plugins/MovieScene/MovieRenderPipeline/Source/MovieRenderPipelineRenderPasses/Private/MoviePipelineDeferredPasses.cpp b/Engine/Plugins/MovieScene/MovieRenderPipeline/Source/MovieRenderPipelineRenderPasses/Private/MoviePipelineDeferredPasses.cpp
--- a/Engine/Plugins/MovieScene/MovieRenderPipeline/Source/MovieRenderPipelineRenderPasses/Private/MoviePipelineDeferredPasses.cpp
+++ b/Engine/Plugins/MovieScene/MovieRenderPipeline/Source/MovieRenderPipelineRenderPasses/Private/MoviePipelineDeferredPasses.cpp
@@ -373,7 +373,7 @@
 		for (UMaterialInterface* VisMaterial : View->FinalPostProcessSettings.BufferVisualizationOverviewMaterials)
 		{
 			// If this was just to contribute to the history buffer, no need to go any further.
-			if (InSampleState.bDiscardResult)
+			if (InOutSampleState.bDiscardResult)
 			{
 				continue;
 			}
