MoviePipeline: Fixed log spam caused by querying GPU info every frame.

diff --git a/Engine/Plugins/MovieScene/MovieRenderPipeline/Source/MovieRenderPipelineCore/Private/MoviePipelineUtils.cpp b/Engine/Plugins/MovieScene/MovieRenderPipeline/Source/MovieRenderPipelineCore/Private/MoviePipelineUtils.cpp
--- a/Engine/Plugins/MovieScene/MovieRenderPipeline/Source/MovieRenderPipelineCore/Private/MoviePipelineUtils.cpp
+++ b/Engine/Plugins/MovieScene/MovieRenderPipeline/Source/MovieRenderPipelineCore/Private/MoviePipelineUtils.cpp
@@ -786,6 +786,27 @@
 			}
 		}
 
+		void GetCachedGPUDriverInfo(TMap<FString, FString>& InFileMetadata)
+		{
+			static TMap<FString, FString> CachedInfo;
+			if (CachedInfo.Num() > 0)
+			{
+				InFileMetadata.Append(CachedInfo);
+				return;
+			}
+
+			FGPUDriverInfo DriverInfo = FPlatformMisc::GetGPUDriverInfo(GRHIAdapterName);
+			CachedInfo.Add(TEXT("unreal/system/gpu/vendorId"), LexToString(DriverInfo.VendorId));
+			CachedInfo.Add(TEXT("unreal/system/gpu/deviceDescription"), DriverInfo.DeviceDescription);
+			CachedInfo.Add(TEXT("unreal/system/gpu/providerName"), DriverInfo.ProviderName);
+			CachedInfo.Add(TEXT("unreal/system/gpu/internalDriverVersion"), DriverInfo.InternalDriverVersion);
+			CachedInfo.Add(TEXT("unreal/system/gpu/userDriverVersion"), DriverInfo.UserDriverVersion);
+			CachedInfo.Add(TEXT("unreal/system/gpu/driverDate"), DriverInfo.DriverDate);
+			CachedInfo.Add(TEXT("unreal/system/gpu/rhiName"), DriverInfo.RHIName);
+
+			InFileMetadata.Append(CachedInfo);
+		}
+
 		void GetHardwareUsageMetadata(TMap<FString, FString>& InFileMetadata, const FString& InOutputDir)
 		{
 			SCOPE_CYCLE_COUNTER(STAT_HardwareMetadata);
@@ -813,14 +834,9 @@
 				InFileMetadata.Add(TEXT("unreal/stats/outputDirectoryTotalFreeMB"), LexToString(NumFreeBytes / MBDivider));
 			}
 
-			FGPUDriverInfo DriverInfo = FPlatformMisc::GetGPUDriverInfo(GRHIAdapterName);
-			InFileMetadata.Add(TEXT("unreal/system/gpu/vendorId"), LexToString(DriverInfo.VendorId));
-			InFileMetadata.Add(TEXT("unreal/system/gpu/deviceDescription"), DriverInfo.DeviceDescription);
-			InFileMetadata.Add(TEXT("unreal/system/gpu/providerName"), DriverInfo.ProviderName);
-			InFileMetadata.Add(TEXT("unreal/system/gpu/internalDriverVersion"), DriverInfo.InternalDriverVersion);
-			InFileMetadata.Add(TEXT("unreal/system/gpu/userDriverVersion"), DriverInfo.UserDriverVersion);
-			InFileMetadata.Add(TEXT("unreal/system/gpu/driverDate"), DriverInfo.DriverDate);
-			InFileMetadata.Add(TEXT("unreal/system/gpu/rhiName"), DriverInfo.RHIName);
+			// We cache this as each time we fetch it from FPlatformMisc it prints to the log, and
+			// these values aren't going to change anyways during runtime.
+			GetCachedGPUDriverInfo(InFileMetadata);
 
 			FPlatformMemoryStats MemoryStats = FPlatformMemory::GetStats();
 			InFileMetadata.Add(TEXT("unreal/stats/memory/availablePhysicalMB"), LexToString(MemoryStats.AvailablePhysical / MBDivider));
@@ -829,7 +845,6 @@
 			InFileMetadata.Add(TEXT("unreal/stats/memory/totalVirtualMB"), LexToString(MemoryStats.TotalVirtual / MBDivider));
 			InFileMetadata.Add(TEXT("unreal/stats/memory/peakUsedPhysicalMB"), LexToString(MemoryStats.PeakUsedPhysical / MBDivider));
 			InFileMetadata.Add(TEXT("unreal/stats/memory/peakUsedVirtualMB"), LexToString(MemoryStats.PeakUsedVirtual / MBDivider));
-
 		}
 	}
 }
