Movie Render Queue: Fix potential crash from null pointer dereference when searching MRQ setting classes

diff --git a/Engine/Plugins/MovieScene/MovieRenderPipeline/Source/MovieRenderPipelineCore/Private/MoviePipelineUtils.cpp b/Engine/Plugins/MovieScene/MovieRenderPipeline/Source/MovieRenderPipelineCore/Private/MoviePipelineUtils.cpp
--- a/Engine/Plugins/MovieScene/MovieRenderPipeline/Source/MovieRenderPipelineCore/Private/MoviePipelineUtils.cpp
+++ b/Engine/Plugins/MovieScene/MovieRenderPipeline/Source/MovieRenderPipelineCore/Private/MoviePipelineUtils.cpp
@@ -62,7 +62,7 @@ namespace UE
 				{
 					ParentClass = UClass::TryFindTypeSlow<UClass>(FPackageName::ExportTextPathToObjectPath(ParentClassName));
 
-					if (!ParentClass->IsChildOf(UMoviePipelineSetting::StaticClass()))
+					if (ParentClass == nullptr || !ParentClass->IsChildOf(UMoviePipelineSetting::StaticClass()))
 					{
 						continue;
 					}
