MoviePipeline: Improvements to the API for creating renders using MRQ during runtime. This greatly simplifies common workflows and removes a lot of the tripping points for people. This also fixes CancelJobs not working on the InProcess executor which is required to not deadlock the render thread on map change/shutdown during PIE if a render is in progress.

diff --git a/Engine/Plugins/MovieScene/MovieRenderPipeline/Source/MovieRenderPipelineCore/Private/MoviePipeline.cpp b/Engine/Plugins/MovieScene/MovieRenderPipeline/Source/MovieRenderPipelineCore/Private/MoviePipeline.cpp
--- a/Engine/Plugins/MovieScene/MovieRenderPipeline/Source/MovieRenderPipelineCore/Private/MoviePipeline.cpp
+++ b/Engine/Plugins/MovieScene/MovieRenderPipeline/Source/MovieRenderPipelineCore/Private/MoviePipeline.cpp
@@ -288,7 +288,7 @@ void UMoviePipeline::Initialize(UMoviePipelineExecutorJob* InJob)
 	
 	if (UGameViewportClient* Viewport = GetWorld()->GetGameViewport())
 	{
-		Viewport->bDisableWorldRendering = true;
+		Viewport->bDisableWorldRendering = !ViewportInitArgs.bRenderViewport;
 	}
 
 	for (ULevelStreaming* Level : GetWorld()->GetStreamingLevels())
@@ -1279,7 +1279,7 @@ bool UMoviePipeline::DebugFrameStepPreTick()
 
 void UMoviePipeline::LoadDebugWidget()
 {
-	TSubclassOf<UMovieRenderDebugWidget> DebugWidgetClassToUse = DebugWidgetClass;
+	TSubclassOf<UMovieRenderDebugWidget> DebugWidgetClassToUse = ViewportInitArgs.DebugWidgetClass;
 	if (DebugWidgetClassToUse.Get() == nullptr)
 	{
 		DebugWidgetClassToUse = LoadClass<UMovieRenderDebugWidget>(nullptr, *DefaultDebugWidgetAsset, nullptr, LOAD_None, nullptr);
diff --git a/Engine/Plugins/MovieScene/MovieRenderPipeline/Source/MovieRenderPipelineCore/Private/MoviePipelineInProcessExecutor.cpp b/Engine/Plugins/MovieScene/MovieRenderPipeline/Source/MovieRenderPipelineCore/Private/MoviePipelineInProcessExecutor.cpp
--- a/Engine/Plugins/MovieScene/MovieRenderPipeline/Source/MovieRenderPipelineCore/Private/MoviePipelineInProcessExecutor.cpp
+++ b/Engine/Plugins/MovieScene/MovieRenderPipeline/Source/MovieRenderPipelineCore/Private/MoviePipelineInProcessExecutor.cpp
@@ -122,7 +122,7 @@ void UMoviePipelineInProcessExecutor::OnMapLoadFinished(UWorld* NewWorld)
 	}
 
 	ActiveMoviePipeline = NewObject<UMoviePipeline>(NewWorld, MoviePipelineClass);
-	ActiveMoviePipeline->DebugWidgetClass = DebugWidgetClass;
+	ActiveMoviePipeline->SetViewportInitArgs(ViewportInitArgs);
 
 	// We allow users to set a multi-frame delay before we actually run the Initialization function and start thinking.
 	// This solves cases where there are engine systems that need to finish loading before we do anything.
@@ -186,6 +186,8 @@ void UMoviePipelineInProcessExecutor::OnMoviePipelineFinished(FMoviePipelineOutp
 {
 	FCoreDelegates::OnBeginFrame.RemoveAll(this);
 	UMoviePipeline* MoviePipeline = ActiveMoviePipeline;
+	
+	OnIndividualJobFinishedDelegateNative.Broadcast(InOutputData);
 
 	if (ActiveMoviePipeline)
 	{
@@ -280,5 +282,22 @@ void UMoviePipelineInProcessExecutor::RestoreState()
 		}
 	}
 }
+void UMoviePipelineInProcessExecutor::CancelAllJobs_Implementation()
+{
+	// Setting this flag will make the completion of the current job skip the remaining jobs
+	bIsCanceling = true;
+	CancelCurrentJob();
+}
 
+void UMoviePipelineInProcessExecutor::CancelCurrentJob_Implementation()
+{
+	if (ActiveMoviePipeline && ActiveMoviePipeline->GetPipelineState() != EMovieRenderPipelineState::Finished)
+	{
+		UE_LOG(LogMovieRenderPipeline, Log, TEXT("MoviePipelineInProcessExecutor: Application quit while Movie Pipeline was still active. Stalling to do full shutdown."));
+
+		// This will flush any outstanding work on the movie pipeline (file writes) immediately
+		ActiveMoviePipeline->RequestShutdown(); // Set the Shutdown Requested flag.
+		ActiveMoviePipeline->Shutdown(); // Flush the shutdown.
+	}
+}
 #undef LOCTEXT_NAMESPACE // "MoviePipelineInProcessExecutor"
diff --git a/Engine/Plugins/MovieScene/MovieRenderPipeline/Source/MovieRenderPipelineCore/Private/MoviePipelineQueueEngineSubsystem.cpp b/Engine/Plugins/MovieScene/MovieRenderPipeline/Source/MovieRenderPipelineCore/Private/MoviePipelineQueueEngineSubsystem.cpp
--- a/Engine/Plugins/MovieScene/MovieRenderPipeline/Source/MovieRenderPipelineCore/Private/MoviePipelineQueueEngineSubsystem.cpp
+++ b/Engine/Plugins/MovieScene/MovieRenderPipeline/Source/MovieRenderPipelineCore/Private/MoviePipelineQueueEngineSubsystem.cpp
@@ -3,16 +3,23 @@
 #include "MoviePipelineQueueEngineSubsystem.h"
 #include "Modules/ModuleManager.h"
 #include "MoviePipeline.h"
+#include "MoviePipelineInProcessExecutor.h"
 
 UMoviePipelineExecutorBase* UMoviePipelineQueueEngineSubsystem::RenderQueueWithExecutor(TSubclassOf<UMoviePipelineExecutorBase> InExecutorType)
 {
+	if (!ensureMsgf(!InExecutorType.Get(), TEXT("RenderQueueWithExecutor cannot be called with a null class type!")))
+	{
+		FFrame::KismetExecutionMessage(TEXT("RenderQueueWithExecutor cannot be called with a null class type!"), ELogVerbosity::Error);
+		return nullptr;
+	}
+
 	RenderQueueWithExecutorInstance(NewObject<UMoviePipelineExecutorBase>(this, InExecutorType));
 	return ActiveExecutor;
 }
 
 void UMoviePipelineQueueEngineSubsystem::RenderQueueWithExecutorInstance(UMoviePipelineExecutorBase* InExecutor)
 {
-	if(!ensureMsgf(!IsRendering(), TEXT("RenderQueueWithExecutor cannot be called while already rendering!")))
+	if(!ensureMsgf(!IsRendering(), TEXT("RenderQueueWithExecutorInstance cannot be called while already rendering!")))
 	{
 		FFrame::KismetExecutionMessage(TEXT("Render already in progress."), ELogVerbosity::Error);
 		return;
@@ -25,6 +32,11 @@ void UMoviePipelineQueueEngineSubsystem::RenderQueueWithExecutorInstance(UMovieP
 	}
 
 	ActiveExecutor = InExecutor;
+	UE::MoviePipeline::FViewportArgs InitArgs;
+	InitArgs.bRenderViewport = bCachedRenderPlayerViewport;
+	InitArgs.DebugWidgetClass = CachedProgressWidgetClass;
+	ActiveExecutor->SetViewportInitArgs(InitArgs);
+
 	ActiveExecutor->OnExecutorFinished().AddUObject(this, &UMoviePipelineQueueEngineSubsystem::OnExecutorFinished);
 	ActiveExecutor->Execute(GetQueue());
 }
@@ -32,5 +44,73 @@ void UMoviePipelineQueueEngineSubsystem::RenderQueueWithExecutorInstance(UMovieP
 void UMoviePipelineQueueEngineSubsystem::OnExecutorFinished(UMoviePipelineExecutorBase* InPipelineExecutor, bool bSuccess)
 {
 	ActiveExecutor = nullptr;
+	
+	// Wait until ActiveExecutor has been nulled so that if they immediately try to allocate new jobs as a result of
+	// the callback it doesn't throw an exception about a render being in progress.
+	if(CachedRenderJobResults.IsSet())
+	{
+		OnRenderFinished.Broadcast(CachedRenderJobResults.GetValue());
+		CachedRenderJobResults.Reset();
+	}
+}
+
+UMoviePipelineExecutorJob* UMoviePipelineQueueEngineSubsystem::AllocateJob(ULevelSequence* InSequence)
+{
+	if(!InSequence)
+	{
+		FFrame::KismetExecutionMessage(TEXT("AllocateJob cannot be called with a null sequence!"), ELogVerbosity::Error);
+		return nullptr;
+	}
+	
+	if(IsRendering())
+	{
+		FFrame::KismetExecutionMessage(TEXT("AllocateJob cannot be called while rendering!"), ELogVerbosity::Error);
+		return nullptr;
+	}
+	
+	CurrentQueue->DeleteAllJobs();
+	
+	UMoviePipelineExecutorJob* NewJob = CurrentQueue->AllocateNewJob(UMoviePipelineExecutorJob::StaticClass());
+	NewJob->SetSequence(FSoftObjectPath(InSequence));
+	NewJob->Map = FSoftObjectPath(GetWorld());
+	NewJob->JobName = NewJob->Sequence.GetAssetName();
+	
+	
+	return NewJob;
+}
+
+void UMoviePipelineQueueEngineSubsystem::RenderJob(UMoviePipelineExecutorJob* InJob)
+{
+	if(!CurrentQueue->GetJobs().Contains(InJob))
+	{
+		FFrame::KismetExecutionMessage(TEXT("RenderJob cannot be called with a job not created by AllocateJob!"), ELogVerbosity::Error);
+		return;
+	}
+	
+	if(IsRendering())
+	{
+		FFrame::KismetExecutionMessage(TEXT("RenderJob cannot be called while rendering!"), ELogVerbosity::Error);
+		return;
+	}
+	
+	UMoviePipelineInProcessExecutor* InProcessExecutor = NewObject<UMoviePipelineInProcessExecutor>(this);
+	InProcessExecutor->bUseCurrentLevel = true;
+	InProcessExecutor->OnIndividualJobFinished().AddUObject(this, &UMoviePipelineQueueEngineSubsystem::OnIndividualJobFinished);
+	RenderQueueWithExecutorInstance(InProcessExecutor);
+}
+
+void UMoviePipelineQueueEngineSubsystem::OnIndividualJobFinished(FMoviePipelineOutputData Results)
+{
+	// OnExecutorFinished will be called after this (since this callback is per-job), but we want users
+	// to be able to queue up another job immediately without needing to wait 1 tick for ActiveExecutor to
+	// be set to nullptr. If they queue up a new job right now, the executor will get nulled out. So instead
+	// we cache the result of this individual job, and then broadcast the actual Blueprint event from OnExecutorFinished.
+	CachedRenderJobResults = Results;
+}
+
+void UMoviePipelineQueueEngineSubsystem::SetConfiguration(TSubclassOf<UMovieRenderDebugWidget> InProgressWidgetClass, const bool bRenderPlayerViewport)
+{
+	CachedProgressWidgetClass = InProgressWidgetClass;
+	bCachedRenderPlayerViewport = bRenderPlayerViewport;
 }
 
diff --git a/Engine/Plugins/MovieScene/MovieRenderPipeline/Source/MovieRenderPipelineCore/Private/MovieRenderDebugWidget.h b/Engine/Plugins/MovieScene/MovieRenderPipeline/Source/MovieRenderPipelineCore/Private/MovieRenderDebugWidget.h
deleted file mode 100644
diff --git a/Engine/Plugins/MovieScene/MovieRenderPipeline/Source/MovieRenderPipelineCore/Public/MoviePipeline.h b/Engine/Plugins/MovieScene/MovieRenderPipeline/Source/MovieRenderPipelineCore/Public/MoviePipeline.h
--- a/Engine/Plugins/MovieScene/MovieRenderPipeline/Source/MovieRenderPipelineCore/Public/MoviePipeline.h
+++ b/Engine/Plugins/MovieScene/MovieRenderPipeline/Source/MovieRenderPipelineCore/Public/MoviePipeline.h
@@ -237,6 +237,12 @@ public:
 	*/
 	void ResolveFilenameFormatArguments(const FString& InFormatString, const TMap<FString, FString>& InFormatOverrides, FString& OutFinalPath, FMoviePipelineFormatArgs& OutFinalFormatArgs, const FMoviePipelineFrameOutputState* InOutputState=nullptr, const int32 InFrameNumberOffset=0) const;
 
+	/**
+	* Allows initialization of some viewport-related arguments that aren't related to the job. Needs to
+	* be called before the Initialize function. Optional.
+	*/
+	void SetViewportInitArgs(const UE::MoviePipeline::FViewportArgs& InArgs) { ViewportInitArgs = InArgs; }
+
 protected:
 	/**
 	* This function should be called by the Executor when execution has finished (this should still be called in the event of an error)
@@ -462,6 +468,7 @@ public:
 	IImageWriteQueue* ImageWriteQueue;
 
 	/** Optional widget for feedback during render */
+	UE_DEPRECATED(5.1, "Use SetViewportInitArgs instead.")
 	UPROPERTY(Transient)
 	TSubclassOf<UMovieRenderDebugWidget> DebugWidgetClass;
 
@@ -480,6 +487,8 @@ private:
 	/** Files that we've requested be written to disk but have not yet finished writing. */
 	TArray<FMoviePipelineOutputFuture> OutputFutures;
 
+	UE::MoviePipeline::FViewportArgs ViewportInitArgs;
+
 	TSharedPtr<MoviePipeline::FCameraCutSubSectionHierarchyNode> CachedSequenceHierarchyRoot;
 
 	struct FClothSimSettingsCache
diff --git a/Engine/Plugins/MovieScene/MovieRenderPipeline/Source/MovieRenderPipelineCore/Public/MoviePipelineExecutor.h b/Engine/Plugins/MovieScene/MovieRenderPipeline/Source/MovieRenderPipelineCore/Public/MoviePipelineExecutor.h
--- a/Engine/Plugins/MovieScene/MovieRenderPipeline/Source/MovieRenderPipelineCore/Public/MoviePipelineExecutor.h
+++ b/Engine/Plugins/MovieScene/MovieRenderPipeline/Source/MovieRenderPipelineCore/Public/MoviePipelineExecutor.h
@@ -7,6 +7,7 @@
 #include "Http.h"
 #include "Misc/CoreDelegates.h"
 #include "Serialization/ArrayReader.h"
+#include "MovieRenderDebugWidget.h"
 #include "MoviePipelineExecutor.generated.h"
 
 class UMoviePipelineMasterConfig;
@@ -14,7 +15,6 @@ class UMoviePipelineExecutorBase;
 class UMoviePipelineExecutorJob;
 class UMoviePipeline;
 class UMoviePipelineQueue;
-class UMovieRenderDebugWidget;
 
 DECLARE_MULTICAST_DELEGATE_TwoParams(FOnMoviePipelineExecutorFinishedNative, UMoviePipelineExecutorBase* /*PipelineExecutor*/, bool /*bSuccess*/);
 DECLARE_DYNAMIC_MULTICAST_DELEGATE_TwoParams(FOnMoviePipelineExecutorFinished, UMoviePipelineExecutorBase*, PipelineExecutor, bool, bSuccess);
@@ -339,9 +339,16 @@ private:
 
 public:
 	/** Optional widget for feedback during render */
+	UE_DEPRECATED(5.1, "Use SetViewportInitArgs instead.")
 	UPROPERTY(BlueprintReadWrite, Category = "Movie Render Pipeline")
 	TSubclassOf<UMovieRenderDebugWidget> DebugWidgetClass;
 
+	void SetViewportInitArgs(const UE::MoviePipeline::FViewportArgs& InArgs) { ViewportInitArgs = InArgs; }
+	/**
+	* Some global initialization args that get passed to the UMoviePipeline before the call to ::Initialize()
+	*/
+	UE::MoviePipeline::FViewportArgs ViewportInitArgs;
+
 	/**
 	* Arbitrary data that can be associated with the executor. Not used by default implementations, nor read.
 	* This can be used to attach third party metadata such as job ids from remote farms.
diff --git a/Engine/Plugins/MovieScene/MovieRenderPipeline/Source/MovieRenderPipelineCore/Public/MoviePipelineInProcessExecutor.h b/Engine/Plugins/MovieScene/MovieRenderPipeline/Source/MovieRenderPipelineCore/Public/MoviePipelineInProcessExecutor.h
--- a/Engine/Plugins/MovieScene/MovieRenderPipeline/Source/MovieRenderPipelineCore/Public/MoviePipelineInProcessExecutor.h
+++ b/Engine/Plugins/MovieScene/MovieRenderPipeline/Source/MovieRenderPipelineCore/Public/MoviePipelineInProcessExecutor.h
@@ -2,6 +2,7 @@
 #pragma once
 
 #include "MoviePipelineLinearExecutor.h"
+#include "MoviePipeline.h"
 #include "MoviePipelineInProcessExecutor.generated.h"
 
 class UWorld;
@@ -30,9 +31,15 @@ public:
 	UPROPERTY(BlueprintReadWrite, Category = "Movie Render Pipeline")
 	bool bUseCurrentLevel;
 
+	FMoviePipelineWorkFinishedNative& OnIndividualJobFinished()
+	{
+		return OnIndividualJobFinishedDelegateNative;
+	}
+
 protected:
 	virtual void Start(const UMoviePipelineExecutorJob* InJob) override;
-
+	virtual void CancelAllJobs_Implementation() override;
+	virtual void CancelCurrentJob_Implementation() override;
 private:
 	void OnMapLoadFinished(UWorld* NewWorld);
 	void OnMoviePipelineFinished(FMoviePipelineOutputData InOutputData);
@@ -62,4 +69,5 @@ private:
 	};
 
 	FSavedState SavedState;
+	FMoviePipelineWorkFinishedNative OnIndividualJobFinishedDelegateNative;
 };
\ No newline at end of file
diff --git a/Engine/Plugins/MovieScene/MovieRenderPipeline/Source/MovieRenderPipelineCore/Public/MoviePipelineQueueEngineSubsystem.h b/Engine/Plugins/MovieScene/MovieRenderPipeline/Source/MovieRenderPipelineCore/Public/MoviePipelineQueueEngineSubsystem.h
--- a/Engine/Plugins/MovieScene/MovieRenderPipeline/Source/MovieRenderPipelineCore/Public/MoviePipelineQueueEngineSubsystem.h
+++ b/Engine/Plugins/MovieScene/MovieRenderPipeline/Source/MovieRenderPipelineCore/Public/MoviePipelineQueueEngineSubsystem.h
@@ -4,9 +4,21 @@
 #include "Subsystems/EngineSubsystem.h"
 #include "MoviePipelineQueue.h"
 #include "MoviePipelineExecutor.h"
+#include "MoviePipeline.h"
 #include "MoviePipelineQueueEngineSubsystem.generated.h"
 
-UCLASS(BlueprintType)
+
+/**
+* This subsystem is intended for use when rendering in a shipping game (but can also be used in PIE
+* during development/testing). See UMoviePipelineQueueSubsystem for the Editor-only queue which is
+* bound to the Movie Render Queue UI. To do simple renders at runtime, call AllocateJob(...)
+* with the Level Sequence you want to render, then call FindOrAddSettingByClass on the job to add
+* the settings (such as render pass, output type, and output directory) that you want for the job.
+* Finally, call RenderJob with the job you just configured. Register a delegate to OnRenderFinished 
+* to be notified when the render finished. You can optionally call SetConfiguration(...) before
+* RenderJob to configure some advanced options.
+*/
+UCLASS(BlueprintType, meta=(DisplayName="MoviePipeline Runtime Subsystem"))
 class MOVIERENDERPIPELINECORE_API UMoviePipelineQueueEngineSubsystem : public UEngineSubsystem
 {
 	GENERATED_BODY()
@@ -16,6 +28,8 @@ public:
 	UMoviePipelineQueueEngineSubsystem()
 	{
 		CurrentQueue = CreateDefaultSubobject<UMoviePipelineQueue>("EngineMoviePipelineQueue");
+		const bool bRenderViewport = false;
+		SetConfiguration(nullptr, bRenderViewport);
 	}
 	
 	/** Returns the queue of Movie Pipelines that need to be rendered. */
@@ -64,10 +78,71 @@ public:
 	{
 		return ActiveExecutor ? ActiveExecutor->IsRendering() : false;
 	}
+	
+	/**
+	* Sets some advanced configuration options that are used only occasionally to have better control over integrating it into
+	* your game/application. This applies to both RenderQueueWithExecutor(Instance) and the simplified RenderJob API. This persists
+	* until you call it again with different settings, and needs to be called before the Render* functions.
+	*
+	* @param ProgressWidget 		- Create a User Widget that inherits from UMovieRenderDebugWidget and it will be created. 
+	*						  		  Passing nullptr will use the default widget, if you wish to hide the widget then create a
+	*						  		  new widget which inherits from UMovieRenderDebugWidget and put nothing in it.
+	* @param bRenderPlayerViewport  - If true, we will render the regular player viewport in addition to the off-screen MRQ render.
+	*								  This is significantly performance heavy (doubles render times) but can be useful in the event
+	*								  that you want to keep rendering the player viewport to better integrate the render into your
+	*								  own application.
+	*
+	*/
+	UFUNCTION(BlueprintCallable, Category = "Movie Render Pipeline|Rendering")
+	void SetConfiguration(TSubclassOf<UMovieRenderDebugWidget> InProgressWidgetClass = nullptr, const bool bRenderPlayerViewport = false);
+
+	/**
+	* Convenience function for creating a UMoviePipelineExecutorJob out of the given Level Sequence asset. The 
+	* newly created job will be initialized to render on the current level, and will not have any default settings
+	* added to it - instead you will need to call FindOrAddSettingByClass on the job's configuration to add
+	* settings such as render passes (UMoviePipelineDeferredPassBase), output types (UMoviePipelineImageSequenceOutput_PNG),
+	* and configure the output directory (UMoviePipelineOutputSetting). Once configuration is complete, register
+	* a delegate to OnRenderFinished and then call RenderJob.
+	*
+	* Calling this function will clear the internal queue, see RenderJob for more details.
+	* 
+	* Using this function while IsRendering() returns true will result in an exception being thrown and no attempt
+	* being made to create the job.
+	*/
+	UFUNCTION(BlueprintCallable, Category = "Movie Render Pipeline|Rendering")
+	UMoviePipelineExecutorJob* AllocateJob(ULevelSequence* InSequence);
+
+	/**
+	* Convenience function for rendering the specified job. Calling this will render the specified job (if it was
+	* allocated using AllocateJob) and then it will reset the queue once finished. If you need to render multiple 
+	* jobs (in a queue) then you need to either implement the queue behavior yourself, or use 
+	* GetQueue()->AllocateJob(...) instead and use the non-convenience functions.
+	*
+	* Calling this function will clear the queue (after completion).
+	*
+	* Using this function while IsRendering() returns true will result in an exception being thrown and no attempt
+	* being made to render the job.
+	*/
+	UFUNCTION(BlueprintCallable, Category = "Movie Render Pipeline|Rendering")
+	void RenderJob(UMoviePipelineExecutorJob* InJob);
+
+	/**
+	* Assign a function to this delegate to get notified when each individual job is finished.
+	*
+	* THIS WILL ONLY BE CALLED IF USING THE RENDERJOB CONVENIENCE FUNCTION.
+	*
+	* Because there can only be one job in the queue when using RenderJob, this will be called when
+	* the render is complete, and the executor has been released. This allows you to queue up another
+	* job immediately as a result of the OnRenderFinished callback.
+	*/
+	UPROPERTY(BlueprintAssignable)
+	FMoviePipelineWorkFinished OnRenderFinished;
+
 
 private:
 	/** Called when the executor is finished so that we can release it and stop reporting IsRendering() == true. */
 	void OnExecutorFinished(UMoviePipelineExecutorBase* InPipelineExecutor, bool bSuccess);
+	void OnIndividualJobFinished(FMoviePipelineOutputData Results);
 
 private:
 	UPROPERTY(Transient)
@@ -75,4 +150,8 @@ private:
 
 	UPROPERTY(Transient, Instanced)
 	UMoviePipelineQueue* CurrentQueue;
+	
+	TOptional<FMoviePipelineOutputData> CachedRenderJobResults;
+	TSubclassOf<UMovieRenderDebugWidget> CachedProgressWidgetClass;
+	bool bCachedRenderPlayerViewport;
 };
\ No newline at end of file
diff --git a/Engine/Plugins/MovieScene/MovieRenderPipeline/Source/MovieRenderPipelineCore/Public/MovieRenderDebugWidget.h b/Engine/Plugins/MovieScene/MovieRenderPipeline/Source/MovieRenderPipelineCore/Public/MovieRenderDebugWidget.h
new file mode 100644
--- /dev/null
+++ b/Engine/Plugins/MovieScene/MovieRenderPipeline/Source/MovieRenderPipelineCore/Public/MovieRenderDebugWidget.h
@@ -0,0 +1,21 @@
+// Copyright Epic Games, Inc. All Rights Reserved.
+#pragma once
+
+#include "Blueprint/UserWidget.h"
+#include "MovieRenderDebugWidget.generated.h"
+
+class UMoviePipeline;
+
+/**
+* C++ Base Class for the debug widget that is drawn onto the game viewport
+* (but not burned into the output files) that allow us to easily visualize
+* the current state of the pipeline.
+*/
+UCLASS(Blueprintable, Abstract)
+class UMovieRenderDebugWidget : public UUserWidget
+{
+	GENERATED_BODY()
+public:
+	UFUNCTION(BlueprintImplementableEvent)
+	void OnInitializedForPipeline(UMoviePipeline* ForPipeline);
+};
\ No newline at end of file
diff --git a/Engine/Plugins/MovieScene/MovieRenderPipeline/Source/MovieRenderPipelineCore/Public/MovieRenderPipelineDataTypes.h b/Engine/Plugins/MovieScene/MovieRenderPipeline/Source/MovieRenderPipelineCore/Public/MovieRenderPipelineDataTypes.h
--- a/Engine/Plugins/MovieScene/MovieRenderPipeline/Source/MovieRenderPipelineCore/Public/MovieRenderPipelineDataTypes.h
+++ b/Engine/Plugins/MovieScene/MovieRenderPipeline/Source/MovieRenderPipelineCore/Public/MovieRenderPipelineDataTypes.h
@@ -12,6 +12,7 @@
 #include "OpenColorIOColorSpace.h"
 #include "Async/ParallelFor.h"
 #include "MovieSceneSequenceID.h"
+#include "MovieRenderDebugWidget.h"
 #include "MovieRenderPipelineDataTypes.generated.h"
 
 class UMovieSceneCinematicShotSection;
@@ -1203,6 +1204,17 @@ namespace MoviePipeline
 	};
 }
 
+namespace UE
+{
+	namespace MoviePipeline
+	{
+		struct FViewportArgs
+		{
+			TSubclassOf<UMovieRenderDebugWidget> DebugWidgetClass;
+			bool bRenderViewport = false;
+		};
+	}
+}
 /**
 * Contains information about the to-disk output generated by a movie pipeline. This structure is used both for per-shot work finished
 * callbacks and for the final render finished callback. When used as a per-shot callback ShotData will only have one entry (for the
diff --git a/Engine/Plugins/MovieScene/MovieRenderPipeline/Source/MovieRenderPipelineEditor/Private/MoviePipelinePIEExecutor.cpp b/Engine/Plugins/MovieScene/MovieRenderPipeline/Source/MovieRenderPipelineEditor/Private/MoviePipelinePIEExecutor.cpp
--- a/Engine/Plugins/MovieScene/MovieRenderPipeline/Source/MovieRenderPipelineEditor/Private/MoviePipelinePIEExecutor.cpp
+++ b/Engine/Plugins/MovieScene/MovieRenderPipeline/Source/MovieRenderPipelineEditor/Private/MoviePipelinePIEExecutor.cpp
@@ -191,7 +191,6 @@ void UMoviePipelinePIEExecutor::OnPIEStartupFinished(bool)
 
 	// This Pipeline belongs to the world being created so that they have context for things they execute.
 	ActiveMoviePipeline = NewObject<UMoviePipeline>(ExecutingWorld, PipelineClass);
-	ActiveMoviePipeline->DebugWidgetClass = DebugWidgetClass;
 	
 	// We allow users to set a multi-frame delay before we actually run the Initialization function and start thinking.
 	// This solves cases where there are engine systems that need to finish loading before we do anything.
@@ -203,6 +202,7 @@ void UMoviePipelinePIEExecutor::OnPIEStartupFinished(bool)
 	// Listen for when the pipeline thinks it has finished.
 	ActiveMoviePipeline->OnMoviePipelineWorkFinished().AddUObject(this, &UMoviePipelinePIEExecutor::OnPIEMoviePipelineFinished);
 	ActiveMoviePipeline->OnMoviePipelineShotWorkFinished().AddUObject(this, &UMoviePipelinePIEExecutor::OnJobShotFinished);
+	ActiveMoviePipeline->SetViewportInitArgs(ViewportInitArgs);
 	
 	if (ExecutorSettings->InitialDelayFrameCount == 0)
 	{
