MoviePipeline: PNG exports can now choose to not write the alpha channel. This has been added as alpha does not correctly propagate through post processing with some effects and png is the only easily accessible lossless compression output (exrs require special software to read most of the time).

diff --git a/Engine/Plugins/MovieScene/MovieRenderPipeline/Source/MovieRenderPipelineRenderPasses/Private/MoviePipelineImageSequenceOutput.cpp b/Engine/Plugins/MovieScene/MovieRenderPipeline/Source/MovieRenderPipelineRenderPasses/Private/MoviePipelineImageSequenceOutput.cpp
--- a/Engine/Plugins/MovieScene/MovieRenderPipeline/Source/MovieRenderPipelineRenderPasses/Private/MoviePipelineImageSequenceOutput.cpp
+++ b/Engine/Plugins/MovieScene/MovieRenderPipeline/Source/MovieRenderPipelineRenderPasses/Private/MoviePipelineImageSequenceOutput.cpp
@@ -243,6 +243,22 @@ void UMoviePipelineImageSequenceOutputBase::OnReceiveImageDataImpl(FMoviePipelin
 			}
 		}
 
+		if (!IsAlphaAllowed())
+		{
+			switch (QuantizedPixelType)
+			{
+			case EImagePixelType::Color:
+				TileImageTask->PixelPreProcessors.Add(TAsyncAlphaWrite<FColor>(255));
+				break;
+			case EImagePixelType::Float16:
+				TileImageTask->PixelPreProcessors.Add(TAsyncAlphaWrite<FFloat16Color>(1.0f));
+				break;
+			case EImagePixelType::Float32:
+				TileImageTask->PixelPreProcessors.Add(TAsyncAlphaWrite<FLinearColor>(1.0f));
+				break;
+			}
+		}
+
 
 		TileImageTask->PixelData = MoveTemp(QuantizedPixelData);
 		
diff --git a/Engine/Plugins/MovieScene/MovieRenderPipeline/Source/MovieRenderPipelineRenderPasses/Public/MoviePipelineImageSequenceOutput.h b/Engine/Plugins/MovieScene/MovieRenderPipeline/Source/MovieRenderPipelineRenderPasses/Public/MoviePipelineImageSequenceOutput.h
--- a/Engine/Plugins/MovieScene/MovieRenderPipeline/Source/MovieRenderPipelineRenderPasses/Public/MoviePipelineImageSequenceOutput.h
+++ b/Engine/Plugins/MovieScene/MovieRenderPipeline/Source/MovieRenderPipelineRenderPasses/Public/MoviePipelineImageSequenceOutput.h
@@ -25,6 +25,7 @@ protected:
 	// ~UMovieRenderPipelineOutputContainer interface
 
 	virtual void GetFormatArguments(FMoviePipelineFormatArgs& InOutFormatArgs) const override;
+	virtual bool IsAlphaAllowed() const { return false; }
 protected:
 	/** The format of the image to write out */
 	EImageFormat OutputFormat;
@@ -60,11 +61,17 @@ public:
 #if WITH_EDITOR
 	virtual FText GetDisplayText() const override { return NSLOCTEXT("MovieRenderPipeline", "ImgSequencePNGSettingDisplayName", ".png Sequence [8bit]"); }
 #endif
+	virtual bool IsAlphaAllowed() const override { return bWriteAlpha; }
+
 public:
 	UMoviePipelineImageSequenceOutput_PNG()
 	{
 		OutputFormat = EImageFormat::PNG;
+		bWriteAlpha = true;
 	}
+
+	UPROPERTY(EditAnywhere, BlueprintReadWrite, Category = "PNG")
+	bool bWriteAlpha;
 };
 
 UCLASS()
