MoviePipeline: Added support for rendering multiple cameras at the same time. If this feature is enabled, all camera components found adjacent to the camera cut for each shot will also be rendered as additional passes (simultaneously).

This can be useful for capturing multiple camera angles of a dynamic event as each camera will show the same contents due to them being rendered at the same time as the main camera.

diff --git a/Engine/Plugins/MovieScene/MoviePipelineMaskRenderPass/Source/MoviePipelineMaskRenderPass/Private/MoviePipelinePanoramicPass.cpp b/Engine/Plugins/MovieScene/MoviePipelineMaskRenderPass/Source/MoviePipelineMaskRenderPass/Private/MoviePipelinePanoramicPass.cpp
--- a/Engine/Plugins/MovieScene/MoviePipelineMaskRenderPass/Source/MoviePipelineMaskRenderPass/Private/MoviePipelinePanoramicPass.cpp
+++ b/Engine/Plugins/MovieScene/MoviePipelineMaskRenderPass/Source/MoviePipelineMaskRenderPass/Private/MoviePipelinePanoramicPass.cpp
@@ -392,7 +392,7 @@ FSceneView* UMoviePipelinePanoramicPass::GetSceneViewForSampleState(FSceneViewFa
 	View->PreviousViewTransform = FTransform(PanoPane->PrevCameraRotation, PanoPane->PrevCameraLocation);
 
 	View->StartFinalPostprocessSettings(View->ViewLocation);
-	BlendPostProcessSettings(View);
+	BlendPostProcessSettings(View, InOutSampleState, OptPayload);
 
 	// Scaling sensor size inversely with the the projection matrix [0][0] should physically
 	// cause the circle of confusion to be unchanged.
diff --git a/Engine/Plugins/MovieScene/MovieRenderPipeline/Source/MovieRenderPipelineCore/Private/MoviePipeline.cpp b/Engine/Plugins/MovieScene/MovieRenderPipeline/Source/MovieRenderPipelineCore/Private/MoviePipeline.cpp
--- a/Engine/Plugins/MovieScene/MovieRenderPipeline/Source/MovieRenderPipelineCore/Private/MoviePipeline.cpp
+++ b/Engine/Plugins/MovieScene/MovieRenderPipeline/Source/MovieRenderPipelineCore/Private/MoviePipeline.cpp
@@ -1497,6 +1497,7 @@ void UMoviePipeline::ResolveFilenameFormatArguments(const FString& InFormatStrin
 		Params.FileMetadata = InOutputState->FileMetadata;
 		Params.ShotOverride = ActiveShotList[InOutputState->ShotIndex];
 		Params.InitializationVersion = ActiveShotList[InOutputState->ShotIndex]->ShotInfo.VersionNumber;
+		Params.CameraIndex = InOutputState->CameraIndex;
 	}
 
 	UMoviePipelineOutputSetting* OutputSetting = GetPipelineMasterConfig()->FindSetting<UMoviePipelineOutputSetting>();
@@ -1682,4 +1683,65 @@ void UMoviePipeline::RestoreSkeletalMeshClothSubSteps()
 	}
 }
 
+void UMoviePipeline::GetSidecarCameraData(UMoviePipelineExecutorShot* InShot, int32 InCameraIndex, FMinimalViewInfo& OutViewInfo, UCameraComponent** OutCameraComponent) const
+{
+	UMovieScene* MovieScene = GetTargetSequence()->GetMovieScene();
+	if (!InShot || !MovieScene)
+	{
+		return;
+	}
+
+	const FGuid CameraBindingRef = InShot->SidecarCameras[InCameraIndex].BindingId;
+	const FMovieSceneSequenceID CameraSequenceRef = InShot->SidecarCameras[InCameraIndex].SequenceId;
+	{
+		TArrayView<TWeakObjectPtr<UObject>> BoundCameras = LevelSequenceActor->GetSequencePlayer()->FindBoundObjects(CameraBindingRef, CameraSequenceRef);
+		if (BoundCameras.Num() > 0)
+		{
+			UCameraComponent* BoundCamera = Cast<UCameraComponent>(BoundCameras[0].Get());
+			if (BoundCamera)
+			{
+				*OutCameraComponent = BoundCamera;
+				BoundCamera->GetCameraView(GetWorld()->GetDeltaSeconds(), OutViewInfo);
+				
+				// We override the current/previous transform based on cached data though to ensure we accurately handle cur/next frame positions
+				OutViewInfo.Location = FrameInfo.CurrSidecarViewLocations[InCameraIndex];
+				OutViewInfo.Rotation = FrameInfo.CurrSidecarViewRotations[InCameraIndex];
+				OutViewInfo.PreviousViewTransform = FTransform(FrameInfo.PrevSidecarViewRotations[InCameraIndex], FrameInfo.PrevSidecarViewLocations[InCameraIndex]);
+			}
+		}
+	}
+}
+
+void UMoviePipeline::GetSidecarCameraViewPoints(UMoviePipelineExecutorShot* InShot, TArray<FVector>& OutSidecarViewLocations, TArray<FRotator>& OutSidecarViewRotations) const
+{
+	OutSidecarViewLocations.Reset();
+	OutSidecarViewRotations.Reset();
+
+	if (!InShot)
+	{
+		return;
+	}
+
+	// We can't use GetSidecarCameraData because it relies on data generated by this function to ensure the previous/next locations are correct.
+	for (int32 Index = 0; Index < InShot->SidecarCameras.Num(); Index++)
+	{
+		const FGuid CameraBindingRef = InShot->SidecarCameras[Index].BindingId;
+		const FMovieSceneSequenceID CameraSequenceRef = InShot->SidecarCameras[Index].SequenceId;
+
+		TArrayView<TWeakObjectPtr<UObject>> BoundCameras = LevelSequenceActor->GetSequencePlayer()->FindBoundObjects(CameraBindingRef, CameraSequenceRef);
+		if (BoundCameras.Num() > 0)
+		{
+			UCameraComponent* BoundCamera = Cast<UCameraComponent>(BoundCameras[0].Get());
+			if (BoundCamera)
+			{
+				FMinimalViewInfo ViewInfo;
+				BoundCamera->GetCameraView(GetWorld()->GetDeltaSeconds(), ViewInfo);
+
+				OutSidecarViewLocations.Add(ViewInfo.Location);
+				OutSidecarViewRotations.Add(ViewInfo.Rotation);
+			}
+		}
+	}
+}
+
 #undef LOCTEXT_NAMESPACE // "MoviePipeline"
diff --git a/Engine/Plugins/MovieScene/MovieRenderPipeline/Source/MovieRenderPipelineCore/Private/MoviePipelineBlueprintLibrary.cpp b/Engine/Plugins/MovieScene/MovieRenderPipeline/Source/MovieRenderPipelineCore/Private/MoviePipelineBlueprintLibrary.cpp
--- a/Engine/Plugins/MovieScene/MovieRenderPipeline/Source/MovieRenderPipelineCore/Private/MoviePipelineBlueprintLibrary.cpp
+++ b/Engine/Plugins/MovieScene/MovieRenderPipeline/Source/MovieRenderPipelineCore/Private/MoviePipelineBlueprintLibrary.cpp
@@ -40,6 +40,8 @@
 #include "CineCameraActor.h"
 #include "CineCameraComponent.h"
 
+PRAGMA_DISABLE_OPTIMIZATION
+
 EMovieRenderPipelineState UMoviePipelineBlueprintLibrary::GetPipelineState(const UMoviePipeline* InPipeline)
 {
 	if (InPipeline)
@@ -457,6 +459,7 @@ void UMoviePipelineBlueprintLibrary::UpdateJobShotListFromSequence(ULevelSequenc
 			TRange<FFrameNumber> Range;
 			UMovieSceneCameraCutSection* CameraCut;
 			FMovieSceneSequenceID SequenceID;
+			TArray<FMoviePipelineSidecarCamera> SidecarCameras;
 
 			// These are only updated after the final linearized entities are chosen
 			TTuple<FString, FString> Name;
@@ -538,6 +541,72 @@ void UMoviePipelineBlueprintLibrary::UpdateJobShotListFromSequence(ULevelSequenc
 			LeafNode->CameraCutSection = Entity.CameraCut;
 			MoviePipeline::BuildSectionHierarchyRecursive(SequenceHierarchyCache, InSequence, Entity.SequenceID, MovieSceneSequenceID::Invalid, LeafNode);
 
+			// If we did find a camera cut section, we're going to look for any object bindings in the same sequence that contain
+			// camera components and track them as sidecar cameras. We do this even if sidecar data isn't being used.
+			if (LeafNode->CameraCutSection.IsValid() && LeafNode->MovieScene.IsValid())
+			{
+				int32 MainCameraIndex = INDEX_NONE;
+				FMovieSceneObjectBindingID MainBinding = LeafNode->CameraCutSection->GetCameraBindingID();
+
+				for (int32 Index = 0; Index < LeafNode->MovieScene->GetPossessableCount(); Index++)
+				{
+					FMovieScenePossessable& Possessable = LeafNode->MovieScene->GetPossessable(Index);
+					FGuid ValidBinding = FGuid();
+
+#if WITH_EDITOR
+					// This is only available in the editor, so in the editor we'll auto-detect the class
+					// but in runtime builds they'll have to manually tag it with "Camera" :(
+					const UClass* PossessedClass = Possessable.GetPossessedObjectClass();
+					if (PossessedClass->IsChildOf<UCameraComponent>())
+					{
+						ValidBinding = Possessable.GetGuid();
+					}
+#endif
+					if (Possessable.Tags.Contains(TEXT("Camera")))
+					{
+						// Don't duplicate them up, but also run this name check in the editor
+						ValidBinding = Possessable.GetGuid();
+					}
+
+					if (ValidBinding.IsValid())
+					{
+						FMoviePipelineSidecarCamera& SidecarCamera = Entity.SidecarCameras.AddDefaulted_GetRef();
+						SidecarCamera.BindingId = ValidBinding;
+						SidecarCamera.SequenceId = Entity.SequenceID;
+						
+						FGuid ParentGuid = Possessable.GetParent();
+						while (ParentGuid.IsValid())
+						{
+							// This comparison doesn't support cross-sequence references, technically we should create a 
+							// FMovieSceneObjectBindingID etc, but because we only support finding sidecar cameras in the
+							// same sequence as the main camera cut track, this is fine for now.
+							if (ParentGuid == MainBinding.GetGuid())
+							{
+								// We want to know which sidecar camera this is, not the possessable index.
+								MainCameraIndex = Entity.SidecarCameras.Num() - 1;
+							}
+
+							if (FMovieScenePossessable* ParentAsPossessable = LeafNode->MovieScene->FindPossessable(ParentGuid))
+							{
+								ParentGuid = ParentAsPossessable->GetParent();
+								SidecarCamera.Name = ParentAsPossessable->GetName();
+							}
+							else if (FMovieSceneSpawnable* ParentAsSpawnable = LeafNode->MovieScene->FindSpawnable(ParentGuid))
+							{
+								// Spawnables will never have a parent
+								ParentGuid.Invalidate(); 
+								SidecarCamera.Name = ParentAsSpawnable->GetName();
+							}
+						}
+					}
+				}
+
+				// Ensure the "main" camera (specified by the Camera Cut Track) is always index 0, so that it's rendered
+				// first, and the stable sort (on output files) still tries to keep it as the primary layer.
+				ensureAlwaysMsgf(MainCameraIndex != INDEX_NONE, TEXT("Main Camera Binding not found in camera array!"));
+				Entity.SidecarCameras.Swap(MainCameraIndex, 0);
+			}
+
 			FMovieSceneTimeTransform InnerToOuterTransform = FMovieSceneTimeTransform();
 			FMovieSceneSubSequenceData* SubSequenceData = SequenceHierarchyCache.FindSubData(Entity.SequenceID);
 			if (SubSequenceData)
@@ -645,6 +714,7 @@ void UMoviePipelineBlueprintLibrary::UpdateJobShotListFromSequence(ULevelSequenc
 			NewShot->ShotInfo.TotalOutputRangeMaster = Entity.Range;
 			NewShot->ShotInfo.WarmupRangeMaster = Entity.CameraCutWarmUpRange;
 			NewShot->ShotInfo.OuterToInnerTransform = Entity.InnerToOuterTransform.Inverse();
+			NewShot->SidecarCameras = Entity.SidecarCameras;
 			UE_LOG(LogMovieRenderPipeline, Log, TEXT("Registering range: %s (InnerName: %s OuterName: %s)"), *LexToString(NewShot->ShotInfo.TotalOutputRangeMaster), *NewShot->InnerName, *NewShot->OuterName);
 		}
 
@@ -849,7 +919,7 @@ void UMoviePipelineBlueprintLibrary::ResolveFilenameFormatArguments(const FStrin
 			FrameNumber = FrameNumberRel;
 			FrameNumberShot = FrameNumberShotRel;
 		}
-		FString CameraName = InParams.ShotOverride ? InParams.ShotOverride->InnerName : FString();
+		FString CameraName = InParams.ShotOverride ? InParams.ShotOverride->GetCameraName(InParams.CameraIndex) : FString();
 		CameraName = CameraName.Len() > 0 ? CameraName : TEXT("NoCamera");
 		if(InParams.CameraNameOverride.Len() > 0)
 		{
@@ -866,6 +936,7 @@ void UMoviePipelineBlueprintLibrary::ResolveFilenameFormatArguments(const FStrin
 		MoviePipeline::GetOutputStateFormatArgs(OutMergedFormatArgs, FrameNumber, FrameNumberShot, FrameNumberRel, FrameNumberShotRel, CameraName, ShotName);
 	}
 
+
 	// And from ourself
 	{
 		// Use the shared function as UMoviePipelineMasterConfig::GetFormatArguments to ensure all the time variables get overwritten.
@@ -998,3 +1069,4 @@ FText UMoviePipelineBlueprintLibrary::GetMoviePipelineEngineChangelistLabel(cons
 
 	return FText();
 }
+PRAGMA_ENABLE_OPTIMIZATION
\ No newline at end of file
diff --git a/Engine/Plugins/MovieScene/MovieRenderPipeline/Source/MovieRenderPipelineCore/Private/MoviePipelineRendering.cpp b/Engine/Plugins/MovieScene/MovieRenderPipeline/Source/MovieRenderPipelineCore/Private/MoviePipelineRendering.cpp
--- a/Engine/Plugins/MovieScene/MovieRenderPipeline/Source/MovieRenderPipelineCore/Private/MoviePipelineRendering.cpp
+++ b/Engine/Plugins/MovieScene/MovieRenderPipeline/Source/MovieRenderPipelineCore/Private/MoviePipelineRendering.cpp
@@ -182,14 +182,22 @@ void UMoviePipeline::RenderFrame()
 	{
 		return;
 	}
+	
+	{
+		// Sidecar Cameras get updated below after rendering, they're still separate for backwards compat reasons
+		FrameInfo.PrevViewLocation = FrameInfo.CurrViewLocation;
+		FrameInfo.PrevViewRotation = FrameInfo.CurrViewRotation;
 
-	FrameInfo.PrevViewLocation = FrameInfo.CurrViewLocation;
-	FrameInfo.PrevViewRotation = FrameInfo.CurrViewRotation;
+		// Update the Sidecar Cameras
+		FrameInfo.PrevSidecarViewLocations = FrameInfo.CurrSidecarViewLocations;
+		FrameInfo.PrevSidecarViewRotations = FrameInfo.CurrSidecarViewRotations;
 
-	bool bWriteAllSamples = DebugSettings ? DebugSettings->bWriteAllSamples : false;
+		// Update our current view location
+		LocalPlayerController->GetPlayerViewPoint(FrameInfo.CurrViewLocation, FrameInfo.CurrViewRotation);
+		GetSidecarCameraViewPoints(ActiveShotList[CurrentShotIndex], FrameInfo.CurrSidecarViewLocations, FrameInfo.CurrSidecarViewRotations);
+	}
 
-	// Update our current view location
-	LocalPlayerController->GetPlayerViewPoint(FrameInfo.CurrViewLocation, FrameInfo.CurrViewRotation);
+	bool bWriteAllSamples = DebugSettings ? DebugSettings->bWriteAllSamples : false;
 
 	// Add appropriate metadata here that is shared by all passes.
 	{
@@ -205,20 +213,12 @@ void UMoviePipeline::RenderFrame()
 
 		UE::MoviePipeline::GetHardwareUsageMetadata(CachedOutputState.FileMetadata, ResolvedOutputDirectory);
 
-		CachedOutputState.FileMetadata.Add(TEXT("unreal/camera/curPos/x"), FString::SanitizeFloat(FrameInfo.CurrViewLocation.X));
-		CachedOutputState.FileMetadata.Add(TEXT("unreal/camera/curPos/y"), FString::SanitizeFloat(FrameInfo.CurrViewLocation.Y));
-		CachedOutputState.FileMetadata.Add(TEXT("unreal/camera/curPos/z"), FString::SanitizeFloat(FrameInfo.CurrViewLocation.Z));
-		CachedOutputState.FileMetadata.Add(TEXT("unreal/camera/curRot/pitch"), FString::SanitizeFloat(FrameInfo.CurrViewRotation.Pitch));
-		CachedOutputState.FileMetadata.Add(TEXT("unreal/camera/curRot/yaw"), FString::SanitizeFloat(FrameInfo.CurrViewRotation.Yaw));
-		CachedOutputState.FileMetadata.Add(TEXT("unreal/camera/curRot/roll"), FString::SanitizeFloat(FrameInfo.CurrViewRotation.Roll));
 
-		CachedOutputState.FileMetadata.Add(TEXT("unreal/camera/prevPos/x"), FString::SanitizeFloat(FrameInfo.PrevViewLocation.X));
-		CachedOutputState.FileMetadata.Add(TEXT("unreal/camera/prevPos/y"), FString::SanitizeFloat(FrameInfo.PrevViewLocation.Y));
-		CachedOutputState.FileMetadata.Add(TEXT("unreal/camera/prevPos/z"), FString::SanitizeFloat(FrameInfo.PrevViewLocation.Z));
-		CachedOutputState.FileMetadata.Add(TEXT("unreal/camera/prevRot/pitch"), FString::SanitizeFloat(FrameInfo.PrevViewRotation.Pitch));
-		CachedOutputState.FileMetadata.Add(TEXT("unreal/camera/prevRot/yaw"), FString::SanitizeFloat(FrameInfo.PrevViewRotation.Yaw));
-		CachedOutputState.FileMetadata.Add(TEXT("unreal/camera/prevRot/roll"), FString::SanitizeFloat(FrameInfo.PrevViewRotation.Roll));
+		// We'll leave these in for legacy, when this tracks the 'Main' camera (of the player), render passes that support
+		// multiple cameras will have to write each camera name into their metadata.
+		UE::MoviePipeline::GetMetadataFromCameraLocRot(TEXT("camera"), TEXT(""), FrameInfo.CurrViewLocation, FrameInfo.CurrViewRotation, FrameInfo.PrevViewLocation, FrameInfo.PrevViewRotation, CachedOutputState.FileMetadata);
 
+		// This is still global regardless, individual cameras don't get their own motion blur amount because the engine tick is tied to it.
 		CachedOutputState.FileMetadata.Add(TEXT("unreal/camera/shutterAngle"), FString::SanitizeFloat(CachedOutputState.TimeData.MotionBlurFraction * 360.0f));
 	}
 
@@ -227,10 +227,6 @@ void UMoviePipeline::RenderFrame()
 		// We can optimize some of the settings for 'special' frames we may be rendering, ie: we render once for motion vectors, but
 		// we don't need that per-tile so we can set the tile count to 1, and spatial sample count to 1 for that particular frame.
 		{
-			// Tiling is only needed when actually producing frames.
-			TileCount.X = 1;
-			TileCount.Y = 1;
-
 			// Spatial Samples aren't needed when not producing frames (caveat: Render Warmup Frame, handled below)
 			NumSpatialSamples = 1;
 		}
@@ -269,6 +265,7 @@ void UMoviePipeline::RenderFrame()
 		TimeStats.StartTime = FDateTime::UtcNow();
 	}
 
+	// Support for RenderDoc captures of just the MRQ work
 #if WITH_EDITOR && !UE_BUILD_SHIPPING
 	TUniquePtr<RenderCaptureInterface::FScopedCapture> ScopedGPUCapture;
 	if (CachedOutputState.bCaptureRendering)
@@ -378,6 +375,7 @@ void UMoviePipeline::RenderFrame()
 				SampleState.AntiAliasingMethod = AntiAliasingMethod;
 				SampleState.SceneCaptureSource = (ColorSettings && ColorSettings->bDisableToneCurve) ? ESceneCaptureSource::SCS_FinalColorHDR : ESceneCaptureSource::SCS_FinalToneCurveHDR;
 				SampleState.OutputState = CachedOutputState;
+				SampleState.OutputState.CameraIndex = 0; // Initialize to a sane default for non multi-cam passes.
 				SampleState.ProjectionMatrixJitterAmount = FVector2D((float)(SpatialShiftX) * 2.0f / BackbufferResolution.X, (float)SpatialShiftY * -2.0f / BackbufferResolution.Y);
 				SampleState.TileIndexes = FIntPoint(TileX, TileY);
 				SampleState.TileCounts = TileCount;
@@ -472,8 +470,11 @@ void UMoviePipeline::ProcessOutstandingFinishedFrames()
 	}
 }
 
+
 void UMoviePipeline::OnSampleRendered(TUniquePtr<FImagePixelData>&& OutputSample)
 {
+	// This function handles the "Write all Samples" feature which lets you inspect data
+	// pre-accumulation.
 	UMoviePipelineOutputSetting* OutputSettings = GetPipelineMasterConfig()->FindSetting<UMoviePipelineOutputSetting>();
 	check(OutputSettings);
 
@@ -504,7 +505,7 @@ void UMoviePipeline::OnSampleRendered(TUniquePtr<FImagePixelData>&& OutputSample
 	if (Shot)
 	{
 		FormatOverrides.Add(TEXT("shot_name"), Shot->OuterName);
-		FormatOverrides.Add(TEXT("camera_name"), Shot->InnerName);
+		FormatOverrides.Add(TEXT("camera_name"), Shot->GetCameraName(InFrameData->SampleState.OutputState.CameraIndex));
 	}
 	FMoviePipelineFormatArgs FinalFormatArgs;
 
diff --git a/Engine/Plugins/MovieScene/MovieRenderPipeline/Source/MovieRenderPipelineCore/Private/MoviePipelineTiming.cpp b/Engine/Plugins/MovieScene/MovieRenderPipeline/Source/MovieRenderPipelineCore/Private/MoviePipelineTiming.cpp
--- a/Engine/Plugins/MovieScene/MovieRenderPipeline/Source/MovieRenderPipelineCore/Private/MoviePipelineTiming.cpp
+++ b/Engine/Plugins/MovieScene/MovieRenderPipeline/Source/MovieRenderPipelineCore/Private/MoviePipelineTiming.cpp
@@ -205,6 +205,11 @@ void UMoviePipeline::TickProducingFrames()
 			
 			FrameInfo.PrevViewLocation = FrameInfo.CurrViewLocation;
 			FrameInfo.PrevViewRotation = FrameInfo.CurrViewRotation;
+
+			// Cache data for all the sidecar cameras too.
+			GetSidecarCameraViewPoints(CurrentCameraCut, FrameInfo.CurrSidecarViewLocations, FrameInfo.CurrSidecarViewRotations);
+			FrameInfo.PrevSidecarViewLocations = FrameInfo.CurrSidecarViewLocations;
+			FrameInfo.PrevSidecarViewRotations = FrameInfo.CurrSidecarViewRotations;
 		}
 
 		// We can safely fall through to the below states as they're OK to process the same frame we set up.
@@ -753,5 +758,7 @@ void UMoviePipeline::CalculateFrameNumbersForOutputState(const MoviePipeline::FF
 	}
 
 	InOutOutputState.ShotName = InCameraCut->OuterName;
+
+	// Continue to use the InnerName (which is the master camera for the shot), individual render passes can override it later if rendering more than one camera.
 	InOutOutputState.CameraName = InCameraCut->InnerName;
-}
+}
\ No newline at end of file
diff --git a/Engine/Plugins/MovieScene/MovieRenderPipeline/Source/MovieRenderPipelineCore/Private/MoviePipelineUtils.cpp b/Engine/Plugins/MovieScene/MovieRenderPipeline/Source/MovieRenderPipelineCore/Private/MoviePipelineUtils.cpp
--- a/Engine/Plugins/MovieScene/MovieRenderPipeline/Source/MovieRenderPipelineCore/Private/MoviePipelineUtils.cpp
+++ b/Engine/Plugins/MovieScene/MovieRenderPipeline/Source/MovieRenderPipelineCore/Private/MoviePipelineUtils.cpp
@@ -25,6 +25,7 @@
 #include "GenericPlatform/GenericPlatformDriver.h"
 #include "MovieRenderPipelineCoreModule.h"
 #include "RHI.h"
+#include "CineCameraComponent.h"
 
 namespace UE
 {
@@ -657,7 +658,7 @@ namespace UE
 		DECLARE_CYCLE_STAT(TEXT("STAT_MoviePipeline_HardwareMetadata"), STAT_HardwareMetadata, STATGROUP_MoviePipeline);
 
 
-		void ValidateOutputFormatString(FString& InOutFilenameFormatString, const bool bTestRenderPass, const bool bTestFrameNumber)
+		void ValidateOutputFormatString(FString& InOutFilenameFormatString, const bool bTestRenderPass, const bool bTestFrameNumber, const bool bTestCameraName)
 		{
 			const FString FrameNumberIdentifiers[] = { TEXT("{frame_number}"), TEXT("{frame_number_shot}"), TEXT("{frame_number_rel}"), TEXT("{frame_number_shot_rel}") };
 
@@ -693,6 +694,38 @@ namespace UE
 				}
 			}
 
+			// Could be more than one camera being written to, make sure there's {camera_name} so it doesn't stomp over each other.
+			if (bTestCameraName)
+			{
+				if (!InOutFilenameFormatString.Contains(TEXT("{camera_name}"), ESearchCase::IgnoreCase))
+				{
+					UE_LOG(LogMovieRenderPipeline, Warning, TEXT("Multiple cameras exported but no {camera_name} format found. Automatically adding!"));
+
+					// Search for a frame number in the output string
+					int32 FrameNumberIndex = INDEX_NONE;
+					for (const FString& Identifier : FrameNumberIdentifiers)
+					{
+						FrameNumberIndex = InOutFilenameFormatString.Find(Identifier, ESearchCase::IgnoreCase, ESearchDir::FromEnd);
+						if (FrameNumberIndex != INDEX_NONE)
+						{
+							break;
+						}
+					}
+
+					if (FrameNumberIndex == INDEX_NONE)
+					{
+						// No frame number found, so just append render_pass
+						InOutFilenameFormatString += TEXT("{camera_name}");
+					}
+					else
+					{
+						// If a frame number is found, we need to insert render_pass first before it, so various editing
+						// software will still be able to identify if this is an image sequence
+						InOutFilenameFormatString.InsertAt(FrameNumberIndex, TEXT("{camera_name}."));
+					}
+				}
+			}
+
 			if (bTestFrameNumber)
 			{
 				// Ensure there is a frame number in the output string somewhere to uniquely identify individual files in an image sequence.
@@ -843,5 +876,46 @@ namespace UE
 			InFileMetadata.Add(TEXT("unreal/stats/memory/peakUsedPhysicalMB"), LexToString(MemoryStats.PeakUsedPhysical / MBDivider));
 			InFileMetadata.Add(TEXT("unreal/stats/memory/peakUsedVirtualMB"), LexToString(MemoryStats.PeakUsedVirtual / MBDivider));
 		}
+		
+		void GetMetadataFromCineCamera(UCineCameraComponent* InComponent, const FString& InCameraName, const FString& InRenderPassName, TMap<FString, FString>& InOutMetadata)
+		{
+			if (InComponent)
+			{
+				InOutMetadata.Add(FString::Printf(TEXT("unreal/%s/%s/sensorWidth"), *InCameraName, *InRenderPassName), FString::SanitizeFloat(InComponent->Filmback.SensorWidth));
+				InOutMetadata.Add(FString::Printf(TEXT("unreal/%s/%s/sensorHeight"), *InCameraName, *InRenderPassName), FString::SanitizeFloat(InComponent->Filmback.SensorHeight));
+				InOutMetadata.Add(FString::Printf(TEXT("unreal/%s/%s/sensorAspectRatio"), *InCameraName, *InRenderPassName), FString::SanitizeFloat(InComponent->Filmback.SensorAspectRatio));
+				InOutMetadata.Add(FString::Printf(TEXT("unreal/%s/%s/minFocalLength"), *InCameraName, *InRenderPassName), FString::SanitizeFloat(InComponent->LensSettings.MinFocalLength));
+				InOutMetadata.Add(FString::Printf(TEXT("unreal/%s/%s/maxFocalLength"), *InCameraName, *InRenderPassName), FString::SanitizeFloat(InComponent->LensSettings.MaxFocalLength));
+				InOutMetadata.Add(FString::Printf(TEXT("unreal/%s/%s/minFStop"), *InCameraName, *InRenderPassName), FString::SanitizeFloat(InComponent->LensSettings.MinFStop));
+				InOutMetadata.Add(FString::Printf(TEXT("unreal/%s/%s/maxFStop"), *InCameraName, *InRenderPassName), FString::SanitizeFloat(InComponent->LensSettings.MaxFStop));
+				InOutMetadata.Add(FString::Printf(TEXT("unreal/%s/%s/dofDiaphragmBladeCount"), *InCameraName, *InRenderPassName), FString::FromInt(InComponent->LensSettings.DiaphragmBladeCount));
+				InOutMetadata.Add(FString::Printf(TEXT("unreal/%s/%s/focalLength"), *InCameraName, *InRenderPassName), FString::SanitizeFloat(InComponent->CurrentFocalLength));
+			}
+		}
+		
+		void GetMetadataFromCameraLocRot(const FString& InCameraName, const FString& InRenderPassName, const FVector& InCurLoc, const FRotator& InCurRot, const FVector& InPrevLoc, const FRotator& InPrevRot, TMap<FString, FString>& InOutMetadata)
+		{
+			// InRenderPassName could be empty (for global camera stuff), and to support proper backwards compat metadata
+			// we need to not end up with an extra "/" in it, so we pre-format it.
+			FString CamName = FString::Printf(TEXT("%s/%s"), *InCameraName, *InRenderPassName);
+			if (InRenderPassName.IsEmpty())
+			{
+				CamName = InCameraName;
+			}
+			InOutMetadata.Add(FString::Printf(TEXT("unreal/%s/curPos/x"), *CamName), FString::SanitizeFloat(InCurLoc.X));
+			InOutMetadata.Add(FString::Printf(TEXT("unreal/%s/curPos/y"), *CamName), FString::SanitizeFloat(InCurLoc.Y));
+			InOutMetadata.Add(FString::Printf(TEXT("unreal/%s/curPos/z"), *CamName), FString::SanitizeFloat(InCurLoc.Z));
+			InOutMetadata.Add(FString::Printf(TEXT("unreal/%s/curRot/pitch"), *CamName),  FString::SanitizeFloat(InCurRot.Pitch));
+			InOutMetadata.Add(FString::Printf(TEXT("unreal/%s/curRot/yaw"), *CamName), FString::SanitizeFloat(InCurRot.Yaw));
+			InOutMetadata.Add(FString::Printf(TEXT("unreal/%s/curRot/roll"), *CamName), FString::SanitizeFloat(InCurRot.Roll));
+
+			InOutMetadata.Add(FString::Printf(TEXT("unreal/%s/prevPos/x"), *CamName), FString::SanitizeFloat(InPrevLoc.X));
+			InOutMetadata.Add(FString::Printf(TEXT("unreal/%s/prevPos/y"), *CamName), FString::SanitizeFloat(InPrevLoc.Y));
+			InOutMetadata.Add(FString::Printf(TEXT("unreal/%s/prevPos/z"), *CamName), FString::SanitizeFloat(InPrevLoc.Z));
+			InOutMetadata.Add(FString::Printf(TEXT("unreal/%s/prevRot/pitch"), *CamName), FString::SanitizeFloat(InPrevRot.Pitch));
+			InOutMetadata.Add(FString::Printf(TEXT("unreal/%s/prevRot/yaw"), *CamName), FString::SanitizeFloat(InPrevRot.Yaw));
+			InOutMetadata.Add(FString::Printf(TEXT("unreal/%s/prevRot/roll"), *CamName), FString::SanitizeFloat(InPrevRot.Roll));
+		}
+
 	}
 }
\ No newline at end of file
diff --git a/Engine/Plugins/MovieScene/MovieRenderPipeline/Source/MovieRenderPipelineCore/Public/MoviePipeline.h b/Engine/Plugins/MovieScene/MovieRenderPipeline/Source/MovieRenderPipelineCore/Public/MoviePipeline.h
--- a/Engine/Plugins/MovieScene/MovieRenderPipeline/Source/MovieRenderPipelineCore/Public/MoviePipeline.h
+++ b/Engine/Plugins/MovieScene/MovieRenderPipeline/Source/MovieRenderPipelineCore/Public/MoviePipeline.h
@@ -185,6 +185,9 @@ public:
 	EMovieRenderPipelineState GetPipelineState() const { return PipelineState; }
 	FMoviePipelineOutputData GetOutputDataParams();
 
+	void GetSidecarCameraData(UMoviePipelineExecutorShot* InShot, int32 InCameraIndex, FMinimalViewInfo& OutViewInfo, class UCameraComponent** OutCameraComponent) const;
+	void GetSidecarCameraViewPoints(UMoviePipelineExecutorShot* InShot, TArray<FVector>& OutSidecarViewLocations, TArray<FRotator>& OutSidecarViewRotations) const;
+
 #if WITH_EDITOR
 	const FMovieSceneExportMetadata& GetOutputMetadata() const { return OutputMetadata; }
 #endif
diff --git a/Engine/Plugins/MovieScene/MovieRenderPipeline/Source/MovieRenderPipelineCore/Public/MoviePipelineCameraSetting.h b/Engine/Plugins/MovieScene/MovieRenderPipeline/Source/MovieRenderPipelineCore/Public/MoviePipelineCameraSetting.h
--- a/Engine/Plugins/MovieScene/MovieRenderPipeline/Source/MovieRenderPipelineCore/Public/MoviePipelineCameraSetting.h
+++ b/Engine/Plugins/MovieScene/MovieRenderPipeline/Source/MovieRenderPipelineCore/Public/MoviePipelineCameraSetting.h
@@ -54,4 +54,11 @@ public:
 	UPROPERTY(EditAnywhere, BlueprintReadWrite, meta = (UIMin = "0", UIMax = "1", ClampMin = "0", ClampMax = "1"), Category = "Camera Settings")
 	float OverscanPercentage;
 	
+	/**
+	* If true, when a Camera Cut section is found we will also render any other cameras within the same sequence (not parent, nor child sequences though).
+	* These cameras are rendered at the same time as the primary camera meaning all cameras capture the same world state. Do note that this multiplies
+	* render times and memory requirements!
+	*/
+	UPROPERTY(EditAnywhere, BlueprintReadWrite, Category = "Camera Settings")
+	bool bRenderAllCameras;
 };
\ No newline at end of file
diff --git a/Engine/Plugins/MovieScene/MovieRenderPipeline/Source/MovieRenderPipelineCore/Public/MoviePipelineHighResSetting.h b/Engine/Plugins/MovieScene/MovieRenderPipeline/Source/MovieRenderPipelineCore/Public/MoviePipelineHighResSetting.h
--- a/Engine/Plugins/MovieScene/MovieRenderPipeline/Source/MovieRenderPipelineCore/Public/MoviePipelineHighResSetting.h
+++ b/Engine/Plugins/MovieScene/MovieRenderPipeline/Source/MovieRenderPipelineCore/Public/MoviePipelineHighResSetting.h
@@ -20,6 +20,7 @@ public:
 		, OverlapRatio(0.f)
 		, bOverrideSubSurfaceScattering(false)
 		, BurleySampleCount(64)
+		, bAllocateHistoryPerTile(true)
 	{
 	}
 	
@@ -192,7 +193,12 @@ public:
 	UPROPERTY(EditAnywhere, BlueprintReadWrite, meta = (UIMin = "64", ClampMin = "0", UIMax = "1024", EditCondition="bOverrideSubSurfaceScattering"), Category = "Render Settings")
 	int32 BurleySampleCount;
 	
-
+	/*
+	* If true, allocate a unique history for each tile. This is needed to make some render features work, but should be disabled
+	* when dealing with extremely large resolutions as you will spend all of your GPU memory on history buffers.
+	*/
+	UPROPERTY(EditAnywhere, BlueprintReadWrite, Category = "Render Settings")
+	bool bAllocateHistoryPerTile;
 
 
 private:
diff --git a/Engine/Plugins/MovieScene/MovieRenderPipeline/Source/MovieRenderPipelineCore/Public/MoviePipelineQueue.h b/Engine/Plugins/MovieScene/MovieRenderPipeline/Source/MovieRenderPipelineCore/Public/MoviePipelineQueue.h
--- a/Engine/Plugins/MovieScene/MovieRenderPipeline/Source/MovieRenderPipelineCore/Public/MoviePipelineQueue.h
+++ b/Engine/Plugins/MovieScene/MovieRenderPipeline/Source/MovieRenderPipelineCore/Public/MoviePipelineQueue.h
@@ -7,6 +7,7 @@
 #include "MoviePipelineMasterConfig.h"
 #include "MoviePipelineShotConfig.h"
 #include "MoviePipelineConfigBase.h"
+#include "MovieSceneSequenceID.h"
 
 #include "MoviePipelineQueue.generated.h"
 
@@ -14,6 +15,27 @@ class UMoviePipelineMasterConfig;
 class ULevel;
 class ULevelSequence;
 
+USTRUCT(BlueprintType)
+struct MOVIERENDERPIPELINECORE_API FMoviePipelineSidecarCamera
+{
+	GENERATED_BODY()
+public:
+	UPROPERTY(EditAnywhere, BlueprintReadWrite, Category = "Movie Render Pipeline")
+	FGuid BindingId;
+
+	// FMovieSceneSequenceID isn't exposed to Blueprints and we need the full support
+	// of a FMovieSceneSequenceID so we can't just store a USequence* like the scripting
+	// layer does - we need to be able to handle the same sequence being in a sequence
+	// multiple times (which scripting does not). This data structure gets regenerated
+	// each time a render starts, so the property should be valid when we want to use
+	// it, even if it's not exposed to scripting.
+	// UPROPERTY(EditAnywhere, BlueprintReadWrite, Category = "Movie Render Pipeline")
+	FMovieSceneSequenceID SequenceId;
+
+	UPROPERTY(EditAnywhere, BlueprintReadWrite, Category = "Movie Render Pipeline")
+	FString Name;
+};
+
 /**
 * This class represents a segment of work within the Executor Job. This should be owned
 * by the UMoviePipelineExecutorJob and can be created before the movie pipeline starts to
@@ -116,6 +138,15 @@ public:
 		return bEnabled;
 	}
 
+	UFUNCTION(BlueprintPure, Category = "Movie Render Pipeline")
+	FString GetCameraName(int32 InCameraIndex) const
+	{
+		if (InCameraIndex >= 0 && InCameraIndex < SidecarCameras.Num())
+		{
+			return SidecarCameras[InCameraIndex].Name;
+		}
+		return InnerName;
+	}
 protected:
 	// UMoviePipipelineExecutorShot Interface
 	virtual void SetStatusMessage_Implementation(const FString& InMessage) { StatusMessage = InMessage; }
@@ -137,11 +168,14 @@ public:
 	/** The name of the camera cut section that this shot represents. Can be empty. */
 	UPROPERTY(BlueprintReadWrite, EditAnywhere, Category = "Movie Render Pipeline")
 	FString InnerName;
+
+	/** List of cameras to render for this shot. Only used if the setting flag is set in the Camera setting. */
+	UPROPERTY(BlueprintReadWrite, EditAnywhere, Category = "Movie Render Pipeline")
+	TArray<FMoviePipelineSidecarCamera> SidecarCameras;
 public:
 	/** Transient information used by the active Movie Pipeline working on this shot. */
 	FMoviePipelineCameraCutInfo ShotInfo;
 
-
 protected:
 	UPROPERTY(Transient)
 	float Progress;
diff --git a/Engine/Plugins/MovieScene/MovieRenderPipeline/Source/MovieRenderPipelineCore/Public/MoviePipelineUtils.h b/Engine/Plugins/MovieScene/MovieRenderPipeline/Source/MovieRenderPipelineCore/Public/MoviePipelineUtils.h
--- a/Engine/Plugins/MovieScene/MovieRenderPipeline/Source/MovieRenderPipelineCore/Public/MoviePipelineUtils.h
+++ b/Engine/Plugins/MovieScene/MovieRenderPipeline/Source/MovieRenderPipelineCore/Public/MoviePipelineUtils.h
@@ -11,6 +11,7 @@
 #include "Misc/FrameRate.h"
 #include "Evaluation/MovieSceneSequenceHierarchy.h"
 #include "MovieSceneSequenceID.h"
+#include "CineCameraComponent.h"
 #include "MoviePipelineQueue.h"
 
 // Forward Declare
@@ -127,12 +128,14 @@ namespace UE
 
 	namespace MoviePipeline
 	{
-		MOVIERENDERPIPELINECORE_API void ValidateOutputFormatString(FString& InOutFilenameFormatString, const bool bTestRenderPass, const bool bTestFrameNumber);
+		MOVIERENDERPIPELINECORE_API void ValidateOutputFormatString(FString& InOutFilenameFormatString, const bool bTestRenderPass, const bool bTestFrameNumber, const bool bIncludeCameraName = false);
 		MOVIERENDERPIPELINECORE_API void RemoveFrameNumberFormatStrings(FString& InOutFilenameFormatString, const bool bIncludeShots);
 
 		MOVIERENDERPIPELINECORE_API FString GetJobAuthor(const UMoviePipelineExecutorJob* InJob);
 		MOVIERENDERPIPELINECORE_API void GetSharedFormatArguments(TMap<FString, FString>& InFilenameArguments, TMap<FString, FString>& InFileMetadata, const FDateTime& InDateTime, const int32 InVersionNumber, const UMoviePipelineExecutorJob* InJob);
 		MOVIERENDERPIPELINECORE_API void GetHardwareUsageMetadata(TMap<FString, FString>& InFileMetadata, const FString& InOutputDir);
+		MOVIERENDERPIPELINECORE_API void GetMetadataFromCineCamera(class UCineCameraComponent* InComponent, const FString& InCameraName, const FString& InRenderPassName, TMap<FString, FString>& InOutMetadata);
+		MOVIERENDERPIPELINECORE_API void GetMetadataFromCameraLocRot(const FString& InCameraName, const FString& InRenderPassName, const FVector& InCurLoc, const FRotator& InCurRot, const FVector& InPrevLoc, const FRotator& InPrevRot, TMap<FString, FString>& InOutMetadata);
 	}
 }
 
diff --git a/Engine/Plugins/MovieScene/MovieRenderPipeline/Source/MovieRenderPipelineCore/Public/MovieRenderPipelineDataTypes.h b/Engine/Plugins/MovieScene/MovieRenderPipeline/Source/MovieRenderPipelineCore/Public/MovieRenderPipelineDataTypes.h
--- a/Engine/Plugins/MovieScene/MovieRenderPipeline/Source/MovieRenderPipelineCore/Public/MovieRenderPipelineDataTypes.h
+++ b/Engine/Plugins/MovieScene/MovieRenderPipeline/Source/MovieRenderPipelineCore/Public/MovieRenderPipelineDataTypes.h
@@ -96,14 +96,16 @@ struct FMoviePipelinePassIdentifier
 	FMoviePipelinePassIdentifier()
 	{}
 
-	FMoviePipelinePassIdentifier(const FString& InPassName)
+	// Name defaults to "camera" for backwards compatiblity with non-multicam metadata
+	FMoviePipelinePassIdentifier(const FString& InPassName, const FString& InCameraName = TEXT("camera"))
 		: Name(InPassName)
+		, CameraName(InCameraName)
 	{
 	}
 
 	bool operator == (const FMoviePipelinePassIdentifier& InRHS) const
 	{
-		return Name == InRHS.Name;
+		return Name == InRHS.Name && CameraName == InRHS.CameraName;
 	}
 
 	bool operator != (const FMoviePipelinePassIdentifier& InRHS) const
@@ -113,12 +115,18 @@ struct FMoviePipelinePassIdentifier
 
 	friend uint32 GetTypeHash(FMoviePipelinePassIdentifier OutputState)
 	{
-		return GetTypeHash(OutputState.Name);
+		return HashCombine(GetTypeHash(OutputState.Name), GetTypeHash(OutputState.CameraName));
 	}
 
 public:
+	// The name of the pass such as "FinalImage" or "ObjectId", etc.
 	UPROPERTY(EditAnywhere, BlueprintReadWrite, Category = "Movie Pipeline")
 	FString Name;
+
+	// The name of the camera that this pass is for. Stored here so we can differentiate between 
+	// multiple cameras within a single pass.
+	UPROPERTY(EditAnywhere, BlueprintReadWrite, Category = "Movie Pipeline")
+	FString CameraName;
 };
 
 namespace MoviePipeline
@@ -148,6 +156,11 @@ namespace MoviePipeline
 
 		/** Previous frame camera view rotation **/
 		FRotator PrevViewRotation;
+
+		TArray<FVector> CurrSidecarViewLocations;
+		TArray<FRotator> CurrSidecarViewRotations;
+		TArray<FVector> PrevSidecarViewLocations;
+		TArray<FRotator> PrevSidecarViewRotations;
 	};
 
 	/**
@@ -710,6 +723,7 @@ public:
 	
 	FTimecode CurrentShotSourceTimeCode;
 
+	int32 CameraIndex;
 
 
 	bool operator == (const FMoviePipelineFrameOutputState& InRHS) const
@@ -778,6 +792,7 @@ struct FMoviePipelineFilenameResolveParams
 		, InitializationTime(0)
 		, InitializationVersion(0)
 		, Job(nullptr)
+		, CameraIndex(-1)
 		, ShotOverride(nullptr)
 		, AdditionalFrameNumberOffset(0)
 	{
@@ -838,6 +853,9 @@ struct FMoviePipelineFilenameResolveParams
 	UPROPERTY(BlueprintReadWrite, EditAnywhere, Category = "Movie Render Pipeline")
 	UMoviePipelineExecutorJob* Job;
 
+	/** If this shot has multiple cameras in the sidecar array, which camera is this for? -1 will return the InnerName/Main camera for the shot. */
+	int32 CameraIndex;
+
 	/** Optional. If specified, settings will be pulled from this shot (if overriden by the shot). If null, always use the master configuration in the job. */
 	UPROPERTY(BlueprintReadWrite, EditAnywhere, Category = "Movie Render Pipeline")
 	UMoviePipelineExecutorShot* ShotOverride;
@@ -1086,6 +1104,19 @@ public:
 	{
 	}
 
+	bool HasDataFromMultipleCameras() const
+	{
+		TMap<FString, int32> CameraNameUseCounts;
+		for(const FMoviePipelinePassIdentifier& PassIdentifier : ExpectedRenderPasses)
+		{
+			CameraNameUseCounts.FindOrAdd(PassIdentifier.CameraName) += 1;
+		}
+
+		// ToDo: This logic might get more complicated because of burn ins, etc. that don't have camera names, 
+		// need to check.
+		return CameraNameUseCounts.Num() > 1;
+	}
+
 private:
 	// Explicitly delete copy operators since we own unique data.
 	// void operator=(const FMoviePipelineMergerOutputFrame&);
diff --git a/Engine/Plugins/MovieScene/MovieRenderPipeline/Source/MovieRenderPipelineRenderPasses/Private/MoviePipelineDeferredPasses.cpp b/Engine/Plugins/MovieScene/MovieRenderPipeline/Source/MovieRenderPipelineRenderPasses/Private/MoviePipelineDeferredPasses.cpp
--- a/Engine/Plugins/MovieScene/MovieRenderPipeline/Source/MovieRenderPipelineRenderPasses/Private/MoviePipelineDeferredPasses.cpp
+++ b/Engine/Plugins/MovieScene/MovieRenderPipeline/Source/MovieRenderPipelineRenderPasses/Private/MoviePipelineDeferredPasses.cpp
@@ -22,6 +22,7 @@
 #include "FinalPostProcessSettings.h"
 #include "Materials/Material.h"
 #include "MoviePipelineCameraSetting.h"
+#include "MoviePipelineHighResSetting.h"
 #include "MoviePipelineQueue.h"
 #include "MoviePipelineAntiAliasingSetting.h"
 #include "MoviePipelineOutputSetting.h"
@@ -30,11 +31,17 @@
 #include "Components/PrimitiveComponent.h"
 #include "EngineUtils.h"
 #include "Engine/RendererSettings.h"
+#include "Camera/CameraComponent.h"
+#include "CineCameraComponent.h"
+#include "Interfaces/Interface_PostProcessVolume.h"
+#include "MoviePipelineUtils.h"
 
 FString UMoviePipelineDeferredPassBase::StencilLayerMaterialAsset = TEXT("/MovieRenderPipeline/Materials/MoviePipeline_StencilCutout.MoviePipeline_StencilCutout");
 FString UMoviePipelineDeferredPassBase::DefaultDepthAsset = TEXT("/MovieRenderPipeline/Materials/MovieRenderQueue_WorldDepth.MovieRenderQueue_WorldDepth");
 FString UMoviePipelineDeferredPassBase::DefaultMotionVectorsAsset = TEXT("/MovieRenderPipeline/Materials/MovieRenderQueue_MotionVectors.MovieRenderQueue_MotionVectors");
 
+PRAGMA_DISABLE_OPTIMIZATION
+
 UMoviePipelineDeferredPassBase::UMoviePipelineDeferredPassBase() 
 	: UMoviePipelineImagePassBase()
 {
@@ -51,8 +58,8 @@
 		NewPass.Material = TSoftObjectPtr<UMaterialInterface>(FSoftObjectPath(MaterialPath));
 		NewPass.bEnabled = false;
 	}
+	bRenderMainPass = true;
 	bUse32BitPostProcessMaterials = false;
-	CurrentLayerIndex = INDEX_NONE;
 }
 
 void UMoviePipelineDeferredPassBase::MoviePipelineRenderShowFlagOverride(FEngineShowFlags& OutShowFlag)
@@ -105,7 +112,7 @@
 			UE_LOG(LogMovieRenderPipeline, Error, TEXT("Failed to load Stencil Mask material, stencil layers will be incorrect. Path: %s"), *StencilMatRef.ToString());
 		}
 	}
-
+	
 	for (FMoviePipelinePostProcessPass& AdditionalPass : AdditionalPostProcessMaterials)
 	{
 		if (AdditionalPass.bEnabled)
@@ -120,27 +127,74 @@
 
 	SurfaceQueue = MakeShared<FMoviePipelineSurfaceQueue>(InPassInitSettings.BackbufferResolution, EPixelFormat::PF_FloatRGBA, 3, true);
 
-	// Each stencil layer uses its own view state to keep TAA history etc separate
-	if (bAddDefaultLayer)
+	// Create a view state. Each individual camera, tile, and stencil layer need their own unique state as this includes visual history for anti-aliasing, etc. 
+	UMoviePipelineExecutorShot* CurrentShot = GetPipeline()->GetActiveShotList()[GetPipeline()->GetCurrentShotIndex()];
+	UMoviePipelineCameraSetting* CameraSettings = GetPipeline()->FindOrAddSettingForShot<UMoviePipelineCameraSetting>(CurrentShot);
+	UMoviePipelineHighResSetting* HighResSettings = GetPipeline()->FindOrAddSettingForShot<UMoviePipelineHighResSetting>(CurrentShot);
+	int32 NumCameras = CameraSettings->bRenderAllCameras ? CurrentShot->SidecarCameras.Num() : 1;
+
+	int32 TotalNumberOfAccumulators = 0;
+	for (int32 CamIndex = 0; CamIndex < NumCameras; CamIndex++)
 	{
-		StencilLayerViewStates.AddDefaulted();
+		FMultiCameraViewStateData& CameraData = CameraViewStateData.AddDefaulted_GetRef();
+
+		// We don't always want to allocate a unique history per tile as very large resolutions can OOM the GPU in backbuffer images alone.
+		// But we do need the history for some features (like Lumen) to work, so it's optional.
+		int32 NumHighResTiles = HighResSettings->bAllocateHistoryPerTile ? (HighResSettings->TileCount * HighResSettings->TileCount) : 1;
+		for (int32 TileIndexX = 0; TileIndexX < NumHighResTiles; TileIndexX++)
+		{
+			for (int32 TileIndexY = 0; TileIndexY < NumHighResTiles; TileIndexY++)
+			{
+				FMultiCameraViewStateData::FPerTile& PerTile = CameraData.TileData.FindOrAdd(FIntPoint(TileIndexX, TileIndexY));
+				// If they want to render the main pass (most likely) add a view state for it
+				if (bRenderMainPass)
+				{
+					PerTile.SceneViewStates.AddDefaulted();
+				}
+
+				// If they want to render a "default" stencil layer (that has everything not in another layer) add that...
+				if (StencilLayers.Num() > 0 && bAddDefaultLayer)
+				{
+					PerTile.SceneViewStates.AddDefaulted();
+				}
+
+				// Finally all of the other stencil layers
+				for (int32 Index = 0; Index < StencilLayers.Num(); Index++)
+				{
+					PerTile.SceneViewStates.AddDefaulted();
+				}
+			}
+		}
+
+		// We have to add up the number of accumulators needed separately, because we don't make
+		// one accumulator per high-res tile.
+		if (bRenderMainPass)
+		{
+			TotalNumberOfAccumulators++;
+		}
+		if (StencilLayers.Num() > 0 && bAddDefaultLayer)
+		{
+			TotalNumberOfAccumulators++;
+		}
+		for (int32 Index = 0; Index < StencilLayers.Num(); Index++)
+		{
+			TotalNumberOfAccumulators++;
+		}
+
+		// Now that we have an array of view states, allocate each one.
+		for (TPair<FIntPoint, FMultiCameraViewStateData::FPerTile>& Pair : CameraData.TileData)
+		{
+			for (int32 Index = 0; Index < Pair.Value.SceneViewStates.Num(); Index++)
+			{
+				Pair.Value.SceneViewStates[Index].Allocate(InPassInitSettings.FeatureLevel);
+			}
+		}
 	}
 
-	for (int32 Index = 0; Index < StencilLayers.Num(); Index++)
-	{
-		StencilLayerViewStates.AddDefaulted();
-	}
-
-	for (int32 Index = 0; Index < StencilLayerViewStates.Num(); Index++)
-	{
-		StencilLayerViewStates[Index].Allocate(InPassInitSettings.FeatureLevel);
-	}
-
-
-	// We must have at least enough accumulators to render all of the requested post process materials, because work doesn't begin
-	// until they're actually submitted to the render thread (which happens all at once) but we tie up an accumulator as we get ready to submit.
-	// If there aren't enough accumulators then we block until one is free but since submission hasn't gone through they'll never be free.
-	int32 PoolSize = (StencilLayerViewStates.Num() + ActivePostProcessMaterials.Num() + 1) * 3;
+	// We must allocate one accumulator per output, because when we submit a sample we tie up an accumulator, but because of temporal sampling
+	// the accumulators can be tied up for multiple game frames, thus we must have at least one per output and we can only reuse them between
+	// actual output frames (not engine frames). This doesn't allocate memory until they're actually used so it's ok to over-allocate.
+	int32 PoolSize = (TotalNumberOfAccumulators + (ActivePostProcessMaterials.Num()*NumCameras) + 1) * 3;
 	AccumulatorPool = MakeShared<TAccumulatorPool<FImageOverlappedAccumulator>, ESPMode::ThreadSafe>(PoolSize);
 	
 	PreviousCustomDepthValue.Reset();
@@ -151,7 +205,8 @@
 	// One Extension per sequence, since each sequence has its own OCIO settings.
 	OCIOSceneViewExtension = FSceneViewExtensions::NewExtension<FOpenColorIODisplayExtension>();
 
-	if (StencilLayerViewStates.Num() > 0)
+	const bool bEnableStencilPass = bAddDefaultLayer || StencilLayers.Num() > 0;
+	if (bEnableStencilPass)
 	{
 		IConsoleVariable* CVar = IConsoleManager::Get().FindConsoleVariable(TEXT("r.CustomDepth"));
 		if (CVar)
@@ -203,17 +258,22 @@
 
 	ActivePostProcessMaterials.Reset();
 
-	for (int32 Index = 0; Index < StencilLayerViewStates.Num(); Index++)
+	for (FMultiCameraViewStateData& CameraData : CameraViewStateData)
 	{
-		FSceneViewStateInterface* Ref = StencilLayerViewStates[Index].GetReference();
-		if (Ref)
+		for (TPair<FIntPoint, FMultiCameraViewStateData::FPerTile>& Pair : CameraData.TileData)
 		{
-			Ref->ClearMIDPool();
+			for (int32 Index = 0; Index < Pair.Value.SceneViewStates.Num(); Index++)
+			{
+				FSceneViewStateInterface* Ref = Pair.Value.SceneViewStates[Index].GetReference();
+				if (Ref)
+				{
+					Ref->ClearMIDPool();
+				}
+				Pair.Value.SceneViewStates[Index].Destroy();
+			}
 		}
-		StencilLayerViewStates[Index].Destroy();
 	}
-	StencilLayerViewStates.Reset();
-	CurrentLayerIndex = INDEX_NONE;
+	CameraViewStateData.Reset();
 	TileRenderTargets.Reset();
 	
 	OCIOSceneViewExtension.Reset();
@@ -256,68 +316,114 @@
 	Super::AddReferencedObjects(InThis, Collector);
 
 	UMoviePipelineDeferredPassBase& This = *CastChecked<UMoviePipelineDeferredPassBase>(InThis);
-	for (int32 Index = 0; Index < This.StencilLayerViewStates.Num(); Index++)
+
+	for (FMultiCameraViewStateData& CameraData : This.CameraViewStateData)
 	{
-		FSceneViewStateInterface* Ref = This.StencilLayerViewStates[Index].GetReference();
-		if (Ref)
+		for (TPair<FIntPoint, FMultiCameraViewStateData::FPerTile>& Pair : CameraData.TileData)
 		{
-			Ref->AddReferencedObjects(Collector);
+			for (int32 Index = 0; Index < Pair.Value.SceneViewStates.Num(); Index++)
+			{
+				FSceneViewStateInterface* Ref = Pair.Value.SceneViewStates[Index].GetReference();
+				if (Ref)
+				{
+					Ref->AddReferencedObjects(Collector);
+				}
+			}
 		}
 	}
 }
 
+namespace UE
+{
+namespace MoviePipeline
+{
+struct FDeferredPassRenderStatePayload : public UMoviePipelineImagePassBase::IViewCalcPayload
+{
+	int32 CameraIndex;
+	FIntPoint TileIndex; // Will always be 1,1 if no history-per-tile is enabled
+	int32 SceneViewIndex;
+};
+}
+}
+
+
 FSceneViewStateInterface* UMoviePipelineDeferredPassBase::GetSceneViewStateInterface(IViewCalcPayload* OptPayload)
 {
-	if (CurrentLayerIndex == INDEX_NONE)
+	UE::MoviePipeline::FDeferredPassRenderStatePayload* Payload = (UE::MoviePipeline::FDeferredPassRenderStatePayload*)OptPayload;
+	check(Payload);
+
+	FMultiCameraViewStateData& CameraData = CameraViewStateData[Payload->CameraIndex];
+	if (FMultiCameraViewStateData::FPerTile* TileData = CameraData.TileData.Find(Payload->TileIndex))
 	{
-		return Super::GetSceneViewStateInterface();
+		return TileData->SceneViewStates[Payload->SceneViewIndex].GetReference();
 	}
-	else
-	{
-		return StencilLayerViewStates[CurrentLayerIndex].GetReference();
-	}
+
+	return nullptr;
 }
 
 UTextureRenderTarget2D* UMoviePipelineDeferredPassBase::GetViewRenderTarget(IViewCalcPayload* OptPayload) const
 {
-	if (CurrentLayerIndex == INDEX_NONE)
-	{
-		return TileRenderTargets[0];
-	}
-	else
-	{
-		return TileRenderTargets[CurrentLayerIndex + 1];
-	}
+	return TileRenderTargets[0];
+	UE::MoviePipeline::FDeferredPassRenderStatePayload* Payload = (UE::MoviePipeline::FDeferredPassRenderStatePayload*)OptPayload;
+	check(Payload);
+
+	// ToDo: Is it really necessary to have a unique RT for each render? Only the last one gets shown (right now)
+	// if (CurrentLayerIndex == INDEX_NONE)
+	// {
+	// 	return TileRenderTargets[0];
+	// }
+	// else
+	// {
+	// 	return TileRenderTargets[CurrentLayerIndex + 1];
+	// }
 }
 
-
 void UMoviePipelineDeferredPassBase::GatherOutputPassesImpl(TArray<FMoviePipelinePassIdentifier>& ExpectedRenderPasses)
 {
-	// Add the default backbuffer
-	Super::GatherOutputPassesImpl(ExpectedRenderPasses);
+	// No super call here because multiple cameras makes this all complicated
+	// Super::GatherOutputPassesImpl(ExpectedRenderPasses);
 
-	TArray<FString> RenderPasses;
-	for (UMaterialInterface* Material : ActivePostProcessMaterials)
+	UMoviePipelineExecutorShot* CurrentShot = GetPipeline()->GetActiveShotList()[GetPipeline()->GetCurrentShotIndex()];
+	UMoviePipelineCameraSetting* CameraSettings = GetPipeline()->FindOrAddSettingForShot<UMoviePipelineCameraSetting>(CurrentShot);
+	int32 NumCameras = CameraSettings->bRenderAllCameras ? CurrentShot->SidecarCameras.Num() : 1;
+
+	for (int32 CameraIndex = 0; CameraIndex < NumCameras; CameraIndex++)
 	{
-		if (Material)
+		FMoviePipelinePassIdentifier PassIdentifierForCurrentCamera;
+		PassIdentifierForCurrentCamera.Name = PassIdentifier.Name;
+		PassIdentifierForCurrentCamera.CameraName = CurrentShot->GetCameraName(CameraIndex);
+
+		// Add the default backbuffer
+		if (bRenderMainPass)
 		{
-			RenderPasses.Add(Material->GetName());
+			ExpectedRenderPasses.Add(PassIdentifierForCurrentCamera);
 		}
-	}
 
-	for (const FString& Pass : RenderPasses)
-	{
-		ExpectedRenderPasses.Add(FMoviePipelinePassIdentifier(PassIdentifier.Name + Pass));
-	}
+		// Each camera will render everything in the Post Process Material stack.
+		TArray<FString> RenderPasses;
+		for (UMaterialInterface* Material : ActivePostProcessMaterials)
+		{
+			if (Material)
+			{
+				RenderPasses.Add(Material->GetName());
+			}
+		}
 
-	if (bAddDefaultLayer)
-	{
-		ExpectedRenderPasses.Add(FMoviePipelinePassIdentifier(PassIdentifier.Name + TEXT("DefaultLayer")));
-	}
+		for (const FString& Pass : RenderPasses)
+		{
+			ExpectedRenderPasses.Add(FMoviePipelinePassIdentifier(PassIdentifierForCurrentCamera.Name + Pass, PassIdentifierForCurrentCamera.CameraName));
+		}
 
-	for (const FActorLayer& Layer : StencilLayers)
-	{
-		ExpectedRenderPasses.Add(FMoviePipelinePassIdentifier(PassIdentifier.Name + Layer.Name.ToString()));
+		// Stencil Layer Time!
+		if (StencilLayers.Num() > 0 && bAddDefaultLayer)
+		{
+			ExpectedRenderPasses.Add(FMoviePipelinePassIdentifier(PassIdentifierForCurrentCamera.Name + TEXT("DefaultLayer"), PassIdentifierForCurrentCamera.CameraName));
+		}
+
+		for (const FActorLayer& Layer : StencilLayers)
+		{
+			ExpectedRenderPasses.Add(FMoviePipelinePassIdentifier(PassIdentifierForCurrentCamera.Name + Layer.Name.ToString(), PassIdentifierForCurrentCamera.CameraName));
+		}
 	}
 }
 
@@ -349,183 +455,210 @@
 		SurfaceQueue->BlockUntilAnyAvailable();
 	}
 
+	UMoviePipelineExecutorShot* CurrentShot = GetPipeline()->GetActiveShotList()[GetPipeline()->GetCurrentShotIndex()];
+	UMoviePipelineCameraSetting* CameraSettings = GetPipeline()->FindOrAddSettingForShot<UMoviePipelineCameraSetting>(CurrentShot);
+	int32 NumCameras = CameraSettings->bRenderAllCameras ? CurrentShot->SidecarCameras.Num() : 1;
 
-	// Main Render Pass
+	for (int32 CameraIndex = 0; CameraIndex < NumCameras; CameraIndex++)
 	{
-		FMoviePipelineRenderPassMetrics InOutSampleState = InSampleState;
+		FMoviePipelinePassIdentifier PassIdentifierForCurrentCamera;
+		PassIdentifierForCurrentCamera.Name = PassIdentifier.Name;
+		PassIdentifierForCurrentCamera.CameraName = CurrentShot->GetCameraName(CameraIndex);
 
-		CurrentLayerIndex = INDEX_NONE;
-		TSharedPtr<FSceneViewFamilyContext> ViewFamily = CalculateViewFamily(InOutSampleState);
+		// Main Render Pass
+		if (bRenderMainPass)
+		{
+			FMoviePipelineRenderPassMetrics InOutSampleState = InSampleState;
+			// InOutSampleState.OutputState.CameraCount = NumCameras;
+			InOutSampleState.OutputState.CameraIndex = CameraIndex;
 
-		// Add post-processing materials if needed
-		FSceneView* View = const_cast<FSceneView*>(ViewFamily->Views[0]);
-		View->FinalPostProcessSettings.BufferVisualizationOverviewMaterials.Empty();
-		View->FinalPostProcessSettings.BufferVisualizationPipes.Empty();
+			UE::MoviePipeline::FDeferredPassRenderStatePayload Payload;
+			Payload.CameraIndex = CameraIndex;
+			Payload.TileIndex = InOutSampleState.TileIndexes;
 
-		for (UMaterialInterface* Material : ActivePostProcessMaterials)
-		{
-			if (Material)
+			// Main renders use index 0.
+			Payload.SceneViewIndex = 0;
+
+			TSharedPtr<FSceneViewFamilyContext> ViewFamily = CalculateViewFamily(InOutSampleState, &Payload);
+
+			// Add post-processing materials if needed
+			FSceneView* View = const_cast<FSceneView*>(ViewFamily->Views[0]);
+			View->FinalPostProcessSettings.BufferVisualizationOverviewMaterials.Empty();
+			View->FinalPostProcessSettings.BufferVisualizationPipes.Empty();
+
+			for (UMaterialInterface* Material : ActivePostProcessMaterials)
 			{
-				View->FinalPostProcessSettings.BufferVisualizationOverviewMaterials.Add(Material);
+				if (Material)
+				{
+					View->FinalPostProcessSettings.BufferVisualizationOverviewMaterials.Add(Material);
+				}
 			}
-		}
 
-		for (UMaterialInterface* VisMaterial : View->FinalPostProcessSettings.BufferVisualizationOverviewMaterials)
-		{
-			// If this was just to contribute to the history buffer, no need to go any further.
-			if (InOutSampleState.bDiscardResult)
+			for (UMaterialInterface* VisMaterial : View->FinalPostProcessSettings.BufferVisualizationOverviewMaterials)
 			{
-				continue;
+				// If this was just to contribute to the history buffer, no need to go any further.
+				if (InOutSampleState.bDiscardResult)
+				{
+					continue;
+				}
+				FMoviePipelinePassIdentifier LayerPassIdentifier = FMoviePipelinePassIdentifier(PassIdentifier.Name + VisMaterial->GetName(), PassIdentifierForCurrentCamera.CameraName);
+
+				auto BufferPipe = MakeShared<FImagePixelPipe, ESPMode::ThreadSafe>();
+				BufferPipe->AddEndpoint(MakeForwardingEndpoint(LayerPassIdentifier, InOutSampleState));
+
+				View->FinalPostProcessSettings.BufferVisualizationPipes.Add(VisMaterial->GetFName(), BufferPipe);
 			}
-			FMoviePipelinePassIdentifier LayerPassIdentifier = FMoviePipelinePassIdentifier(PassIdentifier.Name + VisMaterial->GetName());
 
-			auto BufferPipe = MakeShared<FImagePixelPipe, ESPMode::ThreadSafe>();
-			BufferPipe->AddEndpoint(MakeForwardingEndpoint(LayerPassIdentifier, InSampleState));
 
-			View->FinalPostProcessSettings.BufferVisualizationPipes.Add(VisMaterial->GetFName(), BufferPipe);
+			int32 NumValidMaterials = View->FinalPostProcessSettings.BufferVisualizationPipes.Num();
+			View->FinalPostProcessSettings.bBufferVisualizationDumpRequired = NumValidMaterials > 0;
+
+			// Submit to be rendered. Main render pass always uses target 0.
+			FRenderTarget* RenderTarget = GetViewRenderTarget()->GameThread_GetRenderTargetResource();
+			FCanvas Canvas = FCanvas(RenderTarget, nullptr, GetPipeline()->GetWorld(), View->GetFeatureLevel(), FCanvas::CDM_DeferDrawing, 1.0f);
+			GetRendererModule().BeginRenderingViewFamily(&Canvas, ViewFamily.Get());
+
+			// Readback + Accumulate.
+			PostRendererSubmission(InOutSampleState, PassIdentifierForCurrentCamera, GetOutputFileSortingOrder(), Canvas);
 		}
 
 
-		int32 NumValidMaterials = View->FinalPostProcessSettings.BufferVisualizationPipes.Num();
-		View->FinalPostProcessSettings.bBufferVisualizationDumpRequired = NumValidMaterials > 0;
+		// Now do the stencil layer submission (which doesn't support additional post processing materials)
+		{
+			FMoviePipelineRenderPassMetrics InOutSampleState = InSampleState;
+			InOutSampleState.OutputState.CameraIndex = CameraIndex;
 
-		// Submit to be rendered. Main render pass always uses target 0.
-		FRenderTarget* RenderTarget = GetViewRenderTarget()->GameThread_GetRenderTargetResource();
-		FCanvas Canvas = FCanvas(RenderTarget, nullptr, GetPipeline()->GetWorld(), View->GetFeatureLevel(), FCanvas::CDM_DeferDrawing, 1.0f);
-		GetRendererModule().BeginRenderingViewFamily(&Canvas, ViewFamily.Get());
+			struct FStencilValues
+			{
+				FStencilValues()
+					: bRenderCustomDepth(false)
+					, StencilMask(ERendererStencilMask::ERSM_Default)
+					, CustomStencil(0)
+				{
+				}
 
-		// Readback + Accumulate.
-		PostRendererSubmission(InOutSampleState, PassIdentifier, GetOutputFileSortingOrder(), Canvas);
-	}
+				bool bRenderCustomDepth;
+				ERendererStencilMask StencilMask;
+				int32 CustomStencil;
+			};
 
-	// Now submit stencil layers if needed
-	{
-		FMoviePipelineRenderPassMetrics InOutSampleState = InSampleState;
-
-		struct FStencilValues
-		{
-			FStencilValues()
-				: bRenderCustomDepth(false)
-				, StencilMask(ERendererStencilMask::ERSM_Default)
-				, CustomStencil(0)
+			// Now for each stencil layer we reconfigure all the actors custom depth/stencil 
+			TArray<FActorLayer> AllStencilLayers = StencilLayers;
+			if (bAddDefaultLayer)
 			{
+				FActorLayer& DefaultLayer = AllStencilLayers.AddDefaulted_GetRef();
+				DefaultLayer.Name = FName("DefaultLayer");
 			}
 
-			bool bRenderCustomDepth;
-			ERendererStencilMask StencilMask;
-			int32 CustomStencil;
-		};
-
-		// If we're going to be using stencil layers, we need to cache all of the users
-		// custom stencil/depth settings since we're changing them to do the mask.
-		TMap<UPrimitiveComponent*, FStencilValues> PreviousValues;
-		if (StencilLayers.Num() > 0)
-		{
-			for (TActorIterator<AActor> ActorItr(GetWorld()); ActorItr; ++ActorItr)
+			// If we're going to be using stencil layers, we need to cache all of the users
+			// custom stencil/depth settings since we're changing them to do the mask.
+			TMap<UPrimitiveComponent*, FStencilValues> PreviousValues;
+			if (AllStencilLayers.Num() > 0)
 			{
-				AActor* Actor = *ActorItr;
-				if (Actor)
+				for (TActorIterator<AActor> ActorItr(GetWorld()); ActorItr; ++ActorItr)
 				{
-					for (UActorComponent* Component : Actor->GetComponents())
+					AActor* Actor = *ActorItr;
+					if (Actor)
 					{
-						if (Component && Component->IsA<UPrimitiveComponent>())
+						for (UActorComponent* Component : Actor->GetComponents())
 						{
-							UPrimitiveComponent* PrimitiveComponent = CastChecked<UPrimitiveComponent>(Component);
-							FStencilValues& Values = PreviousValues.Add(PrimitiveComponent);
-							Values.StencilMask = PrimitiveComponent->CustomDepthStencilWriteMask;
-							Values.CustomStencil = PrimitiveComponent->CustomDepthStencilValue;
-							Values.bRenderCustomDepth = PrimitiveComponent->bRenderCustomDepth;
-						}
-					}
-				}
-			}
-		}
-
-		// Now for each stencil layer we reconfigure all the actors custom depth/stencil 
-		TArray<FActorLayer> AllStencilLayers = StencilLayers;
-		if (bAddDefaultLayer)
-		{
-			FActorLayer& DefaultLayer = AllStencilLayers.AddDefaulted_GetRef();
-			DefaultLayer.Name = FName("DefaultLayer");
-		}
-
-		CurrentLayerIndex = 0;
-		for (const FActorLayer& Layer : AllStencilLayers)
-		{
-			FMoviePipelinePassIdentifier LayerPassIdentifier = FMoviePipelinePassIdentifier(PassIdentifier.Name + Layer.Name.ToString());
-
-			for (TActorIterator<AActor> ActorItr(GetWorld()); ActorItr; ++ActorItr)
-			{
-				AActor* Actor = *ActorItr;
-				if (Actor)
-				{
-					// The way stencil masking works is that we draw the actors on the given layer to the stencil buffer.
-					// Then we apply a post-processing material which colors pixels outside those actors black, before
-					// post processing. Then, TAA, Motion Blur, etc. is applied to all pixels. An alpha channel can preserve
-					// which pixels were the geometry and which are dead space which lets you apply that as a mask later.
-					bool bInLayer = true;
-					if (bAddDefaultLayer && Layer.Name == FName("DefaultLayer"))
-					{
-						// If we're trying to render the default layer, the logic is different - we only add objects who
-						// aren't in any of the stencil layers.
-						for (const FActorLayer& AllLayer : StencilLayers)
-						{
-							bInLayer = !Actor->Layers.Contains(AllLayer.Name);
-							if (!bInLayer)
+							if (Component && Component->IsA<UPrimitiveComponent>())
 							{
-								break;
+								UPrimitiveComponent* PrimitiveComponent = CastChecked<UPrimitiveComponent>(Component);
+								FStencilValues& Values = PreviousValues.Add(PrimitiveComponent);
+								Values.StencilMask = PrimitiveComponent->CustomDepthStencilWriteMask;
+								Values.CustomStencil = PrimitiveComponent->CustomDepthStencilValue;
+								Values.bRenderCustomDepth = PrimitiveComponent->bRenderCustomDepth;
 							}
-						}
-					}
-					else
-					{
-						// If this a normal layer, we only add the actor if it exists on this layer.
-						bInLayer = Actor->Layers.Contains(Layer.Name);
-					}
-
-					for (UActorComponent* Component : Actor->GetComponents())
-					{
-						if (Component && Component->IsA<UPrimitiveComponent>())
-						{
-							UPrimitiveComponent* PrimitiveComponent = CastChecked<UPrimitiveComponent>(Component);
-							// We want to render all objects not on the layer to stencil too so that foreground objects mask.
-							PrimitiveComponent->SetCustomDepthStencilValue(bInLayer ? 1 : 0);
-							PrimitiveComponent->SetCustomDepthStencilWriteMask(ERendererStencilMask::ERSM_Default);
-							PrimitiveComponent->SetRenderCustomDepth(true);
 						}
 					}
 				}
 			}
 
-			if (StencilLayerMaterial)
+
+			for (int32 StencilLayerIndex = 0; StencilLayerIndex < AllStencilLayers.Num(); StencilLayerIndex++)
 			{
-				TSharedPtr<FSceneViewFamilyContext> ViewFamily = CalculateViewFamily(InOutSampleState);
-				FSceneView* View = const_cast<FSceneView*>(ViewFamily->Views[0]);
+				const FActorLayer& Layer = AllStencilLayers[StencilLayerIndex];
+				FMoviePipelinePassIdentifier LayerPassIdentifier = FMoviePipelinePassIdentifier(PassIdentifierForCurrentCamera.Name + Layer.Name.ToString());
+				LayerPassIdentifier.CameraName = PassIdentifierForCurrentCamera.CameraName;
 
-				// Now that we've modified all of the stencil values, we can submit them to be rendered.
-				View->FinalPostProcessSettings.AddBlendable(StencilLayerMaterial, 1.0f);
-				IBlendableInterface* BlendableInterface = Cast<IBlendableInterface>(StencilLayerMaterial);
-				BlendableInterface->OverrideBlendableSettings(*View, 1.f);
+				// Modify all of the actors in this world so they have the right stencil settings (so we can use the stencil buffer as a mask later)
+				for (TActorIterator<AActor> ActorItr(GetWorld()); ActorItr; ++ActorItr)
+				{
+					AActor* Actor = *ActorItr;
+					if (Actor)
+					{
+						// The way stencil masking works is that we draw the actors on the given layer to the stencil buffer.
+						// Then we apply a post-processing material which colors pixels outside those actors black, before
+						// post processing. Then, TAA, Motion Blur, etc. is applied to all pixels. An alpha channel can preserve
+						// which pixels were the geometry and which are dead space which lets you apply that as a mask later.
+						bool bInLayer = true;
+						if (bAddDefaultLayer && Layer.Name == FName("DefaultLayer"))
+						{
+							// If we're trying to render the default layer, the logic is different - we only add objects who
+							// aren't in any of the stencil layers.
+							for (const FActorLayer& AllLayer : StencilLayers)
+							{
+								bInLayer = !Actor->Layers.Contains(AllLayer.Name);
+								if (!bInLayer)
+								{
+									break;
+								}
+							}
+						}
+						else
+						{
+							// If this a normal layer, we only add the actor if it exists on this layer.
+							bInLayer = Actor->Layers.Contains(Layer.Name);
+						}
 
+						for (UActorComponent* Component : Actor->GetComponents())
+						{
+							if (Component && Component->IsA<UPrimitiveComponent>())
+							{
+								UPrimitiveComponent* PrimitiveComponent = CastChecked<UPrimitiveComponent>(Component);
+								// We want to render all objects not on the layer to stencil too so that foreground objects mask.
+								PrimitiveComponent->SetCustomDepthStencilValue(bInLayer ? 1 : 0);
+								PrimitiveComponent->SetCustomDepthStencilWriteMask(ERendererStencilMask::ERSM_Default);
+								PrimitiveComponent->SetRenderCustomDepth(true);
+							}
+						}
+					}
+				}
+
+				// Submit the actual render now
+				if (StencilLayerMaterial)
 				{
-					FRenderTarget* RenderTarget = GetViewRenderTarget()->GameThread_GetRenderTargetResource();
-					FCanvas Canvas = FCanvas(RenderTarget, nullptr, GetPipeline()->GetWorld(), ERHIFeatureLevel::SM5, FCanvas::CDM_DeferDrawing, 1.0f);
-					GetRendererModule().BeginRenderingViewFamily(&Canvas, ViewFamily.Get());
+					UE::MoviePipeline::FDeferredPassRenderStatePayload Payload;
+					Payload.CameraIndex = CameraIndex;
+					Payload.TileIndex = InOutSampleState.TileIndexes;
+					Payload.SceneViewIndex = StencilLayerIndex + (bRenderMainPass ? 1 : 0);
+					TSharedPtr<FSceneViewFamilyContext> ViewFamily = CalculateViewFamily(InOutSampleState, &Payload);
+					FSceneView* View = const_cast<FSceneView*>(ViewFamily->Views[0]);
 
-					// Readback + Accumulate.
-					PostRendererSubmission(InSampleState, LayerPassIdentifier, GetOutputFileSortingOrder() + 1, Canvas);
+					// Now that we've modified all of the stencil values, we can submit them to be rendered.
+					View->FinalPostProcessSettings.AddBlendable(StencilLayerMaterial, 1.0f);
+					IBlendableInterface* BlendableInterface = Cast<IBlendableInterface>(StencilLayerMaterial);
+					BlendableInterface->OverrideBlendableSettings(*View, 1.f);
+
+					{
+						FRenderTarget* RenderTarget = GetViewRenderTarget()->GameThread_GetRenderTargetResource();
+						FCanvas Canvas = FCanvas(RenderTarget, nullptr, GetPipeline()->GetWorld(), View->GetFeatureLevel(), FCanvas::CDM_DeferDrawing, 1.0f);
+						GetRendererModule().BeginRenderingViewFamily(&Canvas, ViewFamily.Get());
+
+						// Readback + Accumulate.
+						PostRendererSubmission(InOutSampleState, LayerPassIdentifier, GetOutputFileSortingOrder() + 1, Canvas);
+					}
 				}
 			}
 
-			CurrentLayerIndex++;
-		}
-
-		// Now we can restore the custom depth/stencil/etc. values so that the main render pass acts as the user expects next time.
-		for (TPair<UPrimitiveComponent*, FStencilValues>& KVP : PreviousValues)
-		{
-			KVP.Key->SetCustomDepthStencilValue(KVP.Value.CustomStencil);
-			KVP.Key->SetCustomDepthStencilWriteMask(KVP.Value.StencilMask);
-			KVP.Key->SetRenderCustomDepth(KVP.Value.bRenderCustomDepth);
+			// Now that all stencil layers have been rendered, we can restore the custom depth/stencil/etc. values so that the main render pass acts as the user expects next time.
+			for (TPair<UPrimitiveComponent*, FStencilValues>& KVP : PreviousValues)
+			{
+				KVP.Key->SetCustomDepthStencilValue(KVP.Value.CustomStencil);
+				KVP.Key->SetCustomDepthStencilWriteMask(KVP.Value.StencilMask);
+				KVP.Key->SetRenderCustomDepth(KVP.Value.bRenderCustomDepth);
+			}
 		}
 	}
 }
@@ -608,6 +741,112 @@
 	return Callback;
 }
 
+UE::MoviePipeline::FImagePassCameraViewData UMoviePipelineDeferredPassBase::GetCameraInfo(FMoviePipelineRenderPassMetrics& InOutSampleState, IViewCalcPayload* OptPayload) const
+{
+	UMoviePipelineExecutorShot* CurrentShot = GetPipeline()->GetActiveShotList()[GetPipeline()->GetCurrentShotIndex()];
+	UMoviePipelineCameraSetting* CameraSettings = GetPipeline()->FindOrAddSettingForShot<UMoviePipelineCameraSetting>(CurrentShot);
+	int32 NumCameras = CameraSettings->bRenderAllCameras ? CurrentShot->SidecarCameras.Num() : 1;
+	
+	if (NumCameras == 1)
+	{
+		// If there's only one camera being used we can use the parent class which assumes the camera comes from the PlayerCameraManager
+		return Super::GetCameraInfo(InOutSampleState, OptPayload);
+	}
+	else
+	{
+		UE::MoviePipeline::FImagePassCameraViewData OutCameraData;
+
+		// Here's where it gets a lot more complicated. There's a number of properties we need to fetch from a camera manually to fill out the minimal view info.
+		UCameraComponent* OutCamera = nullptr;
+
+		GetPipeline()->GetSidecarCameraData(CurrentShot, InOutSampleState.OutputState.CameraIndex, OutCameraData.ViewInfo, &OutCamera);
+		if (OutCamera)
+		{
+			// This has to come from the main camera for consistency's sake, and it's not a per-camera setting in the editor.
+			OutCameraData.ViewActor = GetPipeline()->GetWorld()->GetFirstPlayerController()->GetViewTarget();
+
+			// Try adding cine-camera specific metadata (not all animated cameras are cine cameras though)
+			UCineCameraComponent* CineCameraComponent = Cast<UCineCameraComponent>(OutCamera);
+			if (CineCameraComponent)
+			{
+				// ToDo: This is still wrong, PassIdentifier.CameraName needs to come in from the InOutSampleState somewhere.
+				UE::MoviePipeline::GetMetadataFromCineCamera(CineCameraComponent, PassIdentifier.CameraName, PassIdentifier.Name, OutCameraData.FileMetadata);
+
+				// We only do this in the multi-camera case because the single camera case is covered by the main Rendering loop.
+				UE::MoviePipeline::GetMetadataFromCameraLocRot(PassIdentifier.CameraName, PassIdentifier.Name, OutCameraData.ViewInfo.Location, OutCameraData.ViewInfo.Rotation, OutCameraData.ViewInfo.PreviousViewTransform->GetLocation(), FRotator(OutCameraData.ViewInfo.PreviousViewTransform->GetRotation()), OutCameraData.FileMetadata);
+			}
+		}
+		else
+		{
+			UE_LOG(LogMovieRenderPipeline, Error, TEXT("Failed to find Camera Component for Shot: %d CameraIndex: %d"), GetPipeline()->GetCurrentShotIndex(), InOutSampleState.OutputState.CameraIndex);
+		}
+
+		return OutCameraData;
+	}
+}
+
+void UMoviePipelineDeferredPassBase::BlendPostProcessSettings(FSceneView* InView, FMoviePipelineRenderPassMetrics& InOutSampleState, IViewCalcPayload* OptPayload)
+{
+	UMoviePipelineExecutorShot* CurrentShot = GetPipeline()->GetActiveShotList()[GetPipeline()->GetCurrentShotIndex()];
+	UMoviePipelineCameraSetting* CameraSettings = GetPipeline()->FindOrAddSettingForShot<UMoviePipelineCameraSetting>(CurrentShot);
+	int32 NumCameras = CameraSettings->bRenderAllCameras ? CurrentShot->SidecarCameras.Num() : 1;
+
+	UCameraComponent* OutCamera = nullptr;
+	FMinimalViewInfo OutViewInfo;
+
+	GetPipeline()->GetSidecarCameraData(CurrentShot, InOutSampleState.OutputState.CameraIndex, OutViewInfo, &OutCamera);
+	if (!OutCamera)
+	{
+		// GetCameraInfo will have already printed a warning
+		return;
+	}
+
+	// The primary camera should still respect the world post processing volumes and should already be the viewtarget.
+	if (NumCameras == 1)
+	{
+		// If there's only one camera being used we can use the parent class which assumes the camera comes from the PlayerCameraManager
+		Super::GetCameraInfo(InOutSampleState, OptPayload);
+	}
+	else
+	{
+		// For sidecar cameras we need to do the blending of PP volumes
+		FVector ViewLocation = OutCamera->GetComponentLocation();
+		for (IInterface_PostProcessVolume* PPVolume : GetWorld()->PostProcessVolumes)
+		{
+			const FPostProcessVolumeProperties VolumeProperties = PPVolume->GetProperties();
+
+			// Skip any volumes which are disabled
+			if (!VolumeProperties.bIsEnabled)
+			{
+				continue;
+			}
+
+			float LocalWeight = FMath::Clamp(VolumeProperties.BlendWeight, 0.0f, 1.0f);
+
+			if (!VolumeProperties.bIsUnbound)
+			{
+				float DistanceToPoint = 0.0f;
+				PPVolume->EncompassesPoint(ViewLocation, 0.0f, &DistanceToPoint);
+
+				if (DistanceToPoint >= 0 && DistanceToPoint < VolumeProperties.BlendRadius)
+				{
+					LocalWeight *= FMath::Clamp(1.0f - DistanceToPoint / VolumeProperties.BlendRadius, 0.0f, 1.0f);
+				}
+				else
+				{
+					LocalWeight = 0.0f;
+				}
+			}
+
+			InView->OverridePostProcessSettings(*VolumeProperties.Settings, LocalWeight);
+		}
+
+		// After blending all post processing volumes, blend the camera's post process settings too
+		InView->OverridePostProcessSettings(OutCamera->PostProcessSettings, OutCamera->PostProcessBlendWeight);
+	}
+}
+
+
 void UMoviePipelineDeferredPassBase::PostRendererSubmission(const FMoviePipelineRenderPassMetrics& InSampleState, const FMoviePipelinePassIdentifier InPassIdentifier, const int32 InSortingOrder, FCanvas& InCanvas)
 {
 	// If this was just to contribute to the history buffer, no need to go any further.
@@ -793,3 +1032,5 @@
 
 	Super::SetupImpl(InPassInitSettings);
 }
+
+PRAGMA_ENABLE_OPTIMIZATION
diff --git a/Engine/Plugins/MovieScene/MovieRenderPipeline/Source/MovieRenderPipelineRenderPasses/Private/MoviePipelineImagePassBase.cpp b/Engine/Plugins/MovieScene/MovieRenderPipeline/Source/MovieRenderPipelineRenderPasses/Private/MoviePipelineImagePassBase.cpp
--- a/Engine/Plugins/MovieScene/MovieRenderPipeline/Source/MovieRenderPipelineRenderPasses/Private/MoviePipelineImagePassBase.cpp
+++ b/Engine/Plugins/MovieScene/MovieRenderPipeline/Source/MovieRenderPipelineRenderPasses/Private/MoviePipelineImagePassBase.cpp
@@ -22,6 +22,8 @@
 // For Cine Camera Variables in Metadata
 #include "CineCameraActor.h"
 #include "CineCameraComponent.h"
+#include "MoviePipelineUtils.h"
+
 
 DECLARE_CYCLE_STAT(TEXT("STAT_MoviePipeline_AccumulateSample_TT"), STAT_AccumulateSample_TaskThread, STATGROUP_MoviePipeline);
 
@@ -288,12 +290,14 @@ FSceneView* UMoviePipelineImagePassBase::GetSceneViewForSampleState(FSceneViewFa
 	int32 TileSizeX = InOutSampleState.BackbufferSize.X;
 	int32 TileSizeY = InOutSampleState.BackbufferSize.Y;
 
+	UE::MoviePipeline::FImagePassCameraViewData CameraInfo = GetCameraInfo(InOutSampleState, OptPayload);
+
 	FSceneViewInitOptions ViewInitOptions;
 	ViewInitOptions.ViewFamily = ViewFamily;
-	ViewInitOptions.ViewOrigin = InOutSampleState.FrameInfo.CurrViewLocation;
+	ViewInitOptions.ViewOrigin = CameraInfo.ViewInfo.Location;
 	ViewInitOptions.SetViewRectangle(FIntRect(FIntPoint(0, 0), FIntPoint(TileSizeX, TileSizeY)));
-	ViewInitOptions.ViewRotationMatrix = FInverseRotationMatrix(InOutSampleState.FrameInfo.CurrViewRotation);
-	ViewInitOptions.ViewActor = LocalPlayerController ? LocalPlayerController->GetViewTarget() : nullptr;
+	ViewInitOptions.ViewRotationMatrix = FInverseRotationMatrix(CameraInfo.ViewInfo.Rotation);
+	ViewInitOptions.ViewActor = CameraInfo.ViewActor;
 
 	// Rotate the view 90 degrees (reason: unknown)
 	ViewInitOptions.ViewRotationMatrix = ViewInitOptions.ViewRotationMatrix * FMatrix(
@@ -301,106 +305,119 @@ FSceneView* UMoviePipelineImagePassBase::GetSceneViewForSampleState(FSceneViewFa
 		FPlane(1, 0, 0, 0),
 		FPlane(0, 1, 0, 0),
 		FPlane(0, 0, 0, 1));
-	float ViewFOV = 90.f;
-	if (GetPipeline()->GetWorld()->GetFirstPlayerController()->PlayerCameraManager)
-	{
-		ViewFOV = GetPipeline()->GetWorld()->GetFirstPlayerController()->PlayerCameraManager->GetFOVAngle();
-	}
+	float ViewFOV = CameraInfo.ViewInfo.DesiredFOV;
 
 	// Inflate our FOV to support the overscan 
 	ViewFOV = 2.0f * FMath::RadiansToDegrees(FMath::Atan((1.0f + InOutSampleState.OverscanPercentage) * FMath::Tan(FMath::DegreesToRadians( ViewFOV * 0.5f ))));
 
 	float DofSensorScale = 1.0f;
 
-	// Calculate a Projection Matrix
+	// Calculate a Projection Matrix. This code unfortunately ends up similar to, but not quite the same as FMinimalViewInfo::CalculateProjectionMatrixGivenView
 	{
-		float XAxisMultiplier;
-		float YAxisMultiplier;
+		FMatrix BaseProjMatrix;
 
-		check(GetPipeline()->GetWorld());
-		check(GetPipeline()->GetWorld()->GetFirstPlayerController());
-		APlayerCameraManager* PlayerCameraManager = GetPipeline()->GetWorld()->GetFirstPlayerController()->PlayerCameraManager;
-		
-		// Stretch the fovs if the view is constrained to the camera's aspect ratio
-		if (PlayerCameraManager && PlayerCameraManager->GetCameraCacheView().bConstrainAspectRatio)
+		if (CameraInfo.ViewInfo.ProjectionMode == ECameraProjectionMode::Orthographic)
 		{
-			const FMinimalViewInfo CameraCache = PlayerCameraManager->GetCameraCacheView();
-			const float DestAspectRatio = ViewInitOptions.GetViewRect().Width() / (float)ViewInitOptions.GetViewRect().Height();
+			const float YScale = 1.0f / CameraInfo.ViewInfo.AspectRatio;
 
-			// If the camera's aspect ratio has a thinner width, then stretch the horizontal fov more than usual to 
-			// account for the extra with of (before constraining - after constraining)
-			if (CameraCache.AspectRatio < DestAspectRatio)
-			{
-				const float ConstrainedWidth = ViewInitOptions.GetViewRect().Height() * CameraCache.AspectRatio;
-				XAxisMultiplier = ConstrainedWidth / (float)ViewInitOptions.GetViewRect().Width();
-				YAxisMultiplier = CameraCache.AspectRatio;
-			}
-			// Simplified some math here but effectively functions similarly to the above, the unsimplified code would look like:
-			// const float ConstrainedHeight = ViewInitOptions.GetViewRect().Width() / CameraCache.AspectRatio;
-			// YAxisMultiplier = (ConstrainedHeight / ViewInitOptions.GetViewRect.Height()) * CameraCache.AspectRatio;
-			else
-			{
-				XAxisMultiplier = 1.0f;
-				YAxisMultiplier = ViewInitOptions.GetViewRect().Width() / (float)ViewInitOptions.GetViewRect().Height();
-			}
+			const float HalfOrthoWidth = CameraInfo.ViewInfo.OrthoWidth / 2.0f;
+			const float ScaledOrthoHeight = CameraInfo.ViewInfo.OrthoWidth / 2.0f * YScale;
+
+			const float NearPlane = CameraInfo.ViewInfo.OrthoNearClipPlane;
+			const float FarPlane = CameraInfo.ViewInfo.OrthoFarClipPlane;
+
+			const float ZScale = 1.0f / (FarPlane - NearPlane);
+			const float ZOffset = -NearPlane;
+
+			BaseProjMatrix = FReversedZOrthoMatrix(
+				HalfOrthoWidth,
+				ScaledOrthoHeight,
+				ZScale,
+				ZOffset
+			);
 		}
 		else
 		{
-			const int32 DestSizeX = ViewInitOptions.GetViewRect().Width();
-			const int32 DestSizeY = ViewInitOptions.GetViewRect().Height();
-			const EAspectRatioAxisConstraint AspectRatioAxisConstraint = GetDefault<ULocalPlayer>()->AspectRatioAxisConstraint;
-			if (((DestSizeX > DestSizeY) && (AspectRatioAxisConstraint == AspectRatio_MajorAxisFOV)) || (AspectRatioAxisConstraint == AspectRatio_MaintainXFOV))
+			float XAxisMultiplier;
+			float YAxisMultiplier;
+
+			if (CameraInfo.ViewInfo.bConstrainAspectRatio)
 			{
-				//if the viewport is wider than it is tall
-				XAxisMultiplier = 1.0f;
-				YAxisMultiplier = ViewInitOptions.GetViewRect().Width() / (float)ViewInitOptions.GetViewRect().Height();
+				const float DestAspectRatio = ViewInitOptions.GetViewRect().Width() / (float)ViewInitOptions.GetViewRect().Height();
+
+				// If the camera's aspect ratio has a thinner width, then stretch the horizontal fov more than usual to 
+				// account for the extra with of (before constraining - after constraining)
+				if (CameraInfo.ViewInfo.AspectRatio < DestAspectRatio)
+				{
+					const float ConstrainedWidth = ViewInitOptions.GetViewRect().Height() * CameraInfo.ViewInfo.AspectRatio;
+					XAxisMultiplier = ConstrainedWidth / (float)ViewInitOptions.GetViewRect().Width();
+					YAxisMultiplier = CameraInfo.ViewInfo.AspectRatio;
+				}
+				// Simplified some math here but effectively functions similarly to the above, the unsimplified code would look like:
+				// const float ConstrainedHeight = ViewInitOptions.GetViewRect().Width() / CameraCache.AspectRatio;
+				// YAxisMultiplier = (ConstrainedHeight / ViewInitOptions.GetViewRect.Height()) * CameraCache.AspectRatio;
+				else
+				{
+					XAxisMultiplier = 1.0f;
+					YAxisMultiplier = ViewInitOptions.GetViewRect().Width() / (float)ViewInitOptions.GetViewRect().Height();
+				}
 			}
 			else
 			{
-				//if the viewport is taller than it is wide
-				XAxisMultiplier = ViewInitOptions.GetViewRect().Height() / (float)ViewInitOptions.GetViewRect().Width();
-				YAxisMultiplier = 1.0f;
+				const int32 DestSizeX = ViewInitOptions.GetViewRect().Width();
+				const int32 DestSizeY = ViewInitOptions.GetViewRect().Height();
+				const EAspectRatioAxisConstraint AspectRatioAxisConstraint = GetDefault<ULocalPlayer>()->AspectRatioAxisConstraint;
+				if (((DestSizeX > DestSizeY) && (AspectRatioAxisConstraint == AspectRatio_MajorAxisFOV)) || (AspectRatioAxisConstraint == AspectRatio_MaintainXFOV))
+				{
+					//if the viewport is wider than it is tall
+					XAxisMultiplier = 1.0f;
+					YAxisMultiplier = ViewInitOptions.GetViewRect().Width() / (float)ViewInitOptions.GetViewRect().Height();
+				}
+				else
+				{
+					//if the viewport is taller than it is wide
+					XAxisMultiplier = ViewInitOptions.GetViewRect().Height() / (float)ViewInitOptions.GetViewRect().Width();
+					YAxisMultiplier = 1.0f;
+				}
 			}
-		}
 
-		float MinZ = GNearClippingPlane;
-		if (PlayerCameraManager)
-		{
-			float NearClipPlane = PlayerCameraManager->GetCameraCacheView().PerspectiveNearClipPlane;
-			MinZ = NearClipPlane > 0 ? NearClipPlane : MinZ;
-		}
-		const float MaxZ = MinZ;
-		// Avoid zero ViewFOV's which cause divide by zero's in projection matrix
-		const float MatrixFOV = FMath::Max(0.001f, ViewFOV) * (float)PI / 360.0f;
+			const float MinZ = CameraInfo.ViewInfo.GetFinalPerspectiveNearClipPlane();
+			const float MaxZ = MinZ;
+			// Avoid zero ViewFOV's which cause divide by zero's in projection matrix
+			const float MatrixFOV = FMath::Max(0.001f, ViewFOV) * (float)PI / 360.0f;
 
-		FMatrix BaseProjMatrix;
 
-		if ((bool)ERHIZBuffer::IsInverted)
-		{
-			BaseProjMatrix = FReversedZPerspectiveMatrix(
-				MatrixFOV,
-				MatrixFOV,
-				XAxisMultiplier,
-				YAxisMultiplier,
-				MinZ,
-				MaxZ
-			);
-		}
-		else
-		{
-			BaseProjMatrix = FPerspectiveMatrix(
-				MatrixFOV,
-				MatrixFOV,
-				XAxisMultiplier,
-				YAxisMultiplier,
-				MinZ,
-				MaxZ
-			);
-		}
+			if ((bool)ERHIZBuffer::IsInverted)
+			{
+				BaseProjMatrix = FReversedZPerspectiveMatrix(
+					MatrixFOV,
+					MatrixFOV,
+					XAxisMultiplier,
+					YAxisMultiplier,
+					MinZ,
+					MaxZ
+				);
+			}
+			else
+			{
+				BaseProjMatrix = FPerspectiveMatrix(
+					MatrixFOV,
+					MatrixFOV,
+					XAxisMultiplier,
+					YAxisMultiplier,
+					MinZ,
+					MaxZ
+				);
+			}
 
-		// Modify the perspective matrix to do an off center projection, with overlap for high-res tiling
-		ModifyProjectionMatrixForTiling(InOutSampleState, /*InOut*/ BaseProjMatrix, DofSensorScale);
+			// Modify the perspective matrix to do an off center projection, with overlap for high-res tiling
+			ModifyProjectionMatrixForTiling(InOutSampleState, /*InOut*/ BaseProjMatrix, DofSensorScale);
+			// ToDo: Does orthographic support tiling in the same way or do I need to modify the values before creating the ortho view.
+		}
+		
+		// BaseProjMatrix may be perspective or orthographic.
 		ViewInitOptions.ProjectionMatrix = BaseProjMatrix;
+		
 	}
 
 	ViewInitOptions.SceneViewStateInterface = GetSceneViewStateInterface(OptPayload);
@@ -408,13 +425,15 @@ FSceneView* UMoviePipelineImagePassBase::GetSceneViewForSampleState(FSceneViewFa
 
 	FSceneView* View = new FSceneView(ViewInitOptions);
 	ViewFamily->Views.Add(View);
-	View->ViewLocation = InOutSampleState.FrameInfo.CurrViewLocation;
-	View->ViewRotation = InOutSampleState.FrameInfo.CurrViewRotation;
+
+	
+	View->ViewLocation = CameraInfo.ViewInfo.Location;
+	View->ViewRotation = CameraInfo.ViewInfo.Rotation;
 	// Override previous/current view transforms so that tiled renders don't use the wrong occlusion/motion blur information.
-	View->PreviousViewTransform = FTransform(InOutSampleState.FrameInfo.PrevViewRotation, InOutSampleState.FrameInfo.PrevViewLocation);
+	View->PreviousViewTransform = CameraInfo.ViewInfo.PreviousViewTransform;
 
 	View->StartFinalPostprocessSettings(View->ViewLocation);
-	BlendPostProcessSettings(View);
+	BlendPostProcessSettings(View, InOutSampleState, OptPayload);
 
 	// Scaling sensor size inversely with the the projection matrix [0][0] should physically
 	// cause the circle of confusion to be unchanged.
@@ -425,39 +444,17 @@ FSceneView* UMoviePipelineImagePassBase::GetSceneViewForSampleState(FSceneViewFa
 
 	// This metadata is per-file and not per-view, but we need the blended result from the view to actually match what we rendered.
 	// To solve this, we'll insert metadata per renderpass, separated by render pass name.
-	InOutSampleState.OutputState.FileMetadata.Add(FString::Printf(TEXT("unreal/camera/%s/fstop"), *PassIdentifier.Name), FString::SanitizeFloat(View->FinalPostProcessSettings.DepthOfFieldFstop));
-	InOutSampleState.OutputState.FileMetadata.Add(FString::Printf(TEXT("unreal/camera/%s/fov"), *PassIdentifier.Name), FString::SanitizeFloat(ViewInitOptions.FOV));
-	InOutSampleState.OutputState.FileMetadata.Add(FString::Printf(TEXT("unreal/camera/%s/focalDistance"), *PassIdentifier.Name), FString::SanitizeFloat(View->FinalPostProcessSettings.DepthOfFieldFocalDistance));
-	InOutSampleState.OutputState.FileMetadata.Add(FString::Printf(TEXT("unreal/camera/%s/sensorWidth"), *PassIdentifier.Name), FString::SanitizeFloat(View->FinalPostProcessSettings.DepthOfFieldSensorWidth));
-	InOutSampleState.OutputState.FileMetadata.Add(FString::Printf(TEXT("unreal/camera/%s/overscanPercent"), *PassIdentifier.Name), FString::SanitizeFloat(InOutSampleState.OverscanPercentage));
-
-	if (GetWorld()->GetFirstPlayerController()->PlayerCameraManager)
-	{
-		// This only works if you use a Cine Camera (which is almost guranteed with Sequencer) and it's easier (and less human error prone) than re-deriving the information
-		ACineCameraActor* CineCameraActor = Cast<ACineCameraActor>(GetWorld()->GetFirstPlayerController()->PlayerCameraManager->GetViewTarget());
-		if (CineCameraActor)
-		{
-			UCineCameraComponent* CineCameraComponent = CineCameraActor->GetCineCameraComponent();
-			if (CineCameraComponent)
-			{
-				InOutSampleState.OutputState.FileMetadata.Add(FString::Printf(TEXT("unreal/camera/%s/sensorWidth"), *PassIdentifier.Name), FString::SanitizeFloat(CineCameraComponent->Filmback.SensorWidth));
-				InOutSampleState.OutputState.FileMetadata.Add(FString::Printf(TEXT("unreal/camera/%s/sensorHeight"), *PassIdentifier.Name), FString::SanitizeFloat(CineCameraComponent->Filmback.SensorHeight));
-				InOutSampleState.OutputState.FileMetadata.Add(FString::Printf(TEXT("unreal/camera/%s/sensorAspectRatio"), *PassIdentifier.Name), FString::SanitizeFloat(CineCameraComponent->Filmback.SensorAspectRatio));
-				InOutSampleState.OutputState.FileMetadata.Add(FString::Printf(TEXT("unreal/camera/%s/minFocalLength"), *PassIdentifier.Name), FString::SanitizeFloat(CineCameraComponent->LensSettings.MinFocalLength));
-				InOutSampleState.OutputState.FileMetadata.Add(FString::Printf(TEXT("unreal/camera/%s/maxFocalLength"), *PassIdentifier.Name), FString::SanitizeFloat(CineCameraComponent->LensSettings.MaxFocalLength));
-				InOutSampleState.OutputState.FileMetadata.Add(FString::Printf(TEXT("unreal/camera/%s/minFStop"), *PassIdentifier.Name), FString::SanitizeFloat(CineCameraComponent->LensSettings.MinFStop));
-				InOutSampleState.OutputState.FileMetadata.Add(FString::Printf(TEXT("unreal/camera/%s/maxFStop"), *PassIdentifier.Name), FString::SanitizeFloat(CineCameraComponent->LensSettings.MaxFStop));
-				InOutSampleState.OutputState.FileMetadata.Add(FString::Printf(TEXT("unreal/camera/%s/dofDiaphragmBladeCount"), *PassIdentifier.Name), FString::FromInt(CineCameraComponent->LensSettings.DiaphragmBladeCount));
-				InOutSampleState.OutputState.FileMetadata.Add(FString::Printf(TEXT("unreal/camera/%s/focalLength"), *PassIdentifier.Name), FString::SanitizeFloat(CineCameraComponent->CurrentFocalLength));
-			}
-		}
-	}
-
+	InOutSampleState.OutputState.FileMetadata.Add(FString::Printf(TEXT("unreal/%s/%s/fstop"), *PassIdentifier.CameraName, *PassIdentifier.Name), FString::SanitizeFloat(View->FinalPostProcessSettings.DepthOfFieldFstop));
+	InOutSampleState.OutputState.FileMetadata.Add(FString::Printf(TEXT("unreal/%s/%s/fov"), *PassIdentifier.CameraName, *PassIdentifier.Name), FString::SanitizeFloat(ViewInitOptions.FOV));
+	InOutSampleState.OutputState.FileMetadata.Add(FString::Printf(TEXT("unreal/%s/%s/focalDistance"), *PassIdentifier.CameraName, *PassIdentifier.Name), FString::SanitizeFloat(View->FinalPostProcessSettings.DepthOfFieldFocalDistance));
+	InOutSampleState.OutputState.FileMetadata.Add(FString::Printf(TEXT("unreal/%s/%s/sensorWidth"), *PassIdentifier.CameraName, *PassIdentifier.Name), FString::SanitizeFloat(View->FinalPostProcessSettings.DepthOfFieldSensorWidth));
+	InOutSampleState.OutputState.FileMetadata.Add(FString::Printf(TEXT("unreal/%s/%s/overscanPercent"), *PassIdentifier.CameraName, *PassIdentifier.Name), FString::SanitizeFloat(InOutSampleState.OverscanPercentage));
 
+	InOutSampleState.OutputState.FileMetadata.Append(CameraInfo.FileMetadata);
 	return View;
 }
 
-void UMoviePipelineImagePassBase::BlendPostProcessSettings(FSceneView* InView)
+void UMoviePipelineImagePassBase::BlendPostProcessSettings(FSceneView* InView, FMoviePipelineRenderPassMetrics& InOutSampleState, IViewCalcPayload* OptPayload)
 {
 	check(InView);
 
@@ -544,6 +541,40 @@ void UMoviePipelineImagePassBase::ModifyProjectionMatrixForTiling(const FMoviePi
 	InOutProjectionMatrix.M[2][1] += OffsetY / PadRatioX;
 }
 
+UE::MoviePipeline::FImagePassCameraViewData UMoviePipelineImagePassBase::GetCameraInfo(FMoviePipelineRenderPassMetrics& InOutSampleState, IViewCalcPayload* OptPayload) const
+{
+	UE::MoviePipeline::FImagePassCameraViewData OutCameraData;
+
+	// Default implementation doesn't support multi-camera and always provides the information from the current PlayerCameraManager
+	if (GetPipeline()->GetWorld()->GetFirstPlayerController()->PlayerCameraManager)
+	{
+		OutCameraData.ViewInfo = GetPipeline()->GetWorld()->GetFirstPlayerController()->PlayerCameraManager->GetCameraCacheView();
+
+		// Now override some of the properties with things that come from MRQ
+		OutCameraData.ViewInfo.Location = InOutSampleState.FrameInfo.CurrViewLocation;
+		OutCameraData.ViewInfo.Rotation = InOutSampleState.FrameInfo.CurrViewRotation;
+		OutCameraData.ViewInfo.PreviousViewTransform = FTransform(InOutSampleState.FrameInfo.PrevViewRotation, InOutSampleState.FrameInfo.PrevViewLocation);
+
+		// And some fields that aren't in FMinimalViewInfo
+		OutCameraData.ViewActor = GetPipeline()->GetWorld()->GetFirstPlayerController()->GetViewTarget();
+
+		// This only works if you use a Cine Camera (which is almost guranteed with Sequencer) and it's easier (and less human error prone) than re-deriving the information
+		ACineCameraActor* CineCameraActor = Cast<ACineCameraActor>(GetWorld()->GetFirstPlayerController()->PlayerCameraManager->GetViewTarget());
+		if (CineCameraActor)
+		{
+			UCineCameraComponent* CineCameraComponent = CineCameraActor->GetCineCameraComponent();
+			if (CineCameraComponent)
+			{
+				// Add camera-specific metadata
+				UE::MoviePipeline::GetMetadataFromCineCamera(CineCameraComponent, PassIdentifier.CameraName, PassIdentifier.Name, OutCameraData.FileMetadata);
+			}
+		}
+	}
+
+	return OutCameraData;
+}
+
+
 TSharedPtr<FAccumulatorPool::FAccumulatorInstance, ESPMode::ThreadSafe> FAccumulatorPool::BlockAndGetAccumulator_GameThread(int32 InFrameNumber, const FMoviePipelinePassIdentifier& InPassIdentifier)
 {
 	FScopeLock ScopeLock(&CriticalSection);
diff --git a/Engine/Plugins/MovieScene/MovieRenderPipeline/Source/MovieRenderPipelineRenderPasses/Private/MoviePipelineImageSequenceOutput.cpp b/Engine/Plugins/MovieScene/MovieRenderPipeline/Source/MovieRenderPipelineRenderPasses/Private/MoviePipelineImageSequenceOutput.cpp
--- a/Engine/Plugins/MovieScene/MovieRenderPipeline/Source/MovieRenderPipelineRenderPasses/Private/MoviePipelineImageSequenceOutput.cpp
+++ b/Engine/Plugins/MovieScene/MovieRenderPipeline/Source/MovieRenderPipelineRenderPasses/Private/MoviePipelineImageSequenceOutput.cpp
@@ -24,7 +24,6 @@
 #include "Misc/Paths.h"
 
 DECLARE_CYCLE_STAT(TEXT("ImgSeqOutput_RecieveImageData"), STAT_ImgSeqRecieveImageData, STATGROUP_MoviePipeline);
-
 struct FAsyncImageQuantization
 {
 	FAsyncImageQuantization(FImageWriteTask* InWriteTask, const bool bInApplysRGB)
@@ -155,9 +154,11 @@ void UMoviePipelineImageSequenceOutputBase::OnReceiveImageDataImpl(FMoviePipelin
 			// If we're writing more than one render pass out, we need to ensure the file name has the format string in it so we don't
 			// overwrite the same file multiple times. Burn In overlays don't count if they are getting composited on top of an existing file.
 			const bool bIncludeRenderPass = InMergedOutputFrame->ImageOutputData.Num() - CompositedPasses.Num() > 1;
+			const bool bIncludeCameraName = InMergedOutputFrame->HasDataFromMultipleCameras();
 			const bool bTestFrameNumber = true;
 
-			UE::MoviePipeline::ValidateOutputFormatString(FileNameFormatString, bIncludeRenderPass, bTestFrameNumber);
+			// ToDo: This validation may need to include the OutputDirectory, because they could have put {camera_name} in the output folder
+			UE::MoviePipeline::ValidateOutputFormatString(FileNameFormatString, bIncludeRenderPass, bTestFrameNumber, bIncludeCameraName);
 
 			// Create specific data that needs to override 
 			TMap<FString, FString> FormatOverrides;
@@ -167,13 +168,13 @@ void UMoviePipelineImageSequenceOutputBase::OnReceiveImageDataImpl(FMoviePipelin
 
 			// Resolve for XMLs
 			{
-				GetPipeline()->ResolveFilenameFormatArguments(FileNameFormatString, FormatOverrides, XMLData.ImageSequenceFileName, FinalFormatArgs, &InMergedOutputFrame->FrameOutputState, -InMergedOutputFrame->FrameOutputState.ShotOutputFrameNumber);
+				GetPipeline()->ResolveFilenameFormatArguments(FileNameFormatString, FormatOverrides, XMLData.ImageSequenceFileName, FinalFormatArgs, &Payload->SampleState.OutputState, -Payload->SampleState.OutputState.ShotOutputFrameNumber);
 			}
 			
 			// Resolve the final absolute file path to write this to
 			{
 				FString FormatString = OutputDirectory / FileNameFormatString;
-				GetPipeline()->ResolveFilenameFormatArguments(FormatString, FormatOverrides, OutputData.FilePath, FinalFormatArgs, &InMergedOutputFrame->FrameOutputState);
+				GetPipeline()->ResolveFilenameFormatArguments(FormatString, FormatOverrides, OutputData.FilePath, FinalFormatArgs, &Payload->SampleState.OutputState);
 
 				if (FPaths::IsRelative(OutputData.FilePath))
 				{
@@ -184,7 +185,7 @@ void UMoviePipelineImageSequenceOutputBase::OnReceiveImageDataImpl(FMoviePipelin
 			// More XML resolving. Create a deterministic clipname by removing frame numbers, file extension, and any trailing .'s
 			{
 				UE::MoviePipeline::RemoveFrameNumberFormatStrings(FileNameFormatString, true);
-				GetPipeline()->ResolveFilenameFormatArguments(FileNameFormatString, FormatOverrides, XMLData.ClipName, FinalFormatArgs, &InMergedOutputFrame->FrameOutputState);
+				GetPipeline()->ResolveFilenameFormatArguments(FileNameFormatString, FormatOverrides, XMLData.ClipName, FinalFormatArgs, &Payload->SampleState.OutputState);
 				XMLData.ClipName.RemoveFromEnd(Extension);
 				XMLData.ClipName.RemoveFromEnd(".");
 			}
@@ -263,7 +264,7 @@ void UMoviePipelineImageSequenceOutputBase::OnReceiveImageDataImpl(FMoviePipelin
 		TileImageTask->PixelData = MoveTemp(QuantizedPixelData);
 		
 #if WITH_EDITOR
-		GetPipeline()->AddFrameToOutputMetadata(XMLData.ClipName, XMLData.ImageSequenceFileName, InMergedOutputFrame->FrameOutputState, Extension, Payload->bRequireTransparentOutput);
+		GetPipeline()->AddFrameToOutputMetadata(XMLData.ClipName, XMLData.ImageSequenceFileName, Payload->SampleState.OutputState, Extension, Payload->bRequireTransparentOutput);
 #endif
 
 		GetPipeline()->AddOutputFuture(ImageWriteQueue->Enqueue(MoveTemp(TileImageTask)), OutputData);
diff --git a/Engine/Plugins/MovieScene/MovieRenderPipeline/Source/MovieRenderPipelineRenderPasses/Public/MoviePipelineDeferredPasses.h b/Engine/Plugins/MovieScene/MovieRenderPipeline/Source/MovieRenderPipelineRenderPasses/Public/MoviePipelineDeferredPasses.h
--- a/Engine/Plugins/MovieScene/MovieRenderPipeline/Source/MovieRenderPipelineRenderPasses/Public/MoviePipelineDeferredPasses.h
+++ b/Engine/Plugins/MovieScene/MovieRenderPipeline/Source/MovieRenderPipelineRenderPasses/Public/MoviePipelineDeferredPasses.h
@@ -55,6 +55,8 @@ protected:
 	virtual UTextureRenderTarget2D* GetViewRenderTarget(IViewCalcPayload* OptPayload = nullptr) const override;
 	virtual void AddViewExtensions(FSceneViewFamilyContext& InContext, FMoviePipelineRenderPassMetrics& InOutSampleState) override;
 	virtual bool IsAutoExposureAllowed(const FMoviePipelineRenderPassMetrics& InSampleState) const override;
+	virtual void BlendPostProcessSettings(FSceneView* InView, FMoviePipelineRenderPassMetrics& InOutSampleState, IViewCalcPayload* OptPayload = nullptr);
+	virtual UE::MoviePipeline::FImagePassCameraViewData GetCameraInfo(FMoviePipelineRenderPassMetrics& InOutSampleState, IViewCalcPayload* OptPayload = nullptr) const;
 	// ~UMoviePipelineRenderPass
 
 	// FGCObject Interface
@@ -96,6 +98,12 @@ public:
 	UPROPERTY(EditAnywhere, BlueprintReadWrite, Category = "Deferred Renderer Data")
 	TArray<FMoviePipelinePostProcessPass> AdditionalPostProcessMaterials;
 
+	/**
+	* This can be turned off if you're only doing a stencil-layer based render and don't need the main non-stencil approach.
+	*/
+	UPROPERTY(EditAnywhere, BlueprintReadWrite, Category = "Stencil Clip Layers")
+	bool bRenderMainPass;
+
 	/**
 	* If true, an additional stencil layer will be rendered which contains all objects which do not belong to layers
 	* specified in the Stencil Layers. This is useful for wanting to isolate one or two layers but still have everything
@@ -124,23 +132,28 @@ protected:
 	UPROPERTY(Transient, DuplicateTransient)
 	TArray<UTextureRenderTarget2D*> TileRenderTargets;
 
+	struct FMultiCameraViewStateData
+	{
+		struct FPerTile
+		{
+			TArray<FSceneViewStateReference> SceneViewStates;
+		};
 
-	TSharedPtr<FAccumulatorPool, ESPMode::ThreadSafe> AccumulatorPool;
+		TMap<FIntPoint, FPerTile> TileData;
+	};
 
-	int32 CurrentLayerIndex;
-	TArray<FSceneViewStateReference> StencilLayerViewStates;
-
-	/** The lifetime of this SceneViewExtension is only during the rendering process. It is destroyed as part of TearDown. */
-	TSharedPtr<FOpenColorIODisplayExtension, ESPMode::ThreadSafe> OCIOSceneViewExtension;
+	TArray<FMultiCameraViewStateData> CameraViewStateData;
 
 	// Cache the custom stencil value. Only has meaning if they have stencil layers.
 	TOptional<int32> PreviousCustomDepthValue;
-	
 	/** Cache the previous dump frames as HDR value. Only used if using 32-bit post processing. */
 	TOptional<int32> PreviousDumpFramesValue;
 	/** Cache the previous color format value. Only used if using 32-bit post processing. */
 	TOptional<int32> PreviousColorFormatValue;
 
+	TSharedPtr<FAccumulatorPool, ESPMode::ThreadSafe> AccumulatorPool;
+	/** The lifetime of this SceneViewExtension is only during the rendering process. It is destroyed as part of TearDown. */
+	TSharedPtr<FOpenColorIODisplayExtension, ESPMode::ThreadSafe> OCIOSceneViewExtension;
 public:
 	static FString StencilLayerMaterialAsset;
 	static FString DefaultDepthAsset;
diff --git a/Engine/Plugins/MovieScene/MovieRenderPipeline/Source/MovieRenderPipelineRenderPasses/Public/MoviePipelineImagePassBase.h b/Engine/Plugins/MovieScene/MovieRenderPipeline/Source/MovieRenderPipelineRenderPasses/Public/MoviePipelineImagePassBase.h
--- a/Engine/Plugins/MovieScene/MovieRenderPipeline/Source/MovieRenderPipelineRenderPasses/Public/MoviePipelineImagePassBase.h
+++ b/Engine/Plugins/MovieScene/MovieRenderPipeline/Source/MovieRenderPipelineRenderPasses/Public/MoviePipelineImagePassBase.h
@@ -24,6 +24,23 @@ class FSceneViewFamily;
 class FSceneView;
 struct FAccumulatorPool;
 
+namespace UE
+{
+	namespace MoviePipeline
+	{
+		struct FImagePassCameraViewData
+		{
+			FImagePassCameraViewData()
+				: ViewActor(nullptr)
+			{}
+			
+			FMinimalViewInfo ViewInfo;
+			TMap<FString, FString> FileMetadata;
+			AActor* ViewActor;
+		};
+	}
+}
+
 class FMoviePipelineBackgroundAccumulateTask
 {
 public:
@@ -96,7 +113,9 @@ protected:
 protected:
 	virtual void GetViewShowFlags(FEngineShowFlags& OutShowFlag, EViewModeIndex& OutViewModeIndex) const;
 	virtual TSharedPtr<FSceneViewFamilyContext> CalculateViewFamily(FMoviePipelineRenderPassMetrics& InOutSampleState, IViewCalcPayload* OptPayload = nullptr);
-	virtual void BlendPostProcessSettings(FSceneView* InView);
+	UE_DEPRECATED(5.1, "Use the overload that provides the InOutSampleState and IViewCalcPayload instead.")
+	virtual void BlendPostProcessSettings(FSceneView* InView) {}
+	virtual void BlendPostProcessSettings(FSceneView* InView, FMoviePipelineRenderPassMetrics& InOutSampleState, IViewCalcPayload* OptPayload = nullptr);
 	virtual void SetupViewForViewModeOverride(FSceneView* View);
 	virtual void MoviePipelineRenderShowFlagOverride(FEngineShowFlags& OutShowFlag) {}
 	virtual bool IsScreenPercentageSupported() const { return true; }	
@@ -107,6 +126,7 @@ protected:
 	virtual void AddViewExtensions(FSceneViewFamilyContext& InContext, FMoviePipelineRenderPassMetrics& InOutSampleState) { }
 	virtual bool IsAutoExposureAllowed(const FMoviePipelineRenderPassMetrics& InSampleState) const { return true; }
 	virtual FSceneView* GetSceneViewForSampleState(FSceneViewFamily* ViewFamily, FMoviePipelineRenderPassMetrics& InOutSampleState, IViewCalcPayload* OptPayload = nullptr);
+	virtual UE::MoviePipeline::FImagePassCameraViewData GetCameraInfo(FMoviePipelineRenderPassMetrics& InOutSampleState, IViewCalcPayload* OptPayload = nullptr) const;
 public:
 	
 
