MoviePipeline: Fixed a crash related to multi-camera support when a sequence had no Camera Cut track. Fixed log spam warning about multiple passes when using burn ins and not using multiple cameras.

diff --git a/Engine/Plugins/MovieScene/MovieRenderPipeline/Source/MovieRenderPipelineCore/Private/MoviePipeline.cpp b/Engine/Plugins/MovieScene/MovieRenderPipeline/Source/MovieRenderPipelineCore/Private/MoviePipeline.cpp
--- a/Engine/Plugins/MovieScene/MovieRenderPipeline/Source/MovieRenderPipelineCore/Private/MoviePipeline.cpp
+++ b/Engine/Plugins/MovieScene/MovieRenderPipeline/Source/MovieRenderPipelineCore/Private/MoviePipeline.cpp
@@ -1691,25 +1691,42 @@ void UMoviePipeline::GetSidecarCameraData(UMoviePipelineExecutorShot* InShot, in
 		return;
 	}
 
-	const FGuid CameraBindingRef = InShot->SidecarCameras[InCameraIndex].BindingId;
-	const FMovieSceneSequenceID CameraSequenceRef = InShot->SidecarCameras[InCameraIndex].SequenceId;
+	UObject* PrimaryViewTarget = nullptr;
+	// Sequences don't always have a camera cut track, at which point we fall back to the regular player controller.
+	if (InCameraIndex >= InShot->SidecarCameras.Num())
 	{
-		TArrayView<TWeakObjectPtr<UObject>> BoundCameras = LevelSequenceActor->GetSequencePlayer()->FindBoundObjects(CameraBindingRef, CameraSequenceRef);
-		if (BoundCameras.Num() > 0)
+		if (APlayerController* PlayerController = GetWorld()->GetFirstPlayerController())
 		{
-			UCameraComponent* BoundCamera = Cast<UCameraComponent>(BoundCameras[0].Get());
-			if (BoundCamera)
+			if (APlayerCameraManager* PlayerCameraManager = PlayerController->PlayerCameraManager)
 			{
-				*OutCameraComponent = BoundCamera;
-				BoundCamera->GetCameraView(GetWorld()->GetDeltaSeconds(), OutViewInfo);
-				
-				// We override the current/previous transform based on cached data though to ensure we accurately handle cur/next frame positions
-				OutViewInfo.Location = FrameInfo.CurrSidecarViewLocations[InCameraIndex];
-				OutViewInfo.Rotation = FrameInfo.CurrSidecarViewRotations[InCameraIndex];
-				OutViewInfo.PreviousViewTransform = FTransform(FrameInfo.PrevSidecarViewRotations[InCameraIndex], FrameInfo.PrevSidecarViewLocations[InCameraIndex]);
+				PrimaryViewTarget = PlayerCameraManager->GetViewTarget();
+			}
+		}
+	}
+	else
+	{
+		const FGuid CameraBindingRef = InShot->SidecarCameras[InCameraIndex].BindingId;
+		const FMovieSceneSequenceID CameraSequenceRef = InShot->SidecarCameras[InCameraIndex].SequenceId;
+		{
+			TArrayView<TWeakObjectPtr<UObject>> BoundCameras = LevelSequenceActor->GetSequencePlayer()->FindBoundObjects(CameraBindingRef, CameraSequenceRef);
+			if (BoundCameras.Num() > 0)
+			{
+				PrimaryViewTarget = BoundCameras[0].Get();
 			}
 		}
 	}
+
+	UCameraComponent* BoundCamera = Cast<UCameraComponent>(PrimaryViewTarget);
+	if (BoundCamera)
+	{
+		*OutCameraComponent = BoundCamera;
+		BoundCamera->GetCameraView(GetWorld()->GetDeltaSeconds(), OutViewInfo);
+				
+		// We override the current/previous transform based on cached data though to ensure we accurately handle cur/next frame positions
+		OutViewInfo.Location = FrameInfo.CurrSidecarViewLocations[InCameraIndex];
+		OutViewInfo.Rotation = FrameInfo.CurrSidecarViewRotations[InCameraIndex];
+		OutViewInfo.PreviousViewTransform = FTransform(FrameInfo.PrevSidecarViewRotations[InCameraIndex], FrameInfo.PrevSidecarViewLocations[InCameraIndex]);
+	}
 }
 
 void UMoviePipeline::GetSidecarCameraViewPoints(UMoviePipelineExecutorShot* InShot, TArray<FVector>& OutSidecarViewLocations, TArray<FRotator>& OutSidecarViewRotations) const
diff --git a/Engine/Plugins/MovieScene/MovieRenderPipeline/Source/MovieRenderPipelineRenderPasses/Private/MoviePipelineEXROutput.cpp b/Engine/Plugins/MovieScene/MovieRenderPipeline/Source/MovieRenderPipelineRenderPasses/Private/MoviePipelineEXROutput.cpp
--- a/Engine/Plugins/MovieScene/MovieRenderPipeline/Source/MovieRenderPipelineRenderPasses/Private/MoviePipelineEXROutput.cpp
+++ b/Engine/Plugins/MovieScene/MovieRenderPipeline/Source/MovieRenderPipelineRenderPasses/Private/MoviePipelineEXROutput.cpp
@@ -10,6 +10,7 @@
 #include "Math/Float16.h"
 #include "MovieRenderPipelineCoreModule.h"
 #include "MoviePipelineOutputSetting.h"
+#include "MoviePipelineCameraSetting.h"
 #include "ImageWriteQueue.h"
 #include "MoviePipeline.h"
 #include "MoviePipelineImageQuantization.h"
@@ -498,22 +499,36 @@ void UMoviePipelineImageSequenceOutput_EXR::OnReceiveImageDataImpl(FMoviePipelin
 
 			// No quantization required, just copy the data as we will move it into the image write task.
 			TUniquePtr<FImagePixelData> PixelData = RenderPassData.Value->CopyImageData();
+			FImagePixelDataPayload* Payload = RenderPassData.Value->GetPayload<FImagePixelDataPayload>();
+			ShotIndex = Payload->SampleState.OutputState.ShotIndex;
 
 			// If there is more than one layer, then we will prefix the layer. The first layer is not prefixed (and gets inserted as RGBA)
 			// as most programs that handle EXRs expect the main image data to be in an unnamed layer.
 			if (LayerIndex == 0)
 			{
 				// Only check the main image pass for transparent output since that's generally considered the 'preview'.
-				FImagePixelDataPayload* Payload = RenderPassData.Value->GetPayload<FImagePixelDataPayload>();
 				bRequiresTransparentOutput = Payload->bRequireTransparentOutput;
-				ShotIndex = Payload->SampleState.OutputState.ShotIndex;
 				MultiLayerImageTask->OverscanPercentage = Payload->SampleState.OverscanPercentage;
 			}
 			else
 			{
 				// If there is more than one layer, then we will prefix the layer. The first layer is not prefixed (and gets inserted as RGBA)
-				// as most programs that handle EXRs expect the main image data to be in an unnamed layer.
-				MultiLayerImageTask->LayerNames.FindOrAdd(PixelData.Get(), RenderPassData.Key.Name);
+				// as most programs that handle EXRs expect the main image data to be in an unnamed layer. We only postfix with cameraname
+				// if there's multiple cameras, as pipelines may be already be built around the generic "one camera" support.
+				UMoviePipelineExecutorShot* CurrentShot = GetPipeline()->GetActiveShotList()[ShotIndex];
+				UMoviePipelineCameraSetting* CameraSettings = GetPipeline()->FindOrAddSettingForShot<UMoviePipelineCameraSetting>(CurrentShot);
+				int32 NumCameras = CameraSettings->bRenderAllCameras ? CurrentShot->SidecarCameras.Num() : 1;
+				
+				FString CombinedName;
+				if (NumCameras == 1)
+				{
+					CombinedName = RenderPassData.Key.Name;
+				}
+				else
+				{
+					CombinedName = FString::Printf(TEXT("%s_%s"), *RenderPassData.Key.Name, *RenderPassData.Key.CameraName);
+				}
+				MultiLayerImageTask->LayerNames.FindOrAdd(PixelData.Get(), CombinedName);
 			}
 
 			MultiLayerImageTask->Width = Resolutions[Index].X;
diff --git a/Engine/Plugins/MovieScene/MovieRenderPipeline/Source/MovieRenderPipelineSettings/Private/MoviePipelineBurnInSetting.cpp b/Engine/Plugins/MovieScene/MovieRenderPipeline/Source/MovieRenderPipelineSettings/Private/MoviePipelineBurnInSetting.cpp
--- a/Engine/Plugins/MovieScene/MovieRenderPipeline/Source/MovieRenderPipelineSettings/Private/MoviePipelineBurnInSetting.cpp
+++ b/Engine/Plugins/MovieScene/MovieRenderPipeline/Source/MovieRenderPipelineSettings/Private/MoviePipelineBurnInSetting.cpp
@@ -6,6 +6,7 @@
 #include "MovieRenderPipelineDataTypes.h"
 #include "MoviePipelineBurnInWidget.h"
 #include "MoviePipelineOutputSetting.h"
+#include "MoviePipelineCameraSetting.h"
 #include "MoviePipelineMasterConfig.h"
 #include "MoviePipelineBlueprintLibrary.h"
 #include "MoviePipeline.h"
@@ -13,6 +14,7 @@
 #include "MoviePipelineOutputBuilder.h"
 #include "ImagePixelData.h"
 #include "MovieRenderPipelineCoreModule.h"
+#include "MoviePipelineQueue.h"
 
 FString UMoviePipelineBurnInSetting::DefaultBurnInWidgetAsset = TEXT("/MovieRenderPipeline/Blueprints/DefaultBurnIn.DefaultBurnIn_C");
 
@@ -20,7 +22,17 @@ void UMoviePipelineBurnInSetting::GatherOutputPassesImpl(TArray<FMoviePipelinePa
 {
 	if (BurnInClass.IsValid() && WidgetRenderer != nullptr)
 	{
-		ExpectedRenderPasses.Add(FMoviePipelinePassIdentifier(TEXT("BurnInOverlay")));
+		UMoviePipelineExecutorShot* CurrentShot = GetPipeline()->GetActiveShotList()[GetPipeline()->GetCurrentShotIndex()];
+		UMoviePipelineCameraSetting* CameraSettings = GetPipeline()->FindOrAddSettingForShot<UMoviePipelineCameraSetting>(CurrentShot);
+		int32 NumCameras = CameraSettings->bRenderAllCameras ? CurrentShot->SidecarCameras.Num() : 1;
+
+		for (int32 CameraIndex = 0; CameraIndex < NumCameras; CameraIndex++)
+		{
+			FMoviePipelinePassIdentifier PassIdentifierForCurrentCamera;
+			PassIdentifierForCurrentCamera.Name = TEXT("BurnInOverlay");
+			PassIdentifierForCurrentCamera.CameraName = CurrentShot->GetCameraName(CameraIndex);
+			ExpectedRenderPasses.Add(PassIdentifierForCurrentCamera);
+		}
 	}
 }
 
@@ -42,40 +54,54 @@ void UMoviePipelineBurnInSetting::RenderSample_GameThreadImpl(const FMoviePipeli
 
 	if (bFirstTile && bFirstSpatial && bFirstTemporal)
 	{
-		// Update the Widget with the latest frame information
-		BurnInWidgetInstance->OnOutputFrameStarted(GetPipeline());
+		UMoviePipelineExecutorShot* CurrentShot = GetPipeline()->GetActiveShotList()[GetPipeline()->GetCurrentShotIndex()];
+		UMoviePipelineCameraSetting* CameraSettings = GetPipeline()->FindOrAddSettingForShot<UMoviePipelineCameraSetting>(CurrentShot);
+		int32 NumCameras = CameraSettings->bRenderAllCameras ? CurrentShot->SidecarCameras.Num() : 1;
+		for (int32 CameraIndex = 0; CameraIndex < NumCameras; CameraIndex++)
+		{
+			FMoviePipelinePassIdentifier PassIdentifierForCurrentCamera;
+			PassIdentifierForCurrentCamera.Name = TEXT("BurnInOverlay");
+			PassIdentifierForCurrentCamera.CameraName = CurrentShot->GetCameraName(CameraIndex);
 
-		// Draw the widget to the render target
-		WidgetRenderer->DrawWindow(RenderTarget, VirtualWindow->GetHittestGrid(), VirtualWindow.ToSharedRef(), 1.f, OutputResolution, InSampleState.OutputState.TimeData.FrameDeltaTime);
+			UMoviePipelineBurnInWidget* CurrentWidget = BurnInWidgetInstances[CameraIndex];
+			// Put the widget in our window
+			VirtualWindow->SetContent(CurrentWidget->TakeWidget());
 
-		FRenderTarget* BackbufferRenderTarget = RenderTarget->GameThread_GetRenderTargetResource();
-		TSharedPtr<FMoviePipelineOutputMerger, ESPMode::ThreadSafe> OutputBuilder = GetPipeline()->OutputBuilder;
+			// Update the Widget with the latest frame information
+			CurrentWidget->OnOutputFrameStarted(GetPipeline());
 
-		ENQUEUE_RENDER_COMMAND(BurnInRenderTargetResolveCommand)(
-			[InSampleState, bComposite = bCompositeOntoFinalImage, BackbufferRenderTarget, OutputBuilder](FRHICommandListImmediate& RHICmdList)
-		{
-			FIntRect SourceRect = FIntRect(0, 0, BackbufferRenderTarget->GetSizeXY().X, BackbufferRenderTarget->GetSizeXY().Y);
+			// Draw the widget to the render target
+			WidgetRenderer->DrawWindow(RenderTarget, VirtualWindow->GetHittestGrid(), VirtualWindow.ToSharedRef(), 1.f, OutputResolution, InSampleState.OutputState.TimeData.FrameDeltaTime);
 
-			// Read the data back to the CPU
-			TArray<FColor> RawPixels;
-			RawPixels.SetNum(SourceRect.Width() * SourceRect.Height());
+			FRenderTarget* BackbufferRenderTarget = RenderTarget->GameThread_GetRenderTargetResource();
+			TSharedPtr<FMoviePipelineOutputMerger, ESPMode::ThreadSafe> OutputBuilder = GetPipeline()->OutputBuilder;
 
-			FReadSurfaceDataFlags ReadDataFlags(ERangeCompressionMode::RCM_MinMax);
-			ReadDataFlags.SetLinearToGamma(false);
+			ENQUEUE_RENDER_COMMAND(BurnInRenderTargetResolveCommand)(
+				[InSampleState, PassIdentifierForCurrentCamera, bComposite = bCompositeOntoFinalImage, BackbufferRenderTarget, OutputBuilder](FRHICommandListImmediate& RHICmdList)
+				{
+					FIntRect SourceRect = FIntRect(0, 0, BackbufferRenderTarget->GetSizeXY().X, BackbufferRenderTarget->GetSizeXY().Y);
 
-			RHICmdList.ReadSurfaceData(BackbufferRenderTarget->GetRenderTargetTexture(), SourceRect, RawPixels, ReadDataFlags);
+					// Read the data back to the CPU
+					TArray<FColor> RawPixels;
+					RawPixels.SetNum(SourceRect.Width() * SourceRect.Height());
 
-			TSharedRef<FImagePixelDataPayload, ESPMode::ThreadSafe> FrameData = MakeShared<FImagePixelDataPayload, ESPMode::ThreadSafe>();
-			FrameData->PassIdentifier = FMoviePipelinePassIdentifier(TEXT("BurnInOverlay"));
-			FrameData->SampleState = InSampleState;
-			FrameData->bRequireTransparentOutput = true;
-			FrameData->SortingOrder = 4;
-			FrameData->bCompositeToFinalImage = bComposite;
+					FReadSurfaceDataFlags ReadDataFlags(ERangeCompressionMode::RCM_MinMax);
+					ReadDataFlags.SetLinearToGamma(false);
 
-			TUniquePtr<FImagePixelData> PixelData = MakeUnique<TImagePixelData<FColor>>(InSampleState.BackbufferSize, TArray64<FColor>(MoveTemp(RawPixels)), FrameData);
+					RHICmdList.ReadSurfaceData(BackbufferRenderTarget->GetRenderTargetTexture(), SourceRect, RawPixels, ReadDataFlags);
 
-			OutputBuilder->OnCompleteRenderPassDataAvailable_AnyThread(MoveTemp(PixelData));
-		});
+					TSharedRef<FImagePixelDataPayload, ESPMode::ThreadSafe> FrameData = MakeShared<FImagePixelDataPayload, ESPMode::ThreadSafe>();
+					FrameData->PassIdentifier = PassIdentifierForCurrentCamera;
+					FrameData->SampleState = InSampleState;
+					FrameData->bRequireTransparentOutput = true;
+					FrameData->SortingOrder = 4;
+					FrameData->bCompositeToFinalImage = bComposite;
+
+					TUniquePtr<FImagePixelData> PixelData = MakeUnique<TImagePixelData<FColor>>(InSampleState.BackbufferSize, TArray64<FColor>(MoveTemp(RawPixels)), FrameData);
+
+					OutputBuilder->OnCompleteRenderPassDataAvailable_AnyThread(MoveTemp(PixelData));
+				});
+		}
 	}
 }
 
@@ -101,10 +127,15 @@ void UMoviePipelineBurnInSetting::SetupImpl(const MoviePipeline::FMoviePipelineR
 		return;
 	}
 
-	BurnInWidgetInstance = CreateWidget<UMoviePipelineBurnInWidget>(GetWorld(), BurnIn);
+	UMoviePipelineExecutorShot* CurrentShot = GetPipeline()->GetActiveShotList()[GetPipeline()->GetCurrentShotIndex()];
+	UMoviePipelineCameraSetting* CameraSettings = GetPipeline()->FindOrAddSettingForShot<UMoviePipelineCameraSetting>(CurrentShot);
+	int32 NumCameras = CameraSettings->bRenderAllCameras ? CurrentShot->SidecarCameras.Num() : 1;
+	for (int32 CameraIndex = 0; CameraIndex < NumCameras; CameraIndex++)
+	{
+		BurnInWidgetInstances.Add(CreateWidget<UMoviePipelineBurnInWidget>(GetWorld(), BurnIn));
+	}
 
 	VirtualWindow = SNew(SVirtualWindow).Size(FVector2D(OutputResolution.X, OutputResolution.Y));
-	VirtualWindow->SetContent(BurnInWidgetInstance->TakeWidget());
 
 	if (FSlateApplication::IsInitialized())
 	{
@@ -134,5 +165,5 @@ void UMoviePipelineBurnInSetting::TeardownImpl()
 	
 	WidgetRenderer = nullptr;
 	RenderTarget = nullptr;
-	BurnInWidgetInstance = nullptr;
+	BurnInWidgetInstances.Reset();
 }
\ No newline at end of file
diff --git a/Engine/Plugins/MovieScene/MovieRenderPipeline/Source/MovieRenderPipelineSettings/Public/MoviePipelineBurnInSetting.h b/Engine/Plugins/MovieScene/MovieRenderPipeline/Source/MovieRenderPipelineSettings/Public/MoviePipelineBurnInSetting.h
--- a/Engine/Plugins/MovieScene/MovieRenderPipeline/Source/MovieRenderPipelineSettings/Public/MoviePipelineBurnInSetting.h
+++ b/Engine/Plugins/MovieScene/MovieRenderPipeline/Source/MovieRenderPipelineSettings/Public/MoviePipelineBurnInSetting.h
@@ -57,7 +57,7 @@ private:
 	UTextureRenderTarget2D* RenderTarget;
 
 	UPROPERTY(Transient)
-	UMoviePipelineBurnInWidget* BurnInWidgetInstance;
+	TArray<UMoviePipelineBurnInWidget*> BurnInWidgetInstances;
 public:
 	static FString DefaultBurnInWidgetAsset;
 };
\ No newline at end of file
