MoviePipeline: Fixed a situation where high temporal sample counts could have their delta times clamped during Level Tick which created mismatches between calculated component velocities and render thread delta times.

diff --git a/Engine/Plugins/MovieScene/MovieRenderPipeline/Source/MovieRenderPipelineCore/Private/MoviePipeline.cpp b/Engine/Plugins/MovieScene/MovieRenderPipeline/Source/MovieRenderPipelineCore/Private/MoviePipeline.cpp
--- a/Engine/Plugins/MovieScene/MovieRenderPipeline/Source/MovieRenderPipelineCore/Private/MoviePipeline.cpp
+++ b/Engine/Plugins/MovieScene/MovieRenderPipeline/Source/MovieRenderPipelineCore/Private/MoviePipeline.cpp
@@ -789,6 +789,14 @@ bool UMoviePipelineCustomTimeStep::UpdateTimeStep(UEngine* /*InEngine*/)
 		FApp::UpdateLastTime();
 		FApp::SetDeltaTime(TimeCache.UndilatedDeltaTime);
 		FApp::SetCurrentTime(FApp::GetCurrentTime() + FApp::GetDeltaTime());
+
+		// The UWorldSettings can clamp the delta time inside the Level Tick function, this creates a mismatch between component
+		// velocity and render thread velocity and becomes an issue at high temporal sample counts. 
+		if (AWorldSettings* WorldSettings = GetWorld()->GetWorldSettings())
+		{
+			WorldSettings->MinUndilatedFrameTime = TimeCache.UndilatedDeltaTime;
+			WorldSettings->MaxUndilatedFrameTime = TimeCache.UndilatedDeltaTime;
+		}
 	}
 
 	// Clear our cached time to ensure we're always explicitly telling this what to do and never relying on the last set value.
