MoviePipeline: Temporary fix for a crash when a camera cut track pointed to a camera in a sub-sequence.

diff --git a/Engine/Plugins/MovieScene/MovieRenderPipeline/Source/MovieRenderPipelineCore/Private/MoviePipelineBlueprintLibrary.cpp b/Engine/Plugins/MovieScene/MovieRenderPipeline/Source/MovieRenderPipelineCore/Private/MoviePipelineBlueprintLibrary.cpp
--- a/Engine/Plugins/MovieScene/MovieRenderPipeline/Source/MovieRenderPipelineCore/Private/MoviePipelineBlueprintLibrary.cpp
+++ b/Engine/Plugins/MovieScene/MovieRenderPipeline/Source/MovieRenderPipelineCore/Private/MoviePipelineBlueprintLibrary.cpp
@@ -601,9 +601,19 @@ void UMoviePipelineBlueprintLibrary::UpdateJobShotListFromSequence(ULevelSequenc
 					}
 				}
 
+				if (MainCameraIndex == INDEX_NONE)
+				{
+					// When the camera cut points to a camera in another Sequence it won't get detected by the above loop,
+					// so we manually add it as the only sidecar camera.
+					FMoviePipelineSidecarCamera& SidecarCamera = Entity.SidecarCameras.AddDefaulted_GetRef();
+					SidecarCamera.BindingId = LeafNode->CameraCutSection->GetCameraBindingID().GetGuid();
+					SidecarCamera.SequenceId = LeafNode->CameraCutSection->GetCameraBindingID().GetRelativeSequenceID();
+					SidecarCamera.Name = TEXT("main_camera");
+					MainCameraIndex = Entity.SidecarCameras.Num() - 1;
+				}
+
 				// Ensure the "main" camera (specified by the Camera Cut Track) is always index 0, so that it's rendered
 				// first, and the stable sort (on output files) still tries to keep it as the primary layer.
-				ensureAlwaysMsgf(MainCameraIndex != INDEX_NONE, TEXT("Main Camera Binding not found in camera array!"));
 				Entity.SidecarCameras.Swap(MainCameraIndex, 0);
 			}
 
