MoviePipeline: Fixed CineCamera FocalLength not being respected - fallout from the multi-camera support. FMinimalViewInfo::DesiredFoV is only assigned in ULocalPlayer::GetViewPoint so when querying the Camera Components directly for their FMinimalViewInfo we instead need to read from FMinimalViewInfo::FoV (which the CineCameraComponent writes the Focal Length->FoV into).

diff --git a/Engine/Plugins/MovieScene/MovieRenderPipeline/Source/MovieRenderPipelineRenderPasses/Private/MoviePipelineImagePassBase.cpp b/Engine/Plugins/MovieScene/MovieRenderPipeline/Source/MovieRenderPipelineRenderPasses/Private/MoviePipelineImagePassBase.cpp
--- a/Engine/Plugins/MovieScene/MovieRenderPipeline/Source/MovieRenderPipelineRenderPasses/Private/MoviePipelineImagePassBase.cpp
+++ b/Engine/Plugins/MovieScene/MovieRenderPipeline/Source/MovieRenderPipelineRenderPasses/Private/MoviePipelineImagePassBase.cpp
@@ -305,7 +305,7 @@ FSceneView* UMoviePipelineImagePassBase::GetSceneViewForSampleState(FSceneViewFa
 		FPlane(1, 0, 0, 0),
 		FPlane(0, 1, 0, 0),
 		FPlane(0, 0, 0, 1));
-	float ViewFOV = CameraInfo.ViewInfo.DesiredFOV;
+	float ViewFOV = CameraInfo.ViewInfo.FOV;
 
 	// Inflate our FOV to support the overscan 
 	ViewFOV = 2.0f * FMath::RadiansToDegrees(FMath::Atan((1.0f + InOutSampleState.OverscanPercentage) * FMath::Tan(FMath::DegreesToRadians( ViewFOV * 0.5f ))));
@@ -422,6 +422,7 @@ FSceneView* UMoviePipelineImagePassBase::GetSceneViewForSampleState(FSceneViewFa
 
 	ViewInitOptions.SceneViewStateInterface = GetSceneViewStateInterface(OptPayload);
 	ViewInitOptions.FOV = ViewFOV;
+	ViewInitOptions.DesiredFOV = ViewFOV;
 
 	FSceneView* View = new FSceneView(ViewInitOptions);
 	ViewFamily->Views.Add(View);
