Added a way to hide any UI-related elements in MoviePipelinePIEExecutor.

diff --git a/Engine/Plugins/MovieScene/MovieRenderPipeline/Source/MovieRenderPipelineEditor/Private/MoviePipelinePIEExecutor.cpp b/Engine/Plugins/MovieScene/MovieRenderPipeline/Source/MovieRenderPipelineEditor/Private/MoviePipelinePIEExecutor.cpp
--- a/Engine/Plugins/MovieScene/MovieRenderPipeline/Source/MovieRenderPipelineEditor/Private/MoviePipelinePIEExecutor.cpp
+++ b/Engine/Plugins/MovieScene/MovieRenderPipeline/Source/MovieRenderPipelineEditor/Private/MoviePipelinePIEExecutor.cpp
@@ -58,6 +58,19 @@ void UMoviePipelinePIEExecutor::FValidationMessageGatherer::Serialize(const TCHA
 	}
 }
 
+UMoviePipelinePIEExecutor::UMoviePipelinePIEExecutor()
+	: UMoviePipelineLinearExecutorBase()
+	, bRenderOffscreen(false)
+	, RemainingInitializationFrames(-1)
+	, bPreviousUseFixedTimeStep(false)
+	, PreviousFixedTimeStepDelta(1 / 30.0)
+{
+	if (!FApp::CanEverRender() || FSlateApplication::Get().IsRenderingOffScreen())
+	{
+		SetIsRenderingOffscreen(true);
+	}
+}
+
 void UMoviePipelinePIEExecutor::Start(const UMoviePipelineExecutorJob* InJob)
 {
 	Super::Start(InJob);
@@ -66,8 +79,11 @@ void UMoviePipelinePIEExecutor::Start(const UMoviePipelineExecutorJob* InJob)
 	ULevelSequence* LevelSequence = Cast<ULevelSequence>(InJob->Sequence.TryLoad());
 	if (!LevelSequence)
 	{
-		FText FailureReason = LOCTEXT("InvalidSequenceFailureDialog", "One or more jobs in the queue has an invalid/null sequence.");
-		FMessageDialog::Open(EAppMsgType::Ok, FailureReason);
+		if(!IsRenderingOffscreen())
+		{
+			FText FailureReason = LOCTEXT("InvalidSequenceFailureDialog", "One or more jobs in the queue has an invalid/null sequence.");
+			FMessageDialog::Open(EAppMsgType::Ok, FailureReason);
+		}
 
 		OnExecutorFinishedImpl();
 		return;
@@ -79,8 +95,11 @@ void UMoviePipelinePIEExecutor::Start(const UMoviePipelineExecutorJob* InJob)
 	const bool bAllMapsValid = UMoviePipelineEditorBlueprintLibrary::IsMapValidForRemoteRender(Queue->GetJobs());
 	if (!bAllMapsValid)
 	{
-		FText FailureReason = LOCTEXT("UnsavedMapFailureDialog", "One or more jobs in the queue have an unsaved map as their target map. Maps must be saved at least once before rendering.");
-		FMessageDialog::Open(EAppMsgType::Ok, FailureReason);
+		if(!IsRenderingOffscreen())
+		{
+			FText FailureReason = LOCTEXT("UnsavedMapFailureDialog", "One or more jobs in the queue have an unsaved map as their target map. Maps must be saved at least once before rendering.");
+			FMessageDialog::Open(EAppMsgType::Ok, FailureReason);
+		}
 
 		OnExecutorFinishedImpl();
 		return;
@@ -102,7 +121,7 @@ void UMoviePipelinePIEExecutor::Start(const UMoviePipelineExecutorJob* InJob)
 		.SizingRule(ESizingRule::UserSized);
 
 	WeakCustomWindow = CustomWindow;
-	FSlateApplication::Get().AddWindow(CustomWindow);
+	FSlateApplication::Get().AddWindow(CustomWindow, !IsRenderingOffscreen());
 
 	// Initialize our own copy of the Editor Play settings which we will adjust defaults on.
 	ULevelEditorPlaySettings* PlayInEditorSettings = NewObject<ULevelEditorPlaySettings>();
diff --git a/Engine/Plugins/MovieScene/MovieRenderPipeline/Source/MovieRenderPipelineEditor/Public/MoviePipelinePIEExecutor.h b/Engine/Plugins/MovieScene/MovieRenderPipeline/Source/MovieRenderPipelineEditor/Public/MoviePipelinePIEExecutor.h
--- a/Engine/Plugins/MovieScene/MovieRenderPipeline/Source/MovieRenderPipelineEditor/Public/MoviePipelinePIEExecutor.h
+++ b/Engine/Plugins/MovieScene/MovieRenderPipeline/Source/MovieRenderPipelineEditor/Public/MoviePipelinePIEExecutor.h
@@ -29,13 +29,7 @@ class MOVIERENDERPIPELINEEDITOR_API UMoviePipelinePIEExecutor : public UMoviePip
 	GENERATED_BODY()
 	
 public:
-	UMoviePipelinePIEExecutor()
-		: UMoviePipelineLinearExecutorBase()
-		, RemainingInitializationFrames(-1)
-		, bPreviousUseFixedTimeStep(false)
-		, PreviousFixedTimeStepDelta(1/30.0)
-	{
-	}
+	UMoviePipelinePIEExecutor();
 
 public:
 	/** Deprecated. Use OnIndividualJobWorkFinished instead. */
@@ -66,6 +60,14 @@ public:
 	UFUNCTION(BlueprintCallable, Category = "Movie Render Pipeline")
 	void SetInitializationTime(const FDateTime& InInitializationTime) { CustomInitializationTime = InInitializationTime; }
 
+	/** Should it render without any UI elements showing up (such as the rendering progress window)? */
+	UFUNCTION(BlueprintCallable, Category = "Movie Render Pipeline")
+	void SetIsRenderingOffscreen(const bool bInRenderOffscreen) { bRenderOffscreen = bInRenderOffscreen; }
+
+	/** Will it render without any UI elements showing up (such as the rendering progress window)? */
+	UFUNCTION(BlueprintPure, Category = "Movie Render Pipeline")
+	bool IsRenderingOffscreen() const { return bRenderOffscreen; }
+
 protected:
 	virtual void Start(const UMoviePipelineExecutorJob* InJob) override;
 
@@ -95,6 +97,7 @@ private:
 	/** Called a short period of time after OnPIEMoviePipelineFinished to allow Editor the time to fully close PIE before we make a new request. */
 	void DelayedFinishNotification();
 private:
+	bool bRenderOffscreen;
 	/** If using delayed initialization, how many frames are left before we call Initialize. Will be -1 if not actively counting down. */
 	int32 RemainingInitializationFrames;
 	bool bPreviousUseFixedTimeStep;
