Fix crash in UMoviePipelinePIEExecutor() when converting a level to World Partition. It was introduced by 20590381.

#rnx

diff --git a/Engine/Plugins/MovieScene/MovieRenderPipeline/Source/MovieRenderPipelineEditor/Private/MoviePipelinePIEExecutor.cpp b/Engine/Plugins/MovieScene/MovieRenderPipeline/Source/MovieRenderPipelineEditor/Private/MoviePipelinePIEExecutor.cpp
--- a/Engine/Plugins/MovieScene/MovieRenderPipeline/Source/MovieRenderPipelineEditor/Private/MoviePipelinePIEExecutor.cpp
+++ b/Engine/Plugins/MovieScene/MovieRenderPipeline/Source/MovieRenderPipelineEditor/Private/MoviePipelinePIEExecutor.cpp
@@ -65,7 +65,7 @@ UMoviePipelinePIEExecutor::UMoviePipelinePIEExecutor()
 	, bPreviousUseFixedTimeStep(false)
 	, PreviousFixedTimeStepDelta(1 / 30.0)
 {
-	if (!FApp::CanEverRender() || FSlateApplication::Get().IsRenderingOffScreen())
+	if (!FApp::CanEverRender() || (FSlateApplication::IsInitialized() && FSlateApplication::Get().IsRenderingOffScreen()))
 	{
 		SetIsRenderingOffscreen(true);
 	}
