MoviePipeline: Fixed a crash when adding a level sequence to the UI that contained possessables with unloaded crashes.

diff --git a/Engine/Plugins/MovieScene/MovieRenderPipeline/Source/MovieRenderPipelineCore/Private/MoviePipelineBlueprintLibrary.cpp b/Engine/Plugins/MovieScene/MovieRenderPipeline/Source/MovieRenderPipelineCore/Private/MoviePipelineBlueprintLibrary.cpp
--- a/Engine/Plugins/MovieScene/MovieRenderPipeline/Source/MovieRenderPipelineCore/Private/MoviePipelineBlueprintLibrary.cpp
+++ b/Engine/Plugins/MovieScene/MovieRenderPipeline/Source/MovieRenderPipelineCore/Private/MoviePipelineBlueprintLibrary.cpp
@@ -557,6 +557,14 @@ void UMoviePipelineBlueprintLibrary::UpdateJobShotListFromSequence(ULevelSequenc
 					// This is only available in the editor, so in the editor we'll auto-detect the class
 					// but in runtime builds they'll have to manually tag it with "Camera" :(
 					const UClass* PossessedClass = Possessable.GetPossessedObjectClass();
+
+					// A possessable might point to a unloaded class, but we can safely skip it
+					// because we only need to look for UCameraComponents and those are core engine so
+					// they will always be loaded.
+					if (!PossessedClass)
+					{
+						continue;
+					}
 					if (PossessedClass->IsChildOf<UCameraComponent>())
 					{
 						ValidBinding = Possessable.GetGuid();
