MoviePipeline: Fixed a few multi-camera issues including superfluous warnings about filename formats being printed and incorrect post process blending for Cine Cameras.

diff --git a/Engine/Plugins/MovieScene/MovieRenderPipeline/Source/MovieRenderPipelineCore/Public/MovieRenderPipelineDataTypes.h b/Engine/Plugins/MovieScene/MovieRenderPipeline/Source/MovieRenderPipelineCore/Public/MovieRenderPipelineDataTypes.h
--- a/Engine/Plugins/MovieScene/MovieRenderPipeline/Source/MovieRenderPipelineCore/Public/MovieRenderPipelineDataTypes.h
+++ b/Engine/Plugins/MovieScene/MovieRenderPipeline/Source/MovieRenderPipelineCore/Public/MovieRenderPipelineDataTypes.h
@@ -996,6 +996,14 @@ namespace MoviePipeline
 		ERHIFeatureLevel::Type FeatureLevel;
 	};
 
+	struct FCompositePassInfo
+	{
+		FCompositePassInfo() {}
+
+		FMoviePipelinePassIdentifier PassIdentifier;
+		TUniquePtr<FImagePixelData> PixelData;
+	};
+
 }
 
 struct FImagePixelDataPayload : IImagePixelDataPayload, public TSharedFromThis<FImagePixelDataPayload, ESPMode::ThreadSafe>
@@ -1117,6 +1125,23 @@ public:
 		return CameraNameUseCounts.Num() > 1;
 	}
 
+	bool HasDataFromMultipleRenderPasses(const TArray<MoviePipeline::FCompositePassInfo>& InCompositedPasses) const
+	{
+		TMap<FString, int32> RenderPassUseCounts;
+		for (const FMoviePipelinePassIdentifier& PassIdentifier : ExpectedRenderPasses)
+		{
+			RenderPassUseCounts.FindOrAdd(PassIdentifier.Name) += 1;
+		}
+
+		// Remove any render passes that will be composited on later
+		for (const MoviePipeline::FCompositePassInfo& CompositePass : InCompositedPasses)
+		{
+			RenderPassUseCounts.Remove(CompositePass.PassIdentifier.Name);
+		}
+
+		return RenderPassUseCounts.Num() > 1;
+	}
+
 private:
 	// Explicitly delete copy operators since we own unique data.
 	// void operator=(const FMoviePipelineMergerOutputFrame&);
@@ -1185,14 +1210,6 @@ namespace MoviePipeline
 		/** An array of active submixes we are recording for this shot. Gets cleared when recording stops on a shot. */
 		TArray<TWeakPtr<Audio::FMixerSubmix, ESPMode::ThreadSafe>> ActiveSubmixes;
 	};
-
-	struct FCompositePassInfo
-	{
-		FCompositePassInfo() {}
-
-		FMoviePipelinePassIdentifier PassIdentifier;
-		TUniquePtr<FImagePixelData> PixelData;
-	};
 }
 
 USTRUCT(BlueprintType)
diff --git a/Engine/Plugins/MovieScene/MovieRenderPipeline/Source/MovieRenderPipelineRenderPasses/Private/MoviePipelineDeferredPasses.cpp b/Engine/Plugins/MovieScene/MovieRenderPipeline/Source/MovieRenderPipelineRenderPasses/Private/MoviePipelineDeferredPasses.cpp
--- a/Engine/Plugins/MovieScene/MovieRenderPipeline/Source/MovieRenderPipelineRenderPasses/Private/MoviePipelineDeferredPasses.cpp
+++ b/Engine/Plugins/MovieScene/MovieRenderPipeline/Source/MovieRenderPipelineRenderPasses/Private/MoviePipelineDeferredPasses.cpp
@@ -805,7 +805,7 @@ void UMoviePipelineDeferredPassBase::BlendPostProcessSettings(FSceneView* InView
 	if (NumCameras == 1)
 	{
 		// If there's only one camera being used we can use the parent class which assumes the camera comes from the PlayerCameraManager
-		Super::GetCameraInfo(InOutSampleState, OptPayload);
+		Super::BlendPostProcessSettings(InView, InOutSampleState, OptPayload);
 	}
 	else
 	{
@@ -842,7 +842,7 @@ void UMoviePipelineDeferredPassBase::BlendPostProcessSettings(FSceneView* InView
 		}
 
 		// After blending all post processing volumes, blend the camera's post process settings too
-		InView->OverridePostProcessSettings(OutCamera->PostProcessSettings, OutCamera->PostProcessBlendWeight);
+		InView->OverridePostProcessSettings(OutViewInfo.PostProcessSettings, OutViewInfo.PostProcessBlendWeight);
 	}
 }
 
diff --git a/Engine/Plugins/MovieScene/MovieRenderPipeline/Source/MovieRenderPipelineRenderPasses/Private/MoviePipelineImageSequenceOutput.cpp b/Engine/Plugins/MovieScene/MovieRenderPipeline/Source/MovieRenderPipelineRenderPasses/Private/MoviePipelineImageSequenceOutput.cpp
--- a/Engine/Plugins/MovieScene/MovieRenderPipeline/Source/MovieRenderPipelineRenderPasses/Private/MoviePipelineImageSequenceOutput.cpp
+++ b/Engine/Plugins/MovieScene/MovieRenderPipeline/Source/MovieRenderPipelineRenderPasses/Private/MoviePipelineImageSequenceOutput.cpp
@@ -149,16 +149,15 @@ void UMoviePipelineImageSequenceOutputBase::OnReceiveImageDataImpl(FMoviePipelin
 		
 		FXMLData XMLData;
 		{
-			FString FileNameFormatString = OutputSettings->FileNameFormat;
+			FString FileNameFormatString = OutputDirectory / OutputSettings->FileNameFormat;
 
 			// If we're writing more than one render pass out, we need to ensure the file name has the format string in it so we don't
 			// overwrite the same file multiple times. Burn In overlays don't count if they are getting composited on top of an existing file.
-			const bool bIncludeRenderPass = InMergedOutputFrame->ImageOutputData.Num() - CompositedPasses.Num() > 1;
+			const bool bIncludeRenderPass = InMergedOutputFrame->HasDataFromMultipleRenderPasses(CompositedPasses);
 			const bool bIncludeCameraName = InMergedOutputFrame->HasDataFromMultipleCameras();
 			const bool bTestFrameNumber = true;
 
-			// ToDo: This validation may need to include the OutputDirectory, because they could have put {camera_name} in the output folder
-			UE::MoviePipeline::ValidateOutputFormatString(FileNameFormatString, bIncludeRenderPass, bTestFrameNumber, bIncludeCameraName);
+			UE::MoviePipeline::ValidateOutputFormatString(/*InOut*/ FileNameFormatString, bIncludeRenderPass, bTestFrameNumber, bIncludeCameraName);
 
 			// Create specific data that needs to override 
 			TMap<FString, FString> FormatOverrides;
@@ -168,13 +167,12 @@ void UMoviePipelineImageSequenceOutputBase::OnReceiveImageDataImpl(FMoviePipelin
 
 			// Resolve for XMLs
 			{
-				GetPipeline()->ResolveFilenameFormatArguments(FileNameFormatString, FormatOverrides, XMLData.ImageSequenceFileName, FinalFormatArgs, &Payload->SampleState.OutputState, -Payload->SampleState.OutputState.ShotOutputFrameNumber);
+				GetPipeline()->ResolveFilenameFormatArguments(/*In*/ FileNameFormatString, FormatOverrides, /*Out*/ XMLData.ImageSequenceFileName, FinalFormatArgs, &Payload->SampleState.OutputState, -Payload->SampleState.OutputState.ShotOutputFrameNumber);
 			}
 			
 			// Resolve the final absolute file path to write this to
 			{
-				FString FormatString = OutputDirectory / FileNameFormatString;
-				GetPipeline()->ResolveFilenameFormatArguments(FormatString, FormatOverrides, OutputData.FilePath, FinalFormatArgs, &Payload->SampleState.OutputState);
+				GetPipeline()->ResolveFilenameFormatArguments(FileNameFormatString, FormatOverrides, OutputData.FilePath, FinalFormatArgs, &Payload->SampleState.OutputState);
 
 				if (FPaths::IsRelative(OutputData.FilePath))
 				{
@@ -223,10 +221,16 @@ void UMoviePipelineImageSequenceOutputBase::OnReceiveImageDataImpl(FMoviePipelin
 
 
 		// We composite before flipping the alpha so that it is consistent for all formats.
-		if (RenderPassData.Key == FMoviePipelinePassIdentifier(TEXT("FinalImage")))
+		if (RenderPassData.Key.Name == TEXT("FinalImage"))
 		{
 			for (const MoviePipeline::FCompositePassInfo& CompositePass : CompositedPasses)
 			{
+				// Match them up by camera name so multiple passes intended for different camera names work.
+				if (RenderPassData.Key.CameraName != CompositePass.PassIdentifier.CameraName)
+				{
+					continue;
+				}
+
 				// We don't need to copy the data here (even though it's being passed to a async system) because we already made a unique copy of the
 				// burn in/widget data when we decided to composite it.
 				switch (QuantizedPixelType)
