MovieScene: Fix FSlateApplication invalid access in commandlet

diff --git a/Engine/Plugins/MovieScene/MovieRenderPipeline/Source/MovieRenderPipelineEditor/Private/MoviePipelinePIEExecutor.cpp b/Engine/Plugins/MovieScene/MovieRenderPipeline/Source/MovieRenderPipelineEditor/Private/MoviePipelinePIEExecutor.cpp
--- a/Engine/Plugins/MovieScene/MovieRenderPipeline/Source/MovieRenderPipelineEditor/Private/MoviePipelinePIEExecutor.cpp
+++ b/Engine/Plugins/MovieScene/MovieRenderPipeline/Source/MovieRenderPipelineEditor/Private/MoviePipelinePIEExecutor.cpp
@@ -65,9 +65,12 @@ UMoviePipelinePIEExecutor::UMoviePipelinePIEExecutor()
 	, bPreviousUseFixedTimeStep(false)
 	, PreviousFixedTimeStepDelta(1 / 30.0)
 {
-	if (!FApp::CanEverRender() || (FSlateApplication::IsInitialized() && FSlateApplication::Get().IsRenderingOffScreen()))
+	if (!IsTemplate() && FSlateApplication::IsInitialized())
 	{
-		SetIsRenderingOffscreen(true);
+		if (!FApp::CanEverRender() || FSlateApplication::Get().IsRenderingOffScreen())
+		{
+			SetIsRenderingOffscreen(true);
+		}
 	}
 }
 
