MoviePipeline: Orthographic Cameras now warn when used outside the path tracer and support high-res tiling and camera overscan.

diff --git a/Engine/Plugins/MovieScene/MovieRenderPipeline/Source/MovieRenderPipelineRenderPasses/Private/MoviePipelineImagePassBase.cpp b/Engine/Plugins/MovieScene/MovieRenderPipeline/Source/MovieRenderPipelineRenderPasses/Private/MoviePipelineImagePassBase.cpp
--- a/Engine/Plugins/MovieScene/MovieRenderPipeline/Source/MovieRenderPipelineRenderPasses/Private/MoviePipelineImagePassBase.cpp
+++ b/Engine/Plugins/MovieScene/MovieRenderPipeline/Source/MovieRenderPipelineRenderPasses/Private/MoviePipelineImagePassBase.cpp
@@ -158,6 +158,19 @@ TSharedPtr<FSceneViewFamilyContext> UMoviePipelineImagePassBase::CalculateViewFa
 		}
 	}
 
+	// Orthographic cameras don't support anti-aliasing outside the path tracer (other than FXAA)
+	const bool bIsOrthographicCamera = !View->IsPerspectiveProjection();
+	if (bIsOrthographicCamera)
+	{
+		bool bIsSupportedAAMethod = View->AntiAliasingMethod == EAntiAliasingMethod::AAM_FXAA;
+		bool bIsPathTracer = OutViewFamily->EngineShowFlags.PathTracing;
+		bool bWarnJitters = InOutSampleState.ProjectionMatrixJitterAmount.SquaredLength() > SMALL_NUMBER;
+		if ((!bIsPathTracer && !bIsSupportedAAMethod) || bWarnJitters)
+		{
+			UE_LOG(LogMovieRenderPipeline, Warning, TEXT("Orthographic Cameras are only supported with PathTracer or Deferred with FXAA Anti-Aliasing"));
+		}
+	}
+
 	OutViewFamily->ViewExtensions = GEngine->ViewExtensions->GatherActiveExtensions(FSceneViewExtensionContext(GetWorld()->Scene));
 
 	AddViewExtensions(*OutViewFamily, InOutSampleState);
@@ -176,7 +189,10 @@ TSharedPtr<FSceneViewFamilyContext> UMoviePipelineImagePassBase::CalculateViewFa
 	{
 		// If we're not using Temporal Anti-Aliasing or Path Tracing we will apply the View Matrix projection jitter. Normally TAA sets this
 		// inside FSceneRenderer::PreVisibilityFrameSetup. Path Tracing does its own anti-aliasing internally.
-		if (!IsTemporalAccumulationBasedMethod(View->AntiAliasingMethod) && !OutViewFamily->EngineShowFlags.PathTracing)
+		bool bApplyProjectionJitter = !bIsOrthographicCamera
+									&& !OutViewFamily->EngineShowFlags.PathTracing 
+									&& !IsTemporalAccumulationBasedMethod(View->AntiAliasingMethod);
+		if (bApplyProjectionJitter)
 		{
 			View->ViewMatrices.HackAddTemporalAAProjectionJitter(InOutSampleState.ProjectionMatrixJitterAmount);
 		}
@@ -319,9 +335,10 @@ FSceneView* UMoviePipelineImagePassBase::GetSceneViewForSampleState(FSceneViewFa
 		if (CameraInfo.ViewInfo.ProjectionMode == ECameraProjectionMode::Orthographic)
 		{
 			const float YScale = 1.0f / CameraInfo.ViewInfo.AspectRatio;
+			const float OverscanScale = 1.0f + InOutSampleState.OverscanPercentage;
 
-			const float HalfOrthoWidth = CameraInfo.ViewInfo.OrthoWidth / 2.0f;
-			const float ScaledOrthoHeight = CameraInfo.ViewInfo.OrthoWidth / 2.0f * YScale;
+			const float HalfOrthoWidth = (CameraInfo.ViewInfo.OrthoWidth / 2.0f) * OverscanScale;
+			const float ScaledOrthoHeight = (CameraInfo.ViewInfo.OrthoWidth / 2.0f) * OverscanScale * YScale;
 
 			const float NearPlane = CameraInfo.ViewInfo.OrthoNearClipPlane;
 			const float FarPlane = CameraInfo.ViewInfo.OrthoFarClipPlane;
@@ -335,6 +352,11 @@ FSceneView* UMoviePipelineImagePassBase::GetSceneViewForSampleState(FSceneViewFa
 				ZScale,
 				ZOffset
 			);
+			ViewInitOptions.bUseFauxOrthoViewPos = true;
+			
+			// Modify the projection matrix to do an off center projection, with overlap for high-res tiling
+			const bool bOrthographic = true;
+			ModifyProjectionMatrixForTiling(InOutSampleState, bOrthographic, /*InOut*/ BaseProjMatrix, DofSensorScale);
 		}
 		else
 		{
@@ -411,13 +433,15 @@ FSceneView* UMoviePipelineImagePassBase::GetSceneViewForSampleState(FSceneViewFa
 			}
 
 			// Modify the perspective matrix to do an off center projection, with overlap for high-res tiling
-			ModifyProjectionMatrixForTiling(InOutSampleState, /*InOut*/ BaseProjMatrix, DofSensorScale);
+			const bool bOrthographic = false;
+			ModifyProjectionMatrixForTiling(InOutSampleState, bOrthographic, /*InOut*/ BaseProjMatrix, DofSensorScale);
 			// ToDo: Does orthographic support tiling in the same way or do I need to modify the values before creating the ortho view.
 		}
 		
 		// BaseProjMatrix may be perspective or orthographic.
 		ViewInitOptions.ProjectionMatrix = BaseProjMatrix;
 		
+		ViewInitOptions.ProjectionMatrix.DebugPrint();
 	}
 
 	ViewInitOptions.SceneViewStateInterface = GetSceneViewStateInterface(OptPayload);
@@ -516,7 +540,7 @@ FVector4 UMoviePipelineImagePassBase::CalculatePrinciplePointOffsetForTiling(con
 	return FVector4(TilePrincipalPointOffset.X, -TilePrincipalPointOffset.Y, TilePrincipalPointScale.X, TilePrincipalPointScale.Y);
 }
 
-void UMoviePipelineImagePassBase::ModifyProjectionMatrixForTiling(const FMoviePipelineRenderPassMetrics& InSampleState, FMatrix& InOutProjectionMatrix, float& OutDoFSensorScale) const
+void UMoviePipelineImagePassBase::ModifyProjectionMatrixForTiling(const FMoviePipelineRenderPassMetrics& InSampleState, const bool bInOrthographic, FMatrix& InOutProjectionMatrix, float& OutDoFSensorScale) const
 {
 	float PadRatioX = 1.0f;
 	float PadRatioY = 1.0f;
@@ -538,8 +562,16 @@ void UMoviePipelineImagePassBase::ModifyProjectionMatrixForTiling(const FMoviePi
 	float OffsetX = -((float(InSampleState.TileIndexes.X) + 0.5f - float(InSampleState.TileCounts.X) / 2.0f) * 2.0f);
 	float OffsetY = ((float(InSampleState.TileIndexes.Y) + 0.5f - float(InSampleState.TileCounts.Y) / 2.0f) * 2.0f);
 
-	InOutProjectionMatrix.M[2][0] += OffsetX / PadRatioX;
-	InOutProjectionMatrix.M[2][1] += OffsetY / PadRatioX;
+	if (bInOrthographic)
+	{
+		InOutProjectionMatrix.M[3][0] += OffsetX / PadRatioX;
+		InOutProjectionMatrix.M[3][1] += OffsetY / PadRatioX;
+	}
+	else
+	{
+		InOutProjectionMatrix.M[2][0] += OffsetX / PadRatioX;
+		InOutProjectionMatrix.M[2][1] += OffsetY / PadRatioX;
+	}
 }
 
 UE::MoviePipeline::FImagePassCameraViewData UMoviePipelineImagePassBase::GetCameraInfo(FMoviePipelineRenderPassMetrics& InOutSampleState, IViewCalcPayload* OptPayload) const
diff --git a/Engine/Plugins/MovieScene/MovieRenderPipeline/Source/MovieRenderPipelineRenderPasses/Public/MoviePipelineImagePassBase.h b/Engine/Plugins/MovieScene/MovieRenderPipeline/Source/MovieRenderPipelineRenderPasses/Public/MoviePipelineImagePassBase.h
--- a/Engine/Plugins/MovieScene/MovieRenderPipeline/Source/MovieRenderPipelineRenderPasses/Public/MoviePipelineImagePassBase.h
+++ b/Engine/Plugins/MovieScene/MovieRenderPipeline/Source/MovieRenderPipelineRenderPasses/Public/MoviePipelineImagePassBase.h
@@ -107,8 +107,7 @@ protected:
 	// ~FGCObject Interface
 
 	FVector4 CalculatePrinciplePointOffsetForTiling(const FMoviePipelineRenderPassMetrics& InSampleState) const;
-	void ModifyProjectionMatrixForTiling(const FMoviePipelineRenderPassMetrics& InSampleState, FMatrix& InOutProjectionMatrix, float& OutDoFSensorScale) const;
-
+	void ModifyProjectionMatrixForTiling(const FMoviePipelineRenderPassMetrics& InSampleState, const bool bInOrthographic, FMatrix& InOutProjectionMatrix, float& OutDoFSensorScale) const;
 
 protected:
 	virtual void GetViewShowFlags(FEngineShowFlags& OutShowFlag, EViewModeIndex& OutViewModeIndex) const;
