MoviePipeline: Clarified header comment required for OnIndividualShotWorkFinished callback.

diff --git a/Engine/Plugins/MovieScene/MovieRenderPipeline/Source/MovieRenderPipelineEditor/Public/MoviePipelinePIEExecutor.h b/Engine/Plugins/MovieScene/MovieRenderPipeline/Source/MovieRenderPipelineEditor/Public/MoviePipelinePIEExecutor.h
--- a/Engine/Plugins/MovieScene/MovieRenderPipeline/Source/MovieRenderPipelineEditor/Public/MoviePipelinePIEExecutor.h
+++ b/Engine/Plugins/MovieScene/MovieRenderPipeline/Source/MovieRenderPipelineEditor/Public/MoviePipelinePIEExecutor.h
@@ -45,7 +45,7 @@ public:
 		return OnIndividualJobWorkFinishedDelegateNative;
 	}
 
-	/** Native C++ event to listen to for when an individual shot has been finished. Only called if the UMoviePipeline is set up correctly, see its headers for details. */
+	/** Native C++ event to listen to for when an individual shot has been finished. Only called if the UMoviePipeline is set up correctly, see its headers for details. This usually means setting bFlushDiskWritesPerShot to true in the UMoviePipelineOutputSetting for each job before rendering. */
 	FMoviePipelineWorkFinishedNative& OnIndividualShotWorkFinished()
 	{
 		return OnIndividualShotWorkFinishedDelegateNative;
