MoviePipeline: Fixed a crash caused by trying to use TSR with SM5 which is currently not supported.

diff --git a/Engine/Plugins/MovieScene/MovieRenderPipeline/Source/MovieRenderPipelineRenderPasses/Private/MoviePipelineImagePassBase.cpp b/Engine/Plugins/MovieScene/MovieRenderPipeline/Source/MovieRenderPipelineRenderPasses/Private/MoviePipelineImagePassBase.cpp
--- a/Engine/Plugins/MovieScene/MovieRenderPipeline/Source/MovieRenderPipelineRenderPasses/Private/MoviePipelineImagePassBase.cpp
+++ b/Engine/Plugins/MovieScene/MovieRenderPipeline/Source/MovieRenderPipelineRenderPasses/Private/MoviePipelineImagePassBase.cpp
@@ -185,6 +185,29 @@ TSharedPtr<FSceneViewFamilyContext> UMoviePipelineImagePassBase::CalculateViewFa
 		OutViewFamily->ViewExtensions[ViewExt]->SetupView(*OutViewFamily.Get(), *View);
 	}
 
+	// The requested configuration may not be supported, warn user and fall back. We can't call
+	// FSceneView::SetupAntiAliasingMethod because it reads the value from the cvar which would
+	// cause the value set by the MoviePipeline UI to be ignored.
+	{
+		bool bMethodWasUnsupported = false;
+		if (View->AntiAliasingMethod == AAM_TemporalAA && !SupportsGen4TAA(View->GetShaderPlatform()))
+		{
+			UE_LOG(LogMovieRenderPipeline, Error, TEXT("TAA was requested but this hardware does not support it."));
+			bMethodWasUnsupported = true;
+		}
+		else if (View->AntiAliasingMethod == AAM_TSR && !SupportsTSR(View->GetShaderPlatform()))
+		{
+			UE_LOG(LogMovieRenderPipeline, Error, TEXT("TSR was requested but this hardware does not support it."));
+			bMethodWasUnsupported = true;
+		}
+
+		if (bMethodWasUnsupported)
+		{
+			View->AntiAliasingMethod = AAM_None;
+		}
+	}
+	View->SetupAntiAliasingMethod();
+
 	// Anti Aliasing
 	{
 		// If we're not using Temporal Anti-Aliasing or Path Tracing we will apply the View Matrix projection jitter. Normally TAA sets this
@@ -440,8 +463,6 @@ FSceneView* UMoviePipelineImagePassBase::GetSceneViewForSampleState(FSceneViewFa
 		
 		// BaseProjMatrix may be perspective or orthographic.
 		ViewInitOptions.ProjectionMatrix = BaseProjMatrix;
-		
-		ViewInitOptions.ProjectionMatrix.DebugPrint();
 	}
 
 	ViewInitOptions.SceneViewStateInterface = GetSceneViewStateInterface(OptPayload);
