MoviePipeline: Adding a job through the Sequencer UI, disabling a setting on it, and then adding the job again no longer causes the disabled settings to duplicate and become active again.

diff --git a/Engine/Plugins/MovieScene/MovieRenderPipeline/Source/MovieRenderPipelineEditor/Private/MoviePipelineEditorBlueprintLibrary.cpp b/Engine/Plugins/MovieScene/MovieRenderPipeline/Source/MovieRenderPipelineEditor/Private/MoviePipelineEditorBlueprintLibrary.cpp
--- a/Engine/Plugins/MovieScene/MovieRenderPipeline/Source/MovieRenderPipelineEditor/Private/MoviePipelineEditorBlueprintLibrary.cpp
+++ b/Engine/Plugins/MovieScene/MovieRenderPipeline/Source/MovieRenderPipelineEditor/Private/MoviePipelineEditorBlueprintLibrary.cpp
@@ -203,7 +203,8 @@ void UMoviePipelineEditorBlueprintLibrary::EnsureJobHasDefaultSettings(UMoviePip
 			continue;
 		}
 
-		UMoviePipelineSetting* ExistingSetting = NewJob->GetConfiguration()->FindSettingByClass(SettingClass);
+		const bool bIncludeDisabledSettings = true;
+		UMoviePipelineSetting* ExistingSetting = NewJob->GetConfiguration()->FindSettingByClass(SettingClass, bIncludeDisabledSettings);
 		if (!ExistingSetting)
 		{
 			NewJob->GetConfiguration()->FindOrAddSettingByClass(SettingClass);
