MoviePipeline: Fixed a bug where having a sequence which was set to "Lock to Display Rate" would cause temporal sub-samples to appear stuttery. The evaluation type needs to be set to WithSubframes before the sequence player starts playing, as it is cached. Removed the caching from the Node Hierarchy because we only want to edit the root moviescene, and by the time the caching is run it will already have been overwritten with WithSubframes so it wouldn't cache the original value anyways.

diff --git a/Engine/Plugins/MovieScene/MovieRenderPipeline/Source/MovieRenderPipelineCore/Private/MoviePipeline.cpp b/Engine/Plugins/MovieScene/MovieRenderPipeline/Source/MovieRenderPipelineCore/Private/MoviePipeline.cpp
--- a/Engine/Plugins/MovieScene/MovieRenderPipeline/Source/MovieRenderPipelineCore/Private/MoviePipeline.cpp
+++ b/Engine/Plugins/MovieScene/MovieRenderPipeline/Source/MovieRenderPipelineCore/Private/MoviePipeline.cpp
@@ -224,6 +224,15 @@ void UMoviePipeline::Initialize(UMoviePipelineExecutorJob* InJob)
 			return;
 		}
 	}
+
+	// Disable frame-locked evaluation. This has to come before we initialize the level sequence actor because
+	// the Sequence Player caches the value when playback starts, and our usual store/restore function only runs
+	// when shot expansion happens, which is too late.
+	{
+		CachedSequenceHierarchyRoot->EvaluationType = TargetSequence->GetMovieScene()->GetEvaluationType();
+		TargetSequence->GetMovieScene()->SetEvaluationType(EMovieSceneEvaluationType::WithSubFrames);
+	}
+
 	
 	// Initialize all of our master config settings. Shot specific ones will be called for their appropriate shot.
 	for (UMoviePipelineSetting* Setting : GetPipelineMasterConfig()->GetAllSettings())
@@ -352,6 +361,10 @@ void UMoviePipeline::RestoreTargetSequenceToOriginalState()
 		return;
 	}
 
+	// This is done outside of the save/restore hierarchy functions as it needs to be edited/cached
+	// before the hierarchy is originally cached.
+	TargetSequence->GetMovieScene()->SetEvaluationType(CachedSequenceHierarchyRoot->EvaluationType);
+
 	MoviePipeline::RestoreCompleteSequenceHierarchy(TargetSequence, CachedSequenceHierarchyRoot);
 }
 
diff --git a/Engine/Plugins/MovieScene/MovieRenderPipeline/Source/MovieRenderPipelineCore/Private/MoviePipelineUtils.cpp b/Engine/Plugins/MovieScene/MovieRenderPipeline/Source/MovieRenderPipelineCore/Private/MoviePipelineUtils.cpp
--- a/Engine/Plugins/MovieScene/MovieRenderPipeline/Source/MovieRenderPipelineCore/Private/MoviePipelineUtils.cpp
+++ b/Engine/Plugins/MovieScene/MovieRenderPipeline/Source/MovieRenderPipelineCore/Private/MoviePipelineUtils.cpp
@@ -256,7 +256,6 @@ namespace MoviePipeline
 						Node->bOriginalMovieScenePackageDirty = OwningPackage->IsDirty();
 					}
 
-					Node->EvaluationType = Node->MovieScene->GetEvaluationType();
 				}
 
 				// Unlock the movie scene so we can make changes to sections below, it'll get re-locked later if needed.
@@ -320,7 +319,6 @@ namespace MoviePipeline
 				if (!bInSave)
 				{
 					Node->MovieScene->SetPlaybackRange(Node->OriginalMovieScenePlaybackRange);
-					Node->MovieScene->SetEvaluationType(Node->EvaluationType);
 #if WITH_EDITOR
 					Node->MovieScene->SetReadOnly(Node->bOriginalMovieSceneReadOnly);
 					Node->MovieScene->SetPlaybackRangeLocked(Node->bOriginalMovieScenePlaybackRangeLocked);
@@ -349,7 +347,6 @@ namespace MoviePipeline
 				Node->MovieScene->SetReadOnly(false);
 				Node->MovieScene->SetPlaybackRangeLocked(false);
 #endif
-				Node->MovieScene->SetEvaluationType(EMovieSceneEvaluationType::WithSubFrames);
 			}
 
 			if (Node->Section.IsValid())
