MoviePipeline: The IsEnabled state on a setting can now be toggled through Python scripting.

diff --git a/Engine/Plugins/MovieScene/MovieRenderPipeline/Source/MovieRenderPipelineCore/Public/MoviePipelineSetting.h b/Engine/Plugins/MovieScene/MovieRenderPipeline/Source/MovieRenderPipelineCore/Public/MoviePipelineSetting.h
--- a/Engine/Plugins/MovieScene/MovieRenderPipeline/Source/MovieRenderPipelineCore/Public/MoviePipelineSetting.h
+++ b/Engine/Plugins/MovieScene/MovieRenderPipeline/Source/MovieRenderPipelineCore/Public/MoviePipelineSetting.h
@@ -147,7 +147,9 @@ public:
 	virtual bool IsSolo() const { return true; }
 	
 	/** Is this setting enabled by the user in the UI? */
+	UFUNCTION(BlueprintPure, Category = "Movie Render Pipeline")
 	virtual bool IsEnabled() const { return bEnabled; }
+	UFUNCTION(BlueprintCallable, Category = "Movie Render Pipeline")
 	virtual void SetIsEnabled(bool bInEnabled) { bEnabled = bInEnabled; }
 
 	/** Has this setting finished any export-related things it needs to do post-finalize? */
