MoviePipeline: Fixed Camera Cut bindings that were bound to Camera Actors instead of Camera Components not rendering correctly.

diff --git a/Engine/Plugins/MovieScene/MovieRenderPipeline/Source/MovieRenderPipelineCore/Private/MoviePipeline.cpp b/Engine/Plugins/MovieScene/MovieRenderPipeline/Source/MovieRenderPipelineCore/Private/MoviePipeline.cpp
--- a/Engine/Plugins/MovieScene/MovieRenderPipeline/Source/MovieRenderPipelineCore/Private/MoviePipeline.cpp
+++ b/Engine/Plugins/MovieScene/MovieRenderPipeline/Source/MovieRenderPipelineCore/Private/MoviePipeline.cpp
@@ -47,6 +47,7 @@
 #include "ClothingSimulationInteractor.h"
 #include "ClothingSimulationInterface.h"
 #include "Components/SkeletalMeshComponent.h"
+#include "MovieSceneCommonHelpers.h"
 
 #if WITH_EDITOR
 #include "MovieSceneExportMetadata.h"
@@ -1724,7 +1725,7 @@ void UMoviePipeline::GetSidecarCameraData(UMoviePipelineExecutorShot* InShot, in
 		}
 	}
 
-	UCameraComponent* BoundCamera = Cast<UCameraComponent>(PrimaryViewTarget);
+	UCameraComponent* BoundCamera = MovieSceneHelpers::CameraComponentFromRuntimeObject(PrimaryViewTarget);
 	if (BoundCamera)
 	{
 		*OutCameraComponent = BoundCamera;
@@ -1756,7 +1757,7 @@ void UMoviePipeline::GetSidecarCameraViewPoints(UMoviePipelineExecutorShot* InSh
 		TArrayView<TWeakObjectPtr<UObject>> BoundCameras = LevelSequenceActor->GetSequencePlayer()->FindBoundObjects(CameraBindingRef, CameraSequenceRef);
 		if (BoundCameras.Num() > 0)
 		{
-			UCameraComponent* BoundCamera = Cast<UCameraComponent>(BoundCameras[0].Get());
+			UCameraComponent* BoundCamera = MovieSceneHelpers::CameraComponentFromRuntimeObject(BoundCameras[0].Get());
 			if (BoundCamera)
 			{
 				FMinimalViewInfo ViewInfo;
diff --git a/Engine/Source/Runtime/MovieScene/Private/MovieSceneCommonHelpers.cpp b/Engine/Source/Runtime/MovieScene/Private/MovieSceneCommonHelpers.cpp
--- a/Engine/Source/Runtime/MovieScene/Private/MovieSceneCommonHelpers.cpp
+++ b/Engine/Source/Runtime/MovieScene/Private/MovieSceneCommonHelpers.cpp
@@ -434,14 +434,14 @@ UCameraComponent* MovieSceneHelpers::CameraComponentFromRuntimeObject(UObject* R
 	if (RuntimeObject)
 	{
 		// find camera we want to control
-		UCameraComponent* const CameraComponent = dynamic_cast<UCameraComponent*>(RuntimeObject);
+		UCameraComponent* const CameraComponent = Cast<UCameraComponent>(RuntimeObject);
 		if (CameraComponent)
 		{
 			return CameraComponent;
 		}
 
 		// see if it's an actor that has a camera component
-		AActor* const Actor = dynamic_cast<AActor*>(RuntimeObject);
+		AActor* const Actor = Cast<AActor>(RuntimeObject);
 		if (Actor)
 		{
 			return CameraComponentFromActor(Actor);
