MoviePipeline: Fixed using the Widget Renderer causing the system to think you had multiple cameras and forcibly adding the {camera_name} tag.

diff --git a/Engine/Plugins/MovieScene/MovieRenderPipeline/Source/MovieRenderPipelineSettings/Private/MoviePipelineWidgetRenderSetting.cpp b/Engine/Plugins/MovieScene/MovieRenderPipeline/Source/MovieRenderPipelineSettings/Private/MoviePipelineWidgetRenderSetting.cpp
--- a/Engine/Plugins/MovieScene/MovieRenderPipeline/Source/MovieRenderPipelineSettings/Private/MoviePipelineWidgetRenderSetting.cpp
+++ b/Engine/Plugins/MovieScene/MovieRenderPipeline/Source/MovieRenderPipelineSettings/Private/MoviePipelineWidgetRenderSetting.cpp
@@ -5,6 +5,7 @@
 #include "MovieRenderPipelineDataTypes.h"
 #include "MoviePipelineBurnInWidget.h"
 #include "MoviePipelineOutputSetting.h"
+#include "MoviePipelineCameraSetting.h"
 #include "MoviePipelineMasterConfig.h"
 #include "MoviePipelineBlueprintLibrary.h"
 #include "MoviePipeline.h"
@@ -14,10 +15,21 @@
 #include "Widgets/SViewport.h"
 #include "Slate/SGameLayerManager.h"
 #include "MovieRenderPipelineCoreModule.h"
+#include "MoviePipelineQueue.h"
 
 void UMoviePipelineWidgetRenderer::GatherOutputPassesImpl(TArray<FMoviePipelinePassIdentifier>& ExpectedRenderPasses)
 {
-	ExpectedRenderPasses.Add(FMoviePipelinePassIdentifier(TEXT("ViewportUI")));
+	UMoviePipelineExecutorShot* CurrentShot = GetPipeline()->GetActiveShotList()[GetPipeline()->GetCurrentShotIndex()];
+	UMoviePipelineCameraSetting* CameraSettings = GetPipeline()->FindOrAddSettingForShot<UMoviePipelineCameraSetting>(CurrentShot);
+	int32 NumCameras = CameraSettings->bRenderAllCameras ? CurrentShot->SidecarCameras.Num() : 1;
+
+	for (int32 CameraIndex = 0; CameraIndex < NumCameras; CameraIndex++)
+	{
+		FMoviePipelinePassIdentifier PassIdentifierForCurrentCamera;
+		PassIdentifierForCurrentCamera.Name = TEXT("ViewportUI");
+		PassIdentifierForCurrentCamera.CameraName = CurrentShot->GetCameraName(CameraIndex);
+		ExpectedRenderPasses.Add(PassIdentifierForCurrentCamera);
+	}
 }
 
 void UMoviePipelineWidgetRenderer::RenderSample_GameThreadImpl(const FMoviePipelineRenderPassMetrics& InSampleState)
@@ -33,43 +45,53 @@ void UMoviePipelineWidgetRenderer::RenderSample_GameThreadImpl(const FMoviePipel
 
 	if (bFirstTile && bFirstSpatial && bFirstTemporal)
 	{
-		// Draw the widget to the render target
-		FRenderTarget* BackbufferRenderTarget = RenderTarget->GameThread_GetRenderTargetResource();
-		
-		ULocalPlayer* LocalPlayer = GetWorld()->GetFirstLocalPlayerFromController();
+		UMoviePipelineExecutorShot* CurrentShot = GetPipeline()->GetActiveShotList()[GetPipeline()->GetCurrentShotIndex()];
+		UMoviePipelineCameraSetting* CameraSettings = GetPipeline()->FindOrAddSettingForShot<UMoviePipelineCameraSetting>(CurrentShot);
+		int32 NumCameras = CameraSettings->bRenderAllCameras ? CurrentShot->SidecarCameras.Num() : 1;
+		for (int32 CameraIndex = 0; CameraIndex < NumCameras; CameraIndex++)
+		{
+			FMoviePipelinePassIdentifier PassIdentifierForCurrentCamera;
+			PassIdentifierForCurrentCamera.Name = TEXT("ViewportUI");
+			PassIdentifierForCurrentCamera.CameraName = CurrentShot->GetCameraName(CameraIndex);
 
-		// Cast the interface to a widget is a little yucky but the implementation is unlikely to change.
-		TSharedPtr<SGameLayerManager> GameLayerManager = StaticCastSharedPtr<SGameLayerManager>(LocalPlayer->ViewportClient->GetGameLayerManager());
+			// Draw the widget to the render target
+			FRenderTarget* BackbufferRenderTarget = RenderTarget->GameThread_GetRenderTargetResource();
 
-		WidgetRenderer->DrawWidget(BackbufferRenderTarget, GameLayerManager.ToSharedRef(), 1.f, FVector2D(RenderTarget->SizeX, RenderTarget->SizeY), (float)InSampleState.OutputState.TimeData.FrameDeltaTime);
+			ULocalPlayer* LocalPlayer = GetWorld()->GetFirstLocalPlayerFromController();
 
-		TSharedPtr<FMoviePipelineOutputMerger, ESPMode::ThreadSafe> OutputBuilder = GetPipeline()->OutputBuilder;
+			// Cast the interface to a widget is a little yucky but the implementation is unlikely to change.
+			TSharedPtr<SGameLayerManager> GameLayerManager = StaticCastSharedPtr<SGameLayerManager>(LocalPlayer->ViewportClient->GetGameLayerManager());
 
-		ENQUEUE_RENDER_COMMAND(BurnInRenderTargetResolveCommand)(
-			[InSampleState, bComposite = bCompositeOntoFinalImage, BackbufferRenderTarget, OutputBuilder](FRHICommandListImmediate& RHICmdList)
-		{
-			FIntRect SourceRect = FIntRect(0, 0, BackbufferRenderTarget->GetSizeXY().X, BackbufferRenderTarget->GetSizeXY().Y);
+			WidgetRenderer->DrawWidget(BackbufferRenderTarget, GameLayerManager.ToSharedRef(), 1.f, FVector2D(RenderTarget->SizeX, RenderTarget->SizeY), (float)InSampleState.OutputState.TimeData.FrameDeltaTime);
+
+			TSharedPtr<FMoviePipelineOutputMerger, ESPMode::ThreadSafe> OutputBuilder = GetPipeline()->OutputBuilder;
+
+			ENQUEUE_RENDER_COMMAND(BurnInRenderTargetResolveCommand)(
+				[InSampleState, PassIdentifierForCurrentCamera, bComposite = bCompositeOntoFinalImage, BackbufferRenderTarget, OutputBuilder](FRHICommandListImmediate& RHICmdList)
+				{
+					FIntRect SourceRect = FIntRect(0, 0, BackbufferRenderTarget->GetSizeXY().X, BackbufferRenderTarget->GetSizeXY().Y);
 
-			// Read the data back to the CPU
-			TArray<FColor> RawPixels;
-			RawPixels.SetNum(SourceRect.Width() * SourceRect.Height());
+					// Read the data back to the CPU
+					TArray<FColor> RawPixels;
+					RawPixels.SetNum(SourceRect.Width() * SourceRect.Height());
 
-			FReadSurfaceDataFlags ReadDataFlags(ERangeCompressionMode::RCM_MinMax);
-			ReadDataFlags.SetLinearToGamma(false);
+					FReadSurfaceDataFlags ReadDataFlags(ERangeCompressionMode::RCM_MinMax);
+					ReadDataFlags.SetLinearToGamma(false);
 
-			RHICmdList.ReadSurfaceData(BackbufferRenderTarget->GetRenderTargetTexture(), SourceRect, RawPixels, ReadDataFlags);
+					RHICmdList.ReadSurfaceData(BackbufferRenderTarget->GetRenderTargetTexture(), SourceRect, RawPixels, ReadDataFlags);
 
-			TSharedRef<FImagePixelDataPayload, ESPMode::ThreadSafe> FrameData = MakeShared<FImagePixelDataPayload, ESPMode::ThreadSafe>();
-			FrameData->PassIdentifier = FMoviePipelinePassIdentifier(TEXT("ViewportUI"));
-			FrameData->SampleState = InSampleState;
-			FrameData->bRequireTransparentOutput = true;
-			FrameData->SortingOrder = 3;
-			FrameData->bCompositeToFinalImage = bComposite;
+					TSharedRef<FImagePixelDataPayload, ESPMode::ThreadSafe> FrameData = MakeShared<FImagePixelDataPayload, ESPMode::ThreadSafe>();
+					FrameData->PassIdentifier = PassIdentifierForCurrentCamera;
+					FrameData->SampleState = InSampleState;
+					FrameData->bRequireTransparentOutput = true;
+					FrameData->SortingOrder = 3;
+					FrameData->bCompositeToFinalImage = bComposite;
 
-			TUniquePtr<FImagePixelData> PixelData = MakeUnique<TImagePixelData<FColor>>(InSampleState.BackbufferSize, TArray64<FColor>(MoveTemp(RawPixels)), FrameData);
+					TUniquePtr<FImagePixelData> PixelData = MakeUnique<TImagePixelData<FColor>>(InSampleState.BackbufferSize, TArray64<FColor>(MoveTemp(RawPixels)), FrameData);
 
-			OutputBuilder->OnCompleteRenderPassDataAvailable_AnyThread(MoveTemp(PixelData));
-		});
+					OutputBuilder->OnCompleteRenderPassDataAvailable_AnyThread(MoveTemp(PixelData));
+				});
+		}
 	}
 }
 
