MoviePipeline: Fixed an issue where the Command Line Encoder could not use {version} tokens in the filename. Added FMoviePipelineFilenameResolveParams::FileNameOverride to allow users to specify what filename they want to use when resolving (instead of being forced into the Output Settin's output directory/file format name). If omitted, no behavior change.

diff --git a/Engine/Plugins/MovieScene/MovieRenderPipeline/Source/MovieRenderPipelineCore/Private/MoviePipelineBlueprintLibrary.cpp b/Engine/Plugins/MovieScene/MovieRenderPipeline/Source/MovieRenderPipelineCore/Private/MoviePipelineBlueprintLibrary.cpp
--- a/Engine/Plugins/MovieScene/MovieRenderPipeline/Source/MovieRenderPipelineCore/Private/MoviePipelineBlueprintLibrary.cpp
+++ b/Engine/Plugins/MovieScene/MovieRenderPipeline/Source/MovieRenderPipelineCore/Private/MoviePipelineBlueprintLibrary.cpp
@@ -761,7 +761,7 @@ int32 UMoviePipelineBlueprintLibrary::ResolveVersionNumber(FMoviePipelineFilenam
 	}
 
 	// Calculate a version number by looking at the output path and then scanning for a version token.
-	FString FileNameFormatString = OutputSettings->OutputDirectory.Path / OutputSettings->FileNameFormat;
+	FString FileNameFormatString = InParams.FileNameOverride.Len() > 0 ? InParams.FileNameOverride : OutputSettings->OutputDirectory.Path / OutputSettings->FileNameFormat;
 
 	FString FinalPath;
 	FMoviePipelineFormatArgs FinalFormatArgs;
@@ -998,7 +998,9 @@ void UMoviePipelineBlueprintLibrary::ResolveFilenameFormatArguments(const FStrin
 		NamedArgs.Add(Argument.Key, Argument.Value);
 	}
 
-	FString BaseFilename = FString::Format(*InFormatString, NamedArgs);
+	FString FileNameFormatString = InParams.FileNameOverride.Len() > 0 ? InParams.FileNameOverride : InFormatString;
+
+	FString BaseFilename = FString::Format(*FileNameFormatString, NamedArgs);
 	FPaths::NormalizeFilename(BaseFilename);
 
 
@@ -1023,7 +1025,7 @@ void UMoviePipelineBlueprintLibrary::ResolveFilenameFormatArguments(const FStrin
 		NamedArgs.Add(TEXT("file_dup"), FString::Printf(TEXT("_(%d)"), DuplicateIndex));
 
 		// Re-resolve the format string now that we've reassigned frame_dup to a number.
-		ThisTry = FString::Format(*InFormatString, NamedArgs) + Extension;
+		ThisTry = FString::Format(*FileNameFormatString, NamedArgs) + Extension;
 
 		// If the file doesn't exist, we can use that, else, increment the index and try again
 		if (CanWriteToFile(*ThisTry, OutputSetting->bOverrideExistingOutput))
diff --git a/Engine/Plugins/MovieScene/MovieRenderPipeline/Source/MovieRenderPipelineCore/Private/MoviePipelineCommandLineEncoder.cpp b/Engine/Plugins/MovieScene/MovieRenderPipeline/Source/MovieRenderPipelineCore/Private/MoviePipelineCommandLineEncoder.cpp
--- a/Engine/Plugins/MovieScene/MovieRenderPipeline/Source/MovieRenderPipelineCore/Private/MoviePipelineCommandLineEncoder.cpp
+++ b/Engine/Plugins/MovieScene/MovieRenderPipeline/Source/MovieRenderPipelineCore/Private/MoviePipelineCommandLineEncoder.cpp
@@ -16,8 +16,14 @@
 #include "MoviePipelineBlueprintLibrary.h"
 
 // Forward Declare
-static TArray<FText> GetErrorTexts();
-static bool HasMultipleRenderPasses(const TArray<FMoviePipelineShotOutputData>& InData);
+namespace UE
+{
+namespace MoviePipeline
+{
+	static TArray<FText> GetErrorTexts();
+	static bool HasMultipleRenderPasses(const TArray<FMoviePipelineShotOutputData>& InData);
+}
+}
 
 UMoviePipelineCommandLineEncoder::UMoviePipelineCommandLineEncoder()
 {
@@ -69,7 +75,7 @@ void UMoviePipelineCommandLineEncoder::StartEncodingProcess(TArray<FMoviePipelin
 
 	// Early out if there's any errors
 	{
-		TArray<FText> ErrorTexts = GetErrorTexts();
+		TArray<FText> ErrorTexts = UE::MoviePipeline::GetErrorTexts();
 		for (const FText& ErrorText : ErrorTexts)
 		{
 			UE_LOG(LogMovieRenderPipelineIO, Error, TEXT("%s"), *ErrorText.ToString());
@@ -87,7 +93,7 @@ void UMoviePipelineCommandLineEncoder::StartEncodingProcess(TArray<FMoviePipelin
 	
 	// If we're writing more than one render pass out, we need to ensure the file name has the format string in it so we don't
 	// overwrite the same file multiple times. 
-	const bool bIncludeRenderPass = HasMultipleRenderPasses(InOutData);
+	const bool bIncludeRenderPass = UE::MoviePipeline::HasMultipleRenderPasses(InOutData);
 	const bool bTestFrameNumber = false;
 
 	UE::MoviePipeline::ValidateOutputFormatString(FileNameFormatString, bIncludeRenderPass, bTestFrameNumber);
@@ -166,6 +172,19 @@ void UMoviePipelineCommandLineEncoder::StartEncodingProcess(TArray<FMoviePipelin
 			FormatOverrides.Add(TEXT("shot_name"), Shot->OuterName);
 			FormatOverrides.Add(TEXT("camera_name"), Shot->InnerName);
 		}
+
+		// We must manually resolve {version} tokens because they're intended to be per-shot/global (ie: individual file names
+		// don't have versions when using image sequences).
+		{
+			FMoviePipelineFilenameResolveParams ResolveParams;
+			ResolveParams.InitializationTime = GetPipeline()->GetInitializationTime();
+			ResolveParams.Job = GetPipeline()->GetCurrentJob();
+			ResolveParams.ShotOverride = RenderPass.Value.Shot.Get();
+			ResolveParams.FileNameOverride = FileNameFormatString;
+			int32 VersionNumber = UMoviePipelineBlueprintLibrary::ResolveVersionNumber(ResolveParams);
+			FileNameFormatString.ReplaceInline(TEXT("{version}"), *FString::Printf(TEXT("v%0*d"), 3, VersionNumber));
+		}
+
 		FMoviePipelineFormatArgs FinalFormatArgs;
 
 		FString FinalFilePath;
@@ -416,55 +435,61 @@ FString UMoviePipelineCommandLineEncoder::GetQualitySettingString() const
 	return FString();
 }
 
-static TArray<FText> GetErrorTexts()
+namespace UE
 {
-	TArray<FText> OutErrors;
-	const UMoviePipelineCommandLineEncoderSettings* EncoderSettings = GetDefault<UMoviePipelineCommandLineEncoderSettings>();
-	if (EncoderSettings->ExecutablePath.Len() == 0)
-	{
-		OutErrors.Add(NSLOCTEXT("MovieRenderPipeline", "CommandLineEncode_MissingExecutable", "No encoder executable has been specified in the Project Settings. Please set an encoder executable in Project Settings > Movie Pipeline CLI Encoder"));
-	}
-	if (EncoderSettings->VideoCodec.Len() == 0)
-	{
-		OutErrors.Add(NSLOCTEXT("MovieRenderPipeline", "CommandLineEncode_MissingVideoCodec", "No video encoding codec has been specified in the Project Settings. Please set an video codec in Project Settings > Movie Pipeline CLI Encoder"));
-	}
-	if (EncoderSettings->AudioCodec.Len() == 0)
-	{
-		OutErrors.Add(NSLOCTEXT("MovieRenderPipeline", "CommandLineEncode_MissingAudioCodec", "No audio encoding codec has been specified in the Project Settings. Please set an audio codec in Project Settings > Movie Pipeline CLI Encoder"));
-	}
-	if (EncoderSettings->OutputFileExtension.Len() == 0)
+namespace MoviePipeline
+{
+	static TArray<FText> GetErrorTexts()
 	{
-		OutErrors.Add(NSLOCTEXT("MovieRenderPipeline", "CommandLineEncode_MissingFileExtension", "No file extension has been specified in the Project Settings. Please set a file extension in Project Settings > Movie Pipeline CLI Encoder"));
-	}
+		TArray<FText> OutErrors;
+		const UMoviePipelineCommandLineEncoderSettings* EncoderSettings = GetDefault<UMoviePipelineCommandLineEncoderSettings>();
+		if (EncoderSettings->ExecutablePath.Len() == 0)
+		{
+			OutErrors.Add(NSLOCTEXT("MovieRenderPipeline", "CommandLineEncode_MissingExecutable", "No encoder executable has been specified in the Project Settings. Please set an encoder executable in Project Settings > Movie Pipeline CLI Encoder"));
+		}
+		if (EncoderSettings->VideoCodec.Len() == 0)
+		{
+			OutErrors.Add(NSLOCTEXT("MovieRenderPipeline", "CommandLineEncode_MissingVideoCodec", "No video encoding codec has been specified in the Project Settings. Please set an video codec in Project Settings > Movie Pipeline CLI Encoder"));
+		}
+		if (EncoderSettings->AudioCodec.Len() == 0)
+		{
+			OutErrors.Add(NSLOCTEXT("MovieRenderPipeline", "CommandLineEncode_MissingAudioCodec", "No audio encoding codec has been specified in the Project Settings. Please set an audio codec in Project Settings > Movie Pipeline CLI Encoder"));
+		}
+		if (EncoderSettings->OutputFileExtension.Len() == 0)
+		{
+			OutErrors.Add(NSLOCTEXT("MovieRenderPipeline", "CommandLineEncode_MissingFileExtension", "No file extension has been specified in the Project Settings. Please set a file extension in Project Settings > Movie Pipeline CLI Encoder"));
+		}
 
-	return OutErrors;
-}
+		return OutErrors;
+	}
 
-static bool HasMultipleRenderPasses(const TArray<FMoviePipelineShotOutputData>& InData)
-{
-	bool bHasMultipleRenderPasses = false;
-	for (const FMoviePipelineShotOutputData& Data : InData)
+	static bool HasMultipleRenderPasses(const TArray<FMoviePipelineShotOutputData>& InData)
 	{
-		int32 RenderPassCount = 0;
-		for (const TTuple<FMoviePipelinePassIdentifier, FMoviePipelineRenderPassOutputData>& Pair : Data.RenderPassData)
+		bool bHasMultipleRenderPasses = false;
+		for (const FMoviePipelineShotOutputData& Data : InData)
 		{
-			if (Pair.Key.Name != TEXT("Audio"))
+			int32 RenderPassCount = 0;
+			for (const TTuple<FMoviePipelinePassIdentifier, FMoviePipelineRenderPassOutputData>& Pair : Data.RenderPassData)
 			{
-				RenderPassCount++;
+				if (Pair.Key.Name != TEXT("Audio"))
+				{
+					RenderPassCount++;
+				}
 			}
+
+			bHasMultipleRenderPasses |= RenderPassCount > 1;
 		}
 
-		bHasMultipleRenderPasses |= RenderPassCount > 1;
+		return bHasMultipleRenderPasses;
 	}
-
-	return bHasMultipleRenderPasses;
+}
 }
 
 void UMoviePipelineCommandLineEncoder::ValidateStateImpl()
 {
 	Super::ValidateStateImpl();
 
-	TArray<FText> Errors = GetErrorTexts();
+	TArray<FText> Errors = UE::MoviePipeline::GetErrorTexts();
 	for (const FText& Error : Errors)
 	{
 		ValidationResults.Add(Error);
diff --git a/Engine/Plugins/MovieScene/MovieRenderPipeline/Source/MovieRenderPipelineCore/Public/MovieRenderPipelineDataTypes.h b/Engine/Plugins/MovieScene/MovieRenderPipeline/Source/MovieRenderPipelineCore/Public/MovieRenderPipelineDataTypes.h
--- a/Engine/Plugins/MovieScene/MovieRenderPipeline/Source/MovieRenderPipelineCore/Public/MovieRenderPipelineDataTypes.h
+++ b/Engine/Plugins/MovieScene/MovieRenderPipeline/Source/MovieRenderPipelineCore/Public/MovieRenderPipelineDataTypes.h
@@ -830,6 +830,10 @@ struct FMoviePipelineFilenameResolveParams
 	UPROPERTY(BlueprintReadWrite, EditAnywhere, Category = "Movie Render Pipeline")
 	bool bForceRelativeFrameNumbers;
 
+	/** Optional. If specified this is the filename that will be used instead of automatically building it from the Job's Output Setting. */
+	UPROPERTY(BlueprintReadWrite, EditAnywhere, Category = "Movie Render Pipeline")
+	FString FileNameOverride;
+
 	/** 
 	* A map between "{format}" tokens and their values. These are applied after the auto-generated ones from the system,
 	* which allows the caller to override things like {.ext} depending or {render_pass} which have dummy names by default.
