MoviePipeline: Fixed separately composited Burn In/Widget Renderers not actually writing an alpha channel when the main output was JPEG/BMP. They are specifically requested to have an alpha channel (so that they are useful), so having the alpha channel overwritten because the main type doesn't support alpha is incorrect.

diff --git a/Engine/Plugins/MovieScene/MovieRenderPipeline/Source/MovieRenderPipelineRenderPasses/Private/MoviePipelineImageSequenceOutput.cpp b/Engine/Plugins/MovieScene/MovieRenderPipeline/Source/MovieRenderPipelineRenderPasses/Private/MoviePipelineImageSequenceOutput.cpp
--- a/Engine/Plugins/MovieScene/MovieRenderPipeline/Source/MovieRenderPipelineRenderPasses/Private/MoviePipelineImageSequenceOutput.cpp
+++ b/Engine/Plugins/MovieScene/MovieRenderPipeline/Source/MovieRenderPipelineRenderPasses/Private/MoviePipelineImageSequenceOutput.cpp
@@ -248,7 +248,9 @@ void UMoviePipelineImageSequenceOutputBase::OnReceiveImageDataImpl(FMoviePipelin
 			}
 		}
 
-		if (!IsAlphaAllowed())
+		// A payload _requiring_ alpha output will override the Write Alpha option, because that flag is used to indicate that the output is
+		// no good without alpha, and we already did logic above to ensure it got turned into a filetype that could write alpha.
+		if (!IsAlphaAllowed() && !Payload->bRequireTransparentOutput)
 		{
 			switch (QuantizedPixelType)
 			{
