MoviePipeline: Added an experimental streaming system flushing mechanism to help ensure there's no blurry textures visible when using texture streaming. Using texture streaming reduces VRAM usages significantly in large scenes. Historically MRQ has disabled Texture Streaming (and will continue to disable it by default, but this can be controlled by the Game Overrides setting set to "Don't Override") due to a delay between requesting a high-res mipmap and the mipmap being delievered to the GPU, resulting in some output frames having blurry textures visible in them (especially on newly spawned objects that were already in camera view).

This change adds a flush step to each frame (as part of the flush async engine systems) to request that streaming managers flush their outstanding changes, and adds a step to the Texture Streaming Manager to ensure async requests have their states synced before the flush. Without the sync step, async work that is still being processed doesn't get flushed (which results in non-updated/blurry textures). This change to the streaming flush is gated behind a cvar (which MRQ automatically manipulates) to ensure it has no potential impact on regular texture streaming (that does not intend to fully block to get all resources onto the GPU).

diff --git a/Engine/Plugins/MovieScene/MovieRenderPipeline/Source/MovieRenderPipelineCore/Private/MoviePipelineGameOverrideSetting.cpp b/Engine/Plugins/MovieScene/MovieRenderPipeline/Source/MovieRenderPipelineCore/Private/MoviePipelineGameOverrideSetting.cpp
--- a/Engine/Plugins/MovieScene/MovieRenderPipeline/Source/MovieRenderPipelineCore/Private/MoviePipelineGameOverrideSetting.cpp
+++ b/Engine/Plugins/MovieScene/MovieRenderPipeline/Source/MovieRenderPipelineCore/Private/MoviePipelineGameOverrideSetting.cpp
@@ -92,6 +92,11 @@ void UMoviePipelineGameOverrideSetting::ApplyCVarSettings(const bool bOverrideVa
 		// This CVAR only exists if the D3D12RHI module is loaded
 		MOVIEPIPELINE_STORE_AND_OVERRIDE_CVAR_INT_IF_EXIST(PreviousGPUTimeout, TEXT("r.D3D12.GPUTimeout"), 0, bOverrideValues);
 	}
+
+	if (bFlushStreamingManagers)
+	{
+		MOVIEPIPELINE_STORE_AND_OVERRIDE_CVAR_INT(PreviousStreamingManagerSyncState, TEXT("r.Streaming.SyncStatesWhenBlocking"), 1, bOverrideValues);
+	}
 	
 #if WITH_EDITOR
 	// To make sure the GeometryCache streamer doesn't skip frames and doesn't pop up notification during rendering
@@ -187,6 +192,11 @@ void UMoviePipelineGameOverrideSetting::BuildNewProcessCommandLineArgsImpl(TArra
 	{
 		InOutDeviceProfileCvars.Add(TEXT("r.D3D12.GPUTimeout=0"));
 	}
+
+	if (bFlushStreamingManagers)
+	{
+		InOutDeviceProfileCvars.Add(TEXT("r.Streaming.SyncStatesWhenBlocking=1"));
+	}
 	
 #if WITH_EDITOR
 	{
diff --git a/Engine/Plugins/MovieScene/MovieRenderPipeline/Source/MovieRenderPipelineCore/Private/MoviePipelineRendering.cpp b/Engine/Plugins/MovieScene/MovieRenderPipeline/Source/MovieRenderPipelineCore/Private/MoviePipelineRendering.cpp
--- a/Engine/Plugins/MovieScene/MovieRenderPipeline/Source/MovieRenderPipelineCore/Private/MoviePipelineRendering.cpp
+++ b/Engine/Plugins/MovieScene/MovieRenderPipeline/Source/MovieRenderPipelineCore/Private/MoviePipelineRendering.cpp
@@ -37,6 +37,7 @@
 #include "ShaderCompiler.h"
 #include "EngineUtils.h"
 #include "Materials/MaterialInterface.h"
+#include "ContentStreaming.h"
 
 
 #define LOCTEXT_NAMESPACE "MoviePipeline"
@@ -543,6 +544,17 @@ void UMoviePipeline::FlushAsyncEngineSystems()
 	// A progressbar is already in place so the user can get feedback while waiting for everything to settle.
 	FAssetCompilingManager::Get().FinishAllCompilation();
 
+	// Flush streaming managers
+	{
+		UMoviePipelineGameOverrideSetting* GameOverrideSettings = FindOrAddSettingForShot<UMoviePipelineGameOverrideSetting>(ActiveShotList[CurrentShotIndex]);
+		if (GameOverrideSettings && GameOverrideSettings->bFlushStreamingManagers)
+		{
+			FStreamingManagerCollection& StreamingManagers = IStreamingManager::Get();
+			StreamingManagers.UpdateResourceStreaming(GetWorld()->GetDeltaSeconds(), /* bProcessEverything */ true);
+			StreamingManagers.BlockTillAllRequestsFinished();
+		}
+	}
+
 	// Flush grass
 	if (CurrentShotIndex < ActiveShotList.Num())
 	{
diff --git a/Engine/Plugins/MovieScene/MovieRenderPipeline/Source/MovieRenderPipelineCore/Public/MoviePipelineGameOverrideSetting.h b/Engine/Plugins/MovieScene/MovieRenderPipeline/Source/MovieRenderPipelineCore/Public/MoviePipelineGameOverrideSetting.h
--- a/Engine/Plugins/MovieScene/MovieRenderPipeline/Source/MovieRenderPipelineCore/Public/MoviePipelineGameOverrideSetting.h
+++ b/Engine/Plugins/MovieScene/MovieRenderPipeline/Source/MovieRenderPipelineCore/Public/MoviePipelineGameOverrideSetting.h
@@ -36,7 +36,8 @@ public:
 		, bOverrideViewDistanceScale(true)
 		, ViewDistanceScale(50)
 		, bFlushGrassStreaming(true)
-		, bDisableGPUTimeout(true) 
+		, bFlushStreamingManagers(true)
+		, bDisableGPUTimeout(true)
 	{
 	}
 
@@ -100,6 +101,10 @@ public:
 	UPROPERTY(EditAnywhere, BlueprintReadWrite, Category = "Rendering")
 	bool bFlushGrassStreaming;
 
+	/** Experimental. If true flush the streaming managers (Texture Streaming) each frame. Allows Texture Streaming to not have visible pop-in in final frames. */
+	UPROPERTY(EditAnywhere, BlueprintReadWrite, Category = "Rendering")
+	bool bFlushStreamingManagers;
+
 	/** Should we disable the GPU Timeout? Currently only applicable when using D3D12 renderer. */
 	bool bDisableGPUTimeout;
 
@@ -124,6 +129,7 @@ private:
 	int32 PreviousSkyLightRealTimeReflectionCaptureTimeSlice;
 	int32 PreviousVolumetricRenderTarget;
 	int32 PreviousIgnoreStreamingPerformance;
+	int32 PreviousStreamingManagerSyncState;
 #if WITH_EDITOR
 	int32 PreviousGeoCacheStreamerShowNotification;
 	int32 PreviousGeoCacheStreamerBlockTillFinish;
diff --git a/Engine/Source/Runtime/Engine/Private/Streaming/StreamingManagerTexture.cpp b/Engine/Source/Runtime/Engine/Private/Streaming/StreamingManagerTexture.cpp
--- a/Engine/Source/Runtime/Engine/Private/Streaming/StreamingManagerTexture.cpp
+++ b/Engine/Source/Runtime/Engine/Private/Streaming/StreamingManagerTexture.cpp
@@ -69,6 +69,12 @@ static FAutoConsoleVariableRef CVarStreamingParallelRenderAssetsNumWorkgroups(
 	TEXT("Though adds overhead to GameThread if too high."),
 	ECVF_Default);
 
+static TAutoConsoleVariable<int32> CVarSyncStatesWhenBlocking(
+	TEXT("r.Streaming.SyncStatesWhenBlocking"),
+	0,
+	TEXT("If true, SyncStates will be called to fully update async states before flushing outstanding streaming requests. Used by Movie Render Queue to ensure all streaming requests are handled each frame to avoid pop-in."),
+	ECVF_Default);
+
 bool TrackRenderAsset( const FString& AssetName );
 bool UntrackRenderAsset( const FString& AssetName );
 void ListTrackedRenderAssets( FOutputDevice& Ar, int32 NumTextures );
@@ -1789,6 +1795,12 @@ int32 FRenderAssetStreamingManager::BlockTillAllRequestsFinished( float TimeLimi
 
 	while (ensure(!IsAssetStreamingSuspended()))
 	{
+		// Optionally synchronize the states of async work before we wait for outstanding work to be completed.
+		if (CVarSyncStatesWhenBlocking.GetValueOnGameThread() != 0)
+		{
+			SyncStates(true);
+		}
+
 		int32 NumOfInFlights = 0;
 
 		for (FStreamingRenderAsset& StreamingRenderAsset : StreamingRenderAssets)
