MoviePipeline: Video Codecs no longer trip an assert if they fail to initialize (due to invalid output file path). This was a threading issue, setting the promise to false after Initialize will cause MRQ to quit rendering, but until the promise is processed in the next tick, OnReceiveImageDataImpl is called again with new data so we can end up back inside UMoviePipelineVideoOutputBase::OnReceiveImageDataImpl with a created, but unsuccessfully initialized writer.

diff --git a/Engine/Plugins/MovieScene/MovieRenderPipeline/Source/MovieRenderPipelineCore/Private/MoviePipelineVideoOutputBase.cpp b/Engine/Plugins/MovieScene/MovieRenderPipeline/Source/MovieRenderPipelineCore/Private/MoviePipelineVideoOutputBase.cpp
--- a/Engine/Plugins/MovieScene/MovieRenderPipeline/Source/MovieRenderPipelineCore/Private/MoviePipelineVideoOutputBase.cpp
+++ b/Engine/Plugins/MovieScene/MovieRenderPipeline/Source/MovieRenderPipelineCore/Private/MoviePipelineVideoOutputBase.cpp
@@ -12,6 +12,12 @@
 #include "HAL/PlatformFile.h"
 #include "HAL/PlatformTime.h"
 #include "MoviePipelineUtils.h"
+
+UMoviePipelineVideoOutputBase::UMoviePipelineVideoOutputBase()
+	: bHasError(false)
+{
+}
+
 void UMoviePipelineVideoOutputBase::OnShotFinishedImpl(const UMoviePipelineExecutorShot* InShot, const bool bFlushToDisk)
 {
 	if (bFlushToDisk)
@@ -42,6 +48,12 @@ void UMoviePipelineVideoOutputBase::OnShotFinishedImpl(const UMoviePipelineExecu
 
 void UMoviePipelineVideoOutputBase::OnReceiveImageDataImpl(FMoviePipelineMergerOutputFrame* InMergedOutputFrame)
 {
+	// Add an early out if the output encountered an error (such as failure to initialize)
+	if (bHasError)
+	{
+		return;
+	}
+
 	UMoviePipelineOutputSetting* OutputSettings = GetPipeline()->GetPipelineMasterConfig()->FindSetting<UMoviePipelineOutputSetting>();
 	check(OutputSettings);
 
@@ -155,10 +167,12 @@ void UMoviePipelineVideoOutputBase::OnReceiveImageDataImpl(FMoviePipelineMergerO
 				OutputWriter->Get<0>().Get()->FormatArgs = FinalFormatArgs;
 
 				// If it fails to initialize, immediately mark the promise as failed so the render queue stops.
-				bool bResults = Initialize_EncodeThread(OutputWriter->Get<0>().Get());
-				if (!bResults)
+				bHasError = !Initialize_EncodeThread(OutputWriter->Get<0>().Get());
+				if (bHasError)
 				{
 					OutputWriter->Get<1>().SetValue(false);
+					UE_LOG(LogMovieRenderPipelineIO, Error, TEXT("Failed to initialize encoder for FileName: %s"), *FinalFilePath);
+					continue;
 				}
 			}
 		}
@@ -234,8 +248,10 @@ void UMoviePipelineVideoOutputBase::FinalizeImpl()
 		//*LastEvent = Task.Execute([this, RawWriter] {
 			this->Finalize_EncodeThread(RawWriter);
 		//	});
-
-			Writer.Get<1>().SetValue(true);
+			if (!bHasError)
+			{
+				Writer.Get<1>().SetValue(true);
+			}
 	}
 	
 	// Stall until all of the events are handled so that they still exist when the Task Graph goes to execute them.
@@ -245,6 +261,7 @@ void UMoviePipelineVideoOutputBase::FinalizeImpl()
 	}
 
 	AllWriters.Empty();
+	bHasError = false;
 }
 
 #if WITH_EDITOR
@@ -257,4 +274,4 @@ FText UMoviePipelineVideoOutputBase::GetFooterText(UMoviePipelineExecutorJob* In
 
 	return FText();
 }
-#endif
\ No newline at end of file
+#endif
diff --git a/Engine/Plugins/MovieScene/MovieRenderPipeline/Source/MovieRenderPipelineCore/Public/MoviePipelineVideoOutputBase.h b/Engine/Plugins/MovieScene/MovieRenderPipeline/Source/MovieRenderPipelineCore/Public/MoviePipelineVideoOutputBase.h
--- a/Engine/Plugins/MovieScene/MovieRenderPipeline/Source/MovieRenderPipelineCore/Public/MoviePipelineVideoOutputBase.h
+++ b/Engine/Plugins/MovieScene/MovieRenderPipeline/Source/MovieRenderPipelineCore/Public/MoviePipelineVideoOutputBase.h
@@ -57,8 +57,8 @@ class MOVIERENDERPIPELINECORE_API UMoviePipelineVideoOutputBase : public UMovieP
 {
 	GENERATED_BODY()
 
-protected:
-
+public:
+	UMoviePipelineVideoOutputBase();
 
 protected:
 	// UMoviePipelineOutputBase Interface
@@ -84,6 +84,7 @@ protected:
 
 private:
 	TArray<FMoviePipelineCodecWriter> AllWriters;
+	bool bHasError;
 
 	FGraphEventArray OutstandingTasks;
 };
\ No newline at end of file
