ObjectPtr upgrade for engine plugins used from UnrealEditor -allmodules.  Automated upgrade process.

diff --git a/Engine/Plugins/Experimental/Water/Source/Runtime/Public/WaterBodyLakeActor.h b/Engine/Plugins/Experimental/Water/Source/Runtime/Public/WaterBodyLakeActor.h
--- a/Engine/Plugins/Experimental/Water/Source/Runtime/Public/WaterBodyLakeActor.h
+++ b/Engine/Plugins/Experimental/Water/Source/Runtime/Public/WaterBodyLakeActor.h
@@ -17,13 +17,13 @@ class UDEPRECATED_LakeGenerator : public UDEPRECATED_WaterBodyGenerator
 	GENERATED_UCLASS_BODY()
 public:
 	UPROPERTY(NonPIEDuplicateTransient)
-	UStaticMeshComponent* LakeMeshComp;
+	TObjectPtr<UStaticMeshComponent> LakeMeshComp;
 
 	UPROPERTY()
-	UBoxComponent* LakeCollisionComp_DEPRECATED;
+	TObjectPtr<UBoxComponent> LakeCollisionComp_DEPRECATED;
 
 	UPROPERTY(NonPIEDuplicateTransient)
-	ULakeCollisionComponent* LakeCollision;
+	TObjectPtr<ULakeCollisionComponent> LakeCollision;
 };
 
 // ----------------------------------------------------------------------------------
@@ -37,6 +37,6 @@ protected:
 
 #if WITH_EDITORONLY_DATA
 	UPROPERTY()
-	UDEPRECATED_LakeGenerator* LakeGenerator_DEPRECATED;
+	TObjectPtr<UDEPRECATED_LakeGenerator> LakeGenerator_DEPRECATED;
 #endif
 };
\ No newline at end of file
diff --git a/Engine/Plugins/Experimental/Water/Source/Runtime/Public/WaterBodyLakeComponent.h b/Engine/Plugins/Experimental/Water/Source/Runtime/Public/WaterBodyLakeComponent.h
--- a/Engine/Plugins/Experimental/Water/Source/Runtime/Public/WaterBodyLakeComponent.h
+++ b/Engine/Plugins/Experimental/Water/Source/Runtime/Public/WaterBodyLakeComponent.h
@@ -39,8 +39,8 @@ protected:
 #endif // WITH_EDITOR
 
 	UPROPERTY(NonPIEDuplicateTransient)
-	UStaticMeshComponent* LakeMeshComp;
+	TObjectPtr<UStaticMeshComponent> LakeMeshComp;
 
 	UPROPERTY(NonPIEDuplicateTransient)
-	ULakeCollisionComponent* LakeCollision;
+	TObjectPtr<ULakeCollisionComponent> LakeCollision;
 };
\ No newline at end of file
diff --git a/Engine/Plugins/Experimental/Water/Source/Runtime/Public/WaterBodyOceanActor.h b/Engine/Plugins/Experimental/Water/Source/Runtime/Public/WaterBodyOceanActor.h
--- a/Engine/Plugins/Experimental/Water/Source/Runtime/Public/WaterBodyOceanActor.h
+++ b/Engine/Plugins/Experimental/Water/Source/Runtime/Public/WaterBodyOceanActor.h
@@ -17,10 +17,10 @@ class UDEPRECATED_OceanGenerator : public UDEPRECATED_WaterBodyGenerator
 	GENERATED_UCLASS_BODY()
 public:
 	UPROPERTY(NonPIEDuplicateTransient)
-	TArray<UOceanBoxCollisionComponent*> CollisionBoxes;
+	TArray<TObjectPtr<UOceanBoxCollisionComponent>> CollisionBoxes;
 
 	UPROPERTY(NonPIEDuplicateTransient)
-	TArray<UOceanCollisionComponent*> CollisionHullSets;
+	TArray<TObjectPtr<UOceanCollisionComponent>> CollisionHullSets;
 };
 
 // ----------------------------------------------------------------------------------
@@ -34,7 +34,7 @@ protected:
 
 #if WITH_EDITORONLY_DATA
 	UPROPERTY()
-	UDEPRECATED_OceanGenerator* OceanGenerator_DEPRECATED;
+	TObjectPtr<UDEPRECATED_OceanGenerator> OceanGenerator_DEPRECATED;
 	
 	UPROPERTY()
 	FVector CollisionExtents_DEPRECATED;
diff --git a/Engine/Plugins/Experimental/Water/Source/Runtime/Public/WaterBodyOceanComponent.h b/Engine/Plugins/Experimental/Water/Source/Runtime/Public/WaterBodyOceanComponent.h
--- a/Engine/Plugins/Experimental/Water/Source/Runtime/Public/WaterBodyOceanComponent.h
+++ b/Engine/Plugins/Experimental/Water/Source/Runtime/Public/WaterBodyOceanComponent.h
@@ -45,10 +45,10 @@ protected:
 #endif
 protected:
 	UPROPERTY(NonPIEDuplicateTransient)
-	TArray<UOceanBoxCollisionComponent*> CollisionBoxes;
+	TArray<TObjectPtr<UOceanBoxCollisionComponent>> CollisionBoxes;
 
 	UPROPERTY(NonPIEDuplicateTransient)
-	TArray<UOceanCollisionComponent*> CollisionHullSets;
+	TArray<TObjectPtr<UOceanCollisionComponent>> CollisionHullSets;
 
 	/** The area over which the ocean should be displayed, centered on the actor */
 	UPROPERTY(Category = Water, EditAnywhere, BlueprintReadOnly)
diff --git a/Engine/Plugins/Experimental/Water/Source/Runtime/Public/WaterBodyRiverActor.h b/Engine/Plugins/Experimental/Water/Source/Runtime/Public/WaterBodyRiverActor.h
--- a/Engine/Plugins/Experimental/Water/Source/Runtime/Public/WaterBodyRiverActor.h
+++ b/Engine/Plugins/Experimental/Water/Source/Runtime/Public/WaterBodyRiverActor.h
@@ -18,7 +18,7 @@ class UDEPRECATED_RiverGenerator : public UDEPRECATED_WaterBodyGenerator
 	GENERATED_UCLASS_BODY()
 public:
 	UPROPERTY(NonPIEDuplicateTransient)
-	TArray<USplineMeshComponent*> SplineMeshComponents;
+	TArray<TObjectPtr<USplineMeshComponent>> SplineMeshComponents;
 };
 
 // ----------------------------------------------------------------------------------
@@ -32,20 +32,20 @@ protected:
 
 #if WITH_EDITORONLY_DATA
 	UPROPERTY()
-	UDEPRECATED_RiverGenerator* RiverGenerator_DEPRECATED;
+	TObjectPtr<UDEPRECATED_RiverGenerator> RiverGenerator_DEPRECATED;
 	
 	/** Material used when a river is overlapping a lake. */
 	UPROPERTY()
-	UMaterialInterface* LakeTransitionMaterial_DEPRECATED;
+	TObjectPtr<UMaterialInterface> LakeTransitionMaterial_DEPRECATED;
 
 	UPROPERTY()
-	UMaterialInstanceDynamic* LakeTransitionMID_DEPRECATED;
+	TObjectPtr<UMaterialInstanceDynamic> LakeTransitionMID_DEPRECATED;
 
 	/** This is the material used when a river is overlapping the ocean. */
 	UPROPERTY()
-	UMaterialInterface* OceanTransitionMaterial_DEPRECATED;
+	TObjectPtr<UMaterialInterface> OceanTransitionMaterial_DEPRECATED;
 
 	UPROPERTY()
-	UMaterialInstanceDynamic* OceanTransitionMID_DEPRECATED;
+	TObjectPtr<UMaterialInstanceDynamic> OceanTransitionMID_DEPRECATED;
 #endif // WITH_EDITORONLY_DATA
 };
\ No newline at end of file
diff --git a/Engine/Plugins/Experimental/Water/Source/Runtime/Public/WaterBodyRiverComponent.h b/Engine/Plugins/Experimental/Water/Source/Runtime/Public/WaterBodyRiverComponent.h
--- a/Engine/Plugins/Experimental/Water/Source/Runtime/Public/WaterBodyRiverComponent.h
+++ b/Engine/Plugins/Experimental/Water/Source/Runtime/Public/WaterBodyRiverComponent.h
@@ -53,19 +53,19 @@ protected:
 
 protected:
 	UPROPERTY(NonPIEDuplicateTransient)
-	TArray<USplineMeshComponent*> SplineMeshComponents;
+	TArray<TObjectPtr<USplineMeshComponent>> SplineMeshComponents;
 
 	/** Material used when a river is overlapping a lake. */
 	UPROPERTY(Category = Rendering, EditAnywhere, BlueprintReadOnly, meta = (DisplayName = "River to Lake Transition"))
-	UMaterialInterface* LakeTransitionMaterial;
+	TObjectPtr<UMaterialInterface> LakeTransitionMaterial;
 
 	UPROPERTY(Category = Debug, VisibleInstanceOnly, Transient, NonPIEDuplicateTransient, TextExportTransient, meta = (DisplayAfter = "LakeTransitionMaterial"))
-	UMaterialInstanceDynamic* LakeTransitionMID;
+	TObjectPtr<UMaterialInstanceDynamic> LakeTransitionMID;
 
 	/** This is the material used when a river is overlapping the ocean. */
 	UPROPERTY(Category = Rendering, EditAnywhere, BlueprintReadOnly, meta = (DisplayName = "River to Ocean Transition"))
-	UMaterialInterface* OceanTransitionMaterial;
+	TObjectPtr<UMaterialInterface> OceanTransitionMaterial;
 
 	UPROPERTY(Category = Debug, VisibleInstanceOnly, Transient, NonPIEDuplicateTransient, TextExportTransient, meta = (DisplayAfter = "OceanTransitionMaterial"))
-	UMaterialInstanceDynamic* OceanTransitionMID;
+	TObjectPtr<UMaterialInstanceDynamic> OceanTransitionMID;
 };
\ No newline at end of file
diff --git a/Engine/Plugins/Experimental/Water/Source/Runtime/Public/WaterBodyWeightmapSettings.h b/Engine/Plugins/Experimental/Water/Source/Runtime/Public/WaterBodyWeightmapSettings.h
--- a/Engine/Plugins/Experimental/Water/Source/Runtime/Public/WaterBodyWeightmapSettings.h
+++ b/Engine/Plugins/Experimental/Water/Source/Runtime/Public/WaterBodyWeightmapSettings.h
@@ -31,7 +31,7 @@ struct FWaterBodyWeightmapSettings
 	float EdgeOffset;
 
 	UPROPERTY(EditAnywhere, BlueprintReadWrite, Category = WaterBodyWeightmapSettings)
-	UTexture2D* ModulationTexture;
+	TObjectPtr<UTexture2D> ModulationTexture;
 
 	UPROPERTY(EditAnywhere, BlueprintReadWrite, Category = WaterBodyWeightmapSettings)
 	float TextureTiling;
diff --git a/Engine/Plugins/Experimental/Water/Source/Runtime/Public/WaterBrushEffects.h b/Engine/Plugins/Experimental/Water/Source/Runtime/Public/WaterBrushEffects.h
--- a/Engine/Plugins/Experimental/Water/Source/Runtime/Public/WaterBrushEffects.h
+++ b/Engine/Plugins/Experimental/Water/Source/Runtime/Public/WaterBrushEffects.h
@@ -69,7 +69,7 @@ struct FWaterBrushEffectCurves
 	bool bUseCurveChannel;
 
 	UPROPERTY(EditAnywhere, BlueprintReadWrite, Category = BrushEffects)
-	UCurveFloat* ElevationCurveAsset;
+	TObjectPtr<UCurveFloat> ElevationCurveAsset;
 
 	UPROPERTY(EditAnywhere, BlueprintReadWrite, Category = BrushEffects)
 	float ChannelEdgeOffset;
@@ -103,7 +103,7 @@ struct FWaterBrushEffectDisplacement
 	float DisplacementTiling;
 
 	UPROPERTY(EditAnywhere, BlueprintReadWrite, Category = BrushEffects)
-	UTexture2D* Texture;
+	TObjectPtr<UTexture2D> Texture;
 
 	UPROPERTY(EditAnywhere, BlueprintReadWrite, Category = BrushEffects)
 	float Midpoint;
diff --git a/Engine/Plugins/Experimental/Water/Source/Runtime/Public/WaterCurveSettings.h b/Engine/Plugins/Experimental/Water/Source/Runtime/Public/WaterCurveSettings.h
--- a/Engine/Plugins/Experimental/Water/Source/Runtime/Public/WaterCurveSettings.h
+++ b/Engine/Plugins/Experimental/Water/Source/Runtime/Public/WaterCurveSettings.h
@@ -17,7 +17,7 @@ struct FWaterCurveSettings
 	bool bUseCurveChannel = true;
 
 	UPROPERTY(EditAnywhere, BlueprintReadWrite, Category = Water)
-	UCurveFloat* ElevationCurveAsset = nullptr;
+	TObjectPtr<UCurveFloat> ElevationCurveAsset = nullptr;
 
 	UPROPERTY(EditAnywhere, BlueprintReadWrite, Category = Water)
 	float ChannelEdgeOffset = 0.0f;
diff --git a/Engine/Plugins/Experimental/Water/Source/Runtime/Public/WaterMeshComponent.h b/Engine/Plugins/Experimental/Water/Source/Runtime/Public/WaterMeshComponent.h
--- a/Engine/Plugins/Experimental/Water/Source/Runtime/Public/WaterMeshComponent.h
+++ b/Engine/Plugins/Experimental/Water/Source/Runtime/Public/WaterMeshComponent.h
@@ -72,7 +72,7 @@ public:
 	int32 ForceCollapseDensityLevel = -1;
 
 	UPROPERTY(EditAnywhere, Category = "Mesh|FarDistance")
-	UMaterialInterface* FarDistanceMaterial = nullptr;
+	TObjectPtr<UMaterialInterface> FarDistanceMaterial = nullptr;
 
 	UPROPERTY(EditAnywhere, Category = "Mesh|FarDistance", meta = (ClampMin = "0"))
 	float FarDistanceMeshExtent = 0.0f;
@@ -100,7 +100,7 @@ private:
 
 	/** Unique list of materials used by this component */
 	UPROPERTY(Transient, NonPIEDuplicateTransient, TextExportTransient)
-	TSet<UMaterialInterface*> UsedMaterials;
+	TSet<TObjectPtr<UMaterialInterface>> UsedMaterials;
 
 	/** Dirty flag which will make sure the water mesh is updated properly */
 	bool bNeedsRebuild = true;
diff --git a/Engine/Plugins/Experimental/Water/Source/Runtime/Public/WaterSubsystem.h b/Engine/Plugins/Experimental/Water/Source/Runtime/Public/WaterSubsystem.h
--- a/Engine/Plugins/Experimental/Water/Source/Runtime/Public/WaterSubsystem.h
+++ b/Engine/Plugins/Experimental/Water/Source/Runtime/Public/WaterSubsystem.h
@@ -211,7 +211,7 @@ private:
 
 public:
 	UPROPERTY(Transient)
-	ABuoyancyManager* BuoyancyManager;
+	TObjectPtr<ABuoyancyManager> BuoyancyManager;
 
 	DECLARE_EVENT_OneParam(UWaterSubsystem, FOnWaterSubsystemInitialized, UWaterSubsystem*)
 	static FOnWaterSubsystemInitialized OnWaterSubsystemInitialized;
@@ -223,10 +223,10 @@ public:
 	FOnWaterScalabilityChanged OnWaterScalabilityChanged;
 
 	UPROPERTY()
-	UStaticMesh* DefaultRiverMesh;
+	TObjectPtr<UStaticMesh> DefaultRiverMesh;
 
 	UPROPERTY()
-	UStaticMesh* DefaultLakeMesh;
+	TObjectPtr<UStaticMesh> DefaultLakeMesh;
 
 private:
 
@@ -247,7 +247,7 @@ private:
 
 	/** The parameter collection asset that holds the global parameters that are updated by this actor */
 	UPROPERTY()
-	UMaterialParameterCollection* MaterialParameterCollection;
+	TObjectPtr<UMaterialParameterCollection> MaterialParameterCollection;
 
 	FUnderwaterPostProcessVolume UnderwaterPostProcessVolume;
 
diff --git a/Engine/Plugins/Experimental/Water/Source/Runtime/Public/WaterWaves.h b/Engine/Plugins/Experimental/Water/Source/Runtime/Public/WaterWaves.h
--- a/Engine/Plugins/Experimental/Water/Source/Runtime/Public/WaterWaves.h
+++ b/Engine/Plugins/Experimental/Water/Source/Runtime/Public/WaterWaves.h
@@ -87,7 +87,7 @@ public:
 
 private:
 	UPROPERTY(EditAnywhere, Category = "Water Waves", Instanced, DisplayName = "Waves Source")
-	UWaterWaves* WaterWaves = nullptr;
+	TObjectPtr<UWaterWaves> WaterWaves = nullptr;
 };
 
 UCLASS(BlueprintType, NotBlueprintable, AutoExpandCategories = "Water Waves Asset")
@@ -102,7 +102,7 @@ public:
 
 private:
 	UPROPERTY(EditAnywhere, Category = "Water Waves Asset")
-	UWaterWavesAsset* WaterWavesAsset;
+	TObjectPtr<UWaterWavesAsset> WaterWavesAsset;
 
 public:
 	virtual const UWaterWaves* GetWaterWaves() const override { return WaterWavesAsset ? WaterWavesAsset->GetWaterWaves() : nullptr; }
diff --git a/Engine/Plugins/Experimental/Water/Source/Runtime/Public/WaterZoneActor.h b/Engine/Plugins/Experimental/Water/Source/Runtime/Public/WaterZoneActor.h
--- a/Engine/Plugins/Experimental/Water/Source/Runtime/Public/WaterZoneActor.h
+++ b/Engine/Plugins/Experimental/Water/Source/Runtime/Public/WaterZoneActor.h
@@ -56,7 +56,7 @@ public:
 	float GetGroundZMin() const { return GroundZMin; }
 
 	UPROPERTY(Transient, DuplicateTransient, VisibleAnywhere, BlueprintReadOnly, Category = Texture)
-	UTextureRenderTarget2D* WaterInfoTexture;
+	TObjectPtr<UTextureRenderTarget2D> WaterInfoTexture;
 private:
 
 	bool UpdateWaterInfoTexture();
@@ -82,7 +82,7 @@ private:
 
 	/** The water mesh component */
 	UPROPERTY(VisibleAnywhere, Category = Water, BlueprintReadOnly, meta = (AllowPrivateAccess = "true"))
-	UWaterMeshComponent* WaterMesh;
+	TObjectPtr<UWaterMeshComponent> WaterMesh;
 
 	/** Radius of the zone bounding box */
 	UPROPERTY(Category = Shape, EditAnywhere, BlueprintReadWrite, meta = (AllowPrivateAccess = "true"))
@@ -108,15 +108,15 @@ private:
 #if WITH_EDITORONLY_DATA
 	/** A manipulatable box for visualizing/editing the water zone bounds */
 	UPROPERTY(Transient)
-	UBoxComponent* BoundsComponent;
+	TObjectPtr<UBoxComponent> BoundsComponent;
 
 	UPROPERTY(Transient)
 	TArray<TWeakObjectPtr<AWaterBody>> SelectedWaterBodies;
 
 	UPROPERTY(Transient)
-	UBillboardComponent* ActorIcon;
+	TObjectPtr<UBillboardComponent> ActorIcon;
 
 	UPROPERTY(Transient)
-	UTexture2D* WaterVelocityTexture_DEPRECATED;
+	TObjectPtr<UTexture2D> WaterVelocityTexture_DEPRECATED;
 #endif // WITH_EDITORONLY_DATA
 };
\ No newline at end of file
diff --git a/Engine/Plugins/Experimental/WaveFunctionCollapse/Source/WaveFunctionCollapse/Public/WaveFunctionCollapseSubsystem.h b/Engine/Plugins/Experimental/WaveFunctionCollapse/Source/WaveFunctionCollapse/Public/WaveFunctionCollapseSubsystem.h
--- a/Engine/Plugins/Experimental/WaveFunctionCollapse/Source/WaveFunctionCollapse/Public/WaveFunctionCollapseSubsystem.h
+++ b/Engine/Plugins/Experimental/WaveFunctionCollapse/Source/WaveFunctionCollapse/Public/WaveFunctionCollapseSubsystem.h
@@ -17,7 +17,7 @@ class WAVEFUNCTIONCOLLAPSE_API UWaveFunctionCollapseSubsystem : public UEditorSu
 public:
 
 	UPROPERTY(EditAnywhere, BlueprintReadWrite, Category = "WFCSettings")
-	UWaveFunctionCollapseModel* WFCModel;
+	TObjectPtr<UWaveFunctionCollapseModel> WFCModel;
 
 	UPROPERTY(EditAnywhere, BlueprintReadWrite, Category = "WFCSettings")
 	FIntVector Resolution;
diff --git a/Engine/Plugins/Experimental/WidgetEditorToolPalette/Source/WidgetEditorToolPalette/Public/DefaultTools/CreateWidgetTool.h b/Engine/Plugins/Experimental/WidgetEditorToolPalette/Source/WidgetEditorToolPalette/Public/DefaultTools/CreateWidgetTool.h
--- a/Engine/Plugins/Experimental/WidgetEditorToolPalette/Source/WidgetEditorToolPalette/Public/DefaultTools/CreateWidgetTool.h
+++ b/Engine/Plugins/Experimental/WidgetEditorToolPalette/Source/WidgetEditorToolPalette/Public/DefaultTools/CreateWidgetTool.h
@@ -69,7 +69,7 @@ public:
 
 protected:
 	UPROPERTY()
-	UCreateWidgetToolProperties* Properties;
+	TObjectPtr<UCreateWidgetToolProperties> Properties;
 
 protected:
 
diff --git a/Engine/Plugins/Experimental/WidgetEditorToolPalette/Source/WidgetEditorToolPalette/Public/DefaultTools/RectangleSelectTool.h b/Engine/Plugins/Experimental/WidgetEditorToolPalette/Source/WidgetEditorToolPalette/Public/DefaultTools/RectangleSelectTool.h
--- a/Engine/Plugins/Experimental/WidgetEditorToolPalette/Source/WidgetEditorToolPalette/Public/DefaultTools/RectangleSelectTool.h
+++ b/Engine/Plugins/Experimental/WidgetEditorToolPalette/Source/WidgetEditorToolPalette/Public/DefaultTools/RectangleSelectTool.h
@@ -70,7 +70,7 @@ public:
 protected:
 	/** Properties of the tool are stored here */
 	UPROPERTY()
-	URectangleSelectProperties* Properties;
+	TObjectPtr<URectangleSelectProperties> Properties;
 
 	/** Cached last end point to optimize comparisons */
 	FVector2D PrevSelectionEndPoint;
diff --git a/Engine/Plugins/FX/CascadeToNiagaraConverter/Source/CascadeToNiagaraConverter/Public/NiagaraStackGraphUtilitiesAdapterLibrary.h b/Engine/Plugins/FX/CascadeToNiagaraConverter/Source/CascadeToNiagaraConverter/Public/NiagaraStackGraphUtilitiesAdapterLibrary.h
--- a/Engine/Plugins/FX/CascadeToNiagaraConverter/Source/CascadeToNiagaraConverter/Public/NiagaraStackGraphUtilitiesAdapterLibrary.h
+++ b/Engine/Plugins/FX/CascadeToNiagaraConverter/Source/CascadeToNiagaraConverter/Public/NiagaraStackGraphUtilitiesAdapterLibrary.h
@@ -441,11 +441,11 @@ struct FStackEntryAddAction
 
 	// If Mode is Module, represents the pending module script to create a stack entry for. Otherwise this value is ignored.
 	UPROPERTY()
-	UNiagaraScriptConversionContext* ScriptConversionContext = nullptr;
+	TObjectPtr<UNiagaraScriptConversionContext> ScriptConversionContext = nullptr;
 
 	// If mode is SetParameter, represents the pending parameter to set directly and create a stack entry for. Otherwise this value is ignored.
 	UPROPERTY()
-	UNiagaraClipboardFunction* ClipboardFunction = nullptr;
+	TObjectPtr<UNiagaraClipboardFunction> ClipboardFunction = nullptr;
 
 	// Info to find the category of the stack to add the stack entry to.
 	UPROPERTY()
@@ -820,7 +820,7 @@ public:
 	void Init(UNiagaraClipboardFunctionInput* InClipboardFunctionInput, const ENiagaraScriptInputType InInputType, const FNiagaraTypeDefinition& InTypeDefinition);
 
 	UPROPERTY()
-	UNiagaraClipboardFunctionInput* ClipboardFunctionInput;
+	TObjectPtr<UNiagaraClipboardFunctionInput> ClipboardFunctionInput;
 
 	UPROPERTY(BlueprintReadOnly, Category = StaticValue)
 	ENiagaraScriptInputType InputType;
diff --git a/Engine/Plugins/Importers/USDImporter/Source/USDClasses/Public/USDAssetCache.h b/Engine/Plugins/Importers/USDImporter/Source/USDClasses/Public/USDAssetCache.h
--- a/Engine/Plugins/Importers/USDImporter/Source/USDClasses/Public/USDAssetCache.h
+++ b/Engine/Plugins/Importers/USDImporter/Source/USDClasses/Public/USDAssetCache.h
@@ -66,10 +66,10 @@ private:
 	// it will need to reload the stage anyway to rebuild its prim links to components and assets, so there is no point in duplicating
 	// the properties here just yet. Obviously we want to duplicate these properties
 	UPROPERTY( NonPIEDuplicateTransient, Transient, VisibleAnywhere, Category = "Assets" )
-	TMap< FString, UObject* > TransientStorage;
+	TMap< FString, TObjectPtr<UObject> > TransientStorage;
 
 	UPROPERTY( NonPIEDuplicateTransient, VisibleAnywhere, Category = "Assets" )
-	TMap< FString, UObject* > PersistentStorage;
+	TMap< FString, TObjectPtr<UObject> > PersistentStorage;
 
 	UPROPERTY( EditAnywhere, Category = "Assets", AdvancedDisplay )
 	bool bAllowPersistentStorage;
diff --git a/Engine/Plugins/Importers/USDImporter/Source/USDClasses/Public/USDAssetImportData.h b/Engine/Plugins/Importers/USDImporter/Source/USDClasses/Public/USDAssetImportData.h
--- a/Engine/Plugins/Importers/USDImporter/Source/USDClasses/Public/USDAssetImportData.h
+++ b/Engine/Plugins/Importers/USDImporter/Source/USDClasses/Public/USDAssetImportData.h
@@ -19,7 +19,7 @@ public:
 	// Likely a UUSDStageImportOptions, but we don't declare it here
 	// to prevent an unnecessary module dependency on USDStageImporter
 	UPROPERTY()
-	class UObject* ImportOptions;
+	TObjectPtr<class UObject> ImportOptions;
 };
 
 UCLASS(config = EditorPerProjectUserSettings, AutoExpandCategories = (Options), MinimalAPI)
diff --git a/Engine/Plugins/Importers/USDImporter/Source/USDStage/Public/USDPrimTwin.h b/Engine/Plugins/Importers/USDImporter/Source/USDStage/Public/USDPrimTwin.h
--- a/Engine/Plugins/Importers/USDImporter/Source/USDStage/Public/USDPrimTwin.h
+++ b/Engine/Plugins/Importers/USDImporter/Source/USDStage/Public/USDPrimTwin.h
@@ -68,7 +68,7 @@ public:
 
 private:
 	UPROPERTY( Transient )
-	TMap< FString, UUsdPrimTwin* > Children;
+	TMap< FString, TObjectPtr<UUsdPrimTwin> > Children;
 
 	UPROPERTY( Transient )
 	TWeakObjectPtr< UUsdPrimTwin > Parent;
diff --git a/Engine/Plugins/Importers/USDImporter/Source/USDStage/Public/USDStageActor.h b/Engine/Plugins/Importers/USDImporter/Source/USDStage/Public/USDStageActor.h
--- a/Engine/Plugins/Importers/USDImporter/Source/USDStage/Public/USDStageActor.h
+++ b/Engine/Plugins/Importers/USDImporter/Source/USDStage/Public/USDStageActor.h
@@ -147,7 +147,7 @@ public:
 
 private:
 	UPROPERTY(Category = UsdStageActor, VisibleAnywhere, BlueprintReadOnly, meta = (ExposeFunctionCategories = "Mesh,Rendering,Physics,Components|StaticMesh", AllowPrivateAccess = "true"))
-	class USceneComponent* SceneComponent;
+	TObjectPtr<class USceneComponent> SceneComponent;
 
 	/* TimeCode to evaluate the USD stage at */
 	UPROPERTY(EditAnywhere, Category = "USD")
@@ -163,7 +163,7 @@ private:
 	float TimeCodesPerSecond_DEPRECATED;
 
 	UPROPERTY(VisibleAnywhere, Category = "USD", Transient)
-	ULevelSequence* LevelSequence;
+	TObjectPtr<ULevelSequence> LevelSequence;
 
 public:
 	DECLARE_EVENT_OneParam( AUsdStageActor, FOnActorLoaded, AUsdStageActor* );
@@ -243,19 +243,19 @@ private:
 private:
 	// Note: This cannot be instanced: Read the comment in the constructor
 	UPROPERTY( Transient )
-	UUsdPrimTwin* RootUsdTwin;
+	TObjectPtr<UUsdPrimTwin> RootUsdTwin;
 
 	UPROPERTY( Transient )
 	TSet< FString > PrimsToAnimate;
 
 	UPROPERTY( Transient )
-	TMap< UObject*, FString > ObjectsToWatch;
+	TMap< TObjectPtr<UObject>, FString > ObjectsToWatch;
 
 	UPROPERTY( VisibleAnywhere, Category = "USD", AdvancedDisplay )
-	UUsdAssetCache* AssetCache;
+	TObjectPtr<UUsdAssetCache> AssetCache;
 
 	UPROPERTY( Transient )
-	UUsdTransactor* Transactor;
+	TObjectPtr<UUsdTransactor> Transactor;
 
 public:
 	UE_DEPRECATED( 4.27, "Use the const version if you don't wish to load the stage on-demand, or use GetOrLoadUsdStage if you do" )
diff --git a/Engine/Plugins/Importers/USDImporter/Source/USDStageEditor/Private/Widgets/SUSDIntegrationsPanel.h b/Engine/Plugins/Importers/USDImporter/Source/USDStageEditor/Private/Widgets/SUSDIntegrationsPanel.h
--- a/Engine/Plugins/Importers/USDImporter/Source/USDStageEditor/Private/Widgets/SUSDIntegrationsPanel.h
+++ b/Engine/Plugins/Importers/USDImporter/Source/USDStageEditor/Private/Widgets/SUSDIntegrationsPanel.h
@@ -25,10 +25,10 @@ public:
 	// Ideally these would be FSoftObjectPaths, but FPropertyEditorModule doesn't support generating
 	// widgets for those properties yet
 	UPROPERTY( EditAnywhere, Transient, Category = Dummy, meta = ( AllowedClasses = "/Script/Engine.AnimBlueprint" ) )
-	UObject* AnimBPProperty;
+	TObjectPtr<UObject> AnimBPProperty;
 
 	UPROPERTY( EditAnywhere, Transient, Category = Dummy, meta = ( AllowedClasses = "/Script/ControlRigDeveloper.ControlRigBlueprint" ) )
-	UObject* ControlRigProperty;
+	TObjectPtr<UObject> ControlRigProperty;
 };
 
 #if USE_USD_SDK
diff --git a/Engine/Plugins/Importers/USDImporter/Source/USDStageImporter/Public/USDStageImportContext.h b/Engine/Plugins/Importers/USDImporter/Source/USDStageImporter/Public/USDStageImportContext.h
--- a/Engine/Plugins/Importers/USDImporter/Source/USDStageImporter/Public/USDStageImportContext.h
+++ b/Engine/Plugins/Importers/USDImporter/Source/USDStageImporter/Public/USDStageImportContext.h
@@ -34,7 +34,7 @@ struct USDSTAGEIMPORTER_API FUsdStageImportContext
 
 	/** Spawned actor that contains the imported scene as a child hierarchy */
 	UPROPERTY()
-	AActor* SceneActor;
+	TObjectPtr<AActor> SceneActor;
 
 	/** Name to use when importing a single mesh */
 	UPROPERTY()
@@ -48,17 +48,17 @@ struct USDSTAGEIMPORTER_API FUsdStageImportContext
 	FString FilePath;
 
 	UPROPERTY()
-	UUsdStageImportOptions* ImportOptions;
+	TObjectPtr<UUsdStageImportOptions> ImportOptions;
 
 	/** Keep track of the last imported object so that we have something valid to return to upstream code that calls the import factories */
 	UPROPERTY()
-	UObject* ImportedAsset;
+	TObjectPtr<UObject> ImportedAsset;
 
 	/** Level sequence that will contain the animation data during the import process */
 	FUsdLevelSequenceHelper LevelSequenceHelper;
 
 	UPROPERTY()
-	UUsdAssetCache* AssetCache;
+	TObjectPtr<UUsdAssetCache> AssetCache;
 
 	/** Caches various information about prims that are expensive to query */
 	TSharedPtr<FUsdInfoCache> InfoCache;
diff --git a/Engine/Plugins/Media/AjaMedia/Source/AjaMedia/Public/AjaCustomTimeStep.h b/Engine/Plugins/Media/AjaMedia/Source/AjaMedia/Public/AjaCustomTimeStep.h
--- a/Engine/Plugins/Media/AjaMedia/Source/AjaMedia/Public/AjaCustomTimeStep.h
+++ b/Engine/Plugins/Media/AjaMedia/Source/AjaMedia/Public/AjaCustomTimeStep.h
@@ -93,7 +93,7 @@ private:
 #if WITH_EDITORONLY_DATA
 	/** Engine used to initialize the CustomTimeStep */
 	UPROPERTY(Transient)
-	UEngine* InitializedEngine;
+	TObjectPtr<UEngine> InitializedEngine;
 
 	/** When Auto synchronize is enabled, the time the last attempt was triggered. */
 	double LastAutoSynchronizeInEditorAppTime;
diff --git a/Engine/Plugins/Media/AjaMedia/Source/AjaMedia/Public/AjaTimecodeProvider.h b/Engine/Plugins/Media/AjaMedia/Source/AjaMedia/Public/AjaTimecodeProvider.h
--- a/Engine/Plugins/Media/AjaMedia/Source/AjaMedia/Public/AjaTimecodeProvider.h
+++ b/Engine/Plugins/Media/AjaMedia/Source/AjaMedia/Public/AjaTimecodeProvider.h
@@ -116,7 +116,7 @@ private:
 #if WITH_EDITORONLY_DATA
 	/** Engine used to initialize the Provider */
 	UPROPERTY(Transient)
-	UEngine* InitializedEngine;
+	TObjectPtr<UEngine> InitializedEngine;
 
 	/** The time the last attempt to auto synchronize was triggered. */
 	double LastAutoSynchronizeInEditorAppTime;
diff --git a/Engine/Plugins/Media/AjaMedia/Source/AjaMediaOutput/Public/AjaMediaFrameGrabberProtocol.h b/Engine/Plugins/Media/AjaMedia/Source/AjaMediaOutput/Public/AjaMediaFrameGrabberProtocol.h
--- a/Engine/Plugins/Media/AjaMedia/Source/AjaMediaOutput/Public/AjaMediaFrameGrabberProtocol.h
+++ b/Engine/Plugins/Media/AjaMedia/Source/AjaMediaOutput/Public/AjaMediaFrameGrabberProtocol.h
@@ -40,9 +40,9 @@ private:
 
 	/** Transient media output pointer to keep the media output alive while this protocol is in use */
 	UPROPERTY(Transient)
-	UAjaMediaOutput* TransientMediaOutputPtr;
+	TObjectPtr<UAjaMediaOutput> TransientMediaOutputPtr;
 
 	/** Transient media capture pointer that will capture the viewport */
 	UPROPERTY(Transient)
-	UAjaMediaCapture* TransientMediaCapturePtr;
+	TObjectPtr<UAjaMediaCapture> TransientMediaCapturePtr;
 };
diff --git a/Engine/Plugins/Media/BinkMedia/Source/BinkMediaPlayer/Public/BinkMediaTexture.h b/Engine/Plugins/Media/BinkMedia/Source/BinkMediaPlayer/Public/BinkMediaTexture.h
--- a/Engine/Plugins/Media/BinkMedia/Source/BinkMediaPlayer/Public/BinkMediaTexture.h
+++ b/Engine/Plugins/Media/BinkMedia/Source/BinkMediaPlayer/Public/BinkMediaTexture.h
@@ -73,7 +73,7 @@ class BINKMEDIAPLAYER_API UBinkMediaTexture : public UTexture
 
 	/** The MediaPlayer asset to stream video from. */
 	UPROPERTY(EditAnywhere, BlueprintReadOnly, Category=MediaPlayer)
-	UBinkMediaPlayer *MediaPlayer;
+	TObjectPtr<UBinkMediaPlayer> MediaPlayer;
 
 	/** The Pixel Format for the texture. */
 	UPROPERTY(EditAnywhere, BlueprintReadOnly, Category=MediaPlayer)
diff --git a/Engine/Plugins/Media/BinkMedia/Source/BinkMediaPlayerEditor/Private/Factories/BinkMediaTextureFactoryNew.h b/Engine/Plugins/Media/BinkMedia/Source/BinkMediaPlayerEditor/Private/Factories/BinkMediaTextureFactoryNew.h
--- a/Engine/Plugins/Media/BinkMedia/Source/BinkMediaPlayerEditor/Private/Factories/BinkMediaTextureFactoryNew.h
+++ b/Engine/Plugins/Media/BinkMedia/Source/BinkMediaPlayerEditor/Private/Factories/BinkMediaTextureFactoryNew.h
@@ -12,7 +12,7 @@ class UBinkMediaTextureFactoryNew : public UFactory
 	GENERATED_UCLASS_BODY()
 public:
 	UPROPERTY()
-	class UBinkMediaPlayer* InitialMediaPlayer;
+	TObjectPtr<class UBinkMediaPlayer> InitialMediaPlayer;
 
 	virtual UObject* FactoryCreateNew( UClass* InClass, UObject* InParent, FName InName, EObjectFlags Flags, UObject* Context, FFeedbackContext* Warn ) override;
 	virtual bool ShouldShowInNewMenu() const override 
diff --git a/Engine/Plugins/Media/MediaFrameworkUtilities/Source/MediaFrameworkUtilities/Public/MediaAssets/ProxyMediaOutput.h b/Engine/Plugins/Media/MediaFrameworkUtilities/Source/MediaFrameworkUtilities/Public/MediaAssets/ProxyMediaOutput.h
--- a/Engine/Plugins/Media/MediaFrameworkUtilities/Source/MediaFrameworkUtilities/Public/MediaAssets/ProxyMediaOutput.h
+++ b/Engine/Plugins/Media/MediaFrameworkUtilities/Source/MediaFrameworkUtilities/Public/MediaAssets/ProxyMediaOutput.h
@@ -22,11 +22,11 @@ private:
 
 	/** Dynamic media output proxy. */
 	UPROPERTY(Transient, DuplicateTransient)
-	UMediaOutput* DynamicProxy;
+	TObjectPtr<UMediaOutput> DynamicProxy;
 
 	/** Media output proxy. */
 	UPROPERTY(EditAnywhere, Category = "Media Proxy")
-	UMediaOutput* Proxy;
+	TObjectPtr<UMediaOutput> Proxy;
 
 public:
 
diff --git a/Engine/Plugins/Media/MediaFrameworkUtilities/Source/MediaFrameworkUtilities/Public/MediaAssets/ProxyMediaSource.h b/Engine/Plugins/Media/MediaFrameworkUtilities/Source/MediaFrameworkUtilities/Public/MediaAssets/ProxyMediaSource.h
--- a/Engine/Plugins/Media/MediaFrameworkUtilities/Source/MediaFrameworkUtilities/Public/MediaAssets/ProxyMediaSource.h
+++ b/Engine/Plugins/Media/MediaFrameworkUtilities/Source/MediaFrameworkUtilities/Public/MediaAssets/ProxyMediaSource.h
@@ -22,11 +22,11 @@ private:
 
 	/** Cached media source proxy. */
 	UPROPERTY(Transient, DuplicateTransient)
-	UMediaSource* DynamicProxy;
+	TObjectPtr<UMediaSource> DynamicProxy;
 
 	/** Media source proxy. */
 	UPROPERTY(EditAnywhere, Category="Media Proxy")
-	UMediaSource* Proxy;
+	TObjectPtr<UMediaSource> Proxy;
 
 public:
 
diff --git a/Engine/Plugins/Media/MediaFrameworkUtilities/Source/MediaFrameworkUtilities/Public/MediaBundle.h b/Engine/Plugins/Media/MediaFrameworkUtilities/Source/MediaFrameworkUtilities/Public/MediaBundle.h
--- a/Engine/Plugins/Media/MediaFrameworkUtilities/Source/MediaFrameworkUtilities/Public/MediaBundle.h
+++ b/Engine/Plugins/Media/MediaFrameworkUtilities/Source/MediaFrameworkUtilities/Public/MediaBundle.h
@@ -82,7 +82,7 @@ public:
 
 	/* MediaBundle internal MediaSource */
 	UPROPERTY(EditAnywhere, Instanced, NoClear, Category="MediaBundle")
-	UMediaSource* MediaSource;
+	TObjectPtr<UMediaSource> MediaSource;
 	
 	/** Controls MediaPlayer looping option */
 	UPROPERTY(EditAnywhere, Category="MediaBundle")
@@ -101,15 +101,15 @@ public:
 protected:
 	/* MediaBundle default MediaPlayer */
 	UPROPERTY(Instanced)
-	UMediaPlayer* MediaPlayer;
+	TObjectPtr<UMediaPlayer> MediaPlayer;
 
 	/* MediaBundle default MediaTexture */
 	UPROPERTY(Instanced)
-	UMediaTexture* MediaTexture;
+	TObjectPtr<UMediaTexture> MediaTexture;
 
 	/* MediaBundle default Material */
 	UPROPERTY(Instanced)
-	UMaterialInterface* Material;
+	TObjectPtr<UMaterialInterface> Material;
 
 	/* Lens parameters of the source */
 	UPROPERTY(EditAnywhere, Category = "LensParameters")
@@ -125,7 +125,7 @@ protected:
 
 	/* Destination of lens distortion result */
 	UPROPERTY(Instanced)
-	UTextureRenderTarget2D* LensDisplacementMap;
+	TObjectPtr<UTextureRenderTarget2D> LensDisplacementMap;
 
 private:
 	/* Internal reference counter of active media player */
@@ -135,11 +135,11 @@ private:
 #if WITH_EDITORONLY_DATA
 	/* Default Material from the plugin*/
 	UPROPERTY(transient)
-	UMaterial* DefaultMaterial;
+	TObjectPtr<UMaterial> DefaultMaterial;
 
 	/* Default Texture from the plugin*/
 	UPROPERTY(transient)
-	UTexture* DefaultFailedTexture;
+	TObjectPtr<UTexture> DefaultFailedTexture;
 
 	/* Default Actor Class from the plugin*/
 	UPROPERTY(transient)
diff --git a/Engine/Plugins/Media/MediaFrameworkUtilities/Source/MediaFrameworkUtilities/Public/MediaBundleActorBase.h b/Engine/Plugins/Media/MediaFrameworkUtilities/Source/MediaFrameworkUtilities/Public/MediaBundleActorBase.h
--- a/Engine/Plugins/Media/MediaFrameworkUtilities/Source/MediaFrameworkUtilities/Public/MediaBundleActorBase.h
+++ b/Engine/Plugins/Media/MediaFrameworkUtilities/Source/MediaFrameworkUtilities/Public/MediaBundleActorBase.h
@@ -57,13 +57,13 @@ public:
 
 	/** Texture containging 2D garbage matte mask */
 	UPROPERTY(EditAnywhere, Category="Compositing")
-	UTextureRenderTarget2D* GarbageMatteMask;
+	TObjectPtr<UTextureRenderTarget2D> GarbageMatteMask;
 	
 protected:
 
 	/** Associated MediaBundle */
 	UPROPERTY(EditAnywhere, NoClear, Category="MediaBundle")
-	UMediaBundle* MediaBundle;
+	TObjectPtr<UMediaBundle> MediaBundle;
 
 	/** Wheter to auto start the MediaPlayer */
 	UPROPERTY(EditDefaultsOnly, Category="MediaBundle")
@@ -75,15 +75,15 @@ protected:
 
 	/** PrimitiveComponent on which to attach our Material */
 	UPROPERTY(EditDefaultsOnly, Category="MediaBundle")
-	UPrimitiveComponent* PrimitiveCmp;
+	TObjectPtr<UPrimitiveComponent> PrimitiveCmp;
 
 	/** MediaSoundComponent associated to play sound of our MediaSource */
 	UPROPERTY(EditDefaultsOnly, NoClear, Category = "MediaBundle")
-	UMediaSoundComponent* MediaSoundCmp;
+	TObjectPtr<UMediaSoundComponent> MediaSoundCmp;
 	
 	/** Dynamic instance of the associated MediaBundle base Material */
 	UPROPERTY(EditDefaultsOnly, Category="MediaBundle")
-	UMaterialInstanceDynamic* Material;
+	TObjectPtr<UMaterialInstanceDynamic> Material;
 
 	/** Index of the Material on the primitive */
 	UPROPERTY(AdvancedDisplay, EditAnywhere, Category="MediaBundle")
diff --git a/Engine/Plugins/Media/MediaFrameworkUtilities/Source/MediaFrameworkUtilities/Public/MediaBundleTimeSynchronizationSource.h b/Engine/Plugins/Media/MediaFrameworkUtilities/Source/MediaFrameworkUtilities/Public/MediaBundleTimeSynchronizationSource.h
--- a/Engine/Plugins/Media/MediaFrameworkUtilities/Source/MediaFrameworkUtilities/Public/MediaBundleTimeSynchronizationSource.h
+++ b/Engine/Plugins/Media/MediaFrameworkUtilities/Source/MediaFrameworkUtilities/Public/MediaBundleTimeSynchronizationSource.h
@@ -20,7 +20,7 @@ class MEDIAFRAMEWORKUTILITIES_API UMediaBundleTimeSynchronizationSource : public
 public:
 	/* Media bundle asset of this input*/
 	UPROPERTY(EditAnywhere, Category=Player)
-	UMediaBundle* MediaBundle;
+	TObjectPtr<UMediaBundle> MediaBundle;
 
 public:
 
diff --git a/Engine/Plugins/Media/MediaFrameworkUtilities/Source/MediaFrameworkUtilities/Public/MediaPlayerTimeSynchronizationSource.h b/Engine/Plugins/Media/MediaFrameworkUtilities/Source/MediaFrameworkUtilities/Public/MediaPlayerTimeSynchronizationSource.h
--- a/Engine/Plugins/Media/MediaFrameworkUtilities/Source/MediaFrameworkUtilities/Public/MediaPlayerTimeSynchronizationSource.h
+++ b/Engine/Plugins/Media/MediaFrameworkUtilities/Source/MediaFrameworkUtilities/Public/MediaPlayerTimeSynchronizationSource.h
@@ -20,11 +20,11 @@ class MEDIAFRAMEWORKUTILITIES_API UMediaPlayerTimeSynchronizationSource : public
 public:
 	/* Media source asset of this input*/
 	UPROPERTY(EditAnywhere, Category=Player)
-	UMediaSource* MediaSource;
+	TObjectPtr<UMediaSource> MediaSource;
 
 	/* Texture linked to the media player*/
 	UPROPERTY(EditAnywhere, Category = Player)
-	UMediaTexture* MediaTexture;
+	TObjectPtr<UMediaTexture> MediaTexture;
 
 public:
 
diff --git a/Engine/Plugins/Media/MediaFrameworkUtilities/Source/MediaFrameworkUtilities/Public/Profile/MediaProfile.h b/Engine/Plugins/Media/MediaFrameworkUtilities/Source/MediaFrameworkUtilities/Public/Profile/MediaProfile.h
--- a/Engine/Plugins/Media/MediaFrameworkUtilities/Source/MediaFrameworkUtilities/Public/Profile/MediaProfile.h
+++ b/Engine/Plugins/Media/MediaFrameworkUtilities/Source/MediaFrameworkUtilities/Public/Profile/MediaProfile.h
@@ -25,11 +25,11 @@ protected:
 
 	/** Media sources. */
 	UPROPERTY(EditAnywhere, Instanced, Category="Inputs", EditFixedSize, meta=(EditFixedOrder))
-	TArray<UMediaSource*> MediaSources;
+	TArray<TObjectPtr<UMediaSource>> MediaSources;
 
 	/** Media outputs. */
 	UPROPERTY(EditAnywhere, Instanced, Category="Outputs", EditFixedSize, meta=(EditFixedOrder))
-	TArray<UMediaOutput*> MediaOutputs;
+	TArray<TObjectPtr<UMediaOutput>> MediaOutputs;
 
 	/** Override the Engine's Timecode provider defined in the project settings. */
 	UPROPERTY(EditAnywhere, Category="Timecode Provider", meta=(DisplayName="Override Project Settings"))
@@ -37,7 +37,7 @@ protected:
 
 	/** Timecode provider. */
 	UPROPERTY(EditAnywhere, Instanced, Category="Timecode Provider", meta=(EditCondition="bOverrideTimecodeProvider"))
-	UTimecodeProvider* TimecodeProvider;
+	TObjectPtr<UTimecodeProvider> TimecodeProvider;
 
 	/** Override the Engine's Custom time step defined in the project settings. */
 	UPROPERTY(EditAnywhere, Category="Genlock", meta=(DisplayName="Override Project Settings"))
@@ -45,7 +45,7 @@ protected:
 
 	/** Custom time step */
 	UPROPERTY(EditAnywhere, Instanced, Category="Genlock", meta=(EditCondition="bOverrideCustomTimeStep"))
-	UEngineCustomTimeStep* CustomTimeStep;
+	TObjectPtr<UEngineCustomTimeStep> CustomTimeStep;
 
 public:
 #if WITH_EDITORONLY_DATA
@@ -133,14 +133,14 @@ private:
 	/** Applied Timecode provider, cached to reset the previous value. */
 	bool bTimecodeProvideWasApplied;
 	UPROPERTY(Transient)
-	UTimecodeProvider* AppliedTimecodeProvider;
+	TObjectPtr<UTimecodeProvider> AppliedTimecodeProvider;
 	UPROPERTY(Transient)
-	UTimecodeProvider* PreviousTimecodeProvider;
+	TObjectPtr<UTimecodeProvider> PreviousTimecodeProvider;
 
 	/** Applied Custom time step, cached to reset the previous value. */
 	bool bCustomTimeStepWasApplied;
 	UPROPERTY(Transient)
-	UEngineCustomTimeStep* AppliedCustomTimeStep;
+	TObjectPtr<UEngineCustomTimeStep> AppliedCustomTimeStep;
 	UPROPERTY(Transient)
-	UEngineCustomTimeStep* PreviousCustomTimeStep;
+	TObjectPtr<UEngineCustomTimeStep> PreviousCustomTimeStep;
 };
diff --git a/Engine/Plugins/Media/MediaFrameworkUtilities/Source/MediaFrameworkUtilitiesEditor/Private/MediaFrameworkWorldSettingsAssetUserData.h b/Engine/Plugins/Media/MediaFrameworkUtilities/Source/MediaFrameworkUtilitiesEditor/Private/MediaFrameworkWorldSettingsAssetUserData.h
--- a/Engine/Plugins/Media/MediaFrameworkUtilities/Source/MediaFrameworkUtilitiesEditor/Private/MediaFrameworkWorldSettingsAssetUserData.h
+++ b/Engine/Plugins/Media/MediaFrameworkUtilities/Source/MediaFrameworkUtilitiesEditor/Private/MediaFrameworkWorldSettingsAssetUserData.h
@@ -26,7 +26,7 @@ struct FMediaFrameworkCaptureCurrentViewportOutputInfo
 	FMediaFrameworkCaptureCurrentViewportOutputInfo();
 
 	UPROPERTY(EditAnywhere, Category="MediaViewportCapture")
-	UMediaOutput* MediaOutput;
+	TObjectPtr<UMediaOutput> MediaOutput;
 
 	UPROPERTY(EditAnywhere, Category="MediaViewportCapture")
 	FMediaCaptureOptions CaptureOptions;
@@ -49,7 +49,7 @@ struct FMediaFrameworkCaptureCameraViewportCameraOutputInfo
 	TArray<TLazyObjectPtr<AActor>> LockedActors;
 
 	UPROPERTY(EditAnywhere, Category="MediaViewportCapture")
-	UMediaOutput* MediaOutput;
+	TObjectPtr<UMediaOutput> MediaOutput;
 
 	UPROPERTY(EditAnywhere, Category="MediaViewportCapture")
 	FMediaCaptureOptions CaptureOptions;
@@ -60,7 +60,7 @@ struct FMediaFrameworkCaptureCameraViewportCameraOutputInfo
 private:
 	//DEPRECATED 4.21 The type of LockedCameraActors has changed and will be removed from the code base in a future release. Use LockedActors.
 	UPROPERTY()
-	TArray<AActor*> LockedCameraActors_DEPRECATED;
+	TArray<TObjectPtr<AActor>> LockedCameraActors_DEPRECATED;
 	friend UMediaFrameworkWorldSettingsAssetUserData;
 };
 
@@ -76,10 +76,10 @@ struct FMediaFrameworkCaptureRenderTargetCameraOutputInfo
 	FMediaFrameworkCaptureRenderTargetCameraOutputInfo();
 
 	UPROPERTY(EditAnywhere, Category="MediaRenderTargetCapture")
-	UTextureRenderTarget2D* RenderTarget;
+	TObjectPtr<UTextureRenderTarget2D> RenderTarget;
 
 	UPROPERTY(EditAnywhere, Category="MediaRenderTargetCapture")
-	UMediaOutput* MediaOutput;
+	TObjectPtr<UMediaOutput> MediaOutput;
 
 	UPROPERTY(EditAnywhere, Category="MediaRenderTargetCapture")
 	FMediaCaptureOptions CaptureOptions;
diff --git a/Engine/Plugins/Media/MediaMovieStreamer/Source/MediaMovieStreamer/Private/MediaMovieAssets.h b/Engine/Plugins/Media/MediaMovieStreamer/Source/MediaMovieStreamer/Private/MediaMovieAssets.h
--- a/Engine/Plugins/Media/MediaMovieStreamer/Source/MediaMovieStreamer/Private/MediaMovieAssets.h
+++ b/Engine/Plugins/Media/MediaMovieStreamer/Source/MediaMovieStreamer/Private/MediaMovieAssets.h
@@ -66,19 +66,19 @@ private:
 
 	/** Holds the player we are using. */
 	UPROPERTY()
-	UMediaPlayer* MediaPlayer;
+	TObjectPtr<UMediaPlayer> MediaPlayer;
 
 	/** Holds the media sound component we are using. */
 	UPROPERTY()
-	UMediaSoundComponent* MediaSoundComponent;
+	TObjectPtr<UMediaSoundComponent> MediaSoundComponent;
 
 	/** Holds the media source we are using. */
 	UPROPERTY()
-	UMediaSource* MediaSource;
+	TObjectPtr<UMediaSource> MediaSource;
 
 	/** Holds the media texture we are using. */
 	UPROPERTY()
-	UMediaTexture* MediaTexture;
+	TObjectPtr<UMediaTexture> MediaTexture;
 
 	/** Holds the movie streamer. */
 	FMediaMovieStreamer* MovieStreamer;
diff --git a/Engine/Plugins/Media/MediaSourceManager/Source/MediaSourceManager/Public/MediaSourceManager.h b/Engine/Plugins/Media/MediaSourceManager/Source/MediaSourceManager/Public/MediaSourceManager.h
--- a/Engine/Plugins/Media/MediaSourceManager/Source/MediaSourceManager/Public/MediaSourceManager.h
+++ b/Engine/Plugins/Media/MediaSourceManager/Source/MediaSourceManager/Public/MediaSourceManager.h
@@ -20,7 +20,7 @@ class MEDIASOURCEMANAGER_API UMediaSourceManager : public UObject
 public:
 	/** Our channels. */
 	UPROPERTY(EditAnywhere, BlueprintReadWrite, Category = "Channels")
-	TArray<UMediaSourceManagerChannel*> Channels;
+	TArray<TObjectPtr<UMediaSourceManagerChannel>> Channels;
 
 	/**
 	 * Call this to make sure everything is set up.
diff --git a/Engine/Plugins/Media/PixelStreamingPlayer/Source/PixelStreamingPlayer/Public/PixelStreamingPeerComponent.h b/Engine/Plugins/Media/PixelStreamingPlayer/Source/PixelStreamingPlayer/Public/PixelStreamingPeerComponent.h
--- a/Engine/Plugins/Media/PixelStreamingPlayer/Source/PixelStreamingPlayer/Public/PixelStreamingPeerComponent.h
+++ b/Engine/Plugins/Media/PixelStreamingPlayer/Source/PixelStreamingPlayer/Public/PixelStreamingPeerComponent.h
@@ -66,7 +66,7 @@ public:
 	 * A sink for the video data received once this connection has finished negotiating.
 	 */
 	UPROPERTY(EditDefaultsOnly, Category = "Properties", META = (DisplayName = "Pixel Streaming Video Sink", AllowPrivateAccess = true))
-	UPixelStreamingMediaTexture* VideoSink = nullptr;
+	TObjectPtr<UPixelStreamingMediaTexture> VideoSink = nullptr;
 
 private:
 	TUniquePtr<FPixelStreamingPeerConnection> PeerConnection;
diff --git a/Engine/Plugins/Media/PixelStreamingPlayer/Source/PixelStreamingPlayer/Public/PixelStreamingSignallingComponent.h b/Engine/Plugins/Media/PixelStreamingPlayer/Source/PixelStreamingPlayer/Public/PixelStreamingSignallingComponent.h
--- a/Engine/Plugins/Media/PixelStreamingPlayer/Source/PixelStreamingPlayer/Public/PixelStreamingSignallingComponent.h
+++ b/Engine/Plugins/Media/PixelStreamingPlayer/Source/PixelStreamingPlayer/Public/PixelStreamingSignallingComponent.h
@@ -94,7 +94,7 @@ public:
 	 * If this media source is set we will use its supplied URL instead of the Url parameter on the connect call.
 	 */
 	UPROPERTY(EditDefaultsOnly, Category = "Properties", META = (DisplayName = "Stream Media Source", AllowPrivateAccess = true))
-	UStreamMediaSource* MediaSource = nullptr;
+	TObjectPtr<UStreamMediaSource> MediaSource = nullptr;
 
 protected:
 	//
diff --git a/Engine/Plugins/Media/PixelStreamingVCam/Source/PixelStreamingVCam/Public/VCamPixelStreamingSession.h b/Engine/Plugins/Media/PixelStreamingVCam/Source/PixelStreamingVCam/Public/VCamPixelStreamingSession.h
--- a/Engine/Plugins/Media/PixelStreamingVCam/Source/PixelStreamingVCam/Public/VCamPixelStreamingSession.h
+++ b/Engine/Plugins/Media/PixelStreamingVCam/Source/PixelStreamingVCam/Public/VCamPixelStreamingSession.h
@@ -55,10 +55,10 @@ public:
 
 protected:
 	UPROPERTY(Transient)
-	UPixelStreamingMediaOutput* MediaOutput = nullptr;
+	TObjectPtr<UPixelStreamingMediaOutput> MediaOutput = nullptr;
 
 	UPROPERTY(Transient)
-	UPixelStreamingMediaCapture* MediaCapture = nullptr;
+	TObjectPtr<UPixelStreamingMediaCapture> MediaCapture = nullptr;
 
 private:
 	void StopSignallingServer();
diff --git a/Engine/Plugins/Media/PixelStreamingVCam/Source/PixelStreamingVCam/Public/VCamPixelStreamingSubsystem.h b/Engine/Plugins/Media/PixelStreamingVCam/Source/PixelStreamingVCam/Public/VCamPixelStreamingSubsystem.h
--- a/Engine/Plugins/Media/PixelStreamingVCam/Source/PixelStreamingVCam/Public/VCamPixelStreamingSubsystem.h
+++ b/Engine/Plugins/Media/PixelStreamingVCam/Source/PixelStreamingVCam/Public/VCamPixelStreamingSubsystem.h
@@ -29,7 +29,7 @@ public:
 private:
 	// Keep track of which output providers are currently active
 	UPROPERTY(Transient)
-	TArray<UVCamPixelStreamingSession*> ActiveOutputProviders;
+	TArray<TObjectPtr<UVCamPixelStreamingSession>> ActiveOutputProviders;
 
 	// Signalling/webserver
 	TSharedPtr<UE::PixelStreamingServers::IServer> Server;
diff --git a/Engine/Plugins/Media/TimecodeSynchronizer/Source/TimecodeSynchronizer/Public/TimecodeSynchronizer.h b/Engine/Plugins/Media/TimecodeSynchronizer/Source/TimecodeSynchronizer/Public/TimecodeSynchronizer.h
--- a/Engine/Plugins/Media/TimecodeSynchronizer/Source/TimecodeSynchronizer/Public/TimecodeSynchronizer.h
+++ b/Engine/Plugins/Media/TimecodeSynchronizer/Source/TimecodeSynchronizer/Public/TimecodeSynchronizer.h
@@ -222,7 +222,7 @@ private:
 
 	/* Associated source pointers */
 	UPROPERTY(VisibleAnywhere, Transient, Category = Debug, Meta=(DisplayName="Input Source"))
-	UTimeSynchronizationSource* InputSource;
+	TObjectPtr<UTimeSynchronizationSource> InputSource;
 
 	FTimecodeSynchronizerActiveTimecodedInputSource(const FTimecodeSynchronizerActiveTimecodedInputSource&) = delete;
 	FTimecodeSynchronizerActiveTimecodedInputSource& operator=(const FTimecodeSynchronizerActiveTimecodedInputSource&) = delete;
@@ -470,7 +470,7 @@ public:
 
 	/** Custom strategy to tick in a interval. */
 	UPROPERTY(EditAnywhere, Instanced, Category="Timecode Provider", Meta=(DisplayName="Timecode Source"))
-	UTimecodeProvider* TimecodeProvider;
+	TObjectPtr<UTimecodeProvider> TimecodeProvider;
 
 	/**
 	 * Index of the source that drives the synchronized Timecode.
@@ -501,12 +501,12 @@ public:
 	//! ONLY MODIFY THESE IN EDITOR
 	//! TODO: Deprecate this and make it private.
 	UPROPERTY(EditAnywhere, Instanced, Category="Input")
-	TArray<UTimeSynchronizationSource*> TimeSynchronizationInputSources;
+	TArray<TObjectPtr<UTimeSynchronizationSource>> TimeSynchronizationInputSources;
 
 private:
 
 	UPROPERTY(Transient)
-	TArray<UTimeSynchronizationSource*> DynamicSources;
+	TArray<TObjectPtr<UTimeSynchronizationSource>> DynamicSources;
 
 	/** What mode will be used for synchronization. */
 	UPROPERTY(EditAnywhere, Category = "Synchronization")
@@ -540,13 +540,13 @@ private:
 	TArray<FTimecodeSynchronizerActiveTimecodedInputSource> NonSynchronizedSources;
 
 	UPROPERTY(Transient)
-	UFixedFrameRateCustomTimeStep* RegisteredCustomTimeStep;
+	TObjectPtr<UFixedFrameRateCustomTimeStep> RegisteredCustomTimeStep;
 
 	UPROPERTY(Transient)
-	UTimecodeProvider* CachedPreviousTimecodeProvider;
+	TObjectPtr<UTimecodeProvider> CachedPreviousTimecodeProvider;
 
 	UPROPERTY(Transient)
-	UTimecodeProvider* CachedProxiedTimecodeProvider;
+	TObjectPtr<UTimecodeProvider> CachedProxiedTimecodeProvider;
 
 	UPROPERTY(Transient, DuplicateTransient, VisibleAnywhere, Category = "Synchronization")
 	int32 ActualFrameOffset;
diff --git a/Engine/Plugins/MovieScene/CustomizableSequencerTracks/Source/CustomizableSequencerTracks/Public/SequencerTrackBP.h b/Engine/Plugins/MovieScene/CustomizableSequencerTracks/Source/CustomizableSequencerTracks/Public/SequencerTrackBP.h
--- a/Engine/Plugins/MovieScene/CustomizableSequencerTracks/Source/CustomizableSequencerTracks/Public/SequencerTrackBP.h
+++ b/Engine/Plugins/MovieScene/CustomizableSequencerTracks/Source/CustomizableSequencerTracks/Public/SequencerTrackBP.h
@@ -37,7 +37,7 @@ public:
 	ECustomSequencerTrackType TrackType;
 
 	UPROPERTY(Category="Sequencer", EditDefaultsOnly, AssetRegistrySearchable, meta=(EditCondition="TrackType==ECustomSequencerTrackType::ObjectTrack"))
-	UClass* SupportedObjectType;
+	TObjectPtr<UClass> SupportedObjectType;
 
 	UPROPERTY(Category="Sequencer", EditDefaultsOnly, AssetRegistrySearchable)
 	TSubclassOf<USequencerSectionBP> DefaultSectionType;
@@ -70,5 +70,5 @@ public:
 private:
 
 	UPROPERTY(Instanced)
-	TArray<UMovieSceneSection*> Sections;
+	TArray<TObjectPtr<UMovieSceneSection>> Sections;
 };
\ No newline at end of file
diff --git a/Engine/Plugins/MovieScene/CustomizableSequencerTracks/Source/CustomizableSequencerTracks/Public/SequencerTrackInstanceBP.h b/Engine/Plugins/MovieScene/CustomizableSequencerTracks/Source/CustomizableSequencerTracks/Public/SequencerTrackInstanceBP.h
--- a/Engine/Plugins/MovieScene/CustomizableSequencerTracks/Source/CustomizableSequencerTracks/Public/SequencerTrackInstanceBP.h
+++ b/Engine/Plugins/MovieScene/CustomizableSequencerTracks/Source/CustomizableSequencerTracks/Public/SequencerTrackInstanceBP.h
@@ -20,7 +20,7 @@ struct FSequencerTrackInstanceInput
 	{}
 
 	UPROPERTY(Category="Sequencer", BlueprintReadOnly)
-	USequencerSectionBP* Section;
+	TObjectPtr<USequencerSectionBP> Section;
 
 	FMovieSceneContext Context;
 };
diff --git a/Engine/Plugins/MovieScene/MoviePipelineMaskRenderPass/Source/MoviePipelineMaskRenderPass/Private/MoviePipelinePanoramicPass.h b/Engine/Plugins/MovieScene/MoviePipelineMaskRenderPass/Source/MoviePipelineMaskRenderPass/Private/MoviePipelinePanoramicPass.h
--- a/Engine/Plugins/MovieScene/MoviePipelineMaskRenderPass/Source/MoviePipelineMaskRenderPass/Private/MoviePipelinePanoramicPass.h
+++ b/Engine/Plugins/MovieScene/MoviePipelineMaskRenderPass/Source/MoviePipelineMaskRenderPass/Private/MoviePipelinePanoramicPass.h
@@ -190,7 +190,7 @@ protected:
 	TSharedPtr<FOpenColorIODisplayExtension, ESPMode::ThreadSafe> OCIOSceneViewExtension;
 
 	UPROPERTY(Transient, DuplicateTransient)
-	UTextureRenderTarget2D* CanvasReadbackTexture;
+	TObjectPtr<UTextureRenderTarget2D> CanvasReadbackTexture;
 
 	TSharedPtr<MoviePipeline::IMoviePipelineOutputMerger> PanoramicOutputBlender;
 	bool bHasWarnedSettings;
diff --git a/Engine/Plugins/MovieScene/MovieRenderPipeline/Source/MovieRenderPipelineCore/Public/MoviePipeline.h b/Engine/Plugins/MovieScene/MovieRenderPipeline/Source/MovieRenderPipelineCore/Public/MoviePipeline.h
--- a/Engine/Plugins/MovieScene/MovieRenderPipeline/Source/MovieRenderPipelineCore/Public/MoviePipeline.h
+++ b/Engine/Plugins/MovieScene/MovieRenderPipeline/Source/MovieRenderPipelineCore/Public/MoviePipeline.h
@@ -382,29 +382,29 @@ private:
 private:
 	/** Custom TimeStep used to drive the engine while rendering. */
 	UPROPERTY(Transient, Instanced)
-	UMoviePipelineCustomTimeStep* CustomTimeStep;
+	TObjectPtr<UMoviePipelineCustomTimeStep> CustomTimeStep;
 
 	/** Custom Time Controller for the Sequence Player, used to match Custom TimeStep without any floating point accumulation errors. */
 	TSharedPtr<FMoviePipelineTimeController> CustomSequenceTimeController;
 
 	/** Hold a reference to the existing custom time step (if any) so we can restore it after we're done using our custom one. */
 	UPROPERTY(Transient)
-	UEngineCustomTimeStep* CachedPrevCustomTimeStep;
+	TObjectPtr<UEngineCustomTimeStep> CachedPrevCustomTimeStep;
 
 	/** This is our duplicated sequence that we're rendering. This will get modified throughout the rendering process. */
 	UPROPERTY(Transient)
-	ULevelSequence* TargetSequence;
+	TObjectPtr<ULevelSequence> TargetSequence;
 
 	/** The Level Sequence Actor we spawned to play our TargetSequence. */
 	UPROPERTY(Transient)
-	ALevelSequenceActor* LevelSequenceActor;
+	TObjectPtr<ALevelSequenceActor> LevelSequenceActor;
 
 	/** The Debug UI Widget that is spawned and placed on the player UI */
 	UPROPERTY(Transient)
-	UMovieRenderDebugWidget* DebugWidget;
+	TObjectPtr<UMovieRenderDebugWidget> DebugWidget;
 
 	UPROPERTY(Transient)
-	UTexture* PreviewTexture;
+	TObjectPtr<UTexture> PreviewTexture;
 
 	/** A list of all of the shots we are going to render out from this sequence. */
 	TArray<UMoviePipelineExecutorShot*> ActiveShotList;
@@ -478,7 +478,7 @@ public:
 private:
 	/** Keep track of which job we're working on. This holds our Configuration + which shots we're supposed to render from it. */
 	UPROPERTY(Transient)
-	UMoviePipelineExecutorJob* CurrentJob;
+	TObjectPtr<UMoviePipelineExecutorJob> CurrentJob;
 
 #if WITH_EDITOR
 	/** Keep track of clips we've exported, for building FCPXML and other project files */
diff --git a/Engine/Plugins/MovieScene/MovieRenderPipeline/Source/MovieRenderPipelineCore/Public/MoviePipelineConfigBase.h b/Engine/Plugins/MovieScene/MovieRenderPipeline/Source/MovieRenderPipelineCore/Public/MoviePipelineConfigBase.h
--- a/Engine/Plugins/MovieScene/MovieRenderPipeline/Source/MovieRenderPipelineCore/Public/MoviePipelineConfigBase.h
+++ b/Engine/Plugins/MovieScene/MovieRenderPipeline/Source/MovieRenderPipelineCore/Public/MoviePipelineConfigBase.h
@@ -158,7 +158,7 @@ public:
 protected:
 	/** Array of settings classes that affect various parts of the output pipeline. */
 	UPROPERTY(VisibleAnywhere, Instanced, Category = "Movie Pipeline")
-	TArray<UMoviePipelineSetting*> Settings;
+	TArray<TObjectPtr<UMoviePipelineSetting>> Settings;
 
 private:
 	int32 SettingsSerialNumber;
diff --git a/Engine/Plugins/MovieScene/MovieRenderPipeline/Source/MovieRenderPipelineCore/Public/MoviePipelineLinearExecutor.h b/Engine/Plugins/MovieScene/MovieRenderPipeline/Source/MovieRenderPipelineCore/Public/MoviePipelineLinearExecutor.h
--- a/Engine/Plugins/MovieScene/MovieRenderPipeline/Source/MovieRenderPipelineCore/Public/MoviePipelineLinearExecutor.h
+++ b/Engine/Plugins/MovieScene/MovieRenderPipeline/Source/MovieRenderPipelineCore/Public/MoviePipelineLinearExecutor.h
@@ -49,11 +49,11 @@ protected:
 	
 	/** List of Pipeline Configs we've been asked to execute. */
 	UPROPERTY(Transient)
-	UMoviePipelineQueue* Queue;
+	TObjectPtr<UMoviePipelineQueue> Queue;
 
 	/** A Movie Pipeline that has been spawned and is running (if any) */
 	UPROPERTY(Transient)
-	UMoviePipeline* ActiveMoviePipeline;
+	TObjectPtr<UMoviePipeline> ActiveMoviePipeline;
 
 	/** Which Pipeline Config are we currently working on. */
 	int32 CurrentPipelineIndex;
diff --git a/Engine/Plugins/MovieScene/MovieRenderPipeline/Source/MovieRenderPipelineCore/Public/MoviePipelineMasterConfig.h b/Engine/Plugins/MovieScene/MovieRenderPipeline/Source/MovieRenderPipelineCore/Public/MoviePipelineMasterConfig.h
--- a/Engine/Plugins/MovieScene/MovieRenderPipeline/Source/MovieRenderPipelineCore/Public/MoviePipelineMasterConfig.h
+++ b/Engine/Plugins/MovieScene/MovieRenderPipeline/Source/MovieRenderPipelineCore/Public/MoviePipelineMasterConfig.h
@@ -82,13 +82,13 @@ protected:
 public:	
 	/** A mapping of Shot Name -> Shot Config to use for rendering specific shots with specific configs. */
 	UPROPERTY(Instanced)
-	TMap<FString, UMoviePipelineShotConfig*> PerShotConfigMapping;
+	TMap<FString, TObjectPtr<UMoviePipelineShotConfig>> PerShotConfigMapping;
 
 private:
 	UPROPERTY(Instanced)
-	UMoviePipelineOutputSetting* OutputSetting;
+	TObjectPtr<UMoviePipelineOutputSetting> OutputSetting;
 
 	/** An array of settings that are available in the engine and have not been edited by the user. */
 	UPROPERTY(Transient)
-	TArray<UMoviePipelineSetting*> TransientSettings;
+	TArray<TObjectPtr<UMoviePipelineSetting>> TransientSettings;
 };
\ No newline at end of file
diff --git a/Engine/Plugins/MovieScene/MovieRenderPipeline/Source/MovieRenderPipelineCore/Public/MoviePipelinePythonHostExecutor.h b/Engine/Plugins/MovieScene/MovieRenderPipeline/Source/MovieRenderPipelineCore/Public/MoviePipelinePythonHostExecutor.h
--- a/Engine/Plugins/MovieScene/MovieRenderPipeline/Source/MovieRenderPipelineCore/Public/MoviePipelinePythonHostExecutor.h
+++ b/Engine/Plugins/MovieScene/MovieRenderPipeline/Source/MovieRenderPipelineCore/Public/MoviePipelinePythonHostExecutor.h
@@ -56,8 +56,8 @@ public:
 	TSubclassOf<UMoviePipelinePythonHostExecutor> ExecutorClass;
 	
 	UPROPERTY(BlueprintReadWrite, Transient, Category = "Movie Render Pipeline")
-	UMoviePipelineQueue* PipelineQueue;
+	TObjectPtr<UMoviePipelineQueue> PipelineQueue;
 private:
 	UPROPERTY(Transient)
-	UWorld* LastLoadedWorld;
+	TObjectPtr<UWorld> LastLoadedWorld;
 };
\ No newline at end of file
diff --git a/Engine/Plugins/MovieScene/MovieRenderPipeline/Source/MovieRenderPipelineCore/Public/MoviePipelineQueue.h b/Engine/Plugins/MovieScene/MovieRenderPipeline/Source/MovieRenderPipelineCore/Public/MoviePipelineQueue.h
--- a/Engine/Plugins/MovieScene/MovieRenderPipeline/Source/MovieRenderPipelineCore/Public/MoviePipelineQueue.h
+++ b/Engine/Plugins/MovieScene/MovieRenderPipeline/Source/MovieRenderPipelineCore/Public/MoviePipelineQueue.h
@@ -184,7 +184,7 @@ protected:
 
 private:
 	UPROPERTY()
-	UMoviePipelineShotConfig* ShotOverrideConfig;
+	TObjectPtr<UMoviePipelineShotConfig> ShotOverrideConfig;
 
 	UPROPERTY()
 	TSoftObjectPtr<UMoviePipelineShotConfig> ShotOverridePresetOrigin;
@@ -395,7 +395,7 @@ public:
 
 	/** (Optional) Shot specific information. If a shot is missing from this list it will assume to be enabled and will be rendered. */
 	UPROPERTY(BlueprintReadWrite, Instanced, Category = "Movie Render Pipeline")
-	TArray<UMoviePipelineExecutorShot*> ShotInfo;
+	TArray<TObjectPtr<UMoviePipelineExecutorShot>> ShotInfo;
 
 	/** 
 	* Arbitrary data that can be associated with the job. Not used by default implementations, nor read.
@@ -415,7 +415,7 @@ private:
 	/** 
 	*/
 	UPROPERTY(Instanced)
-	UMoviePipelineMasterConfig* Configuration;
+	TObjectPtr<UMoviePipelineMasterConfig> Configuration;
 
 	/**
 	*/
@@ -516,7 +516,7 @@ public:
 
 private:
 	UPROPERTY(Instanced)
-	TArray<UMoviePipelineExecutorJob*> Jobs;
+	TArray<TObjectPtr<UMoviePipelineExecutorJob>> Jobs;
 	
 private:
 	int32 QueueSerialNumber;
diff --git a/Engine/Plugins/MovieScene/MovieRenderPipeline/Source/MovieRenderPipelineCore/Public/MoviePipelineQueueEngineSubsystem.h b/Engine/Plugins/MovieScene/MovieRenderPipeline/Source/MovieRenderPipelineCore/Public/MoviePipelineQueueEngineSubsystem.h
--- a/Engine/Plugins/MovieScene/MovieRenderPipeline/Source/MovieRenderPipelineCore/Public/MoviePipelineQueueEngineSubsystem.h
+++ b/Engine/Plugins/MovieScene/MovieRenderPipeline/Source/MovieRenderPipelineCore/Public/MoviePipelineQueueEngineSubsystem.h
@@ -146,10 +146,10 @@ private:
 
 private:
 	UPROPERTY(Transient)
-	UMoviePipelineExecutorBase* ActiveExecutor;
+	TObjectPtr<UMoviePipelineExecutorBase> ActiveExecutor;
 
 	UPROPERTY(Transient, Instanced)
-	UMoviePipelineQueue* CurrentQueue;
+	TObjectPtr<UMoviePipelineQueue> CurrentQueue;
 	
 	TOptional<FMoviePipelineOutputData> CachedRenderJobResults;
 	TSubclassOf<UMovieRenderDebugWidget> CachedProgressWidgetClass;
diff --git a/Engine/Plugins/MovieScene/MovieRenderPipeline/Source/MovieRenderPipelineCore/Public/MovieRenderPipelineDataTypes.h b/Engine/Plugins/MovieScene/MovieRenderPipeline/Source/MovieRenderPipelineCore/Public/MovieRenderPipelineDataTypes.h
--- a/Engine/Plugins/MovieScene/MovieRenderPipeline/Source/MovieRenderPipelineCore/Public/MovieRenderPipelineDataTypes.h
+++ b/Engine/Plugins/MovieScene/MovieRenderPipeline/Source/MovieRenderPipelineCore/Public/MovieRenderPipelineDataTypes.h
@@ -777,7 +777,7 @@ struct FMoviePipelineFormatArgs
 
 	/** Which job is this for? Some settings are specific to the level sequence being rendered. */
 	UPROPERTY(BlueprintReadWrite, EditAnywhere, Category = "Movie Render Pipeline")
-	class UMoviePipelineExecutorJob* InJob;
+	TObjectPtr<class UMoviePipelineExecutorJob> InJob;
 };
 
 USTRUCT(BlueprintType)
@@ -858,14 +858,14 @@ struct FMoviePipelineFilenameResolveParams
 
 	/** Required. This is the job all of the settings should be pulled from.*/
 	UPROPERTY(BlueprintReadWrite, EditAnywhere, Category = "Movie Render Pipeline")
-	UMoviePipelineExecutorJob* Job;
+	TObjectPtr<UMoviePipelineExecutorJob> Job;
 
 	/** If this shot has multiple cameras in the sidecar array, which camera is this for? -1 will return the InnerName/Main camera for the shot. */
 	int32 CameraIndex;
 
 	/** Optional. If specified, settings will be pulled from this shot (if overriden by the shot). If null, always use the master configuration in the job. */
 	UPROPERTY(BlueprintReadWrite, EditAnywhere, Category = "Movie Render Pipeline")
-	UMoviePipelineExecutorShot* ShotOverride;
+	TObjectPtr<UMoviePipelineExecutorShot> ShotOverride;
 
 	/** Additional offset added onto the offset provided by the Output Settings in the Job. Required for some internal things (FCPXML). */
 	UPROPERTY(BlueprintReadWrite, EditAnywhere, Category = "Movie Render Pipeline")
@@ -1294,11 +1294,11 @@ struct FMoviePipelineOutputData
 	* Provided here for backwards compatibility.
 	*/
 	UPROPERTY(VisibleAnywhere, BlueprintReadOnly, Category = "Movie Pipeline")
-	UMoviePipeline* Pipeline;
+	TObjectPtr<UMoviePipeline> Pipeline;
 	
 	/** Job the data is for. Job may still be in progress (if a shot callback) so be careful about modifying properties on it */
 	UPROPERTY(VisibleAnywhere, BlueprintReadOnly, Category = "Movie Pipeline")
-	UMoviePipelineExecutorJob* Job;
+	TObjectPtr<UMoviePipelineExecutorJob> Job;
 	
 	/** Did the job succeed, or was it canceled early due to an error (such as failure to write file to disk)? */
 	UPROPERTY(VisibleAnywhere, BlueprintReadOnly, Category = "Movie Pipeline")
diff --git a/Engine/Plugins/MovieScene/MovieRenderPipeline/Source/MovieRenderPipelineEditor/Private/MoviePipelineAssetEditor.h b/Engine/Plugins/MovieScene/MovieRenderPipeline/Source/MovieRenderPipelineEditor/Private/MoviePipelineAssetEditor.h
--- a/Engine/Plugins/MovieScene/MovieRenderPipeline/Source/MovieRenderPipelineEditor/Private/MoviePipelineAssetEditor.h
+++ b/Engine/Plugins/MovieScene/MovieRenderPipeline/Source/MovieRenderPipelineEditor/Private/MoviePipelineAssetEditor.h
@@ -19,5 +19,5 @@ public:
 	void SetObjectToEdit(UObject* InObject);
 private:
 	UPROPERTY(Transient)
-	UObject* ObjectToEdit;
+	TObjectPtr<UObject> ObjectToEdit;
 };
\ No newline at end of file
diff --git a/Engine/Plugins/MovieScene/MovieRenderPipeline/Source/MovieRenderPipelineEditor/Private/MoviePipelineFunctionalTestBase.h b/Engine/Plugins/MovieScene/MovieRenderPipeline/Source/MovieRenderPipelineEditor/Private/MoviePipelineFunctionalTestBase.h
--- a/Engine/Plugins/MovieScene/MovieRenderPipeline/Source/MovieRenderPipelineEditor/Private/MoviePipelineFunctionalTestBase.h
+++ b/Engine/Plugins/MovieScene/MovieRenderPipeline/Source/MovieRenderPipelineEditor/Private/MoviePipelineFunctionalTestBase.h
@@ -39,12 +39,12 @@ protected:
 protected:
 
 	UPROPERTY(EditAnywhere, BlueprintReadOnly, Category = "Movie Pipeline")
-	UMoviePipelineQueue* QueuePreset;
+	TObjectPtr<UMoviePipelineQueue> QueuePreset;
 
 protected:
 	UPROPERTY(Transient)
-	UMoviePipeline* ActiveMoviePipeline;
+	TObjectPtr<UMoviePipeline> ActiveMoviePipeline;
 
 	UPROPERTY(Transient)
-	UMoviePipelineQueue* ActiveQueue;
+	TObjectPtr<UMoviePipelineQueue> ActiveQueue;
 };
diff --git a/Engine/Plugins/MovieScene/MovieRenderPipeline/Source/MovieRenderPipelineEditor/Public/MoviePipelineQueueSubsystem.h b/Engine/Plugins/MovieScene/MovieRenderPipeline/Source/MovieRenderPipelineEditor/Public/MoviePipelineQueueSubsystem.h
--- a/Engine/Plugins/MovieScene/MovieRenderPipeline/Source/MovieRenderPipelineEditor/Public/MoviePipelineQueueSubsystem.h
+++ b/Engine/Plugins/MovieScene/MovieRenderPipeline/Source/MovieRenderPipelineEditor/Public/MoviePipelineQueueSubsystem.h
@@ -72,8 +72,8 @@ private:
 
 private:
 	UPROPERTY(Transient)
-	UMoviePipelineExecutorBase* ActiveExecutor;
+	TObjectPtr<UMoviePipelineExecutorBase> ActiveExecutor;
 
 	UPROPERTY(Transient, Instanced)
-	UMoviePipelineQueue* CurrentQueue;
+	TObjectPtr<UMoviePipelineQueue> CurrentQueue;
 };
\ No newline at end of file
diff --git a/Engine/Plugins/MovieScene/MovieRenderPipeline/Source/MovieRenderPipelineRenderPasses/Public/MoviePipelineDeferredPasses.h b/Engine/Plugins/MovieScene/MovieRenderPipeline/Source/MovieRenderPipelineRenderPasses/Public/MoviePipelineDeferredPasses.h
--- a/Engine/Plugins/MovieScene/MovieRenderPipeline/Source/MovieRenderPipelineRenderPasses/Public/MoviePipelineDeferredPasses.h
+++ b/Engine/Plugins/MovieScene/MovieRenderPipeline/Source/MovieRenderPipelineRenderPasses/Public/MoviePipelineDeferredPasses.h
@@ -128,13 +128,13 @@ public:
 protected:
 	/** While rendering, store an array of the non-null valid materials loaded from AdditionalPostProcessMaterials. Cleared on teardown. */
 	UPROPERTY(Transient, DuplicateTransient)
-	TArray<UMaterialInterface*> ActivePostProcessMaterials;
+	TArray<TObjectPtr<UMaterialInterface>> ActivePostProcessMaterials;
 
 	UPROPERTY(Transient, DuplicateTransient)
-	UMaterialInterface* StencilLayerMaterial;
+	TObjectPtr<UMaterialInterface> StencilLayerMaterial;
 
 	UPROPERTY(Transient, DuplicateTransient)
-	TArray<UTextureRenderTarget2D*> TileRenderTargets;
+	TArray<TObjectPtr<UTextureRenderTarget2D>> TileRenderTargets;
 
 	struct FMultiCameraViewStateData
 	{
diff --git a/Engine/Plugins/MovieScene/MovieRenderPipeline/Source/MovieRenderPipelineSettings/Public/MoviePipelineBurnInSetting.h b/Engine/Plugins/MovieScene/MovieRenderPipeline/Source/MovieRenderPipelineSettings/Public/MoviePipelineBurnInSetting.h
--- a/Engine/Plugins/MovieScene/MovieRenderPipeline/Source/MovieRenderPipelineSettings/Public/MoviePipelineBurnInSetting.h
+++ b/Engine/Plugins/MovieScene/MovieRenderPipeline/Source/MovieRenderPipelineSettings/Public/MoviePipelineBurnInSetting.h
@@ -54,10 +54,10 @@ private:
 	TSharedPtr<SVirtualWindow> VirtualWindow;
 
 	UPROPERTY(Transient)
-	UTextureRenderTarget2D* RenderTarget;
+	TObjectPtr<UTextureRenderTarget2D> RenderTarget;
 
 	UPROPERTY(Transient)
-	TArray<UMoviePipelineBurnInWidget*> BurnInWidgetInstances;
+	TArray<TObjectPtr<UMoviePipelineBurnInWidget>> BurnInWidgetInstances;
 public:
 	static FString DefaultBurnInWidgetAsset;
 };
\ No newline at end of file
diff --git a/Engine/Plugins/Online/OnlineFramework/Source/Hotfix/Public/OnlineHotfixManager.h b/Engine/Plugins/Online/OnlineFramework/Source/Hotfix/Public/OnlineHotfixManager.h
--- a/Engine/Plugins/Online/OnlineFramework/Source/Hotfix/Public/OnlineHotfixManager.h
+++ b/Engine/Plugins/Online/OnlineFramework/Source/Hotfix/Public/OnlineHotfixManager.h
@@ -355,7 +355,7 @@ public:
 	/** Array of objects that we're forcing to remain resident because we've applied live hotfixes and won't get an
 	    opportunity to reapply changes if the object is evicted from memory. */
 	UPROPERTY(Transient)
-	TArray<UObject*> AssetsHotfixedFromIniFiles;
+	TArray<TObjectPtr<UObject>> AssetsHotfixedFromIniFiles;
 	
 
 	/** 
diff --git a/Engine/Plugins/Online/OnlineFramework/Source/Hotfix/Public/UpdateManager.h b/Engine/Plugins/Online/OnlineFramework/Source/Hotfix/Public/UpdateManager.h
--- a/Engine/Plugins/Online/OnlineFramework/Source/Hotfix/Public/UpdateManager.h
+++ b/Engine/Plugins/Online/OnlineFramework/Source/Hotfix/Public/UpdateManager.h
@@ -448,9 +448,9 @@ private:
 
 	/** String output */
 	UPROPERTY()
-	UEnum* UpdateStateEnum;
+	TObjectPtr<UEnum> UpdateStateEnum;
 	UPROPERTY()
-	UEnum* UpdateCompletionEnum;
+	TObjectPtr<UEnum> UpdateCompletionEnum;
 
 	/** Fire a delegate after a given amount of time */
 	typedef TFunction<void(void)> DelayCb;
diff --git a/Engine/Plugins/Online/OnlineFramework/Source/Lobby/Public/LobbyBeaconClient.h b/Engine/Plugins/Online/OnlineFramework/Source/Lobby/Public/LobbyBeaconClient.h
--- a/Engine/Plugins/Online/OnlineFramework/Source/Lobby/Public/LobbyBeaconClient.h
+++ b/Engine/Plugins/Online/OnlineFramework/Source/Lobby/Public/LobbyBeaconClient.h
@@ -61,10 +61,10 @@ class LOBBY_API ALobbyBeaconClient : public AOnlineBeaconClient
 
 	/** Client view of the lobby state */
 	UPROPERTY(Replicated)
-	ALobbyBeaconState* LobbyState;
+	TObjectPtr<ALobbyBeaconState> LobbyState;
 	/** Player state associated with this beacon (@todo not splitscreen safe) */
 	UPROPERTY(Replicated)
-	ALobbyBeaconPlayerState* PlayerState;
+	TObjectPtr<ALobbyBeaconPlayerState> PlayerState;
 
 	/** Clear out any references to this in PlayerState */
 	virtual void EndPlay(EEndPlayReason::Type Reason) override;
diff --git a/Engine/Plugins/Online/OnlineFramework/Source/Lobby/Public/LobbyBeaconHost.h b/Engine/Plugins/Online/OnlineFramework/Source/Lobby/Public/LobbyBeaconHost.h
--- a/Engine/Plugins/Online/OnlineFramework/Source/Lobby/Public/LobbyBeaconHost.h
+++ b/Engine/Plugins/Online/OnlineFramework/Source/Lobby/Public/LobbyBeaconHost.h
@@ -98,7 +98,7 @@ protected:
 
 	/** Actor representing the state of the lobby (similar to AGameState) */
 	UPROPERTY()
-	ALobbyBeaconState* LobbyState;
+	TObjectPtr<ALobbyBeaconState> LobbyState;
 
 
 	virtual bool PreLogin(const FUniqueNetIdRepl& InUniqueId, const FString& Options);
diff --git a/Engine/Plugins/Online/OnlineFramework/Source/Lobby/Public/LobbyBeaconPlayerState.h b/Engine/Plugins/Online/OnlineFramework/Source/Lobby/Public/LobbyBeaconPlayerState.h
--- a/Engine/Plugins/Online/OnlineFramework/Source/Lobby/Public/LobbyBeaconPlayerState.h
+++ b/Engine/Plugins/Online/OnlineFramework/Source/Lobby/Public/LobbyBeaconPlayerState.h
@@ -46,7 +46,7 @@ class LOBBY_API ALobbyBeaconPlayerState : public AInfo
 
 	/** Reference to the beacon actor related to this player */
 	UPROPERTY(Replicated)
-	AOnlineBeaconClient* ClientActor;
+	TObjectPtr<AOnlineBeaconClient> ClientActor;
 
 	/**
 	 * Delegate fired when this player state has changed in some way
diff --git a/Engine/Plugins/Online/OnlineFramework/Source/Lobby/Public/LobbyBeaconState.h b/Engine/Plugins/Online/OnlineFramework/Source/Lobby/Public/LobbyBeaconState.h
--- a/Engine/Plugins/Online/OnlineFramework/Source/Lobby/Public/LobbyBeaconState.h
+++ b/Engine/Plugins/Online/OnlineFramework/Source/Lobby/Public/LobbyBeaconState.h
@@ -56,7 +56,7 @@ public:
 
 	/** Actual player state actor */
 	UPROPERTY()
-	ALobbyBeaconPlayerState* LobbyPlayerState;
+	TObjectPtr<ALobbyBeaconPlayerState> LobbyPlayerState;
 
 	/** Player state removal */
 	void PreReplicatedRemove(const FLobbyPlayerStateInfoArray& InArraySerializer);
@@ -148,7 +148,7 @@ private:
 
 	/** Owning lobby beacon for this array of players */
 	UPROPERTY(NotReplicated)
-	ALobbyBeaconState* ParentState = nullptr;
+	TObjectPtr<ALobbyBeaconState> ParentState = nullptr;
 
 	friend ALobbyBeaconState;
 	friend FLobbyPlayerStateActorInfo;
diff --git a/Engine/Plugins/Online/OnlineFramework/Source/Party/Public/Chat/SocialChatManager.h b/Engine/Plugins/Online/OnlineFramework/Source/Party/Public/Chat/SocialChatManager.h
--- a/Engine/Plugins/Online/OnlineFramework/Source/Party/Public/Chat/SocialChatManager.h
+++ b/Engine/Plugins/Online/OnlineFramework/Source/Party/Public/Chat/SocialChatManager.h
@@ -37,12 +37,12 @@ struct FSocialChatChannelConfig
 	}
 
 	UPROPERTY()
-	USocialUser* SocialUser;
+	TObjectPtr<USocialUser> SocialUser;
 
 	FString RoomId;
 
 	UPROPERTY()
-	TArray<USocialChatChannel*> ListenChannels;
+	TArray<TObjectPtr<USocialChatChannel>> ListenChannels;
 	FText DisplayName;
 };
 
@@ -174,19 +174,19 @@ private:
 
 private:
 	UPROPERTY()
-	TMap<TWeakObjectPtr<USocialUser>, USocialPrivateMessageChannel*> DirectChannelsByTargetUser;
+	TMap<TWeakObjectPtr<USocialUser>, TObjectPtr<USocialPrivateMessageChannel>> DirectChannelsByTargetUser;
 
 	UPROPERTY()
-	TMap<FString, USocialChatRoom*> ChatRoomsById;
+	TMap<FString, TObjectPtr<USocialChatRoom>> ChatRoomsById;
 
 	UPROPERTY()
-	TMap<FString, USocialReadOnlyChatChannel*> ReadOnlyChannelsByDisplayName;
+	TMap<FString, TObjectPtr<USocialReadOnlyChatChannel>> ReadOnlyChannelsByDisplayName;
 
 	UPROPERTY(config)
 	bool bEnableChatSlashCommands = true;
 	
 	UPROPERTY()
-	TMap<FUniqueNetIdRepl, USocialGroupChannel*> GroupChannels;
+	TMap<FUniqueNetIdRepl, TObjectPtr<USocialGroupChannel>> GroupChannels;
 
 	mutable FOnChatChannelCreated OnChannelCreatedEvent;
 	mutable FOnChatChannelLeft OnChannelLeftEvent;
diff --git a/Engine/Plugins/Online/OnlineFramework/Source/Party/Public/Chat/SocialGroupChannel.h b/Engine/Plugins/Online/OnlineFramework/Source/Party/Public/Chat/SocialGroupChannel.h
--- a/Engine/Plugins/Online/OnlineFramework/Source/Party/Public/Chat/SocialGroupChannel.h
+++ b/Engine/Plugins/Online/OnlineFramework/Source/Party/Public/Chat/SocialGroupChannel.h
@@ -33,7 +33,7 @@ private:
 
 private:
 	UPROPERTY()
-	USocialUser* SocialUser;
+	TObjectPtr<USocialUser> SocialUser;
 
 	UPROPERTY()
 	FUniqueNetIdRepl GroupId;
@@ -42,7 +42,7 @@ private:
 	FText DisplayName;
 
 	UPROPERTY()
-	TArray<USocialUser*> Members;
+	TArray<TObjectPtr<USocialUser>> Members;
 
 	TWeakPtr<IOnlineGroups, ESPMode::ThreadSafe> GroupInterfacePtr;
 };
\ No newline at end of file
diff --git a/Engine/Plugins/Online/OnlineFramework/Source/Party/Public/Chat/SocialPrivateMessageChannel.h b/Engine/Plugins/Online/OnlineFramework/Source/Party/Public/Chat/SocialPrivateMessageChannel.h
--- a/Engine/Plugins/Online/OnlineFramework/Source/Party/Public/Chat/SocialPrivateMessageChannel.h
+++ b/Engine/Plugins/Online/OnlineFramework/Source/Party/Public/Chat/SocialPrivateMessageChannel.h
@@ -24,5 +24,5 @@ private:
 
 	/** The recipient of the current user's messages */
 	UPROPERTY()
-	USocialUser* TargetUser;
+	TObjectPtr<USocialUser> TargetUser;
 };
\ No newline at end of file
diff --git a/Engine/Plugins/Online/OnlineFramework/Source/Party/Public/Party/PartyMember.h b/Engine/Plugins/Online/OnlineFramework/Source/Party/Public/Party/PartyMember.h
--- a/Engine/Plugins/Online/OnlineFramework/Source/Party/Public/Party/PartyMember.h
+++ b/Engine/Plugins/Online/OnlineFramework/Source/Party/Public/Party/PartyMember.h
@@ -153,7 +153,7 @@ private:
 	FOnlinePartyMemberConstPtr OssPartyMember;
 
 	UPROPERTY()
-	USocialUser* SocialUser = nullptr;
+	TObjectPtr<USocialUser> SocialUser = nullptr;
 
 	bool bHasReceivedInitialData = false;
 	mutable FOnPartyMemberStateChanged OnMemberConnectionStatusChangedEvent;
diff --git a/Engine/Plugins/Online/OnlineFramework/Source/Party/Public/Party/SocialParty.h b/Engine/Plugins/Online/OnlineFramework/Source/Party/Public/Party/SocialParty.h
--- a/Engine/Plugins/Online/OnlineFramework/Source/Party/Public/Party/SocialParty.h
+++ b/Engine/Plugins/Online/OnlineFramework/Source/Party/Public/Party/SocialParty.h
@@ -370,7 +370,7 @@ private:
 	FUniqueNetIdRepl CurrentLeaderId;
 
 	UPROPERTY()
-	TMap<FUniqueNetIdRepl, UPartyMember*> PartyMembersById;
+	TMap<FUniqueNetIdRepl, TObjectPtr<UPartyMember>> PartyMembersById;
 
 	UPROPERTY(config)
 	bool bEnableAutomaticPartyRejoin = true;
@@ -431,7 +431,7 @@ private:
 	
 	/** Spectator beacon client instance while getting approval for spectator*/
 	UPROPERTY()
-	ASpectatorBeaconClient* SpectatorBeaconClient = nullptr;
+	TObjectPtr<ASpectatorBeaconClient> SpectatorBeaconClient = nullptr;
 
 	/**
 	 * True when we have limited functionality due to lacking an xmpp connection.
diff --git a/Engine/Plugins/Online/OnlineFramework/Source/Party/Public/SocialManager.h b/Engine/Plugins/Online/OnlineFramework/Source/Party/Public/SocialManager.h
--- a/Engine/Plugins/Online/OnlineFramework/Source/Party/Public/SocialManager.h
+++ b/Engine/Plugins/Online/OnlineFramework/Source/Party/Public/SocialManager.h
@@ -266,10 +266,10 @@ private:
 	static TMap<TWeakObjectPtr<UGameInstance>, TWeakObjectPtr<USocialManager>> AllManagersByGameInstance;
 
 	UPROPERTY()
-	TArray<USocialToolkit*> SocialToolkits;
+	TArray<TObjectPtr<USocialToolkit>> SocialToolkits;
 
 	UPROPERTY()
-	USocialDebugTools* SocialDebugTools;
+	TObjectPtr<USocialDebugTools> SocialDebugTools;
 
 	bool bIsConnectedToPartyService = false;
 	
diff --git a/Engine/Plugins/Online/OnlineFramework/Source/Party/Public/SocialToolkit.h b/Engine/Plugins/Online/OnlineFramework/Source/Party/Public/SocialToolkit.h
--- a/Engine/Plugins/Online/OnlineFramework/Source/Party/Public/SocialToolkit.h
+++ b/Engine/Plugins/Online/OnlineFramework/Source/Party/Public/SocialToolkit.h
@@ -303,17 +303,17 @@ private:
 	static TMap<TWeakObjectPtr<const ULocalPlayer>, TWeakObjectPtr<USocialToolkit>> AllToolkitsByOwningPlayer;
 
 	UPROPERTY()
-	USocialUser* LocalUser;
+	TObjectPtr<USocialUser> LocalUser;
 
 	UPROPERTY()
-	TArray<USocialUser*> AllUsers;
+	TArray<TObjectPtr<USocialUser>> AllUsers;
 	TMap<FUniqueNetIdRepl, TWeakObjectPtr<USocialUser>> UsersBySubsystemIds;
 
 	UPROPERTY()
 	TWeakObjectPtr<ULocalPlayer> LocalPlayerOwner = nullptr;
 
 	UPROPERTY()
-	USocialChatManager* SocialChatManager;
+	TObjectPtr<USocialChatManager> SocialChatManager;
 
 	TSet<IOnlinePartyJoinInfoConstRef> PartyInvitations;
 	mutable TArray<TWeakPtr<FSocialUserList>> CachedSocialUserLists;
diff --git a/Engine/Plugins/Online/OnlineFramework/Source/Qos/Public/QosRegionManager.h b/Engine/Plugins/Online/OnlineFramework/Source/Qos/Public/QosRegionManager.h
--- a/Engine/Plugins/Online/OnlineFramework/Source/Qos/Public/QosRegionManager.h
+++ b/Engine/Plugins/Online/OnlineFramework/Source/Qos/Public/QosRegionManager.h
@@ -673,7 +673,7 @@ private:
 
 	/** Reference to the evaluator for making datacenter determinations (null when not active) */
 	UPROPERTY()
-	UQosEvaluator* Evaluator;
+	TObjectPtr<UQosEvaluator> Evaluator;
 	/** Result of the last datacenter test */
 	UPROPERTY()
 	EQosCompletionResult QosEvalResult;
diff --git a/Engine/Plugins/Protocols/MQTT/Source/MQTTCore/Public/MQTTClientObject.h b/Engine/Plugins/Protocols/MQTT/Source/MQTTCore/Public/MQTTClientObject.h
--- a/Engine/Plugins/Protocols/MQTT/Source/MQTTCore/Public/MQTTClientObject.h
+++ b/Engine/Plugins/Protocols/MQTT/Source/MQTTCore/Public/MQTTClientObject.h
@@ -122,7 +122,7 @@ protected:
 	FOnMessageDelegate OnMessageDelegate;
 
 	UPROPERTY()
-	TArray<UMQTTSubscriptionObject*> Subscriptions; 
+	TArray<TObjectPtr<UMQTTSubscriptionObject>> Subscriptions; 
 
 private:
 	TSharedPtr<IMQTTClient, ESPMode::ThreadSafe> MqttClientPtr;
diff --git a/Engine/Plugins/Protocols/MQTT/Source/MQTTCore/Public/MQTTSubsystem.h b/Engine/Plugins/Protocols/MQTT/Source/MQTTCore/Public/MQTTSubsystem.h
--- a/Engine/Plugins/Protocols/MQTT/Source/MQTTCore/Public/MQTTSubsystem.h
+++ b/Engine/Plugins/Protocols/MQTT/Source/MQTTCore/Public/MQTTSubsystem.h
@@ -37,5 +37,5 @@ public:
 
 private:
 	UPROPERTY(Transient)
-	TMap<FGuid, UMQTTClientObject*> Clients;
+	TMap<FGuid, TObjectPtr<UMQTTClientObject>> Clients;
 };
diff --git a/Engine/Plugins/Runtime/AR/ARUtilities/Source/ARUtilities/Public/ARPassthroughManager.h b/Engine/Plugins/Runtime/AR/ARUtilities/Source/ARUtilities/Public/ARPassthroughManager.h
--- a/Engine/Plugins/Runtime/AR/ARUtilities/Source/ARUtilities/Public/ARPassthroughManager.h
+++ b/Engine/Plugins/Runtime/AR/ARUtilities/Source/ARUtilities/Public/ARPassthroughManager.h
@@ -37,5 +37,5 @@ private:
 		
 private:
 	UPROPERTY()
-	UPassthroughMaterialUpdateComponent* PassthroughUpdateComponent;
+	TObjectPtr<UPassthroughMaterialUpdateComponent> PassthroughUpdateComponent;
 };
diff --git a/Engine/Plugins/Runtime/AR/ARUtilities/Source/ARUtilities/Public/PassthroughMaterialUpdateComponent.h b/Engine/Plugins/Runtime/AR/ARUtilities/Source/ARUtilities/Public/PassthroughMaterialUpdateComponent.h
--- a/Engine/Plugins/Runtime/AR/ARUtilities/Source/ARUtilities/Public/PassthroughMaterialUpdateComponent.h
+++ b/Engine/Plugins/Runtime/AR/ARUtilities/Source/ARUtilities/Public/PassthroughMaterialUpdateComponent.h
@@ -44,11 +44,11 @@ protected:
 	
 	/** Which material to use for a regular camera texture */
 	UPROPERTY(EditAnywhere, Category = "Passthrough")
-	UMaterialInterface* PassthroughMaterial;
+	TObjectPtr<UMaterialInterface> PassthroughMaterial;
 	
 	/** Which material to use for an external camera texture */
 	UPROPERTY(EditAnywhere, Category = "Passthrough")
-	UMaterialInterface* PassthroughMaterialExternalTexture;
+	TObjectPtr<UMaterialInterface> PassthroughMaterialExternalTexture;
 	
 	/**
 	 * The debug color used to modulate the passthrough material.
@@ -59,8 +59,8 @@ protected:
 	
 private:
 	UPROPERTY(Transient)
-	TArray<UPrimitiveComponent*> AffectedComponents;
+	TArray<TObjectPtr<UPrimitiveComponent>> AffectedComponents;
 	
 	UPROPERTY(Transient)
-	TArray<UPrimitiveComponent*> PendingComponents;
+	TArray<TObjectPtr<UPrimitiveComponent>> PendingComponents;
 };
diff --git a/Engine/Plugins/Runtime/AR/AppleAR/AppleARKit/Source/AppleARKit/Private/AppleARKitTrackable.h b/Engine/Plugins/Runtime/AR/AppleAR/AppleARKit/Source/AppleARKit/Private/AppleARKitTrackable.h
--- a/Engine/Plugins/Runtime/AR/AppleAR/AppleARKit/Source/AppleARKit/Private/AppleARKitTrackable.h
+++ b/Engine/Plugins/Runtime/AR/AppleAR/AppleARKit/Source/AppleARKit/Private/AppleARKitTrackable.h
@@ -25,5 +25,5 @@ public:
 private:
 	/** The cube map of the reflected environment */
 	UPROPERTY()
-	UAppleARKitEnvironmentCaptureProbeTexture* ARKitEnvironmentTexture;
+	TObjectPtr<UAppleARKitEnvironmentCaptureProbeTexture> ARKitEnvironmentTexture;
 };
diff --git a/Engine/Plugins/Runtime/AR/AppleAR/AppleARKit/Source/AppleARKit/Private/AppleARKitVideoOverlay.h b/Engine/Plugins/Runtime/AR/AppleAR/AppleARKit/Source/AppleARKit/Private/AppleARKitVideoOverlay.h
--- a/Engine/Plugins/Runtime/AR/AppleAR/AppleARKit/Source/AppleARKit/Private/AppleARKitVideoOverlay.h
+++ b/Engine/Plugins/Runtime/AR/AppleAR/AppleARKit/Source/AppleARKit/Private/AppleARKitVideoOverlay.h
@@ -32,19 +32,19 @@ class UARKitCameraOverlayMaterialLoader : public UObject
 	
 public:
 	UPROPERTY()
-	UMaterialInterface* DefaultCameraOverlayMaterial;
+	TObjectPtr<UMaterialInterface> DefaultCameraOverlayMaterial;
 	
 	UPROPERTY()
-	UMaterialInterface* DepthOcclusionOverlayMaterial;
+	TObjectPtr<UMaterialInterface> DepthOcclusionOverlayMaterial;
 	
 	UPROPERTY()
-	UMaterialInterface* MatteOcclusionOverlayMaterial;
+	TObjectPtr<UMaterialInterface> MatteOcclusionOverlayMaterial;
 	
 	UPROPERTY()
-	UMaterialInterface* SceneDepthOcclusionMaterial;
+	TObjectPtr<UMaterialInterface> SceneDepthOcclusionMaterial;
 	
 	UPROPERTY()
-	UMaterialInterface* SceneDepthColorationMaterial;
+	TObjectPtr<UMaterialInterface> SceneDepthColorationMaterial;
 	
 	UARKitCameraOverlayMaterialLoader()
 	{
diff --git a/Engine/Plugins/Runtime/AR/AppleAR/AppleARKitFaceSupport/Source/AppleARKitFaceSupport/Public/AppleARKitFaceMeshComponent.h b/Engine/Plugins/Runtime/AR/AppleAR/AppleARKitFaceSupport/Source/AppleARKitFaceSupport/Public/AppleARKitFaceMeshComponent.h
--- a/Engine/Plugins/Runtime/AR/AppleAR/AppleARKitFaceSupport/Source/AppleARKitFaceSupport/Public/AppleARKitFaceMeshComponent.h
+++ b/Engine/Plugins/Runtime/AR/AppleAR/AppleARKitFaceSupport/Source/AppleARKitFaceSupport/Public/AppleARKitFaceMeshComponent.h
@@ -216,7 +216,7 @@ public:
 
 	/** Used when rendering the face mesh (mostly debug reasons) */
 	UPROPERTY(EditAnywhere, BlueprintReadOnly, Category = "Components|ARFaceMesh")
-	class UMaterialInterface* FaceMaterial;
+	TObjectPtr<class UMaterialInterface> FaceMaterial;
 
 	/** Used to identify this component's face ar data uniquely as part of the LiveLink animation pipeline */
 	UPROPERTY(EditAnywhere, BlueprintReadOnly, Category = "Components|ARFaceMesh")
diff --git a/Engine/Plugins/Runtime/AR/AzureSpatialAnchors/Source/AzureSpatialAnchors/Public/AzureCloudSpatialAnchor.h b/Engine/Plugins/Runtime/AR/AzureSpatialAnchors/Source/AzureSpatialAnchors/Public/AzureCloudSpatialAnchor.h
--- a/Engine/Plugins/Runtime/AR/AzureSpatialAnchors/Source/AzureSpatialAnchors/Public/AzureCloudSpatialAnchor.h
+++ b/Engine/Plugins/Runtime/AR/AzureSpatialAnchors/Source/AzureSpatialAnchors/Public/AzureCloudSpatialAnchor.h
@@ -28,7 +28,7 @@ public:
 	 * The ARPin with which this cloud anchor is associated, or null.  Null could mean we are still trying to load the anchor or have not yet located it.
 	 */
 	UPROPERTY(BlueprintReadOnly, Transient, Category = "AzureSpatialAnchors", meta = (Keywords = "azure spatial anchor hololens wmr pin ar all"))
-	UARPin* ARPin;
+	TObjectPtr<UARPin> ARPin;
 
 	/**
 	 * The Azure Cloud identifier of the spatial anchor is the persistent identifier by which a cloud anchor can be requested from the azure cloud service.  Empty if the anchor has not been saved or loaded yet.
diff --git a/Engine/Plugins/Runtime/AR/AzureSpatialAnchors/Source/AzureSpatialAnchors/Public/AzureSpatialAnchorsTypes.h b/Engine/Plugins/Runtime/AR/AzureSpatialAnchors/Source/AzureSpatialAnchors/Public/AzureSpatialAnchorsTypes.h
--- a/Engine/Plugins/Runtime/AR/AzureSpatialAnchors/Source/AzureSpatialAnchors/Public/AzureSpatialAnchorsTypes.h
+++ b/Engine/Plugins/Runtime/AR/AzureSpatialAnchors/Source/AzureSpatialAnchors/Public/AzureSpatialAnchorsTypes.h
@@ -158,7 +158,7 @@ struct FAzureSpatialAnchorsLocateCriteria
 	 * Specify (optionally) an anchor around which to locate anchors.
 	 */
 	UPROPERTY(BlueprintReadWrite, Category = "AzureSpatialAnchors | NearAnchor")
-	class UAzureCloudSpatialAnchor* NearAnchor = nullptr;
+	TObjectPtr<class UAzureCloudSpatialAnchor> NearAnchor = nullptr;
 
 	/**
 	 * Specify the distance at which to locate anchors near the NearAnchor, in cm.
diff --git a/Engine/Plugins/Runtime/AR/Google/GoogleARCore/Source/GoogleARCoreBase/Private/GoogleARCoreAPI.h b/Engine/Plugins/Runtime/AR/Google/GoogleARCore/Source/GoogleARCoreBase/Private/GoogleARCoreAPI.h
--- a/Engine/Plugins/Runtime/AR/Google/GoogleARCore/Source/GoogleARCoreBase/Private/GoogleARCoreAPI.h
+++ b/Engine/Plugins/Runtime/AR/Google/GoogleARCore/Source/GoogleARCoreBase/Private/GoogleARCoreAPI.h
@@ -190,7 +190,7 @@ class UGoogleARCoreUObjectManager : public UObject
 	
 public:
 	UPROPERTY()
-	UGoogleARCorePointCloud* LatestPointCloud;
+	TObjectPtr<UGoogleARCorePointCloud> LatestPointCloud;
 	
 	// pointer type is ArTrackable*
 	UPROPERTY()
@@ -198,7 +198,7 @@ public:
 	
 	// pointer type is ArAnchor*
 	UPROPERTY()
-	TMap<FARCorePointer, UARPin*> HandleToAnchorMap;
+	TMap<FARCorePointer, TObjectPtr<UARPin>> HandleToAnchorMap;
 	
 	void ClearTrackables();
 	
diff --git a/Engine/Plugins/Runtime/AR/Google/GoogleARCore/Source/GoogleARCoreBase/Public/GoogleARCoreAugmentedImageDatabase.h b/Engine/Plugins/Runtime/AR/Google/GoogleARCore/Source/GoogleARCoreBase/Public/GoogleARCoreAugmentedImageDatabase.h
--- a/Engine/Plugins/Runtime/AR/Google/GoogleARCore/Source/GoogleARCoreBase/Public/GoogleARCoreAugmentedImageDatabase.h
+++ b/Engine/Plugins/Runtime/AR/Google/GoogleARCore/Source/GoogleARCoreBase/Public/GoogleARCoreAugmentedImageDatabase.h
@@ -34,7 +34,7 @@ struct FGoogleARCoreAugmentedImageDatabaseEntry
 	 * BGRA8.
 	 */
 	UPROPERTY(EditAnywhere, BlueprintReadOnly, Category = "GoogleARCore|AugmentedImages")
-	UTexture2D *ImageAsset;
+	TObjectPtr<UTexture2D> ImageAsset;
 
 	/**
 	 * Width of the image in meters.
diff --git a/Engine/Plugins/Runtime/AR/Google/GoogleARCore/Source/GoogleARCoreBase/Public/GoogleARCoreFaceMeshComponent.h b/Engine/Plugins/Runtime/AR/Google/GoogleARCore/Source/GoogleARCoreBase/Public/GoogleARCoreFaceMeshComponent.h
--- a/Engine/Plugins/Runtime/AR/Google/GoogleARCore/Source/GoogleARCoreBase/Public/GoogleARCoreFaceMeshComponent.h
+++ b/Engine/Plugins/Runtime/AR/Google/GoogleARCore/Source/GoogleARCoreBase/Public/GoogleARCoreFaceMeshComponent.h
@@ -82,7 +82,7 @@ public:
 
 	/** Used when rendering the face mesh (mostly debug reasons) */
 	UPROPERTY(EditAnywhere, BlueprintReadOnly, Category = "Components|ARFaceMesh")
-	class UMaterialInterface* FaceMaterial;
+	TObjectPtr<class UMaterialInterface> FaceMaterial;
 
 private:
 	virtual void TickComponent(float DeltaTime, enum ELevelTick TickType, FActorComponentTickFunction* ThisTickFunction) override;
diff --git a/Engine/Plugins/Runtime/AR/Google/GoogleARCore/Source/GoogleARCoreBase/Public/GoogleARCoreSessionConfig.h b/Engine/Plugins/Runtime/AR/Google/GoogleARCore/Source/GoogleARCoreBase/Public/GoogleARCoreSessionConfig.h
--- a/Engine/Plugins/Runtime/AR/Google/GoogleARCore/Source/GoogleARCoreBase/Public/GoogleARCoreSessionConfig.h
+++ b/Engine/Plugins/Runtime/AR/Google/GoogleARCore/Source/GoogleARCoreBase/Public/GoogleARCoreSessionConfig.h
@@ -46,7 +46,7 @@ public:
 	 * image tracking.
 	 */
 	UPROPERTY(EditAnywhere, Category = "google arcore augmentedimages")
-	UGoogleARCoreAugmentedImageDatabase *AugmentedImageDatabase;
+	TObjectPtr<UGoogleARCoreAugmentedImageDatabase> AugmentedImageDatabase;
 
 	/** Configure which camera will be used to in the AR session. */
 	UPROPERTY(EditAnywhere, Category = "ARCore Settings")
diff --git a/Engine/Plugins/Runtime/AR/Google/GoogleARCore/Source/GoogleARCoreRendering/Public/GoogleARCorePassthroughCameraRenderer.h b/Engine/Plugins/Runtime/AR/Google/GoogleARCore/Source/GoogleARCoreRendering/Public/GoogleARCorePassthroughCameraRenderer.h
--- a/Engine/Plugins/Runtime/AR/Google/GoogleARCore/Source/GoogleARCoreRendering/Public/GoogleARCorePassthroughCameraRenderer.h
+++ b/Engine/Plugins/Runtime/AR/Google/GoogleARCore/Source/GoogleARCoreRendering/Public/GoogleARCorePassthroughCameraRenderer.h
@@ -32,18 +32,18 @@ class GOOGLEARCORERENDERING_API UGoogleARCoreCameraOverlayMaterialLoader : publi
 public:
 	/** A pointer to the camera overlay material that will be used to render the passthrough camera texture as background. */
 	UPROPERTY()
-	UMaterialInterface* RegularOverlayMaterial;
+	TObjectPtr<UMaterialInterface> RegularOverlayMaterial;
 	
 	/** A pointer to the camera overlay material that will be used to render the passthrough camera texture as background. */
 	UPROPERTY()
-	UMaterialInterface* DebugOverlayMaterial;
+	TObjectPtr<UMaterialInterface> DebugOverlayMaterial;
 	
 	UPROPERTY()
-	UMaterialInterface* DepthOcclusionMaterial;
+	TObjectPtr<UMaterialInterface> DepthOcclusionMaterial;
 	
 	/** Material used for rendering the coloration of the depth map. */
 	UPROPERTY()
-	UMaterialInterface* DepthColorationMaterial;
+	TObjectPtr<UMaterialInterface> DepthColorationMaterial;
 
 	UGoogleARCoreCameraOverlayMaterialLoader()
 	{
diff --git a/Engine/Plugins/Runtime/ApexDestruction/Source/ApexDestruction/Public/DestructibleActor.h b/Engine/Plugins/Runtime/ApexDestruction/Source/ApexDestruction/Public/DestructibleActor.h
--- a/Engine/Plugins/Runtime/ApexDestruction/Source/ApexDestruction/Public/DestructibleActor.h
+++ b/Engine/Plugins/Runtime/ApexDestruction/Source/ApexDestruction/Public/DestructibleActor.h
@@ -26,7 +26,7 @@ class APEXDESTRUCTION_API ADestructibleActor : public AActor
 private:
 PRAGMA_DISABLE_DEPRECATION_WARNINGS
 	UPROPERTY(VisibleAnywhere, BlueprintReadOnly, Category = Destruction, meta = (ExposeFunctionCategories = "Destruction,Components|Destructible", AllowPrivateAccess = "true"))
-	UDestructibleComponent* DestructibleComponent;
+	TObjectPtr<UDestructibleComponent> DestructibleComponent;
 PRAGMA_ENABLE_DEPRECATION_WARNINGS
 public:
 
diff --git a/Engine/Plugins/Runtime/ApexDestruction/Source/ApexDestruction/Public/DestructibleComponent.h b/Engine/Plugins/Runtime/ApexDestruction/Source/ApexDestruction/Public/DestructibleComponent.h
--- a/Engine/Plugins/Runtime/ApexDestruction/Source/ApexDestruction/Public/DestructibleComponent.h
+++ b/Engine/Plugins/Runtime/ApexDestruction/Source/ApexDestruction/Public/DestructibleComponent.h
@@ -56,7 +56,7 @@ class APEXDESTRUCTION_API UDestructibleComponent : public USkinnedMeshComponent,
 #if WITH_EDITORONLY_DATA
 	/** Provide a blueprint interface for setting the destructible mesh */
 	UPROPERTY(Transient, BlueprintSetter=SetDestructibleMesh, BlueprintGetter=GetDestructibleMesh, Category = "Components|Destructible", meta = (DeprecatedProperty, DeprecatedMessage = "Destructible mesh is now deprected, use UGeometryCollection instead"))
-	class UDestructibleMesh* DestructibleMesh_DEPRECATED;
+	TObjectPtr<class UDestructibleMesh> DestructibleMesh_DEPRECATED;
 #endif // WITH_EDITORONLY_DATA
 
 #if WITH_EDITOR
diff --git a/Engine/Plugins/Runtime/ApexDestruction/Source/ApexDestruction/Public/DestructibleFractureSettings.h b/Engine/Plugins/Runtime/ApexDestruction/Source/ApexDestruction/Public/DestructibleFractureSettings.h
--- a/Engine/Plugins/Runtime/ApexDestruction/Source/ApexDestruction/Public/DestructibleFractureSettings.h
+++ b/Engine/Plugins/Runtime/ApexDestruction/Source/ApexDestruction/Public/DestructibleFractureSettings.h
@@ -167,7 +167,7 @@ PRAGMA_ENABLE_DEPRECATION_WARNINGS
 
 	/** APEX references materials by name, but we'll bypass that mechanism and use of UE materials instead. */
 	UPROPERTY()
-	TArray<UMaterialInterface*>				Materials;
+	TArray<TObjectPtr<UMaterialInterface>>				Materials;
 
 PRAGMA_DISABLE_DEPRECATION_WARNINGS
 	/** Per-chunk authoring parameters, which should be made writable when a chunk selection GUI is in place. */
diff --git a/Engine/Plugins/Runtime/ApexDestruction/Source/ApexDestruction/Public/DestructibleMesh.h b/Engine/Plugins/Runtime/ApexDestruction/Source/ApexDestruction/Public/DestructibleMesh.h
--- a/Engine/Plugins/Runtime/ApexDestruction/Source/ApexDestruction/Public/DestructibleMesh.h
+++ b/Engine/Plugins/Runtime/ApexDestruction/Source/ApexDestruction/Public/DestructibleMesh.h
@@ -421,11 +421,11 @@ class UE_DEPRECATED(4.26, "APEX is deprecated. Destruction in future will be sup
 #if WITH_EDITORONLY_DATA
 	/** Information used to author an NxDestructibleAsset*/
 	UPROPERTY(instanced)
-	class UDestructibleFractureSettings* FractureSettings;
+	TObjectPtr<class UDestructibleFractureSettings> FractureSettings;
 	
 	/** Static mesh this destructible mesh is created from. Is nullptr if not created from a static mesh */
 	UPROPERTY()
-	UStaticMesh* SourceStaticMesh;
+	TObjectPtr<UStaticMesh> SourceStaticMesh;
 	
 	/** Timestamp of the source static meshes last import at the time this destruction mesh has been generated. */
 	UPROPERTY()
@@ -433,7 +433,7 @@ class UE_DEPRECATED(4.26, "APEX is deprecated. Destruction in future will be sup
 
 	/** Array of static meshes to build the fracture chunks from */
 	UPROPERTY(Transient)
-	TArray<UStaticMesh*> FractureChunkMeshes;
+	TArray<TObjectPtr<UStaticMesh>> FractureChunkMeshes;
 #endif // WITH_EDITORONLY_DATA
 
 public:
diff --git a/Engine/Plugins/Runtime/AudioModulation/Source/AudioModulation/Public/Generators/SoundModulationEnvelopeFollower.h b/Engine/Plugins/Runtime/AudioModulation/Source/AudioModulation/Public/Generators/SoundModulationEnvelopeFollower.h
--- a/Engine/Plugins/Runtime/AudioModulation/Source/AudioModulation/Public/Generators/SoundModulationEnvelopeFollower.h
+++ b/Engine/Plugins/Runtime/AudioModulation/Source/AudioModulation/Public/Generators/SoundModulationEnvelopeFollower.h
@@ -35,7 +35,7 @@ struct FEnvelopeFollowerGeneratorParams
 
 	/** AudioBus to follow amplitude of and generate modulation control signal from. */
 	UPROPERTY(EditAnywhere, Category = Modulation, BlueprintReadWrite, meta = (EditCondition = "!bBypass"))
-	UAudioBus* AudioBus = nullptr;
+	TObjectPtr<UAudioBus> AudioBus = nullptr;
 
 	/** Gain to apply to amplitude signal. */
 	UPROPERTY(EditAnywhere, Category = Modulation, BlueprintReadWrite, meta = (EditCondition = "!bBypass", ClampMin = 0.0f))
diff --git a/Engine/Plugins/Runtime/AudioModulation/Source/AudioModulation/Public/SoundControlBus.h b/Engine/Plugins/Runtime/AudioModulation/Source/AudioModulation/Public/SoundControlBus.h
--- a/Engine/Plugins/Runtime/AudioModulation/Source/AudioModulation/Public/SoundControlBus.h
+++ b/Engine/Plugins/Runtime/AudioModulation/Source/AudioModulation/Public/SoundControlBus.h
@@ -41,10 +41,10 @@ public:
 	FString Address;
 
 	UPROPERTY(EditAnywhere, Category = Generators, BlueprintReadWrite)
-	TArray<USoundModulationGenerator*> Generators;
+	TArray<TObjectPtr<USoundModulationGenerator>> Generators;
 
 	UPROPERTY(EditAnywhere, Category = General, BlueprintReadOnly)
-	USoundModulationParameter* Parameter;
+	TObjectPtr<USoundModulationParameter> Parameter;
 
 	/* UObject Implementation */
 #if WITH_EDITOR
diff --git a/Engine/Plugins/Runtime/AudioModulation/Source/AudioModulation/Public/SoundControlBusMix.h b/Engine/Plugins/Runtime/AudioModulation/Source/AudioModulation/Public/SoundControlBusMix.h
--- a/Engine/Plugins/Runtime/AudioModulation/Source/AudioModulation/Public/SoundControlBusMix.h
+++ b/Engine/Plugins/Runtime/AudioModulation/Source/AudioModulation/Public/SoundControlBusMix.h
@@ -24,7 +24,7 @@ struct AUDIOMODULATION_API FSoundControlBusMixStage
 
 	/* Bus controlled by stage. */
 	UPROPERTY(EditAnywhere, Category = Stage, BlueprintReadWrite)
-	USoundControlBus* Bus;
+	TObjectPtr<USoundControlBus> Bus;
 
 	/* Value mix is set to. */
 	UPROPERTY(EditAnywhere, Category = Stage, BlueprintReadWrite)
diff --git a/Engine/Plugins/Runtime/AudioModulation/Source/AudioModulation/Public/SoundModulationPatch.h b/Engine/Plugins/Runtime/AudioModulation/Source/AudioModulation/Public/SoundModulationPatch.h
--- a/Engine/Plugins/Runtime/AudioModulation/Source/AudioModulation/Public/SoundModulationPatch.h
+++ b/Engine/Plugins/Runtime/AudioModulation/Source/AudioModulation/Public/SoundModulationPatch.h
@@ -38,7 +38,7 @@ struct AUDIOMODULATION_API FSoundControlModulationInput
 
 	/** The input bus */
 	UPROPERTY(EditAnywhere, BlueprintReadWrite, Category = Input)
-	USoundControlBus* Bus = nullptr;
+	TObjectPtr<USoundControlBus> Bus = nullptr;
 
 	const USoundControlBus* GetBus() const;
 	const USoundControlBus& GetBusChecked() const;
@@ -54,7 +54,7 @@ struct AUDIOMODULATION_API FSoundControlModulationPatch
 	bool bBypass = true;
 
 	UPROPERTY(EditAnywhere, BlueprintReadOnly, Category = Output, meta = (DisplayName = "Parameter"))
-	USoundModulationParameter* OutputParameter = nullptr;
+	TObjectPtr<USoundModulationParameter> OutputParameter = nullptr;
 
 	/** Modulation inputs */
 	UPROPERTY(EditAnywhere, BlueprintReadWrite, Category = Inputs)
diff --git a/Engine/Plugins/Runtime/DataRegistry/Source/DataRegistry/Public/DataRegistry.h b/Engine/Plugins/Runtime/DataRegistry/Source/DataRegistry/Public/DataRegistry.h
--- a/Engine/Plugins/Runtime/DataRegistry/Source/DataRegistry/Public/DataRegistry.h
+++ b/Engine/Plugins/Runtime/DataRegistry/Source/DataRegistry/Public/DataRegistry.h
@@ -182,16 +182,16 @@ protected:
 
 	/** Structure type of all for items in this registry */
 	UPROPERTY(EditDefaultsOnly, Category = DataRegistry, AssetRegistrySearchable, meta = (DisplayThumbnail = "false"))
-	const UScriptStruct* ItemStruct;
+	TObjectPtr<const UScriptStruct> ItemStruct;
 
 	/** List of data sources to search for items */
 	UPROPERTY(EditDefaultsOnly, Instanced, Category = DataRegistry)
-	TArray<UDataRegistrySource*> DataSources;
+	TArray<TObjectPtr<UDataRegistrySource>> DataSources;
 
 	// TODO remove VisibleDefaultsOnly or figure out how to stop it from letting you edit the instance properties
 	/** Runtime list of data sources, created from above list and includes sources added at runtime */
 	UPROPERTY(VisibleDefaultsOnly, Instanced, Transient, Category = DataRegistry)
-	TArray<UDataRegistrySource*> RuntimeSources;
+	TArray<TObjectPtr<UDataRegistrySource>> RuntimeSources;
 
 	/** How often to check for cache updates */
 	UPROPERTY(EditDefaultsOnly, Category = Cache)
