MoviePipeline: Fixed a copy/paste error for projection matrix calculation.

diff --git a/Engine/Plugins/MovieScene/MovieRenderPipeline/Source/MovieRenderPipelineRenderPasses/Private/MoviePipelineImagePassBase.cpp b/Engine/Plugins/MovieScene/MovieRenderPipeline/Source/MovieRenderPipelineRenderPasses/Private/MoviePipelineImagePassBase.cpp
--- a/Engine/Plugins/MovieScene/MovieRenderPipeline/Source/MovieRenderPipelineRenderPasses/Private/MoviePipelineImagePassBase.cpp
+++ b/Engine/Plugins/MovieScene/MovieRenderPipeline/Source/MovieRenderPipelineRenderPasses/Private/MoviePipelineImagePassBase.cpp
@@ -594,12 +594,12 @@ void UMoviePipelineImagePassBase::ModifyProjectionMatrixForTiling(const FMoviePi
 	if (bInOrthographic)
 	{
 		InOutProjectionMatrix.M[3][0] += OffsetX / PadRatioX;
-		InOutProjectionMatrix.M[3][1] += OffsetY / PadRatioX;
+		InOutProjectionMatrix.M[3][1] += OffsetY / PadRatioY;
 	}
 	else
 	{
 		InOutProjectionMatrix.M[2][0] += OffsetX / PadRatioX;
-		InOutProjectionMatrix.M[2][1] += OffsetY / PadRatioX;
+		InOutProjectionMatrix.M[2][1] += OffsetY / PadRatioY;
 	}
 }
 
