MoviePipeline: Fixed crash when Render All Cameras was enabled, but there was a Camera Cut Track section that pointed to an invalid/unresolvable binding. This also de-duplicates camera names to avoid issues later where not all frames are submitted because two different render passes submitted with the same camera identifier.

diff --git a/Engine/Plugins/MovieScene/MovieRenderPipeline/Source/MovieRenderPipelineCore/Private/MoviePipelineBlueprintLibrary.cpp b/Engine/Plugins/MovieScene/MovieRenderPipeline/Source/MovieRenderPipelineCore/Private/MoviePipelineBlueprintLibrary.cpp
--- a/Engine/Plugins/MovieScene/MovieRenderPipeline/Source/MovieRenderPipelineCore/Private/MoviePipelineBlueprintLibrary.cpp
+++ b/Engine/Plugins/MovieScene/MovieRenderPipeline/Source/MovieRenderPipelineCore/Private/MoviePipelineBlueprintLibrary.cpp
@@ -378,6 +378,40 @@ UMoviePipelineQueue* UMoviePipelineBlueprintLibrary::LoadManifestFileFromString(
 	return OutQueue;
 }
 
+namespace UE
+{
+namespace MoviePipeline
+{
+	void DeduplicateNameArray(TArray<FString>& InOutNames)
+	{
+		TMap<FString, int32> NameUseCount;
+		for (FString& Name : InOutNames)
+		{
+			int32& Count = NameUseCount.FindOrAdd(Name, 0);
+			if (++Count > 1)
+			{
+				Name.Append(FString::Format(TEXT("({0})"), { NameUseCount[Name] }));
+			}
+		}
+
+		// For any shot names we found duplicates, append 1 to the first to keep naming consistent
+		for (TPair<FString, int32>& Pair : NameUseCount)
+		{
+			if (Pair.Value > 1)
+			{
+				for (FString& Name : InOutNames)
+				{
+					if (Name.Equals(Pair.Key))
+					{
+						Name.Append(TEXT("(1)"));
+						break;
+					}
+				}
+			}
+		}
+	}
+}
+}
 
 void UMoviePipelineBlueprintLibrary::UpdateJobShotListFromSequence(ULevelSequence* InSequence, UMoviePipelineExecutorJob* InJob, bool& bShotsChanged)
 {
@@ -612,17 +646,26 @@ void UMoviePipelineBlueprintLibrary::UpdateJobShotListFromSequence(ULevelSequenc
 				if (MainCameraIndex == INDEX_NONE)
 				{
 					// When the camera cut points to a camera in another Sequence it won't get detected by the above loop,
-					// so we manually add it as the only sidecar camera.
-					FMoviePipelineSidecarCamera& SidecarCamera = Entity.SidecarCameras.AddDefaulted_GetRef();
-					SidecarCamera.BindingId = LeafNode->CameraCutSection->GetCameraBindingID().GetGuid();
-					SidecarCamera.SequenceId = LeafNode->CameraCutSection->GetCameraBindingID().GetRelativeSequenceID();
-					SidecarCamera.Name = TEXT("main_camera");
-					MainCameraIndex = Entity.SidecarCameras.Num() - 1;
+					// so we manually add it as the only sidecar camera. We need to check that the CameraCutSection points
+					// to a real camera though, as they may have an unbound camera.
+					if (LeafNode->CameraCutSection->GetCameraBindingID().GetGuid().IsValid())
+					{
+						FMoviePipelineSidecarCamera& SidecarCamera = Entity.SidecarCameras.AddDefaulted_GetRef();
+						SidecarCamera.BindingId = LeafNode->CameraCutSection->GetCameraBindingID().GetGuid();
+						SidecarCamera.SequenceId = LeafNode->CameraCutSection->GetCameraBindingID().GetRelativeSequenceID();
+
+						// ToDo: We should fetch the actual camera name from the Object Binding in the other sequence.
+						SidecarCamera.Name = TEXT("main_camera");
+						MainCameraIndex = Entity.SidecarCameras.Num() - 1;
+					}
 				}
 
 				// Ensure the "main" camera (specified by the Camera Cut Track) is always index 0, so that it's rendered
 				// first, and the stable sort (on output files) still tries to keep it as the primary layer.
-				Entity.SidecarCameras.Swap(MainCameraIndex, 0);
+				if (MainCameraIndex >= 0)
+				{
+					Entity.SidecarCameras.Swap(MainCameraIndex, 0);
+				}
 			}
 
 			FMovieSceneTimeTransform InnerToOuterTransform = FMovieSceneTimeTransform();
@@ -659,29 +702,39 @@ void UMoviePipelineBlueprintLibrary::UpdateJobShotListFromSequence(ULevelSequenc
 		// We need to generate all of the linearized segments first so that we have all of the names available.
 		// We need all of the names available because we need to ensure unique shot names (for XMLs), and they 
 		// expect to be consistent (ie: if duplicates are found, the first one is retroactively changed to (1)).
-		TMap<FString, int32> ShotNameUseCount;
-		for (FLinearizedEntity& Entity : Entities)
 		{
-			int32& Count = ShotNameUseCount.FindOrAdd(Entity.Name.Get<1>(), 0);
-			if (++Count > 1)
+			TArray<FString> ShotNames;
+			for (FLinearizedEntity& Entity : Entities)
+			{
+				// Insert the outer name portion (which is the shot)
+				ShotNames.Add(Entity.Name.Get<1>());
+			}
+
+			// Now fix up the names to avoid duplicates
+			UE::MoviePipeline::DeduplicateNameArray(ShotNames);
+
+			// Now re-apply the fixed names
+			for (int32 Index = 0; Index < Entities.Num(); Index++)
 			{
-				FString& ShotName = Entity.Name.Get<1>();
-				ShotName.Append(FString::Format(TEXT("({0})"), { ShotNameUseCount[Entity.Name.Get<1>()] }));
+				Entities[Index].Name.Get<1>() = ShotNames[Index];
 			}
 		}
 
-		// For any shot names we found duplicates, append 1 to the first to keep naming consistent
-		for (TPair<FString, int32>& Pair : ShotNameUseCount)
+		// Repeat for camera names within each shot.
 		{
-			if (Pair.Value > 1)
+			for (FLinearizedEntity& Entity : Entities)
 			{
-				for (FLinearizedEntity& Entity : Entities)
+				TArray<FString> CameraNames;
+				for (FMoviePipelineSidecarCamera& SidecarCamera : Entity.SidecarCameras)
 				{
-					if (Entity.Name.Get<1>().Equals(Pair.Key))
-					{
-						Entity.Name.Get<1>().Append(TEXT("(1)"));
-						break;
-					}
+					CameraNames.Add(SidecarCamera.Name);
+				}
+				
+				UE::MoviePipeline::DeduplicateNameArray(CameraNames);
+
+				for (int32 Index = 0; Index < Entity.SidecarCameras.Num(); Index++)
+				{
+					Entity.SidecarCameras[Index].Name = CameraNames[Index];
 				}
 			}
 		}
