MoviePipeline: Fixed an issue where Video Codecs didn't work with the "Overwrite Existing" feature disabled.

diff --git a/Engine/Plugins/MovieScene/MovieRenderPipeline/Source/MovieRenderPipelineCore/Private/MoviePipelineVideoOutputBase.cpp b/Engine/Plugins/MovieScene/MovieRenderPipeline/Source/MovieRenderPipelineCore/Private/MoviePipelineVideoOutputBase.cpp
--- a/Engine/Plugins/MovieScene/MovieRenderPipeline/Source/MovieRenderPipelineCore/Private/MoviePipelineVideoOutputBase.cpp
+++ b/Engine/Plugins/MovieScene/MovieRenderPipeline/Source/MovieRenderPipelineCore/Private/MoviePipelineVideoOutputBase.cpp
@@ -86,6 +86,7 @@ void UMoviePipelineVideoOutputBase::OnReceiveImageDataImpl(FMoviePipelineMergerO
 		// We need to resolve the filename format string. We combine the folder and file name into one long string first
 		FMoviePipelineFormatArgs FinalFormatArgs;
 		FString FinalFilePath;
+		FString StableFilePath;
 		FString FinalVideoFileName;
 		FString ClipName;
 		{
@@ -106,6 +107,29 @@ void UMoviePipelineVideoOutputBase::OnReceiveImageDataImpl(FMoviePipelineMergerO
 			FormatOverrides.Add(TEXT("render_pass"), RenderPassData.Key.Name);
 			FormatOverrides.Add(TEXT("ext"), GetFilenameExtension());
 
+			// This is a bit crummy but we don't have a good way to implement bOverwriteFiles = False for video files. When you don't have the
+			// override file flag set, we need to increment the number by 1, and write to that. This works fine for image sequences as image
+			// sequences are also differentiated by their {frame_number}, but for video files (which strip it all out), instead we end up
+			// with each incoming image sample trying to generate the same filename, finding a file that already exists (from the previous
+			// frame) and incrementing again, giving us 1 video file per frame. To work around this, we're going to force bOverwriteExisting
+			// off so we can generate a "stable" filename for the incoming image sample, then compare it against existing writers, and
+			// if we don't find a writer, then generate a new one (respecting Overwrite Existing) to add the sample to.
+			{
+				UMoviePipelineOutputSetting* OutputSetting = GetPipeline()->GetPipelineMasterConfig()->FindSetting<UMoviePipelineOutputSetting>();
+				const bool bPreviousOverwriteExisting = OutputSetting->bOverrideExistingOutput;
+				OutputSetting->bOverrideExistingOutput = true;
+				
+				GetPipeline()->ResolveFilenameFormatArguments(OutputDirectory / FileNameFormatString, FormatOverrides, StableFilePath, FinalFormatArgs, &InMergedOutputFrame->FrameOutputState);
+				
+				// Restore user setting
+				OutputSetting->bOverrideExistingOutput = bPreviousOverwriteExisting;
+
+				if (FPaths::IsRelative(StableFilePath))
+				{
+					StableFilePath = FPaths::ConvertRelativePathToFull(StableFilePath);
+				}
+			}
+
 			// The FinalVideoFileName is relative to the output directory (ie: if the user puts folders in to the filename path)
 			GetPipeline()->ResolveFilenameFormatArguments(FileNameFormatString, FormatOverrides, FinalVideoFileName, FinalFormatArgs, &InMergedOutputFrame->FrameOutputState);
 
@@ -138,7 +162,7 @@ void UMoviePipelineVideoOutputBase::OnReceiveImageDataImpl(FMoviePipelineMergerO
 		FMoviePipelineCodecWriter* OutputWriter = nullptr;
 		for (int32 Index = 0; Index < AllWriters.Num(); Index++)
 		{
-			if (AllWriters[Index].Get<0>()->FileName == FinalFilePath)
+			if (AllWriters[Index].Get<0>()->StableFileName == StableFilePath)
 			{
 				OutputWriter = &AllWriters[Index];
 				break;
@@ -154,6 +178,9 @@ void UMoviePipelineVideoOutputBase::OnReceiveImageDataImpl(FMoviePipelineMergerO
 
 			if (NewWriter)
 			{
+				// Store the stable filename this was generated with so we can match them up later.
+				NewWriter->StableFileName = StableFilePath;
+
 				TPromise<bool> Completed;
 				MoviePipeline::FMoviePipelineOutputFutureData OutputData;
 				OutputData.Shot = GetPipeline()->GetActiveShotList()[Payload->SampleState.OutputState.ShotIndex];
diff --git a/Engine/Plugins/MovieScene/MovieRenderPipeline/Source/MovieRenderPipelineCore/Public/MoviePipelineVideoOutputBase.h b/Engine/Plugins/MovieScene/MovieRenderPipeline/Source/MovieRenderPipelineCore/Public/MoviePipelineVideoOutputBase.h
--- a/Engine/Plugins/MovieScene/MovieRenderPipeline/Source/MovieRenderPipelineCore/Public/MoviePipelineVideoOutputBase.h
+++ b/Engine/Plugins/MovieScene/MovieRenderPipeline/Source/MovieRenderPipelineCore/Public/MoviePipelineVideoOutputBase.h
@@ -14,7 +14,12 @@ namespace MovieRenderPipeline
 {
 	struct IVideoCodecWriter
 	{
+		// The filename to actually write to disk with. This may include "de-duplication" numbers (MyFile (1).ext etc.)
 		FString FileName;
+		// The filename without de-duplication numbers, used to match up multiple incoming frames back to the same writer.
+		// We use this when looking for existing writers so that we can avoid the de-duplication numbers perpetually increasing
+		// due to the file existing on disk after the first frame comes in, and then the next one de-duplicating to one more than that.
+		FString StableFileName;
 		FMoviePipelineFormatArgs FormatArgs;
 		bool bConvertToSrgb;
 	};
