MoviePipeline: Added support for Data Layers in the Stencil Layer feature, so the user can specify given Data Layers instead of Editor Layers (which are not supported in World Partition).

diff --git a/Engine/Plugins/MovieScene/MovieRenderPipeline/Config/BaseMovieRenderPipeline.ini b/Engine/Plugins/MovieScene/MovieRenderPipeline/Config/BaseMovieRenderPipeline.ini
new file mode 100644
--- /dev/null
+++ b/Engine/Plugins/MovieScene/MovieRenderPipeline/Config/BaseMovieRenderPipeline.ini
@@ -0,0 +1,2 @@
+[CoreRedirects]
++PropertyRedirects=(OldName="MoviePipelineDeferredPassBase.StencilLayers", NewName="ActorLayers")
diff --git a/Engine/Plugins/MovieScene/MovieRenderPipeline/Source/MovieRenderPipelineRenderPasses/Private/MoviePipelineDeferredPasses.cpp b/Engine/Plugins/MovieScene/MovieRenderPipeline/Source/MovieRenderPipelineRenderPasses/Private/MoviePipelineDeferredPasses.cpp
--- a/Engine/Plugins/MovieScene/MovieRenderPipeline/Source/MovieRenderPipelineRenderPasses/Private/MoviePipelineDeferredPasses.cpp
+++ b/Engine/Plugins/MovieScene/MovieRenderPipeline/Source/MovieRenderPipelineRenderPasses/Private/MoviePipelineDeferredPasses.cpp
@@ -35,12 +35,13 @@
 #include "CineCameraComponent.h"
 #include "Interfaces/Interface_PostProcessVolume.h"
 #include "MoviePipelineUtils.h"
+#include "WorldPartition/DataLayer/DataLayerAsset.h"
 
 FString UMoviePipelineDeferredPassBase::StencilLayerMaterialAsset = TEXT("/MovieRenderPipeline/Materials/MoviePipeline_StencilCutout.MoviePipeline_StencilCutout");
 FString UMoviePipelineDeferredPassBase::DefaultDepthAsset = TEXT("/MovieRenderPipeline/Materials/MovieRenderQueue_WorldDepth.MovieRenderQueue_WorldDepth");
 FString UMoviePipelineDeferredPassBase::DefaultMotionVectorsAsset = TEXT("/MovieRenderPipeline/Materials/MovieRenderQueue_MotionVectors.MovieRenderQueue_MotionVectors");
 
-PRAGMA_DISABLE_OPTIMIZATION
+
 
 UMoviePipelineDeferredPassBase::UMoviePipelineDeferredPassBase() 
 	: UMoviePipelineImagePassBase()
@@ -102,7 +103,7 @@ void UMoviePipelineDeferredPassBase::SetupImpl(const MoviePipeline::FMoviePipeli
 
 	// [0] is FinalImage, [1] is Default Layer, [1+] is Stencil Layers. Not used by post processing materials
 	// Render Target that the GBuffer is copied to
-	int32 NumRenderTargets = (bAddDefaultLayer ? 1 : 0) + StencilLayers.Num() + 1;
+	int32 NumRenderTargets = (bAddDefaultLayer ? 1 : 0) + GetNumStencilLayers() + 1;
 	for (int32 Index = 0; Index < NumRenderTargets; Index++)
 	{
 		UTextureRenderTarget2D* NewTarget = NewObject<UTextureRenderTarget2D>(GetTransientPackage());
@@ -173,13 +174,13 @@ void UMoviePipelineDeferredPassBase::SetupImpl(const MoviePipeline::FMoviePipeli
 				}
 
 				// If they want to render a "default" stencil layer (that has everything not in another layer) add that...
-				if (StencilLayers.Num() > 0 && bAddDefaultLayer)
+				if (GetNumStencilLayers() > 0 && bAddDefaultLayer)
 				{
 					PerTile.SceneViewStates.AddDefaulted();
 				}
 
 				// Finally all of the other stencil layers
-				for (int32 Index = 0; Index < StencilLayers.Num(); Index++)
+				for (int32 Index = 0; Index < GetNumStencilLayers(); Index++)
 				{
 					PerTile.SceneViewStates.AddDefaulted();
 				}
@@ -192,11 +193,11 @@ void UMoviePipelineDeferredPassBase::SetupImpl(const MoviePipeline::FMoviePipeli
 		{
 			TotalNumberOfAccumulators++;
 		}
-		if (StencilLayers.Num() > 0 && bAddDefaultLayer)
+		if (GetNumStencilLayers() > 0 && bAddDefaultLayer)
 		{
 			TotalNumberOfAccumulators++;
 		}
-		for (int32 Index = 0; Index < StencilLayers.Num(); Index++)
+		for (int32 Index = 0; Index < GetNumStencilLayers(); Index++)
 		{
 			TotalNumberOfAccumulators++;
 		}
@@ -225,7 +226,7 @@ void UMoviePipelineDeferredPassBase::SetupImpl(const MoviePipeline::FMoviePipeli
 	// One Extension per sequence, since each sequence has its own OCIO settings.
 	OCIOSceneViewExtension = FSceneViewExtensions::NewExtension<FOpenColorIODisplayExtension>();
 
-	const bool bEnableStencilPass = bAddDefaultLayer || StencilLayers.Num() > 0;
+	const bool bEnableStencilPass = bAddDefaultLayer || GetNumStencilLayers() > 0;
 	if (bEnableStencilPass)
 	{
 		IConsoleVariable* CVar = IConsoleManager::Get().FindConsoleVariable(TEXT("r.CustomDepth"));
@@ -432,14 +433,14 @@ void UMoviePipelineDeferredPassBase::GatherOutputPassesImpl(TArray<FMoviePipelin
 		}
 
 		// Stencil Layer Time!
-		if (StencilLayers.Num() > 0 && bAddDefaultLayer)
+		if (GetNumStencilLayers() > 0 && bAddDefaultLayer)
 		{
 			ExpectedRenderPasses.Add(FMoviePipelinePassIdentifier(PassIdentifierForCurrentCamera.Name + TEXT("DefaultLayer"), PassIdentifierForCurrentCamera.CameraName));
 		}
 
-		for (const FActorLayer& Layer : StencilLayers)
+		for (const FString& StencilLayerName : GetStencilLayerNames())
 		{
-			ExpectedRenderPasses.Add(FMoviePipelinePassIdentifier(PassIdentifierForCurrentCamera.Name + Layer.Name.ToString(), PassIdentifierForCurrentCamera.CameraName));
+			ExpectedRenderPasses.Add(FMoviePipelinePassIdentifier(PassIdentifierForCurrentCamera.Name + StencilLayerName, PassIdentifierForCurrentCamera.CameraName));
 		}
 	}
 }
@@ -559,11 +560,10 @@ void UMoviePipelineDeferredPassBase::RenderSample_GameThreadImpl(const FMoviePip
 			};
 
 			// Now for each stencil layer we reconfigure all the actors custom depth/stencil 
-			TArray<FActorLayer> AllStencilLayers = StencilLayers;
+			TArray<FString> AllStencilLayers = GetStencilLayerNames();
 			if (bAddDefaultLayer)
 			{
-				FActorLayer& DefaultLayer = AllStencilLayers.AddDefaulted_GetRef();
-				DefaultLayer.Name = FName("DefaultLayer");
+				AllStencilLayers.Add(TEXT("DefaultLayer"));
 			}
 
 			// If we're going to be using stencil layers, we need to cache all of the users
@@ -594,8 +594,8 @@ void UMoviePipelineDeferredPassBase::RenderSample_GameThreadImpl(const FMoviePip
 
 			for (int32 StencilLayerIndex = 0; StencilLayerIndex < AllStencilLayers.Num(); StencilLayerIndex++)
 			{
-				const FActorLayer& Layer = AllStencilLayers[StencilLayerIndex];
-				FMoviePipelinePassIdentifier LayerPassIdentifier = FMoviePipelinePassIdentifier(PassIdentifierForCurrentCamera.Name + Layer.Name.ToString());
+				const FString& LayerName = AllStencilLayers[StencilLayerIndex];
+				FMoviePipelinePassIdentifier LayerPassIdentifier = FMoviePipelinePassIdentifier(PassIdentifierForCurrentCamera.Name + LayerName);
 				LayerPassIdentifier.CameraName = PassIdentifierForCurrentCamera.CameraName;
 
 				// Modify all of the actors in this world so they have the right stencil settings (so we can use the stencil buffer as a mask later)
@@ -609,23 +609,17 @@ void UMoviePipelineDeferredPassBase::RenderSample_GameThreadImpl(const FMoviePip
 						// post processing. Then, TAA, Motion Blur, etc. is applied to all pixels. An alpha channel can preserve
 						// which pixels were the geometry and which are dead space which lets you apply that as a mask later.
 						bool bInLayer = true;
-						if (bAddDefaultLayer && Layer.Name == FName("DefaultLayer"))
+						if (bAddDefaultLayer && LayerName == TEXT("DefaultLayer"))
 						{
 							// If we're trying to render the default layer, the logic is different - we only add objects who
 							// aren't in any of the stencil layers.
-							for (const FActorLayer& AllLayer : StencilLayers)
-							{
-								bInLayer = !Actor->Layers.Contains(AllLayer.Name);
-								if (!bInLayer)
-								{
-									break;
-								}
-							}
+							bInLayer = IsActorInAnyStencilLayer(Actor);
+						
 						}
 						else
 						{
 							// If this a normal layer, we only add the actor if it exists on this layer.
-							bInLayer = Actor->Layers.Contains(Layer.Name);
+							bInLayer = IsActorInLayer(Actor, StencilLayerIndex);
 						}
 
 						for (UActorComponent* Component : Actor->GetComponents())
@@ -1047,4 +1041,135 @@ void UMoviePipelineDeferredPass_PathTracer::SetupImpl(const MoviePipeline::FMovi
 	Super::SetupImpl(InPassInitSettings);
 }
 
-PRAGMA_ENABLE_OPTIMIZATION
\ No newline at end of file
+bool UMoviePipelineDeferredPassBase::IsUsingDataLayers() const
+{
+	int32 NumDataLayers = 0;
+	for (FSoftObjectPath DataLayerAssetPath : DataLayers)
+	{
+		UDataLayerAsset* DataLayerAsset = Cast<UDataLayerAsset>(DataLayerAssetPath.TryLoad());
+		if (DataLayerAsset)
+		{
+			NumDataLayers++;
+		}
+	}
+	return NumDataLayers > 0;
+}
+
+int32 UMoviePipelineDeferredPassBase::GetNumStencilLayers() const
+{
+	if (IsUsingDataLayers())
+	{
+		// Because DataLayers are an asset, they can actually be null despite being in this list.
+		int32 NumDataLayers = 0;
+		for (FSoftObjectPath DataLayerAssetPath : DataLayers)
+		{
+			UDataLayerAsset* DataLayerAsset = Cast<UDataLayerAsset>(DataLayerAssetPath.TryLoad());
+			if (DataLayerAsset)
+			{
+				NumDataLayers++;
+			}
+		}
+		return NumDataLayers;
+	}
+	return ActorLayers.Num();
+}
+
+TArray<FString> UMoviePipelineDeferredPassBase::GetStencilLayerNames() const
+{
+	TArray<FString> LayerNames;
+	if (IsUsingDataLayers())
+	{
+		for (FSoftObjectPath DataLayerAssetPath : DataLayers)
+		{
+			UDataLayerAsset* DataLayerAsset = Cast<UDataLayerAsset>(DataLayerAssetPath.TryLoad());
+			if (DataLayerAsset)
+			{
+				LayerNames.Add(DataLayerAsset->GetName());
+			}
+		}
+	}
+	else
+	{
+		for (const FActorLayer& Layer : ActorLayers)
+		{
+			LayerNames.Add(Layer.Name.ToString());
+		}
+	}
+
+	return LayerNames;
+}
+
+FSoftObjectPath UMoviePipelineDeferredPassBase::GetValidDataLayerByIndex(const int32 InIndex) const
+{
+	int32 NumValidDataLayers = 0;
+	for (FSoftObjectPath DataLayerAssetPath : DataLayers)
+	{
+		UDataLayerAsset* DataLayerAsset = Cast<UDataLayerAsset>(DataLayerAssetPath.TryLoad());
+		if (DataLayerAsset)
+		{
+			if (InIndex == NumValidDataLayers)
+			{
+				return DataLayerAssetPath;
+			}
+
+			NumValidDataLayers++;
+		}
+	}
+
+	return FSoftObjectPath();
+}
+
+bool UMoviePipelineDeferredPassBase::IsActorInLayer(AActor* InActor, int32 InLayerIndex) const
+{
+	if (IsUsingDataLayers())
+	{
+		FSoftObjectPath DataLayerAssetPath = GetValidDataLayerByIndex(InLayerIndex);
+		UDataLayerAsset* DataLayerAsset = Cast<UDataLayerAsset>(DataLayerAssetPath.TryLoad());
+		if (DataLayerAsset)
+		{
+			return InActor->ContainsDataLayer(DataLayerAsset);
+		}
+	}
+	else
+	{
+		const FName& LayerName = ActorLayers[InLayerIndex].Name;
+		return InActor->Layers.Contains(LayerName);
+	}
+
+	return false;
+}
+
+bool UMoviePipelineDeferredPassBase::IsActorInAnyStencilLayer(AActor* InActor) const
+{
+	bool bInLayer = false;
+	if (IsUsingDataLayers())
+	{
+		for (FSoftObjectPath DataLayerAssetPath : DataLayers)
+		{
+			UDataLayerAsset* DataLayerAsset = Cast<UDataLayerAsset>(DataLayerAssetPath.TryLoad());
+			if (DataLayerAsset)
+			{
+				bInLayer = !InActor->ContainsDataLayer(DataLayerAsset);
+
+				if (!bInLayer)
+				{
+					break;
+				}
+			}
+		}
+	}
+	else
+	{
+		for (const FActorLayer& AllLayer : ActorLayers)
+		{
+			bInLayer = !InActor->Layers.Contains(AllLayer.Name);
+			if (!bInLayer)
+			{
+				break;
+			}
+		}
+	}
+
+	return bInLayer;
+}
+
diff --git a/Engine/Plugins/MovieScene/MovieRenderPipeline/Source/MovieRenderPipelineRenderPasses/Public/MoviePipelineDeferredPasses.h b/Engine/Plugins/MovieScene/MovieRenderPipeline/Source/MovieRenderPipelineRenderPasses/Public/MoviePipelineDeferredPasses.h
--- a/Engine/Plugins/MovieScene/MovieRenderPipeline/Source/MovieRenderPipelineRenderPasses/Public/MoviePipelineDeferredPasses.h
+++ b/Engine/Plugins/MovieScene/MovieRenderPipeline/Source/MovieRenderPipelineRenderPasses/Public/MoviePipelineDeferredPasses.h
@@ -70,6 +70,13 @@ protected:
 	virtual FString GetCameraName(const int32 InCameraIndex) const;
 	virtual FString GetCameraNameOverride(const int32 InCameraIndex) const;
 
+	bool IsUsingDataLayers() const;
+	int32 GetNumStencilLayers() const;
+	TArray<FString> GetStencilLayerNames() const;
+	bool IsActorInLayer(AActor* InActor, int32 InLayerIndex) const;
+	bool IsActorInAnyStencilLayer(AActor* InActor) const;
+	FSoftObjectPath GetValidDataLayerByIndex(const int32 InIndex) const;
+
 public:
 	/**
 	* Should multiple temporal/spatial samples accumulate the alpha channel? This requires r.PostProcessing.PropagateAlpha
@@ -123,7 +130,18 @@ public:
 	* base layer. Only works with materials that can write to custom depth.
 	*/
 	UPROPERTY(EditAnywhere, BlueprintReadWrite, Category = "Stencil Clip Layers")
+	TArray<FActorLayer> ActorLayers;
+	
+	UE_DEPRECATED(5.1, "Use ActorLayers property instead.")
 	TArray<FActorLayer> StencilLayers;
+	
+	/**
+	* If the map you are working with is a World Partition map, you can specify Data layers instead of Actor Layers. If any
+	* Data Layers are specified, this will take precedence over any ActorLayers in this config. Does not affect whether or
+	* not the Data Layers are actually loaded, you must ensure layers are loaded for rendering.
+	*/
+	UPROPERTY(EditAnywhere, BlueprintReadWrite, meta = (AllowedClasses = "/Script/Engine.DataLayerAsset"), Category = "Stencil Clip Layers")
+	TArray<FSoftObjectPath> DataLayers;
 
 protected:
 	/** While rendering, store an array of the non-null valid materials loaded from AdditionalPostProcessMaterials. Cleared on teardown. */
