MovieRenderQueue: Add virtual texture settings to UMoviePipelineGameOverrideSetting
Defaults to a virtual texture feedback factor of 1

diff --git a/Engine/Plugins/MovieScene/MovieRenderPipeline/Source/MovieRenderPipelineCore/Public/MoviePipelineGameOverrideSetting.h b/Engine/Plugins/MovieScene/MovieRenderPipeline/Source/MovieRenderPipelineCore/Public/MoviePipelineGameOverrideSetting.h
--- a/Engine/Plugins/MovieScene/MovieRenderPipeline/Source/MovieRenderPipelineCore/Public/MoviePipelineGameOverrideSetting.h
+++ b/Engine/Plugins/MovieScene/MovieRenderPipeline/Source/MovieRenderPipelineCore/Public/MoviePipelineGameOverrideSetting.h
@@ -37,6 +37,8 @@ public:
 		, ViewDistanceScale(50)
 		, bFlushGrassStreaming(true)
 		, bFlushStreamingManagers(true)
+		, bOverrideVirtualTextureFeedbackFactor(true)
+		, VirtualTextureFeedbackFactor(1)
 		, bDisableGPUTimeout(true)
 	{
 	}
@@ -105,6 +107,14 @@ public:
 	UPROPERTY(EditAnywhere, BlueprintReadWrite, Category = "Rendering")
 	bool bFlushStreamingManagers;
 
+	/** If true then override the virtual texture feedback resolution factor. Otherwise the value from the project renderer settings will be used. */
+	UPROPERTY(EditAnywhere, BlueprintReadWrite, Category = "Rendering", meta = (InlineEditConditionToggle))
+	bool bOverrideVirtualTextureFeedbackFactor;
+
+	/** The virtual texture feedback resolution factor. A lower factor will increase virtual texture feedback resolution. */
+	UPROPERTY(EditAnywhere, BlueprintReadWrite, Category = "Rendering", meta = (editcondition = "bOverrideVirtualTextureFeedbackFactor"))
+	int32 VirtualTextureFeedbackFactor;
+
 	/** Should we disable the GPU Timeout? Currently only applicable when using D3D12 renderer. */
 	bool bDisableGPUTimeout;
 
diff --git a/Engine/Plugins/MovieScene/MovieRenderPipeline/Source/MovieRenderPipelineRenderPasses/Private/MoviePipelineImagePassBase.cpp b/Engine/Plugins/MovieScene/MovieRenderPipeline/Source/MovieRenderPipelineRenderPasses/Private/MoviePipelineImagePassBase.cpp
--- a/Engine/Plugins/MovieScene/MovieRenderPipeline/Source/MovieRenderPipelineRenderPasses/Private/MoviePipelineImagePassBase.cpp
+++ b/Engine/Plugins/MovieScene/MovieRenderPipeline/Source/MovieRenderPipelineRenderPasses/Private/MoviePipelineImagePassBase.cpp
@@ -13,6 +13,7 @@
 #include "MoviePipelineQueue.h"
 #include "LegacyScreenPercentageDriver.h"
 #include "MoviePipelineMasterConfig.h"
+#include "MoviePipelineGameOverrideSetting.h"
 #include "EngineModule.h"
 #include "Engine/LocalPlayer.h"
 #include "Engine/RendererSettings.h"
@@ -86,7 +87,6 @@ TSharedPtr<FSceneViewFamilyContext> UMoviePipelineImagePassBase::CalculateViewFa
 	OutViewFamily->bWorldIsPaused = InOutSampleState.bWorldIsPaused;
 	OutViewFamily->ViewMode = ViewModeIndex;
 	OutViewFamily->bOverrideVirtualTextureThrottle = true;
-	// OutViewFamily->VirtualTextureFeedbackFactor = 1;
 
 	const bool bIsPerspective = true;
 	ApplyViewMode(OutViewFamily->ViewMode, bIsPerspective, OutViewFamily->EngineShowFlags);
@@ -94,7 +94,15 @@ TSharedPtr<FSceneViewFamilyContext> UMoviePipelineImagePassBase::CalculateViewFa
 	EngineShowFlagOverride(ESFIM_Game, OutViewFamily->ViewMode, OutViewFamily->EngineShowFlags, false);
 	
 	const UMoviePipelineExecutorShot* Shot = GetPipeline()->GetActiveShotList()[InOutSampleState.OutputState.ShotIndex];
-	
+
+	for (UMoviePipelineGameOverrideSetting* OverrideSetting : GetPipeline()->FindSettingsForShot<UMoviePipelineGameOverrideSetting>(Shot))
+	{
+		if (OverrideSetting->bOverrideVirtualTextureFeedbackFactor)
+		{
+			OutViewFamily->VirtualTextureFeedbackFactor = OverrideSetting->VirtualTextureFeedbackFactor;
+		}
+	}
+
 	// No need to do anything if screen percentage is not supported. 
 	if (IsScreenPercentageSupported())
 	{
