nDisplay.MRQ: Added Path Tracer pass

diff --git a/Engine/Plugins/MovieScene/MovieRenderPipeline/Source/MovieRenderPipelineRenderPasses/Private/MoviePipelineDeferredPasses.cpp b/Engine/Plugins/MovieScene/MovieRenderPipeline/Source/MovieRenderPipelineRenderPasses/Private/MoviePipelineDeferredPasses.cpp
--- a/Engine/Plugins/MovieScene/MovieRenderPipeline/Source/MovieRenderPipelineRenderPasses/Private/MoviePipelineDeferredPasses.cpp
+++ b/Engine/Plugins/MovieScene/MovieRenderPipeline/Source/MovieRenderPipelineRenderPasses/Private/MoviePipelineDeferredPasses.cpp
@@ -996,31 +996,25 @@ FText UMoviePipelineDeferredPass_PathTracer::GetFooterText(UMoviePipelineExecuto
 		"All other Path Tracer settings are taken from the Post Process settings.");
 }
 #endif
-namespace UE
-{
-namespace MoviePipeline
+
+bool UMoviePipelineDeferredPassBase::CheckIfPathTracerIsSupported() const
 {
-	bool CheckIfPathTracerIsSupported()
+	bool bSupportsPathTracing = false;
+	if (IsRayTracingEnabled())
 	{
-		bool bSupportsPathTracing = false;
-		if (IsRayTracingEnabled())
+		IConsoleVariable* PathTracingCVar = IConsoleManager::Get().FindConsoleVariable(TEXT("r.PathTracing"));
+		if (PathTracingCVar)
 		{
-			IConsoleVariable* PathTracingCVar = IConsoleManager::Get().FindConsoleVariable(TEXT("r.PathTracing"));
-			if (PathTracingCVar)
-			{
-				bSupportsPathTracing = PathTracingCVar->GetInt() != 0;
-			}
+			bSupportsPathTracing = PathTracingCVar->GetInt() != 0;
 		}
-		return bSupportsPathTracing;
 	}
+	return bSupportsPathTracing;
 }
-}
-void UMoviePipelineDeferredPass_PathTracer::ValidateStateImpl()
+
+void UMoviePipelineDeferredPassBase::PathTracerValidationImpl()
 {
-	Super::ValidateStateImpl();
+	const bool bSupportsPathTracing = CheckIfPathTracerIsSupported();
 
-	bool bSupportsPathTracing = UE::MoviePipeline::CheckIfPathTracerIsSupported();
-	
 	if (!bSupportsPathTracing)
 	{
 		const FText ValidationWarning = NSLOCTEXT("MovieRenderPipeline", "PathTracerValidation_Unsupported", "Path Tracing is currently not enabled for this project and this render pass will not work.");
@@ -1029,9 +1023,15 @@ void UMoviePipelineDeferredPass_PathTracer::ValidateStateImpl()
 	}
 }
 
+void UMoviePipelineDeferredPass_PathTracer::ValidateStateImpl()
+{
+	Super::ValidateStateImpl();
+	PathTracerValidationImpl();
+}
+
 void UMoviePipelineDeferredPass_PathTracer::SetupImpl(const MoviePipeline::FMoviePipelineRenderPassInitSettings& InPassInitSettings)
 {
-	if (!UE::MoviePipeline::CheckIfPathTracerIsSupported())
+	if (!CheckIfPathTracerIsSupported())
 	{
 		UE_LOG(LogMovieRenderPipeline, Error, TEXT("Cannot render a Path Tracer pass, Path Tracer is not enabled by this project."));
 		GetPipeline()->Shutdown(true);
diff --git a/Engine/Plugins/MovieScene/MovieRenderPipeline/Source/MovieRenderPipelineRenderPasses/Public/MoviePipelineDeferredPasses.h b/Engine/Plugins/MovieScene/MovieRenderPipeline/Source/MovieRenderPipelineRenderPasses/Public/MoviePipelineDeferredPasses.h
--- a/Engine/Plugins/MovieScene/MovieRenderPipeline/Source/MovieRenderPipelineRenderPasses/Public/MoviePipelineDeferredPasses.h
+++ b/Engine/Plugins/MovieScene/MovieRenderPipeline/Source/MovieRenderPipelineRenderPasses/Public/MoviePipelineDeferredPasses.h
@@ -77,6 +77,9 @@ protected:
 	bool IsActorInAnyStencilLayer(AActor* InActor) const;
 	FSoftObjectPath GetValidDataLayerByIndex(const int32 InIndex) const;
 
+	bool CheckIfPathTracerIsSupported() const;
+	void PathTracerValidationImpl();
+
 public:
 	/**
 	* Should multiple temporal/spatial samples accumulate the alpha channel? This requires r.PostProcessing.PropagateAlpha
diff --git a/Engine/Plugins/Runtime/nDisplay/Source/DisplayClusterMoviePipeline/Private/DisplayClusterMoviePipelineViewportPass.cpp b/Engine/Plugins/Runtime/nDisplay/Source/DisplayClusterMoviePipeline/Private/DisplayClusterMoviePipelineViewportPass.cpp
--- a/Engine/Plugins/Runtime/nDisplay/Source/DisplayClusterMoviePipeline/Private/DisplayClusterMoviePipelineViewportPass.cpp
+++ b/Engine/Plugins/Runtime/nDisplay/Source/DisplayClusterMoviePipeline/Private/DisplayClusterMoviePipelineViewportPass.cpp
@@ -598,3 +598,34 @@ bool UDisplayClusterMoviePipelineViewportPassBase::CalculateDisplayClusterViewpo
 
 	return false;
 }
+
+/////////////////////////////////////////////////////////////////////////////////////////////////////////////////////////////////
+// UDisplayClusterMoviePipelineViewportPass_PathTracer
+/////////////////////////////////////////////////////////////////////////////////////////////////////////////////////////////////
+#if WITH_EDITOR
+FText UDisplayClusterMoviePipelineViewportPass_PathTracer::GetFooterText(UMoviePipelineExecutorJob* InJob) const {
+	return NSLOCTEXT(
+		"MovieRenderPipeline",
+		"DCViewportBasePassSetting_FooterText_PathTracer",
+		"Sampling for the nDisplay Path Tracer is controlled by the Anti-Aliasing settings and the Reference Motion Blur setting.\n"
+		"All other Path Tracer settings are taken from the Post Process settings.");
+}
+#endif
+
+void UDisplayClusterMoviePipelineViewportPass_PathTracer::ValidateStateImpl()
+{
+	Super::ValidateStateImpl();
+	PathTracerValidationImpl();
+}
+
+void UDisplayClusterMoviePipelineViewportPass_PathTracer::SetupImpl(const MoviePipeline::FMoviePipelineRenderPassInitSettings& InPassInitSettings)
+{
+	if (!CheckIfPathTracerIsSupported())
+	{
+		UE_LOG(LogMovieRenderPipeline, Error, TEXT("Cannot render a nDisplay Path Tracer pass, Path Tracer is not enabled by this project."));
+		GetPipeline()->Shutdown(true);
+		return;
+	}
+
+	Super::SetupImpl(InPassInitSettings);
+}
diff --git a/Engine/Plugins/Runtime/nDisplay/Source/DisplayClusterMoviePipeline/Private/DisplayClusterMoviePipelineViewportPass.h b/Engine/Plugins/Runtime/nDisplay/Source/DisplayClusterMoviePipeline/Private/DisplayClusterMoviePipelineViewportPass.h
--- a/Engine/Plugins/Runtime/nDisplay/Source/DisplayClusterMoviePipeline/Private/DisplayClusterMoviePipelineViewportPass.h
+++ b/Engine/Plugins/Runtime/nDisplay/Source/DisplayClusterMoviePipeline/Private/DisplayClusterMoviePipelineViewportPass.h
@@ -190,3 +190,39 @@ public:
 	}
 	virtual int32 GetOutputFileSortingOrder() const override { return 2; }
 };
+
+/**
+ * nDisplay viewport render pass (PathTracer)
+ */
+UCLASS(BlueprintType)
+class UDisplayClusterMoviePipelineViewportPass_PathTracer : public UDisplayClusterMoviePipelineViewportPassBase
+{
+	GENERATED_BODY()
+
+public:
+	UDisplayClusterMoviePipelineViewportPass_PathTracer() : UDisplayClusterMoviePipelineViewportPassBase(TEXT("nDisplayPathTracer"))
+	{ }
+#if WITH_EDITOR
+	virtual FText GetDisplayText() const override { return NSLOCTEXT("MovieRenderPipeline", "DCViewportBasePassSetting_DisplayName_PathTracer", "nDisplay Path Tracer"); }
+	virtual FText GetFooterText(UMoviePipelineExecutorJob* InJob) const override;
+#endif
+	virtual void GetViewShowFlags(FEngineShowFlags& OutShowFlag, EViewModeIndex& OutViewModeIndex) const override
+	{
+		OutShowFlag = FEngineShowFlags(EShowFlagInitMode::ESFIM_Game);
+		OutShowFlag.SetPathTracing(true);
+		OutShowFlag.SetMotionBlur(!bReferenceMotionBlur);
+		OutViewModeIndex = EViewModeIndex::VMI_PathTracing;
+	}
+	virtual int32 GetOutputFileSortingOrder() const override { return 2; }
+
+	virtual bool IsAntiAliasingSupported() const { return false; }
+	virtual void ValidateStateImpl() override;
+	virtual void SetupImpl(const MoviePipeline::FMoviePipelineRenderPassInitSettings& InPassInitSettings) override;
+
+	/** When enabled, the path tracer will blend all spatial and temporal samples prior to the denoising and will disable post-processed motion blur.
+	 *  In this mode it is possible to use higher temporal sample counts to improve the motion blur quality.
+	 *  When this option is disabled, the path tracer will accumulate spatial samples, but denoise them prior to accumulation of temporal samples.
+	 */
+	UPROPERTY(EditAnywhere, BlueprintReadWrite, Category = "Reference Motion Blur")
+	bool bReferenceMotionBlur;
+};
