MoviePipeline: Added a Comment field to Movie Pipeline jobs that is shown on the Burn In (and exr metadata) to allow artists to communicate notes about what the job was for.

diff --git a/Engine/Plugins/MovieScene/MovieRenderPipeline/Source/MovieRenderPipelineCore/Private/MoviePipelineMasterConfig.cpp b/Engine/Plugins/MovieScene/MovieRenderPipeline/Source/MovieRenderPipelineCore/Private/MoviePipelineMasterConfig.cpp
--- a/Engine/Plugins/MovieScene/MovieRenderPipeline/Source/MovieRenderPipelineCore/Private/MoviePipelineMasterConfig.cpp
+++ b/Engine/Plugins/MovieScene/MovieRenderPipeline/Source/MovieRenderPipelineCore/Private/MoviePipelineMasterConfig.cpp
@@ -103,6 +103,7 @@ void UMoviePipelineMasterConfig::GetFormatArguments(FMoviePipelineFormatArgs& In
 		FString LevelName = TEXT("Level Name");
 		FString SequenceName = TEXT("Sequence Name");
 		FString JobName = TEXT("Job Name");
+		FString JobComment = TEXT("Job Comment");
 		double FrameRate = 0.0;
 
 		if (InOutFormatArgs.InJob)
@@ -110,7 +111,8 @@ void UMoviePipelineMasterConfig::GetFormatArguments(FMoviePipelineFormatArgs& In
 			LevelName = InOutFormatArgs.InJob->Map.GetAssetName();
 			SequenceName = InOutFormatArgs.InJob->Sequence.GetAssetName();
 			JobName = InOutFormatArgs.InJob->JobName;
-			
+			JobComment = InOutFormatArgs.InJob->Comment;
+
 			// FrameRate is a combination of Output Settings and Sequence so we do it here instead of in OutputSetting
 			FrameRate = GetEffectiveFrameRate(Cast<ULevelSequence>(InOutFormatArgs.InJob->Sequence.TryLoad())).AsDecimal();
 		}
@@ -124,6 +126,7 @@ void UMoviePipelineMasterConfig::GetFormatArguments(FMoviePipelineFormatArgs& In
 		InOutFormatArgs.FileMetadata.Add(TEXT("unreal/levelName"), LevelName);
 		InOutFormatArgs.FileMetadata.Add(TEXT("unreal/sequenceName"), SequenceName);
 		InOutFormatArgs.FileMetadata.Add(TEXT("unreal/jobName"), JobName);
+		InOutFormatArgs.FileMetadata.Add(TEXT("unreal/jobComment"), JobComment);
 		InOutFormatArgs.FileMetadata.Add(TEXT("unreal/frameRate"), FString::SanitizeFloat(FrameRate));
 
 		// Normally these are filled when resolving the file name by the job (so that the time is shared), but stub them in here so
diff --git a/Engine/Plugins/MovieScene/MovieRenderPipeline/Source/MovieRenderPipelineCore/Public/MoviePipelineQueue.h b/Engine/Plugins/MovieScene/MovieRenderPipeline/Source/MovieRenderPipelineCore/Public/MoviePipelineQueue.h
--- a/Engine/Plugins/MovieScene/MovieRenderPipeline/Source/MovieRenderPipelineCore/Public/MoviePipelineQueue.h
+++ b/Engine/Plugins/MovieScene/MovieRenderPipeline/Source/MovieRenderPipelineCore/Public/MoviePipelineQueue.h
@@ -392,6 +392,10 @@ public:
 	/** (Optional) If left blank, will default to system username. Can be shown in burn in as a first point of contact about the content. */
 	UPROPERTY(BlueprintReadWrite, EditAnywhere, Category = "Movie Render Pipeline")
 	FString Author;
+	
+	/** (Optional) If specified, will be shown in the burn in to allow users to keep track of notes about a render. */
+	UPROPERTY(BlueprintReadWrite, EditAnywhere, Category = "Movie Render Pipeline", meta = (MultiLine = true))
+	FString Comment;
 
 	/** (Optional) Shot specific information. If a shot is missing from this list it will assume to be enabled and will be rendered. */
 	UPROPERTY(BlueprintReadWrite, Instanced, Category = "Movie Render Pipeline")
