MoviePipeline: Changed CameraCut warmup range detection behavior to be the difference between the Camera Cut Section start, and the first evaluated frame of the Camera Cut section itself to better handle sub-sequences being clipped by parent ranges. Fixed a crash where Use Camera Cut for Warmup would evaluate the wrong Sequence time and then fail to get the correct camera locations due the object not being spawned yet.

diff --git a/Engine/Plugins/MovieScene/MovieRenderPipeline/Source/MovieRenderPipelineCore/Private/MoviePipeline.cpp b/Engine/Plugins/MovieScene/MovieRenderPipeline/Source/MovieRenderPipelineCore/Private/MoviePipeline.cpp
--- a/Engine/Plugins/MovieScene/MovieRenderPipeline/Source/MovieRenderPipelineCore/Private/MoviePipeline.cpp
+++ b/Engine/Plugins/MovieScene/MovieRenderPipeline/Source/MovieRenderPipelineCore/Private/MoviePipeline.cpp
@@ -1739,36 +1739,44 @@ void UMoviePipeline::GetSidecarCameraData(UMoviePipelineExecutorShot* InShot, in
 	}
 }
 
-void UMoviePipeline::GetSidecarCameraViewPoints(UMoviePipelineExecutorShot* InShot, TArray<FVector>& OutSidecarViewLocations, TArray<FRotator>& OutSidecarViewRotations) const
+bool UMoviePipeline::GetSidecarCameraViewPoints(UMoviePipelineExecutorShot* InShot, TArray<FVector>& OutSidecarViewLocations, TArray<FRotator>& OutSidecarViewRotations) const
 {
 	OutSidecarViewLocations.Reset();
 	OutSidecarViewRotations.Reset();
 
-	if (!InShot)
+	if (InShot)
 	{
-		return;
-	}
-
-	// We can't use GetSidecarCameraData because it relies on data generated by this function to ensure the previous/next locations are correct.
-	for (int32 Index = 0; Index < InShot->SidecarCameras.Num(); Index++)
-	{
-		const FGuid CameraBindingRef = InShot->SidecarCameras[Index].BindingId;
-		const FMovieSceneSequenceID CameraSequenceRef = InShot->SidecarCameras[Index].SequenceId;
-
-		TArrayView<TWeakObjectPtr<UObject>> BoundCameras = LevelSequenceActor->GetSequencePlayer()->FindBoundObjects(CameraBindingRef, CameraSequenceRef);
-		if (BoundCameras.Num() > 0)
+		// We can't use GetSidecarCameraData because it relies on data generated by this function to ensure the previous/next locations are correct.
+		for (int32 Index = 0; Index < InShot->SidecarCameras.Num(); Index++)
 		{
-			UCameraComponent* BoundCamera = MovieSceneHelpers::CameraComponentFromRuntimeObject(BoundCameras[0].Get());
-			if (BoundCamera)
+			const FGuid CameraBindingRef = InShot->SidecarCameras[Index].BindingId;
+			const FMovieSceneSequenceID CameraSequenceRef = InShot->SidecarCameras[Index].SequenceId;
+
+			TArrayView<TWeakObjectPtr<UObject>> BoundCameras = LevelSequenceActor->GetSequencePlayer()->FindBoundObjects(CameraBindingRef, CameraSequenceRef);
+			if (BoundCameras.Num() > 0)
 			{
-				FMinimalViewInfo ViewInfo;
-				BoundCamera->GetCameraView(GetWorld()->GetDeltaSeconds(), ViewInfo);
+				UCameraComponent* BoundCamera = MovieSceneHelpers::CameraComponentFromRuntimeObject(BoundCameras[0].Get());
+				if (BoundCamera)
+				{
+					FMinimalViewInfo ViewInfo;
+					BoundCamera->GetCameraView(GetWorld()->GetDeltaSeconds(), ViewInfo);
 
-				OutSidecarViewLocations.Add(ViewInfo.Location);
-				OutSidecarViewRotations.Add(ViewInfo.Rotation);
+					OutSidecarViewLocations.Add(ViewInfo.Location);
+					OutSidecarViewRotations.Add(ViewInfo.Rotation);
+				}
 			}
 		}
 	}
+
+	// Jam a default value in but return false so that the caller can know the data was no good.
+	if (OutSidecarViewLocations.Num() == 0)
+	{
+		OutSidecarViewLocations.AddDefaulted();
+		OutSidecarViewRotations.AddDefaulted();
+		return false;
+	}
+
+	return true;
 }
 
 #undef LOCTEXT_NAMESPACE // "MoviePipeline"
diff --git a/Engine/Plugins/MovieScene/MovieRenderPipeline/Source/MovieRenderPipelineCore/Private/MoviePipelineBlueprintLibrary.cpp b/Engine/Plugins/MovieScene/MovieRenderPipeline/Source/MovieRenderPipelineCore/Private/MoviePipelineBlueprintLibrary.cpp
--- a/Engine/Plugins/MovieScene/MovieRenderPipeline/Source/MovieRenderPipelineCore/Private/MoviePipelineBlueprintLibrary.cpp
+++ b/Engine/Plugins/MovieScene/MovieRenderPipeline/Source/MovieRenderPipelineCore/Private/MoviePipelineBlueprintLibrary.cpp
@@ -674,7 +674,27 @@ void UMoviePipelineBlueprintLibrary::UpdateJobShotListFromSequence(ULevelSequenc
 			{
 				InnerToOuterTransform = SubSequenceData->RootToSequenceTransform.InverseFromAllFirstWarps();
 			}
-			Entity.CameraCutWarmUpRange = MoviePipeline::GetCameraWarmUpRangeFromSequence(InSequence, LeafNode->MovieScene->GetPlaybackRange().GetLowerBoundValue(), InnerToOuterTransform);
+
+			// To make the camera cut range detection more consistent, we'll convert the start of this Entity into root sequence space, and then we'll convert the start of the
+			// camera cut into root sequence space, and the difference between them is the range to use. This is a more reliable way than looking at the actual Playback Range.
+			TRange<FFrameNumber> EntityRangeInMaster = Entity.Range;
+			TRange<FFrameNumber> CameraCutRangeInMaster = EntityRangeInMaster;
+			if (LeafNode->CameraCutSection.IsValid())
+			{
+				CameraCutRangeInMaster = LeafNode->CameraCutSection->GetRange() * InnerToOuterTransform;
+			}
+
+			TRange<FFrameNumber> CameraCutWarmUpRange = TRange<FFrameNumber>::Empty();
+			if (EntityRangeInMaster.GetLowerBound().IsClosed() && CameraCutRangeInMaster.GetLowerBound().IsClosed())
+			{
+				CameraCutWarmUpRange = TRange<FFrameNumber>(CameraCutRangeInMaster.GetLowerBound(), EntityRangeInMaster.GetLowerBound());
+			}
+			if (CameraCutRangeInMaster.IsDegenerate())
+			{
+				CameraCutRangeInMaster = TRange<FFrameNumber>::Empty();
+			}
+
+			Entity.CameraCutWarmUpRange = CameraCutWarmUpRange;
 			Entity.LeafNode = LeafNode;
 			Entity.Name = MoviePipeline::GetNameForShot(SequenceHierarchyCache, InSequence, LeafNode);
 			Entity.InnerToOuterTransform = InnerToOuterTransform;
diff --git a/Engine/Plugins/MovieScene/MovieRenderPipeline/Source/MovieRenderPipelineCore/Private/MoviePipelineTiming.cpp b/Engine/Plugins/MovieScene/MovieRenderPipeline/Source/MovieRenderPipelineCore/Private/MoviePipelineTiming.cpp
--- a/Engine/Plugins/MovieScene/MovieRenderPipeline/Source/MovieRenderPipelineCore/Private/MoviePipelineTiming.cpp
+++ b/Engine/Plugins/MovieScene/MovieRenderPipeline/Source/MovieRenderPipelineCore/Private/MoviePipelineTiming.cpp
@@ -147,6 +147,11 @@ void UMoviePipeline::TickProducingFrames()
 		// again after the warm-up frames.
 		UMoviePipelineAntiAliasingSetting* AntiAliasingSettings = FindOrAddSettingForShot<UMoviePipelineAntiAliasingSetting>(CurrentCameraCut);
 
+		// If we're not emulating the first frame motion blur, we back up by an extra frame so that we can fall through to regular frame-advance logic.
+		// However this causes the Sequencer to initially jump to/evaluate the wrong frame (where objects may not be spawned) and we need cameras, etc.
+		// spawned so we can cache current viewpoints. So we store how much we're going to go back so that we can temporarily add it to the sequence eval
+		// time, and then the actual eval will go forward as usual later.
+		FFrameNumber SequencerEvalOffset = FFrameNumber(0);
 		if (!CurrentCameraCut->ShotInfo.bEmulateFirstFrameMotionBlur)
 		{
 			FFrameTime TicksToEndOfPreviousFrame;
@@ -176,11 +181,12 @@ void UMoviePipeline::TickProducingFrames()
 			TicksToEndOfPreviousFrame += FFrameRate::TransformTime(FFrameTime(FFrameNumber(CurrentCameraCut->ShotInfo.NumEngineWarmUpFramesRemaining)), TargetSequence->GetMovieScene()->GetDisplayRate(), TargetSequence->GetMovieScene()->GetTickResolution());
 			AccumulatedTickSubFrameDeltas -= TicksToEndOfPreviousFrame.GetSubFrame();
 			CurrentCameraCut->ShotInfo.CurrentTickInMaster = CurrentCameraCut->ShotInfo.CurrentTickInMaster - TicksToEndOfPreviousFrame.FloorToFrame();
+			SequencerEvalOffset = TicksToEndOfPreviousFrame.FloorToFrame();
 		}
 
 		// Jump to the first frame of the sequence that we will be playing from. This doesn't take into account camera timing offset or temporal sampling, but that is
 		// proably okay as it means the first frame (frame 0) is exactly evaluated which is easier to preview effects in the editor instead of -.25 of a frame.
-		FFrameNumber CurrentMasterSeqTick = CurrentCameraCut->ShotInfo.CurrentTickInMaster;
+		FFrameNumber CurrentMasterSeqTick = CurrentCameraCut->ShotInfo.CurrentTickInMaster + SequencerEvalOffset;
 		FFrameTime TimeInPlayRate = FFrameRate::TransformTime(CurrentMasterSeqTick, TargetSequence->GetMovieScene()->GetTickResolution(), TargetSequence->GetMovieScene()->GetDisplayRate());
 
 		// We tell it to jump so that things responding to different scrub-types work correctly.
@@ -207,9 +213,11 @@ void UMoviePipeline::TickProducingFrames()
 			FrameInfo.PrevViewRotation = FrameInfo.CurrViewRotation;
 
 			// Cache data for all the sidecar cameras too.
-			GetSidecarCameraViewPoints(CurrentCameraCut, FrameInfo.CurrSidecarViewLocations, FrameInfo.CurrSidecarViewRotations);
+			bool bSuccess = GetSidecarCameraViewPoints(CurrentCameraCut, FrameInfo.CurrSidecarViewLocations, FrameInfo.CurrSidecarViewRotations);
 			FrameInfo.PrevSidecarViewLocations = FrameInfo.CurrSidecarViewLocations;
 			FrameInfo.PrevSidecarViewRotations = FrameInfo.CurrSidecarViewRotations;
+
+			ensureMsgf(bSuccess, TEXT("Failed to evaluate camera to create cur/prev camera locations. No camera actor found on Eval Tick: %d"), CurrentMasterSeqTick.Value);
 		}
 
 		// We can safely fall through to the below states as they're OK to process the same frame we set up.
diff --git a/Engine/Plugins/MovieScene/MovieRenderPipeline/Source/MovieRenderPipelineCore/Private/MoviePipelineUtils.cpp b/Engine/Plugins/MovieScene/MovieRenderPipeline/Source/MovieRenderPipelineCore/Private/MoviePipelineUtils.cpp
--- a/Engine/Plugins/MovieScene/MovieRenderPipeline/Source/MovieRenderPipelineCore/Private/MoviePipelineUtils.cpp
+++ b/Engine/Plugins/MovieScene/MovieRenderPipeline/Source/MovieRenderPipelineCore/Private/MoviePipelineUtils.cpp
@@ -572,28 +572,6 @@ namespace MoviePipeline
 		return TTuple<FString, FString>(InnerName, OuterName);
 	}
 
-	TRange<FFrameNumber> GetCameraWarmUpRangeFromSequence(UMovieSceneSequence* InSequence, FFrameNumber InSectionEnd, FMovieSceneTimeTransform InInnerToOuterTransform)
-	{
-		if (!ensureAlways(InSequence && InSequence->GetMovieScene()))
-		{
-			return TRange<FFrameNumber>::Empty();
-		}
-
-		UMovieSceneCameraCutTrack* CamTrack = Cast<UMovieSceneCameraCutTrack>(InSequence->GetMovieScene()->GetCameraCutTrack());
-		if (CamTrack && CamTrack->GetAllSections().Num() > 0)
-		{
-			UMovieSceneCameraCutSection* ValidCameraCutSection = CastChecked<UMovieSceneCameraCutSection>(CamTrack->GetAllSections()[0]);
-			if (ValidCameraCutSection->GetRange().HasLowerBound())
-			{
-				FFrameNumber CameraCutSectionStart = (ValidCameraCutSection->GetRange().GetLowerBoundValue() * InInnerToOuterTransform).FloorToFrame();
-
-				return TRange<FFrameNumber>(CameraCutSectionStart, InSectionEnd);
-			}
-		}
-
-		return TRange<FFrameNumber>::Empty();
-	}
-
 	void BuildSectionHierarchyRecursive(const FMovieSceneSequenceHierarchy& InHierarchy, UMovieSceneSequence* InRootSequence, const FMovieSceneSequenceID InSequenceId, const FMovieSceneSequenceID InChildId, TSharedPtr<FCameraCutSubSectionHierarchyNode> OutSubsectionHierarchy)
 	{
 		const FMovieSceneSequenceHierarchyNode* SequenceNode = InHierarchy.FindNode(InSequenceId);
diff --git a/Engine/Plugins/MovieScene/MovieRenderPipeline/Source/MovieRenderPipelineCore/Public/MoviePipeline.h b/Engine/Plugins/MovieScene/MovieRenderPipeline/Source/MovieRenderPipelineCore/Public/MoviePipeline.h
--- a/Engine/Plugins/MovieScene/MovieRenderPipeline/Source/MovieRenderPipelineCore/Public/MoviePipeline.h
+++ b/Engine/Plugins/MovieScene/MovieRenderPipeline/Source/MovieRenderPipelineCore/Public/MoviePipeline.h
@@ -186,7 +186,7 @@ public:
 	FMoviePipelineOutputData GetOutputDataParams();
 
 	void GetSidecarCameraData(UMoviePipelineExecutorShot* InShot, int32 InCameraIndex, FMinimalViewInfo& OutViewInfo, class UCameraComponent** OutCameraComponent) const;
-	void GetSidecarCameraViewPoints(UMoviePipelineExecutorShot* InShot, TArray<FVector>& OutSidecarViewLocations, TArray<FRotator>& OutSidecarViewRotations) const;
+	bool GetSidecarCameraViewPoints(UMoviePipelineExecutorShot* InShot, TArray<FVector>& OutSidecarViewLocations, TArray<FRotator>& OutSidecarViewRotations) const;
 
 #if WITH_EDITOR
 	const FMovieSceneExportMetadata& GetOutputMetadata() const { return OutputMetadata; }
diff --git a/Engine/Plugins/MovieScene/MovieRenderPipeline/Source/MovieRenderPipelineCore/Public/MoviePipelineUtils.h b/Engine/Plugins/MovieScene/MovieRenderPipeline/Source/MovieRenderPipelineCore/Public/MoviePipelineUtils.h
--- a/Engine/Plugins/MovieScene/MovieRenderPipeline/Source/MovieRenderPipelineCore/Public/MoviePipelineUtils.h
+++ b/Engine/Plugins/MovieScene/MovieRenderPipeline/Source/MovieRenderPipelineCore/Public/MoviePipelineUtils.h
@@ -151,8 +151,6 @@ namespace MoviePipeline
 	void BuildSectionHierarchyRecursive(const FMovieSceneSequenceHierarchy& InHierarchy, UMovieSceneSequence* InRootSequence, const FMovieSceneSequenceID InSequenceId, const FMovieSceneSequenceID InChildId, TSharedPtr<FCameraCutSubSectionHierarchyNode> OutSubsectionHierarchy);
 	/** Gets the inner and outer names for the shot by resolving camera bindings/shot names, etc. Not neccessairly the final names. */
 	TTuple<FString, FString> GetNameForShot(const FMovieSceneSequenceHierarchy& InHierarchy, UMovieSceneSequence* InRootSequence, TSharedPtr<FCameraCutSubSectionHierarchyNode> InSubSectionHierarch);
-	/** Examine the local sequence and see how much overlap the camera cut has with the start of the playback range. Converts to Root space. */
-	TRange<FFrameNumber> GetCameraWarmUpRangeFromSequence(UMovieSceneSequence* InSequence, FFrameNumber InSectionEnd, FMovieSceneTimeTransform InInnerToOuterTransform);
 	/** Given a leaf node, either caches the value of the hierarchy or restores it. Used to save the state before soloing a shot. */
 	void SaveOrRestoreSubSectionHierarchy(TSharedPtr<FCameraCutSubSectionHierarchyNode> InLeaf, const bool bInSave);
 	/** Given a leaf node, appropriately sets the IsActive flags for the whole hierarchy chain up to root for soloing a shot. */
