MoviePipeline: Minor UX improvements regarding default selection behavior. Importing a Queue asset will now automatically select the first job. When the UI is re-created (reopened in the same session) automatically select the first job. When opening the job configuration panel, automatically select the Output setting if available.

diff --git a/Engine/Plugins/MovieScene/MovieRenderPipeline/Source/MovieRenderPipelineEditor/Private/Widgets/SMoviePipelineConfigEditor.cpp b/Engine/Plugins/MovieScene/MovieRenderPipeline/Source/MovieRenderPipelineEditor/Private/Widgets/SMoviePipelineConfigEditor.cpp
--- a/Engine/Plugins/MovieScene/MovieRenderPipeline/Source/MovieRenderPipelineEditor/Private/Widgets/SMoviePipelineConfigEditor.cpp
+++ b/Engine/Plugins/MovieScene/MovieRenderPipeline/Source/MovieRenderPipelineEditor/Private/Widgets/SMoviePipelineConfigEditor.cpp
@@ -46,6 +46,8 @@
 #include "FrameNumberDetailsCustomization.h"
 #include "Editor.h"
 
+#include "MoviePipelineOutputSetting.h"
+
 #define LOCTEXT_NAMESPACE "SMoviePipelineEditor"
 
 
@@ -64,10 +66,24 @@ void SMoviePipelineConfigEditor::Construct(const FArguments& InArgs)
     
 	CheckForNewSettingsObject();
      
-	if (CachedPipelineConfig->GetUserSettings().Num() > 0)
+	// Automatically try to select the Output setting as it's the most commonly edited one.
+	// If that fails, we fall back to the first setting.
 	{
-		// Automatically try to select the first setting so there is something displayed.
-		SettingsWidget->SetSelectedSettings({ CachedPipelineConfig->GetUserSettings()[0] });
+		UMoviePipelineOutputSetting* OutputSetting = CachedPipelineConfig->FindSetting<UMoviePipelineOutputSetting>();
+		if (OutputSetting)
+		{
+			TArray<UMoviePipelineSetting*> SelectedSettings;
+			SelectedSettings.Add(OutputSetting);
+			SettingsWidget->SetSelectedSettings(SelectedSettings);
+		}
+		// Shot overrides may not have Output Settings.
+		else
+		{
+			if (CachedPipelineConfig->GetUserSettings().Num() > 0)
+			{
+				SettingsWidget->SetSelectedSettings({ CachedPipelineConfig->GetUserSettings()[0] });
+			}
+		}
 	}
 	ChildSlot
 	[
diff --git a/Engine/Plugins/MovieScene/MovieRenderPipeline/Source/MovieRenderPipelineEditor/Private/Widgets/SMoviePipelineQueueEditor.h b/Engine/Plugins/MovieScene/MovieRenderPipeline/Source/MovieRenderPipelineEditor/Private/Widgets/SMoviePipelineQueueEditor.h
--- a/Engine/Plugins/MovieScene/MovieRenderPipeline/Source/MovieRenderPipelineEditor/Private/Widgets/SMoviePipelineQueueEditor.h
+++ b/Engine/Plugins/MovieScene/MovieRenderPipeline/Source/MovieRenderPipelineEditor/Private/Widgets/SMoviePipelineQueueEditor.h
@@ -49,6 +49,7 @@ public:
 	TSharedRef<SWidget> OnGenerateNewJobFromAssetMenu();
 
 	TArray<TSharedPtr<IMoviePipelineQueueTreeItem>> GetSelectedItems() const { return TreeView->GetSelectedItems(); }
+	void SetSelectedJobs(const TArray<UMoviePipelineExecutorJob*>& InJobs) { PendingJobsToSelect = InJobs; }
 
 private:
 	// SWidget Interface
diff --git a/Engine/Plugins/MovieScene/MovieRenderPipeline/Source/MovieRenderPipelineEditor/Private/Widgets/SMoviePipelineQueuePanel.cpp b/Engine/Plugins/MovieScene/MovieRenderPipeline/Source/MovieRenderPipelineEditor/Private/Widgets/SMoviePipelineQueuePanel.cpp
--- a/Engine/Plugins/MovieScene/MovieRenderPipeline/Source/MovieRenderPipelineEditor/Private/Widgets/SMoviePipelineQueuePanel.cpp
+++ b/Engine/Plugins/MovieScene/MovieRenderPipeline/Source/MovieRenderPipelineEditor/Private/Widgets/SMoviePipelineQueuePanel.cpp
@@ -69,8 +69,22 @@ void SMoviePipelineQueuePanel::Construct(const FArguments& InArgs)
 		.OnPresetChosen(this, &SMoviePipelineQueuePanel::OnJobPresetChosen)
 		.OnJobSelectionChanged(this, &SMoviePipelineQueuePanel::OnSelectionChanged);
 
-	// Reset us to no selection.
-	OnSelectionChanged(TArray<UMoviePipelineExecutorJob*>());
+
+	{
+		// Automatically select the first job in the queue
+		UMoviePipelineQueueSubsystem* Subsystem = GEditor->GetEditorSubsystem<UMoviePipelineQueueSubsystem>();
+		check(Subsystem);
+
+		TArray<UMoviePipelineExecutorJob*> Jobs;
+		if (Subsystem->GetQueue()->GetJobs().Num() > 0)
+		{
+			Jobs.Add(Subsystem->GetQueue()->GetJobs()[0]);
+		}
+
+		// Go through the UI so it updates the UI selection too and then this will loop back
+		// around to OnSelectionChanged to update ourself.
+		PipelineQueueEditorWidget->SetSelectedJobs(Jobs);
+	}
 
 	ChildSlot
 	[
@@ -659,6 +673,17 @@ void SMoviePipelineQueuePanel::OnImportSavedQueueAssest(const FAssetData& InPres
 				}
 			}
 		}
+
+		// Automatically select the first job in the queue
+		TArray<UMoviePipelineExecutorJob*> Jobs;
+		if (Subsystem->GetQueue()->GetJobs().Num() > 0)
+		{
+			Jobs.Add(Subsystem->GetQueue()->GetJobs()[0]);
+		}
+
+		// Go through the UI so it updates the UI selection too and then this will loop back
+		// around to OnSelectionChanged to update ourself.
+		PipelineQueueEditorWidget->SetSelectedJobs(Jobs);
 	}
 }
 
