MoviePipeline: Create output directories when using the command line encoder, not all CLI tools work without the directory tree already existing.

diff --git a/Engine/Plugins/MovieScene/MovieRenderPipeline/Source/MovieRenderPipelineCore/Private/MoviePipelineCommandLineEncoder.cpp b/Engine/Plugins/MovieScene/MovieRenderPipeline/Source/MovieRenderPipelineCore/Private/MoviePipelineCommandLineEncoder.cpp
--- a/Engine/Plugins/MovieScene/MovieRenderPipeline/Source/MovieRenderPipelineCore/Private/MoviePipelineCommandLineEncoder.cpp
+++ b/Engine/Plugins/MovieScene/MovieRenderPipeline/Source/MovieRenderPipelineCore/Private/MoviePipelineCommandLineEncoder.cpp
@@ -14,6 +14,11 @@
 #include "MoviePipelineQueue.h"
 #include "MoviePipelineDebugSettings.h"
 #include "MoviePipelineBlueprintLibrary.h"
+#include "Misc/FileHelper.h"
+#include "HAL/FileManager.h"
+#include "HAL/PlatformFileManager.h"
+#include "GenericPlatform/GenericPlatformFile.h"
+
 
 // Forward Declare
 namespace UE
@@ -195,6 +200,18 @@ void UMoviePipelineCommandLineEncoder::StartEncodingProcess(TArray<FMoviePipelin
 			FinalFilePath = FPaths::ConvertRelativePathToFull(FinalFilePath);
 		}
 
+		FPaths::NormalizeFilename(FinalFilePath);
+		FPaths::CollapseRelativeDirectories(FinalFilePath);
+
+		FString FinalFileDirectory = FPaths::GetPath(FinalFilePath);
+
+		// Ensure the output directory is created
+		IPlatformFile& FileManager = FPlatformFileManager::Get().GetPlatformFile();
+		if (!FileManager.CreateDirectoryTree(*FinalFileDirectory))
+		{
+			UE_LOG(LogMovieRenderPipelineIO, Error, TEXT("Failed to create directory for output path '%s'"), *FinalFileDirectory);
+		}
+
 		// Manipulate the in/out data in case scripting tries to get access to the files. It's not a perfect solution
 		// because the encoding won't be finished by the time the scripting layer is called. Need to manipulate the original
 		// and not the copy that we're currently iterating through.
