MoviePipeline: Fixed a bug where default editor properties would show up in the editor config even though they aren't intended to be edited at editor time.

#UE-130951

diff --git a/Engine/Plugins/MovieScene/MovieRenderPipeline/Source/MovieRenderPipelineEditor/Private/Widgets/SMoviePipelineConfigEditor.cpp b/Engine/Plugins/MovieScene/MovieRenderPipeline/Source/MovieRenderPipelineEditor/Private/Widgets/SMoviePipelineConfigEditor.cpp
--- a/Engine/Plugins/MovieScene/MovieRenderPipeline/Source/MovieRenderPipelineEditor/Private/Widgets/SMoviePipelineConfigEditor.cpp
+++ b/Engine/Plugins/MovieScene/MovieRenderPipeline/Source/MovieRenderPipelineEditor/Private/Widgets/SMoviePipelineConfigEditor.cpp
@@ -302,6 +302,7 @@ void SMoviePipelineConfigEditor::UpdateDetails()
 	DetailsViewArgs.bHideSelectionTip = true;
 	DetailsViewArgs.bShowScrollBar = false;
 	DetailsViewArgs.ColumnWidth = 0.5f;
+	DetailsViewArgs.DefaultsOnlyVisibility = EEditDefaultsOnlyNodeVisibility::Hide;
 
 	TArray<UMoviePipelineSetting*> SelectedSources;
 	SettingsWidget->GetSelectedSettings(SelectedSources);
