nDisplay.MRQ:
 - added support for nDisplay viewport sizes for output images
 - fixed crash: replaced 'this' argument in ENQUEUE_RENDER_COMMAND
 - Added editor properrty customization

diff --git a/Engine/Plugins/MovieScene/MoviePipelineMaskRenderPass/Source/MoviePipelineMaskRenderPass/Private/MoviePipelineObjectIdPass.cpp b/Engine/Plugins/MovieScene/MoviePipelineMaskRenderPass/Source/MoviePipelineMaskRenderPass/Private/MoviePipelineObjectIdPass.cpp
--- a/Engine/Plugins/MovieScene/MoviePipelineMaskRenderPass/Source/MoviePipelineMaskRenderPass/Private/MoviePipelineObjectIdPass.cpp
+++ b/Engine/Plugins/MovieScene/MoviePipelineMaskRenderPass/Source/MoviePipelineMaskRenderPass/Private/MoviePipelineObjectIdPass.cpp
@@ -119,6 +119,27 @@ UMoviePipelineObjectIdRenderPass::UMoviePipelineObjectIdRenderPass()
 	}
 }
 
+TWeakObjectPtr<UTextureRenderTarget2D> UMoviePipelineObjectIdRenderPass::CreateViewRenderTargetImpl(const FIntPoint& InSize, IViewCalcPayload* OptPayload) const
+{
+	// Crete the render target with the correct bit depth.
+	TWeakObjectPtr<UTextureRenderTarget2D> TileRenderTarget = NewObject<UTextureRenderTarget2D>(GetTransientPackage());
+	TileRenderTarget->ClearColor = FLinearColor(0.0f, 0.0f, 0.0f, 0.0f);
+
+	// Ensure there's no gamma in the RT otherwise the HitProxy color ids don't round trip properly.
+	TileRenderTarget->TargetGamma = 0.f;
+	TileRenderTarget->InitCustomFormat(InSize.X, InSize.Y, EPixelFormat::PF_B8G8R8A8, true);
+	TileRenderTarget->AddToRoot();
+
+	return TileRenderTarget;
+}
+
+TSharedPtr<FMoviePipelineSurfaceQueue, ESPMode::ThreadSafe> UMoviePipelineObjectIdRenderPass::CreateSurfaceQueueImpl(const FIntPoint& InSize, IViewCalcPayload* OptPayload) const
+{
+	TSharedPtr<FMoviePipelineSurfaceQueue, ESPMode::ThreadSafe> SurfaceQueue = MakeShared<FMoviePipelineSurfaceQueue, ESPMode::ThreadSafe>(InSize, EPixelFormat::PF_B8G8R8A8, 3, false);
+
+	return SurfaceQueue;
+}
+
 void UMoviePipelineObjectIdRenderPass::GatherOutputPassesImpl(TArray<FMoviePipelinePassIdentifier>& ExpectedRenderPasses)
 {
 	// Don't call the super which adds the generic PassIdentifier, which in this case is numberless and incorrect for the final output spec.
@@ -130,17 +151,11 @@ void UMoviePipelineObjectIdRenderPass::SetupImpl(const MoviePipeline::FMoviePipe
 {
 	Super::SetupImpl(InPassInitSettings);
 
-	// Re-initialize the render target with the correct bit depth.
-	TileRenderTarget = NewObject<UTextureRenderTarget2D>(GetTransientPackage());
-	TileRenderTarget->ClearColor = FLinearColor(0.0f, 0.0f, 0.0f, 0.0f);
-
-	// Ensure there's no gamma in the RT otherwise the HitProxy color ids don't round trip properly.
-	TileRenderTarget->TargetGamma = 0.f;
-	TileRenderTarget->InitCustomFormat(InPassInitSettings.BackbufferResolution.X, InPassInitSettings.BackbufferResolution.Y, EPixelFormat::PF_B8G8R8A8, true);
-	TileRenderTarget->AddToRoot();
+	// Re-initialize the render target and surface queue
+	GetOrCreateViewRenderTarget(InPassInitSettings.BackbufferResolution);
+	GetOrCreateSurfaceQueue(InPassInitSettings.BackbufferResolution);
 
 	AccumulatorPool = MakeShared<TAccumulatorPool<FMaskOverlappedAccumulator>, ESPMode::ThreadSafe>(6);
-	SurfaceQueue = MakeShared<FMoviePipelineSurfaceQueue>(InPassInitSettings.BackbufferResolution, EPixelFormat::PF_B8G8R8A8, 3, false);
 
 	UE::MoviePipeline::FObjectIdAccelerationData AccelData = UE::MoviePipeline::FObjectIdAccelerationData();
 
@@ -167,22 +182,8 @@ void UMoviePipelineObjectIdRenderPass::SetupImpl(const MoviePipeline::FMoviePipe
 
 void UMoviePipelineObjectIdRenderPass::TeardownImpl()
 {
-	// This may call FlushRenderingCommands if there are outstanding readbacks that need to happen.
-	if (SurfaceQueue.IsValid())
-	{
-		SurfaceQueue->Shutdown();
-	}
-
-	// Stall until the task graph has completed any pending accumulations.
-	FTaskGraphInterface::Get().WaitUntilTasksComplete(OutstandingTasks, ENamedThreads::GameThread);
-	OutstandingTasks.Reset();
 	ManifestAnnotation.RemoveAnnotation(this);
 
-	if (TileRenderTarget.IsValid())
-	{
-		TileRenderTarget->RemoveFromRoot();
-	}
-
 	// Preserve our view state until the rendering thread has been flushed.
 	Super::TeardownImpl();
 }
@@ -303,21 +304,21 @@ void UMoviePipelineObjectIdRenderPass::RenderSample_GameThreadImpl(const FMovieP
 		return;
 	}
 
+	// Wait for a surface to be available to write to. This will stall the game thread while the RHI/Render Thread catch up.
 	Super::RenderSample_GameThreadImpl(InSampleState);
 
 	FMoviePipelineRenderPassMetrics InOutSampleState = InSampleState;
 
-	// Wait for a surface to be available to write to. This will stall the game thread while the RHI/Render Thread catch up.
-	{
-		SCOPE_CYCLE_COUNTER(STAT_MoviePipeline_WaitForAvailableSurface);
-		SurfaceQueue->BlockUntilAnyAvailable();
-	}
-	
 	TSharedPtr<FSceneViewFamilyContext> ViewFamily = CalculateViewFamily(InOutSampleState);
 	
 	// Submit to be rendered. Main render pass always uses target 0. We do this before making the Hitproxy cache because
 	// BeginRenderingViewFamily ensures render state for things are created.
-	FRenderTarget* RenderTarget = GetViewRenderTarget()->GameThread_GetRenderTargetResource();
+	TWeakObjectPtr<UTextureRenderTarget2D> ViewRenderTarget = GetOrCreateViewRenderTarget(InSampleState.BackbufferSize);
+	check(ViewRenderTarget.IsValid());
+
+	FRenderTarget* RenderTarget = ViewRenderTarget->GameThread_GetRenderTargetResource();
+	check(RenderTarget);
+
 	FCanvas Canvas = FCanvas(RenderTarget, nullptr, GetPipeline()->GetWorld(), ViewFamily->GetFeatureLevel(), FCanvas::CDM_DeferDrawing, 1.0f);
 	GetRendererModule().BeginRenderingViewFamily(&Canvas, ViewFamily.Get());
 
@@ -490,7 +491,7 @@ void UMoviePipelineObjectIdRenderPass::RenderSample_GameThreadImpl(const FMovieP
 			this->OutstandingTasks.Add(Event);
 		};
 
-		TSharedPtr<FMoviePipelineSurfaceQueue> LocalSurfaceQueue = SurfaceQueue;
+		TSharedPtr<FMoviePipelineSurfaceQueue, ESPMode::ThreadSafe> LocalSurfaceQueue = GetOrCreateSurfaceQueue(InSampleState.BackbufferSize, (IViewCalcPayload*)(&FramePayload.Get()));
 
 		ENQUEUE_RENDER_COMMAND(CanvasRenderTargetResolveCommand)(
 			[LocalSurfaceQueue, FramePayload, Callback, RenderTarget](FRHICommandListImmediate& RHICmdList) mutable
diff --git a/Engine/Plugins/MovieScene/MoviePipelineMaskRenderPass/Source/MoviePipelineMaskRenderPass/Private/MoviePipelineObjectIdPass.h b/Engine/Plugins/MovieScene/MoviePipelineMaskRenderPass/Source/MoviePipelineMaskRenderPass/Private/MoviePipelineObjectIdPass.h
--- a/Engine/Plugins/MovieScene/MoviePipelineMaskRenderPass/Source/MoviePipelineMaskRenderPass/Private/MoviePipelineObjectIdPass.h
+++ b/Engine/Plugins/MovieScene/MoviePipelineMaskRenderPass/Source/MoviePipelineMaskRenderPass/Private/MoviePipelineObjectIdPass.h
@@ -39,6 +39,9 @@ public:
 	virtual bool IsScreenPercentageSupported() const override { return false; }
 	virtual int32 GetOutputFileSortingOrder() const override { return 10; }
 	
+	virtual TWeakObjectPtr<UTextureRenderTarget2D> CreateViewRenderTargetImpl(const FIntPoint& InSize, IViewCalcPayload* OptPayload = nullptr) const override;
+	virtual TSharedPtr<FMoviePipelineSurfaceQueue, ESPMode::ThreadSafe> CreateSurfaceQueueImpl(const FIntPoint& InSize, IViewCalcPayload* OptPayload = nullptr) const override;
+
 protected:
 	void PostRendererSubmission(const FMoviePipelineRenderPassMetrics& InSampleState);
 	FString ResolveProxyIdGroup(const AActor* InActor, const UPrimitiveComponent* InPrimComponent, const int32 InMaterialIndex, const int32 InSectionIndex) const;
diff --git a/Engine/Plugins/MovieScene/MoviePipelineMaskRenderPass/Source/MoviePipelineMaskRenderPass/Private/MoviePipelinePanoramicPass.cpp b/Engine/Plugins/MovieScene/MoviePipelineMaskRenderPass/Source/MoviePipelineMaskRenderPass/Private/MoviePipelinePanoramicPass.cpp
--- a/Engine/Plugins/MovieScene/MoviePipelineMaskRenderPass/Source/MoviePipelineMaskRenderPass/Private/MoviePipelinePanoramicPass.cpp
+++ b/Engine/Plugins/MovieScene/MoviePipelineMaskRenderPass/Source/MoviePipelineMaskRenderPass/Private/MoviePipelinePanoramicPass.cpp
@@ -124,34 +124,11 @@ void UMoviePipelinePanoramicPass::SetupImpl(const MoviePipeline::FMoviePipelineR
 	Super::SetupImpl(InPassInitSettings);
 	// LLM_SCOPE_BYNAME(TEXT("MoviePipeline/PanoBlendSetup"));
 
-	// Allocate a shared render target for readback:
-	{
-		CanvasReadbackTexture = NewObject<UTextureRenderTarget2D>(GetTransientPackage());
-		CanvasReadbackTexture->ClearColor = FLinearColor(0.0f, 0.0f, 0.0f, 0.0f);
-
-		// We calculate a different resolution than the final output resolution.
-		float HorizontalFoV;
-		float VerticalFoV;
-		GetFieldOfView(HorizontalFoV, VerticalFoV, bStereo);
-
-		// Horizontal FoV is a proportion of the global horizontal resolution
-		// ToDo: We might have to check which is higher, if numVerticalPanes > numHorizontalPanes this math might be backwards.
-		float HorizontalRes = (HorizontalFoV / 360.0f) * InPassInitSettings.BackbufferResolution.X;
-		float Intermediate = FMath::Tan(FMath::DegreesToRadians(VerticalFoV) * 0.5f) / FMath::Tan(FMath::DegreesToRadians(HorizontalFoV) * 0.5f);
-		float VerticalRes = HorizontalRes * Intermediate;
-
-		//float VerticalRes = HorizontalRes * FMath::Tan(FMath::DegreesToRadians(VerticalFoV) * 0.5f) / FMath::Tan(FMath::DegreesToRadians(HorizontalFoV) * 0.5f);
-		// 
-		// ToDo: I think this is just aspect ratio based on FoVs?
-
-		PaneResolution = FIntPoint(FMath::CeilToInt(HorizontalRes), FMath::CeilToInt(VerticalRes));
-		CanvasReadbackTexture->InitCustomFormat(PaneResolution.X, PaneResolution.Y, EPixelFormat::PF_FloatRGBA, false);
-
-		// OCIO: Since this is a manually created Render target we don't need Gamma to be applied.
-		// We use this render target to render to via a display extension that utilizes Display Gamma
-		// which has a default value of 2.2 (DefaultDisplayGamma), therefore we need to set Gamma on this render target to 2.2 to cancel out any unwanted effects.
-		CanvasReadbackTexture->TargetGamma = FOpenColorIODisplayExtension::DefaultDisplayGamma;
-	}
+	const FIntPoint PaneResolution = GetPaneResolution(InPassInitSettings.BackbufferResolution);
+
+	// Re-initialize the render target and surface queue
+	GetOrCreateViewRenderTarget(PaneResolution);
+	GetOrCreateSurfaceQueue(PaneResolution);
 
 	int32 StereoMultiplier = bStereo ? 2 : 1;
 	int32 NumPanes = NumHorizontalSteps * NumVerticalSteps;
@@ -166,13 +143,6 @@ void UMoviePipelinePanoramicPass::SetupImpl(const MoviePipeline::FMoviePipelineR
 		}
 	}
 
-	if (GetPipeline()->GetPreviewTexture() == nullptr)
-	{
-		GetPipeline()->SetPreviewTexture(CanvasReadbackTexture);
-	}
-
-	SurfaceQueue = MakeShared<FMoviePipelineSurfaceQueue>(PaneResolution, EPixelFormat::PF_FloatRGBA, 3, true);
-
 	// We need one accumulator per pano tile if using accumulation.
 	AccumulatorPool = MakeShared<TAccumulatorPool<FImageOverlappedAccumulator>, ESPMode::ThreadSafe>(NumPanoramicPanes);
 
@@ -188,23 +158,8 @@ void UMoviePipelinePanoramicPass::SetupImpl(const MoviePipeline::FMoviePipelineR
 
 void UMoviePipelinePanoramicPass::TeardownImpl()
 {
-	GetPipeline()->SetPreviewTexture(nullptr);
-
-	// This may call FlushRenderingCommands if there are outstanding readbacks that need to happen.
-	if (SurfaceQueue.IsValid())
-	{
-		SurfaceQueue->Shutdown();
-	}
-
-	// Stall until the task graph has completed any pending accumulations. Because the actual
-	// Panoramic blending is also done on the same task as the accumulation that produced it,
-	// this will stall until all blends are finished too and pass them to the output merger.
-	FTaskGraphInterface::Get().WaitUntilTasksComplete(OutstandingTasks, ENamedThreads::GameThread);
-	OutstandingTasks.Reset();
 	PanoramicOutputBlender.Reset();
-
 	AccumulatorPool.Reset();
-	SurfaceQueue.Reset();
 
 	for (int32 Index = 0; Index < OptionalPaneViewStates.Num(); Index++)
 	{
@@ -216,7 +171,6 @@ void UMoviePipelinePanoramicPass::TeardownImpl()
 		OptionalPaneViewStates[Index].Destroy();
 	}
 	OptionalPaneViewStates.Reset();
-	CanvasReadbackTexture = nullptr;
 	
 	OCIOSceneViewExtension.Reset();
 	OCIOSceneViewExtension = nullptr;
@@ -254,11 +208,6 @@ FSceneViewStateInterface* UMoviePipelinePanoramicPass::GetSceneViewStateInterfac
 	}
 }
 
-UTextureRenderTarget2D* UMoviePipelinePanoramicPass::GetViewRenderTarget(IViewCalcPayload* OptPayload) const
-{
-	return CanvasReadbackTexture;
-}
-
 void UMoviePipelinePanoramicPass::GatherOutputPassesImpl(TArray<FMoviePipelinePassIdentifier>& ExpectedRenderPasses)
 {
 	// Despite rendering many views, we only output one image total, which is covered by the base.
@@ -283,8 +232,27 @@ void UMoviePipelinePanoramicPass::AddViewExtensions(FSceneViewFamilyContext& InC
 	}
 }
 
+FIntPoint UMoviePipelinePanoramicPass::GetPaneResolution(const FIntPoint& InSize) const
+{
+	// We calculate a different resolution than the final output resolution.
+	float HorizontalFoV;
+	float VerticalFoV;
+	GetFieldOfView(HorizontalFoV, VerticalFoV, bStereo);
+
+	// Horizontal FoV is a proportion of the global horizontal resolution
+	// ToDo: We might have to check which is higher, if numVerticalPanes > numHorizontalPanes this math might be backwards.
+	float HorizontalRes = (HorizontalFoV / 360.0f) * InSize.X;
+	float Intermediate = FMath::Tan(FMath::DegreesToRadians(VerticalFoV) * 0.5f) / FMath::Tan(FMath::DegreesToRadians(HorizontalFoV) * 0.5f);
+	float VerticalRes = HorizontalRes * Intermediate;
+
+	//float VerticalRes = HorizontalRes * FMath::Tan(FMath::DegreesToRadians(VerticalFoV) * 0.5f) / FMath::Tan(FMath::DegreesToRadians(HorizontalFoV) * 0.5f);
+	// 
+	// ToDo: I think this is just aspect ratio based on FoVs?
+
+	return FIntPoint(FMath::CeilToInt(HorizontalRes), FMath::CeilToInt(VerticalRes));
+}
 
-void UMoviePipelinePanoramicPass::GetFieldOfView(float& OutHorizontal, float& OutVertical, const bool bInStereo)
+void UMoviePipelinePanoramicPass::GetFieldOfView(float& OutHorizontal, float& OutVertical, const bool bInStereo) const
 {
 	// ToDo: These should probably be mathematically derived based on numSteps
 	OutHorizontal = bInStereo ? 30.f : HorzFieldOfView > 0 ? HorzFieldOfView : 90.f;
@@ -411,13 +379,10 @@ FSceneView* UMoviePipelinePanoramicPass::GetSceneViewForSampleState(FSceneViewFa
 
 void UMoviePipelinePanoramicPass::RenderSample_GameThreadImpl(const FMoviePipelineRenderPassMetrics& InSampleState)
 {
+	// Wait for a surface to be available to write to. This will stall the game thread while the RHI/Render Thread catch up.
 	Super::RenderSample_GameThreadImpl(InSampleState);
 
-	// Wait for a surface to be available to write to. This will stall the game thread while the RHI/Render Thread catch up.
-	{
-		SCOPE_CYCLE_COUNTER(STAT_MoviePipeline_WaitForAvailableSurface);
-		SurfaceQueue->BlockUntilAnyAvailable();
-	}
+	const FIntPoint PaneResolution = GetPaneResolution(InSampleState.BackbufferSize);
 	
 	// This function will be called for each sample of every tile. We'll then submit renders for each
 	// Panoramic Pane. This effectively means that all panes will fill in their top left corners first.
@@ -483,7 +448,12 @@ void UMoviePipelinePanoramicPass::RenderSample_GameThreadImpl(const FMoviePipeli
 				}
 
 				// Submit the view to be rendered.
-				FRenderTarget* RenderTarget = CanvasReadbackTexture->GameThread_GetRenderTargetResource();
+				TWeakObjectPtr<UTextureRenderTarget2D> ViewRenderTarget = GetOrCreateViewRenderTarget(PaneResolution);
+				check(ViewRenderTarget.IsValid());
+
+				FRenderTarget* RenderTarget = ViewRenderTarget->GameThread_GetRenderTargetResource();
+				check(RenderTarget);
+
 				FCanvas Canvas = FCanvas(RenderTarget, nullptr, GetPipeline()->GetWorld(), ViewFamily->GetFeatureLevel(), FCanvas::CDM_DeferDrawing, 1.0f);
 				GetRendererModule().BeginRenderingViewFamily(&Canvas, ViewFamily.Get());
 
@@ -516,7 +486,6 @@ void UMoviePipelinePanoramicPass::ScheduleReadbackAndAccumulation(const FMoviePi
 		FMoviePipelinePassIdentifier PanePassIdentifier = FMoviePipelinePassIdentifier(FString::Printf(TEXT("%s_x%d_y%d"), *PassIdentifier.Name, InPane.VerticalStepIndex, InPane.HorizontalStepIndex));
 		SampleAccumulator = AccumulatorPool->BlockAndGetAccumulator_GameThread(InSampleState.OutputState.OutputFrameNumber, PanePassIdentifier);
 	}
-	TSharedPtr<FMoviePipelineSurfaceQueue> LocalSurfaceQueue = SurfaceQueue;
 
 	TSharedRef<FPanoramicImagePixelDataPayload, ESPMode::ThreadSafe> FramePayload = MakeShared<FPanoramicImagePixelDataPayload, ESPMode::ThreadSafe>();
 	FramePayload->PassIdentifier = PassIdentifier;
@@ -538,6 +507,9 @@ void UMoviePipelinePanoramicPass::ScheduleReadbackAndAccumulation(const FMoviePi
 			FramePayload->Pane.VerticalStepIndex, FramePayload->SampleState.OutputState.OutputFrameNumber);
 	}
 	
+	const FIntPoint PaneResolution = GetPaneResolution(InSampleState.BackbufferSize);
+	TSharedPtr<FMoviePipelineSurfaceQueue, ESPMode::ThreadSafe> LocalSurfaceQueue = GetOrCreateSurfaceQueue(PaneResolution, (IViewCalcPayload*)(&FramePayload.Get()));
+
 	MoviePipeline::FImageSampleAccumulationArgs AccumulationArgs;
 	{
 		AccumulationArgs.OutputMerger = PanoramicOutputBlender;
diff --git a/Engine/Plugins/MovieScene/MoviePipelineMaskRenderPass/Source/MoviePipelineMaskRenderPass/Private/MoviePipelinePanoramicPass.h b/Engine/Plugins/MovieScene/MoviePipelineMaskRenderPass/Source/MoviePipelineMaskRenderPass/Private/MoviePipelinePanoramicPass.h
--- a/Engine/Plugins/MovieScene/MoviePipelineMaskRenderPass/Source/MoviePipelineMaskRenderPass/Private/MoviePipelinePanoramicPass.h
+++ b/Engine/Plugins/MovieScene/MoviePipelineMaskRenderPass/Source/MoviePipelineMaskRenderPass/Private/MoviePipelinePanoramicPass.h
@@ -121,7 +121,6 @@ protected:
 	virtual int32 GetOutputFileSortingOrder() const override { return 1; }
 	virtual bool IsAlphaInTonemapperRequiredImpl() const override { return false; }
 	virtual FSceneViewStateInterface* GetSceneViewStateInterface(IViewCalcPayload* OptPayload) override;
-	virtual UTextureRenderTarget2D* GetViewRenderTarget(IViewCalcPayload* OptPayload) const override;
 	virtual void AddViewExtensions(FSceneViewFamilyContext& InContext, FMoviePipelineRenderPassMetrics& InOutSampleState) override;
 	virtual bool IsAutoExposureAllowed(const FMoviePipelineRenderPassMetrics& InSampleState) const override { return false; }
 	virtual FSceneView* GetSceneViewForSampleState(FSceneViewFamily* ViewFamily, FMoviePipelineRenderPassMetrics& InOutSampleState, IViewCalcPayload* OptPayload = nullptr) override;
@@ -132,7 +131,9 @@ protected:
 	// ~FGCObject Interface
 
 	void ScheduleReadbackAndAccumulation(const FMoviePipelineRenderPassMetrics& InSampleState, const FPanoPane& InPane, FCanvas& InCanvas);
-	void GetFieldOfView(float& OutHorizontal, float& OutVertical, const bool bInStereo);
+	void GetFieldOfView(float& OutHorizontal, float& OutVertical, const bool bInStereo) const;
+	FIntPoint GetPaneResolution(const FIntPoint& InSize) const;
+
 public:
 
 	/**
@@ -181,7 +182,6 @@ public:
 
 protected:
 	TSharedPtr<FAccumulatorPool, ESPMode::ThreadSafe> AccumulatorPool;
-	FIntPoint PaneResolution;
 
 	// ToDo: One per high-res tile per pano-pane?
 	TArray<FSceneViewStateReference> OptionalPaneViewStates;
@@ -189,9 +189,6 @@ protected:
 	/** The lifetime of this SceneViewExtension is only during the rendering process. It is destroyed as part of TearDown. */
 	TSharedPtr<FOpenColorIODisplayExtension, ESPMode::ThreadSafe> OCIOSceneViewExtension;
 
-	UPROPERTY(Transient, DuplicateTransient)
-	TObjectPtr<UTextureRenderTarget2D> CanvasReadbackTexture;
-
 	TSharedPtr<MoviePipeline::IMoviePipelineOutputMerger> PanoramicOutputBlender;
 	bool bHasWarnedSettings;
 };
diff --git a/Engine/Plugins/MovieScene/MovieRenderPipeline/Source/MovieRenderPipelineCore/Private/MoviePipelineRendering.cpp b/Engine/Plugins/MovieScene/MovieRenderPipeline/Source/MovieRenderPipelineCore/Private/MoviePipelineRendering.cpp
--- a/Engine/Plugins/MovieScene/MovieRenderPipeline/Source/MovieRenderPipelineCore/Private/MoviePipelineRendering.cpp
+++ b/Engine/Plugins/MovieScene/MovieRenderPipeline/Source/MovieRenderPipelineCore/Private/MoviePipelineRendering.cpp
@@ -362,12 +362,6 @@ void UMoviePipeline::RenderFrame()
 					SpatialShiftY = r * FMath::Sin(Theta);
 				}
 
-				FIntPoint BackbufferResolution = FIntPoint(FMath::CeilToInt((float)OutputResolution.X / (float)OriginalTileCount.X), FMath::CeilToInt((float)OutputResolution.Y / (float)OriginalTileCount.Y));
-				FIntPoint TileResolution = BackbufferResolution;
-
-				// Apply size padding.
-				BackbufferResolution = HighResSettings->CalculatePaddedBackbufferSize(BackbufferResolution);
-
 				// We take all of the information needed to render a single sample and package it into a struct.
 				FMoviePipelineRenderPassMetrics SampleState;
 				SampleState.FrameIndex = FrameIndex;
@@ -377,37 +371,29 @@ void UMoviePipeline::RenderFrame()
 				SampleState.SceneCaptureSource = (ColorSettings && ColorSettings->bDisableToneCurve) ? ESceneCaptureSource::SCS_FinalColorHDR : ESceneCaptureSource::SCS_FinalToneCurveHDR;
 				SampleState.OutputState = CachedOutputState;
 				SampleState.OutputState.CameraIndex = 0; // Initialize to a sane default for non multi-cam passes.
-				SampleState.ProjectionMatrixJitterAmount = FVector2D((float)(SpatialShiftX) * 2.0f / BackbufferResolution.X, (float)SpatialShiftY * -2.0f / BackbufferResolution.Y);
 				SampleState.TileIndexes = FIntPoint(TileX, TileY);
 				SampleState.TileCounts = TileCount;
+				SampleState.OriginalTileCounts = OriginalTileCount;
+				SampleState.SpatialShiftX = SpatialShiftX;
+				SampleState.SpatialShiftY = SpatialShiftY;
 				SampleState.bDiscardResult = CachedOutputState.bDiscardRenderResult;
 				SampleState.SpatialSampleIndex = SpatialSampleIndex;
 				SampleState.SpatialSampleCount = NumSpatialSamples;
 				SampleState.TemporalSampleIndex = CachedOutputState.TemporalSampleIndex;
 				SampleState.TemporalSampleCount = AntiAliasingSettings->TemporalSampleCount;
 				SampleState.AccumulationGamma = AntiAliasingSettings->AccumulationGamma;
-				SampleState.BackbufferSize = BackbufferResolution;
-				SampleState.TileSize = TileResolution;
 				SampleState.FrameInfo = FrameInfo;
 				SampleState.bWriteSampleToDisk = bWriteAllSamples;
 				SampleState.TextureSharpnessBias = HighResSettings->TextureSharpnessBias;
 				SampleState.OCIOConfiguration = ColorSettings ? &ColorSettings->OCIOConfiguration : nullptr;
 				SampleState.GlobalScreenPercentageFraction = FLegacyScreenPercentageDriver::GetCVarResolutionFraction();
-				{
-					SampleState.OverlappedPad = FIntPoint(FMath::CeilToInt(TileResolution.X * HighResSettings->OverlapRatio), 
-														   FMath::CeilToInt(TileResolution.Y * HighResSettings->OverlapRatio));
-					SampleState.OverlappedOffset = FIntPoint(TileX * TileResolution.X - SampleState.OverlappedPad.X,
-															  TileY * TileResolution.Y - SampleState.OverlappedPad.Y);
-
-					// Move the final render by this much in the accumulator to counteract the offset put into the view matrix.
-					// Note that when bAllowSpatialJitter is false, SpatialShiftX/Y will always be zero.
-					SampleState.OverlappedSubpixelShift = FVector2D(0.5f - SpatialShiftX, 0.5f - SpatialShiftY);
-				}
 				SampleState.OverscanPercentage = FMath::Clamp(CameraSettings->OverscanPercentage, 0.0f, 1.0f);
+
 				// Render each output pass
+				FMoviePipelineRenderPassMetrics SampleStateForCurrentResolution = UE::MoviePipeline::GetRenderPassMetrics(GetPipelineMasterConfig(), ActiveShotList[CurrentShotIndex], SampleState, OutputResolution);
 				for (UMoviePipelineRenderPass* RenderPass : InputBuffers)
 				{
-					RenderPass->RenderSample_GameThread(SampleState);
+					RenderPass->RenderSample_GameThread(SampleStateForCurrentResolution);
 				}
 			}
 		}
diff --git a/Engine/Plugins/MovieScene/MovieRenderPipeline/Source/MovieRenderPipelineCore/Private/MoviePipelineSurfaceReader.cpp b/Engine/Plugins/MovieScene/MovieRenderPipeline/Source/MovieRenderPipelineCore/Private/MoviePipelineSurfaceReader.cpp
--- a/Engine/Plugins/MovieScene/MovieRenderPipeline/Source/MovieRenderPipelineCore/Private/MoviePipelineSurfaceReader.cpp
+++ b/Engine/Plugins/MovieScene/MovieRenderPipeline/Source/MovieRenderPipelineCore/Private/MoviePipelineSurfaceReader.cpp
@@ -15,15 +15,13 @@
 
 DECLARE_CYCLE_STAT(TEXT("STAT_MoviePipeline_SurfaceReadback"), STAT_MoviePipeline_SurfaceReadback, STATGROUP_MoviePipeline);
 
-FMoviePipelineSurfaceReader::FMoviePipelineSurfaceReader(EPixelFormat InPixelFormat, FIntPoint InSurfaceSize, bool bInInvertAlpha)
+FMoviePipelineSurfaceReader::FMoviePipelineSurfaceReader(EPixelFormat InPixelFormat, bool bInInvertAlpha)
 {
 	AvailableEvent = nullptr;
 	ReadbackTexture = nullptr;
 	PixelFormat = InPixelFormat;
 	bQueuedForCapture = false;
 	bInvertAlpha = bInInvertAlpha;
-	
-	Resize(InSurfaceSize.X, InSurfaceSize.Y);
 }
 
 FMoviePipelineSurfaceReader::~FMoviePipelineSurfaceReader()
@@ -40,20 +38,19 @@ void FMoviePipelineSurfaceReader::Initialize()
 	AvailableEvent = FPlatformProcess::GetSynchEventFromPool();
 }
 
-void FMoviePipelineSurfaceReader::Resize(uint32 Width, uint32 Height)
+void FMoviePipelineSurfaceReader::ResizeImpl(uint32 Width, uint32 Height)
 {
 	ReadbackTexture.SafeRelease();
 	Size = FIntPoint((int32)Width, (int32)Height);
 
 	// TexCreate_CPUReadback is important to make this texture available for reading later.
-	FMoviePipelineSurfaceReader* This = this;
 	ENQUEUE_RENDER_COMMAND(CreateCaptureFrameTexture)(
-		[Width, Height, This](FRHICommandListImmediate& RHICmdList)
+		[Width, Height, SurfaceReader = SharedThis(this)](FRHICommandListImmediate& RHICmdList)
 		{
 			const FRHITextureCreateDesc Desc =
-				FRHITextureCreateDesc::Create2D(TEXT("FMoviePipelineSurfaceReader"), Width, Height, This->PixelFormat)
+				FRHITextureCreateDesc::Create2D(TEXT("FMoviePipelineSurfaceReader"), Width, Height, SurfaceReader->PixelFormat)
 				.SetFlags(ETextureCreateFlags::CPUReadback);
-			This->ReadbackTexture = RHICreateTexture(Desc);
+			SurfaceReader->ReadbackTexture = RHICreateTexture(Desc);
 		});
 }
 
@@ -287,7 +284,7 @@ void FMoviePipelineSurfaceQueue::BlockUntilAnyAvailable()
 	bool bAnyAvailable = false;
 	for(FResolveSurface& ResolveSurface : Surfaces)
 	{
-		if (ResolveSurface.Surface.IsAvailable())
+		if (ResolveSurface.Surface->IsAvailable())
 		{
 			bAnyAvailable = true;
 			break;
@@ -299,7 +296,7 @@ void FMoviePipelineSurfaceQueue::BlockUntilAnyAvailable()
 	if (!bAnyAvailable)
 	{
 		const int32 OldestIndex = (CurrentFrameIndex + 1) % Surfaces.Num();
-		Surfaces[OldestIndex].Surface.BlockUntilAvailable();
+		Surfaces[OldestIndex].Surface->BlockUntilAvailable();
 	}
 }
 
@@ -312,9 +309,9 @@ void FMoviePipelineSurfaceQueue::Shutdown()
 			[ResolveSurface](FRHICommandListImmediate& RHICmdList)
 			{
 				// Ensure that all surfaces have queued up their readback.
-				if (ResolveSurface->Surface.WasEverQueued())
+				if (ResolveSurface->Surface->WasEverQueued())
 				{
-					ResolveSurface->Surface.CopyReadbackTexture_RenderThread(MoveTemp(ResolveSurface->FunctionCallback), ResolveSurface->FunctionPayload);
+					ResolveSurface->Surface->CopyReadbackTexture_RenderThread(MoveTemp(ResolveSurface->FunctionCallback), ResolveSurface->FunctionPayload);
 				}
 			});
 	}
@@ -325,7 +322,7 @@ void FMoviePipelineSurfaceQueue::Shutdown()
 	for (FResolveSurface& ResolveSurface : Surfaces)
 	{
 		// These should now all be available.
-		ensureMsgf(ResolveSurface.Surface.IsAvailable(), TEXT("Flushed rendering commands but surface reader didn't perform the readback on a surface!"));
+		ensureMsgf(ResolveSurface.Surface->IsAvailable(), TEXT("Flushed rendering commands but surface reader didn't perform the readback on a surface!"));
 	}
 }
 
@@ -335,7 +332,7 @@ void FMoviePipelineSurfaceQueue::OnRenderTargetReady_RenderThread(const FTexture
 
 	// Pick the next destination surface and ensure it's available.
 	FResolveSurface* NextResolveTarget = &Surfaces[CurrentFrameIndex];
-	if (!NextResolveTarget->Surface.IsAvailable())
+	if (!NextResolveTarget->Surface->IsAvailable())
 	{
 		// BlockUntilAnyAvailable should have been called before submitting more work. We can't block
 		// until a surface is available in this callback because we'd be waiting on the render thread
@@ -343,12 +340,12 @@ void FMoviePipelineSurfaceQueue::OnRenderTargetReady_RenderThread(const FTexture
 		check(false);
 	}
 
-	NextResolveTarget->Surface.Initialize();
+	NextResolveTarget->Surface->Initialize();
 	NextResolveTarget->FunctionCallback = MoveTemp(InFunctionCallback);
 	NextResolveTarget->FunctionPayload = InFramePayload;
 
 	// Queue this sample to be copied to the target surface.
-	NextResolveTarget->Surface.ResolveSampleToReadbackTexture_RenderThread(InRenderTarget);
+	NextResolveTarget->Surface->ResolveSampleToReadbackTexture_RenderThread(InRenderTarget);
 
 	// By the time we get to this point, our oldest surface should have successfully been rendered to, and no longer be in use by the GPU.
 	// We can now safely map the surface and copy the data out of it without causing a GPU stall.
@@ -361,10 +358,10 @@ void FMoviePipelineSurfaceQueue::OnRenderTargetReady_RenderThread(const FTexture
 		FResolveSurface* OldestResolveTarget = &Surfaces[PrevCaptureIndex];
 
 		// Only try to do the readback if the target has ever been written to.
-		if (OldestResolveTarget->Surface.WasEverQueued())
+		if (OldestResolveTarget->Surface->WasEverQueued())
 		{
 			SCOPE_CYCLE_COUNTER(STAT_MoviePipeline_SurfaceReadback);
-			OldestResolveTarget->Surface.CopyReadbackTexture_RenderThread(MoveTemp(OldestResolveTarget->FunctionCallback), OldestResolveTarget->FunctionPayload);
+			OldestResolveTarget->Surface->CopyReadbackTexture_RenderThread(MoveTemp(OldestResolveTarget->FunctionCallback), OldestResolveTarget->FunctionPayload);
 		}
 	}
 
diff --git a/Engine/Plugins/MovieScene/MovieRenderPipeline/Source/MovieRenderPipelineCore/Private/MoviePipelineUtils.cpp b/Engine/Plugins/MovieScene/MovieRenderPipeline/Source/MovieRenderPipelineCore/Private/MoviePipelineUtils.cpp
--- a/Engine/Plugins/MovieScene/MovieRenderPipeline/Source/MovieRenderPipelineCore/Private/MoviePipelineUtils.cpp
+++ b/Engine/Plugins/MovieScene/MovieRenderPipeline/Source/MovieRenderPipelineCore/Private/MoviePipelineUtils.cpp
@@ -3,6 +3,10 @@
 #include "MoviePipelineUtils.h"
 #include "Modules/ModuleManager.h"
 #include "AssetRegistry/AssetRegistryModule.h"
+#include "MoviePipelineAntiAliasingSetting.h"
+#include "MoviePipelineBlueprintLibrary.h"
+#include "MoviePipelineCameraSetting.h"
+#include "MoviePipelineHighResSetting.h"
 #include "MoviePipelineSetting.h"
 #include "AssetRegistry/ARFilter.h"
 #include "AssetRegistry/AssetData.h"
@@ -17,6 +21,7 @@
 #include "MovieScene.h"
 #include "MovieSceneSequence.h"
 #include "MovieRenderPipelineCoreModule.h"
+#include "Math/Halton.h"
 #include "Misc/Paths.h"
 #include "Misc/FrameRate.h"
 #include "MoviePipelineQueue.h"
@@ -896,4 +901,46 @@ namespace UE
 		}
 
 	}
-}
\ No newline at end of file
+}
+namespace UE
+{
+	namespace MoviePipeline
+	{
+		FMoviePipelineRenderPassMetrics GetRenderPassMetrics(UMoviePipelineMasterConfig* InMasterConfig, UMoviePipelineExecutorShot* InPipelineExecutorShot, const FMoviePipelineRenderPassMetrics& InRenderPassMetrics, const FIntPoint& InEffectiveOutputResolution)
+		{
+			FMoviePipelineRenderPassMetrics OutRenderPassMetrics = InRenderPassMetrics;
+
+			if (InMasterConfig && InPipelineExecutorShot)
+			{
+				const UMoviePipelineHighResSetting* HighResSettings = Cast<const UMoviePipelineHighResSetting>(UMoviePipelineBlueprintLibrary::FindOrGetDefaultSettingForShot(UMoviePipelineHighResSetting::StaticClass(), InMasterConfig, InPipelineExecutorShot));
+				check(HighResSettings);
+
+				FIntPoint BackbufferResolution = FIntPoint(FMath::CeilToInt((float)InEffectiveOutputResolution.X / (float)InRenderPassMetrics.OriginalTileCounts.X), FMath::CeilToInt((float)InEffectiveOutputResolution.Y / (float)InRenderPassMetrics.OriginalTileCounts.Y));
+				const FIntPoint TileResolution = BackbufferResolution;
+
+				// Apply size padding.
+				BackbufferResolution = HighResSettings->CalculatePaddedBackbufferSize(BackbufferResolution);
+
+
+				OutRenderPassMetrics.TileSize = TileResolution;
+				OutRenderPassMetrics.ProjectionMatrixJitterAmount = FVector2D((float)(InRenderPassMetrics.SpatialShiftX) * 2.0f / BackbufferResolution.X, (float)InRenderPassMetrics.SpatialShiftY * -2.0f / BackbufferResolution.Y);
+
+				OutRenderPassMetrics.BackbufferSize = BackbufferResolution;
+				OutRenderPassMetrics.EffectiveOutputResolution = InEffectiveOutputResolution;
+
+				{
+					OutRenderPassMetrics.OverlappedPad = FIntPoint(FMath::CeilToInt(TileResolution.X * HighResSettings->OverlapRatio),
+						FMath::CeilToInt(TileResolution.Y * HighResSettings->OverlapRatio));
+					OutRenderPassMetrics.OverlappedOffset = FIntPoint(InRenderPassMetrics.TileIndexes.X * TileResolution.X - InRenderPassMetrics.OverlappedPad.X,
+						InRenderPassMetrics.TileIndexes.Y * TileResolution.Y - InRenderPassMetrics.OverlappedPad.Y);
+
+					// Move the final render by this much in the accumulator to counteract the offset put into the view matrix.
+					// Note that when bAllowSpatialJitter is false, SpatialShiftX/Y will always be zero.
+					OutRenderPassMetrics.OverlappedSubpixelShift = FVector2D(0.5f - InRenderPassMetrics.SpatialShiftX, 0.5f - InRenderPassMetrics.SpatialShiftY);
+				}
+			}
+
+			return OutRenderPassMetrics;
+		}
+	}
+}
diff --git a/Engine/Plugins/MovieScene/MovieRenderPipeline/Source/MovieRenderPipelineCore/Public/MoviePipelineRenderPass.h b/Engine/Plugins/MovieScene/MovieRenderPipeline/Source/MovieRenderPipelineCore/Public/MoviePipelineRenderPass.h
--- a/Engine/Plugins/MovieScene/MovieRenderPipeline/Source/MovieRenderPipelineCore/Public/MoviePipelineRenderPass.h
+++ b/Engine/Plugins/MovieScene/MovieRenderPipeline/Source/MovieRenderPipelineCore/Public/MoviePipelineRenderPass.h
@@ -23,6 +23,7 @@ public:
 
 	void Teardown()
 	{
+		WaitUntilTasksComplete();
 		TeardownImpl();
 	}
 
@@ -52,6 +53,7 @@ protected:
 #endif
 protected:
 	virtual void SetupImpl(const MoviePipeline::FMoviePipelineRenderPassInitSettings& InPassInitSettings) {}
+	virtual void WaitUntilTasksComplete() {}
 	virtual void TeardownImpl() {}
 	virtual void GatherOutputPassesImpl(TArray<FMoviePipelinePassIdentifier>& ExpectedRenderPasses) {}
 	virtual void RenderSample_GameThreadImpl(const FMoviePipelineRenderPassMetrics& InSampleState) {}
diff --git a/Engine/Plugins/MovieScene/MovieRenderPipeline/Source/MovieRenderPipelineCore/Public/MoviePipelineSurfaceReader.h b/Engine/Plugins/MovieScene/MovieRenderPipeline/Source/MovieRenderPipelineCore/Public/MoviePipelineSurfaceReader.h
--- a/Engine/Plugins/MovieScene/MovieRenderPipeline/Source/MovieRenderPipelineCore/Public/MoviePipelineSurfaceReader.h
+++ b/Engine/Plugins/MovieScene/MovieRenderPipeline/Source/MovieRenderPipelineCore/Public/MoviePipelineSurfaceReader.h
@@ -16,9 +16,10 @@
 * surface until it has been written to, combined with a round-robin of surfaces to avoid stalls.
 */
 struct MOVIERENDERPIPELINECORE_API FMoviePipelineSurfaceReader
+	: public TSharedFromThis<FMoviePipelineSurfaceReader, ESPMode::ThreadSafe>
 {
 	/** Construct a surface reader with the given pixel format and size. This will be the output specification. */
-	FMoviePipelineSurfaceReader(EPixelFormat InPixelFormat, FIntPoint InSurfaceSize, bool bInInvertAlpha);
+	FMoviePipelineSurfaceReader(EPixelFormat InPixelFormat, bool bInInvertAlpha);
 
 	~FMoviePipelineSurfaceReader();
 
@@ -46,12 +47,13 @@ struct MOVIERENDERPIPELINECORE_API FMoviePipelineSurfaceReader
 	* Maps the ReadbackTexture to the CPU (which should have been resolved to before this point) and copies the data over.
 	*/
 	void CopyReadbackTexture_RenderThread(TUniqueFunction<void(TUniquePtr<FImagePixelData>&&)>&& InFunctionCallback, TSharedPtr<FImagePixelDataPayload, ESPMode::ThreadSafe> InFramePayload);
-protected:
-
 
+protected:
+	friend struct FMoviePipelineSurfaceQueue;
 	/** Set up this surface to the specified width/height */
-	void Resize(uint32 Width, uint32 Height);
+	void ResizeImpl(uint32 Width, uint32 Height);
 
+protected:
 	/** Optional event that is triggered when the surface is no longer in use */
 	FEvent* AvailableEvent;
 
@@ -93,11 +95,12 @@ private:
 		UE_NONCOPYABLE(FResolveSurface);
 
 		FResolveSurface(const EPixelFormat InPixelFormat, const FIntPoint InSurfaceSize, const bool bInInvertAlpha)
-			: Surface(InPixelFormat, InSurfaceSize, bInInvertAlpha)
+			: Surface(new FMoviePipelineSurfaceReader(InPixelFormat, bInInvertAlpha))
 		{
+			Surface->ResizeImpl(InSurfaceSize.X, InSurfaceSize.Y);
 		}
 
-		FMoviePipelineSurfaceReader Surface;
+		const TSharedRef<FMoviePipelineSurfaceReader, ESPMode::ThreadSafe> Surface;
 
 		TUniqueFunction<void(TUniquePtr<FImagePixelData>&&)> FunctionCallback;
 		TSharedPtr<FImagePixelDataPayload, ESPMode::ThreadSafe> FunctionPayload;
diff --git a/Engine/Plugins/MovieScene/MovieRenderPipeline/Source/MovieRenderPipelineCore/Public/MoviePipelineUtils.h b/Engine/Plugins/MovieScene/MovieRenderPipeline/Source/MovieRenderPipelineCore/Public/MoviePipelineUtils.h
--- a/Engine/Plugins/MovieScene/MovieRenderPipeline/Source/MovieRenderPipelineCore/Public/MoviePipelineUtils.h
+++ b/Engine/Plugins/MovieScene/MovieRenderPipeline/Source/MovieRenderPipelineCore/Public/MoviePipelineUtils.h
@@ -136,6 +136,7 @@ namespace UE
 		MOVIERENDERPIPELINECORE_API void GetHardwareUsageMetadata(TMap<FString, FString>& InFileMetadata, const FString& InOutputDir);
 		MOVIERENDERPIPELINECORE_API void GetMetadataFromCineCamera(class UCineCameraComponent* InComponent, const FString& InCameraName, const FString& InRenderPassName, TMap<FString, FString>& InOutMetadata);
 		MOVIERENDERPIPELINECORE_API void GetMetadataFromCameraLocRot(const FString& InCameraName, const FString& InRenderPassName, const FVector& InCurLoc, const FRotator& InCurRot, const FVector& InPrevLoc, const FRotator& InPrevRot, TMap<FString, FString>& InOutMetadata);
+		MOVIERENDERPIPELINECORE_API FMoviePipelineRenderPassMetrics GetRenderPassMetrics(UMoviePipelineMasterConfig* InMasterConfig, UMoviePipelineExecutorShot* InPipelineExecutorShot, const FMoviePipelineRenderPassMetrics& InRenderPassMetrics, const FIntPoint& InEffectiveOutputResolution);
 	}
 }
 
diff --git a/Engine/Plugins/MovieScene/MovieRenderPipeline/Source/MovieRenderPipelineCore/Public/MovieRenderPipelineDataTypes.h b/Engine/Plugins/MovieScene/MovieRenderPipeline/Source/MovieRenderPipelineCore/Public/MovieRenderPipelineDataTypes.h
--- a/Engine/Plugins/MovieScene/MovieRenderPipeline/Source/MovieRenderPipelineCore/Public/MovieRenderPipelineDataTypes.h
+++ b/Engine/Plugins/MovieScene/MovieRenderPipeline/Source/MovieRenderPipelineCore/Public/MovieRenderPipelineDataTypes.h
@@ -893,12 +893,22 @@ public:
 	/** What type of output should the Post Processing output? (HDR /w ToneCurve, HDR w/o ToneCurve) */
 	ESceneCaptureSource SceneCaptureSource;
 
+	/** How many tiles on X and Y are there total. */
+	FIntPoint OriginalTileCounts;
+
 	/** How many tiles on X and Y are there total. */
 	FIntPoint TileCounts;
 
 	/** Of the TileCount, which X/Y Tile is this sample for. */
 	FIntPoint TileIndexes;
 
+	/** SpatialJitter offset X. */
+	float SpatialShiftX;
+
+	/** SpatialJitter offset Y. */
+	float SpatialShiftY;
+
+
 	/** Get a 0-(TileCount-1) version of TileIndex. */
 	FORCEINLINE int32 GetTileIndex() const
 	{
@@ -913,6 +923,9 @@ public:
 	/** How big is the back buffer for this sample? This is the final target size divided by the number of tiles with padding added. */
 	FIntPoint BackbufferSize;
 
+	/** Full Output image resolution for this shot. */
+	FIntPoint EffectiveOutputResolution;
+
 	/** How big is an individual tile for this sample? This is the final target size divided by the number of tiles (without padding added) */
 	FIntPoint TileSize;
 
diff --git a/Engine/Plugins/MovieScene/MovieRenderPipeline/Source/MovieRenderPipelineRenderPasses/Private/MoviePipelineDeferredPasses.cpp b/Engine/Plugins/MovieScene/MovieRenderPipeline/Source/MovieRenderPipelineRenderPasses/Private/MoviePipelineDeferredPasses.cpp
--- a/Engine/Plugins/MovieScene/MovieRenderPipeline/Source/MovieRenderPipelineRenderPasses/Private/MoviePipelineDeferredPasses.cpp
+++ b/Engine/Plugins/MovieScene/MovieRenderPipeline/Source/MovieRenderPipelineRenderPasses/Private/MoviePipelineDeferredPasses.cpp
@@ -63,6 +63,29 @@ UMoviePipelineDeferredPassBase::UMoviePipelineDeferredPassBase()
 	bUse32BitPostProcessMaterials = false;
 }
 
+FIntPoint UMoviePipelineDeferredPassBase::GetEffectiveOutputResolutionForCamera(const int32 InCameraIndex) const
+{
+	// Add here the the output resolution for each camera. Now only one size is implemented, obtained from the settings.
+
+	UMoviePipelineMasterConfig* MasterConfig = GetPipeline()->GetPipelineMasterConfig();
+	UMoviePipelineExecutorShot* CurrentShot = GetPipeline()->GetActiveShotList()[GetPipeline()->GetCurrentShotIndex()];
+
+	const FIntPoint OutputResolution = UMoviePipelineBlueprintLibrary::GetEffectiveOutputResolution(MasterConfig, CurrentShot);
+
+	return OutputResolution;
+}
+
+FMoviePipelineRenderPassMetrics UMoviePipelineDeferredPassBase::GetRenderPassMetricsForCamera(const int32 InCameraIndex, const FMoviePipelineRenderPassMetrics& InSampleState) const
+{
+	// Add per-camera custom backbuffer size support here.
+	UMoviePipelineMasterConfig* MasterConfig = GetPipeline()->GetPipelineMasterConfig();
+	UMoviePipelineExecutorShot* CurrentShot = GetPipeline()->GetActiveShotList()[GetPipeline()->GetCurrentShotIndex()];
+	check(MasterConfig);
+	check(CurrentShot);
+
+	return UE::MoviePipeline::GetRenderPassMetrics(MasterConfig, CurrentShot, InSampleState, GetEffectiveOutputResolutionForCamera(InCameraIndex));
+}
+
 int32 UMoviePipelineDeferredPassBase::GetNumCamerasToRender() const
 {
 	UMoviePipelineExecutorShot* CurrentShot = GetPipeline()->GetActiveShotList()[GetPipeline()->GetCurrentShotIndex()];
@@ -101,31 +124,6 @@ void UMoviePipelineDeferredPassBase::SetupImpl(const MoviePipeline::FMoviePipeli
 	Super::SetupImpl(InPassInitSettings);
 	LLM_SCOPE_BYNAME(TEXT("MoviePipeline/DeferredPassSetup"));
 
-	// [0] is FinalImage, [1] is Default Layer, [1+] is Stencil Layers. Not used by post processing materials
-	// Render Target that the GBuffer is copied to
-	int32 NumRenderTargets = (bAddDefaultLayer ? 1 : 0) + GetNumStencilLayers() + 1;
-	for (int32 Index = 0; Index < NumRenderTargets; Index++)
-	{
-		UTextureRenderTarget2D* NewTarget = NewObject<UTextureRenderTarget2D>(GetTransientPackage());
-		NewTarget->ClearColor = FLinearColor(0.0f, 0.0f, 0.0f, 0.0f);
-
-		// Initialize to the tile size (not final size) and use a 16 bit back buffer to avoid precision issues when accumulating later
-		NewTarget->InitCustomFormat(InPassInitSettings.BackbufferResolution.X, InPassInitSettings.BackbufferResolution.Y, EPixelFormat::PF_FloatRGBA, false);
-
-		// OCIO: Since this is a manually created Render target we don't need Gamma to be applied.
-		// We use this render target to render to via a display extension that utilizes Display Gamma
-		// which has a default value of 2.2 (DefaultDisplayGamma), therefore we need to set Gamma on this render target to 2.2 to cancel out any unwanted effects.
-		NewTarget->TargetGamma = FOpenColorIODisplayExtension::DefaultDisplayGamma;
-
-		TileRenderTargets.Add(NewTarget);
-	}
-
-
-	if (GetPipeline()->GetPreviewTexture() == nullptr)
-	{
-		GetPipeline()->SetPreviewTexture(TileRenderTargets[0]);
-	}
-
 	{
 		TSoftObjectPtr<UMaterialInterface> StencilMatRef = TSoftObjectPtr<UMaterialInterface>(FSoftObjectPath(StencilLayerMaterialAsset));
 		StencilLayerMaterial = StencilMatRef.LoadSynchronous();
@@ -147,8 +145,6 @@ void UMoviePipelineDeferredPassBase::SetupImpl(const MoviePipeline::FMoviePipeli
 		}
 	}
 
-	SurfaceQueue = MakeShared<FMoviePipelineSurfaceQueue>(InPassInitSettings.BackbufferResolution, EPixelFormat::PF_FloatRGBA, 3, true);
-
 	// Create a view state. Each individual camera, tile, and stencil layer need their own unique state as this includes visual history for anti-aliasing, etc. 
 	UMoviePipelineExecutorShot* CurrentShot = GetPipeline()->GetActiveShotList()[GetPipeline()->GetCurrentShotIndex()];
 	UMoviePipelineHighResSetting* HighResSettings = GetPipeline()->FindOrAddSettingForShot<UMoviePipelineHighResSetting>(CurrentShot);
@@ -157,6 +153,10 @@ void UMoviePipelineDeferredPassBase::SetupImpl(const MoviePipeline::FMoviePipeli
 	int32 TotalNumberOfAccumulators = 0;
 	for (int32 CamIndex = 0; CamIndex < NumCameras; CamIndex++)
 	{
+		// Re-initialize the render target and surface queue for the current camera
+		GetOrCreateViewRenderTarget(GetEffectiveOutputResolutionForCamera(CamIndex));
+		CreateSurfaceQueueImpl(GetEffectiveOutputResolutionForCamera(CamIndex));
+
 		FMultiCameraViewStateData& CameraData = CameraViewStateData.AddDefaulted_GetRef();
 
 		// We don't always want to allocate a unique history per tile as very large resolutions can OOM the GPU in backbuffer images alone.
@@ -265,18 +265,6 @@ void UMoviePipelineDeferredPassBase::SetupImpl(const MoviePipeline::FMoviePipeli
 
 void UMoviePipelineDeferredPassBase::TeardownImpl()
 {
-	GetPipeline()->SetPreviewTexture(nullptr);
-
-	// This may call FlushRenderingCommands if there are outstanding readbacks that need to happen.
-	if (SurfaceQueue.IsValid())
-	{
-		SurfaceQueue->Shutdown();
-	}
-
-	// Stall until the task graph has completed any pending accumulations.
-	FTaskGraphInterface::Get().WaitUntilTasksComplete(OutstandingTasks, ENamedThreads::GameThread);
-	OutstandingTasks.Reset();
-
 	ActivePostProcessMaterials.Reset();
 
 	for (FMultiCameraViewStateData& CameraData : CameraViewStateData)
@@ -295,7 +283,6 @@ void UMoviePipelineDeferredPassBase::TeardownImpl()
 		}
 	}
 	CameraViewStateData.Reset();
-	TileRenderTargets.Reset();
 	
 	OCIOSceneViewExtension.Reset();
 	OCIOSceneViewExtension = nullptr;
@@ -382,23 +369,6 @@ FSceneViewStateInterface* UMoviePipelineDeferredPassBase::GetSceneViewStateInter
 	return nullptr;
 }
 
-UTextureRenderTarget2D* UMoviePipelineDeferredPassBase::GetViewRenderTarget(IViewCalcPayload* OptPayload) const
-{
-	return TileRenderTargets[0];
-	UE::MoviePipeline::FDeferredPassRenderStatePayload* Payload = (UE::MoviePipeline::FDeferredPassRenderStatePayload*)OptPayload;
-	check(Payload);
-
-	// ToDo: Is it really necessary to have a unique RT for each render? Only the last one gets shown (right now)
-	// if (CurrentLayerIndex == INDEX_NONE)
-	// {
-	// 	return TileRenderTargets[0];
-	// }
-	// else
-	// {
-	// 	return TileRenderTargets[CurrentLayerIndex + 1];
-	// }
-}
-
 void UMoviePipelineDeferredPassBase::GatherOutputPassesImpl(TArray<FMoviePipelinePassIdentifier>& ExpectedRenderPasses)
 {
 	// No super call here because multiple cameras makes this all complicated
@@ -465,13 +435,8 @@ void UMoviePipelineDeferredPassBase::AddViewExtensions(FSceneViewFamilyContext&
 
 void UMoviePipelineDeferredPassBase::RenderSample_GameThreadImpl(const FMoviePipelineRenderPassMetrics& InSampleState)
 {
-	Super::RenderSample_GameThreadImpl(InSampleState);
-
 	// Wait for a surface to be available to write to. This will stall the game thread while the RHI/Render Thread catch up.
-	{
-		SCOPE_CYCLE_COUNTER(STAT_MoviePipeline_WaitForAvailableSurface);
-		SurfaceQueue->BlockUntilAnyAvailable();
-	}
+	Super::RenderSample_GameThreadImpl(InSampleState);
 
 	const int32 NumCameras = GetNumCamerasToRender();
 	for (int32 CameraIndex = 0; CameraIndex < NumCameras; CameraIndex++)
@@ -483,7 +448,7 @@ void UMoviePipelineDeferredPassBase::RenderSample_GameThreadImpl(const FMoviePip
 		// Main Render Pass
 		if (bRenderMainPass)
 		{
-			FMoviePipelineRenderPassMetrics InOutSampleState = InSampleState;
+			FMoviePipelineRenderPassMetrics InOutSampleState = GetRenderPassMetricsForCamera(CameraIndex, InSampleState);
 			// InOutSampleState.OutputState.CameraCount = NumCameras;
 			InOutSampleState.OutputState.CameraIndex = CameraIndex;
 			InOutSampleState.OutputState.CameraNameOverride = GetCameraNameOverride(CameraIndex);
@@ -530,7 +495,12 @@ void UMoviePipelineDeferredPassBase::RenderSample_GameThreadImpl(const FMoviePip
 			View->FinalPostProcessSettings.bBufferVisualizationDumpRequired = NumValidMaterials > 0;
 
 			// Submit to be rendered. Main render pass always uses target 0.
-			FRenderTarget* RenderTarget = GetViewRenderTarget()->GameThread_GetRenderTargetResource();
+			TWeakObjectPtr<UTextureRenderTarget2D> ViewRenderTarget = GetOrCreateViewRenderTarget(InOutSampleState.BackbufferSize, (IViewCalcPayload*)(&Payload));
+			check(ViewRenderTarget.IsValid());
+
+			FRenderTarget* RenderTarget = ViewRenderTarget->GameThread_GetRenderTargetResource();
+			check(RenderTarget);
+
 			FCanvas Canvas = FCanvas(RenderTarget, nullptr, GetPipeline()->GetWorld(), View->GetFeatureLevel(), FCanvas::CDM_DeferDrawing, 1.0f);
 			GetRendererModule().BeginRenderingViewFamily(&Canvas, ViewFamily.Get());
 
@@ -541,7 +511,7 @@ void UMoviePipelineDeferredPassBase::RenderSample_GameThreadImpl(const FMoviePip
 
 		// Now do the stencil layer submission (which doesn't support additional post processing materials)
 		{
-			FMoviePipelineRenderPassMetrics InOutSampleState = InSampleState;
+			FMoviePipelineRenderPassMetrics InOutSampleState = GetRenderPassMetricsForCamera(CameraIndex, InSampleState);
 			InOutSampleState.OutputState.CameraIndex = CameraIndex;
 			InOutSampleState.OutputState.CameraNameOverride = GetCameraNameOverride(CameraIndex);
 
@@ -652,7 +622,12 @@ void UMoviePipelineDeferredPassBase::RenderSample_GameThreadImpl(const FMoviePip
 					BlendableInterface->OverrideBlendableSettings(*View, 1.f);
 
 					{
-						FRenderTarget* RenderTarget = GetViewRenderTarget()->GameThread_GetRenderTargetResource();
+						TWeakObjectPtr<UTextureRenderTarget2D> ViewRenderTarget = GetOrCreateViewRenderTarget(InOutSampleState.BackbufferSize, (IViewCalcPayload*)(&Payload));
+						check(ViewRenderTarget.IsValid());
+
+						FRenderTarget* RenderTarget = ViewRenderTarget->GameThread_GetRenderTargetResource();
+						check(RenderTarget);
+
 						FCanvas Canvas = FCanvas(RenderTarget, nullptr, GetPipeline()->GetWorld(), View->GetFeatureLevel(), FCanvas::CDM_DeferDrawing, 1.0f);
 						GetRendererModule().BeginRenderingViewFamily(&Canvas, ViewFamily.Get());
 
@@ -684,7 +659,6 @@ TFunction<void(TUniquePtr<FImagePixelData>&&)> UMoviePipelineDeferredPassBase::M
 		SCOPE_CYCLE_COUNTER(STAT_MoviePipeline_WaitForAvailableAccumulator);
 		SampleAccumulator = AccumulatorPool->BlockAndGetAccumulator_GameThread(InSampleState.OutputState.OutputFrameNumber, InPassIdentifier);
 	}
-	TSharedPtr<FMoviePipelineSurfaceQueue> LocalSurfaceQueue = SurfaceQueue;
 
 	TSharedRef<FImagePixelDataPayload, ESPMode::ThreadSafe> FramePayload = MakeShared<FImagePixelDataPayload, ESPMode::ThreadSafe>();
 	FramePayload->PassIdentifier = InPassIdentifier;
@@ -872,12 +846,14 @@ void UMoviePipelineDeferredPassBase::PostRendererSubmission(const FMoviePipeline
 		check(OutputSettings);
 		
 		// Taking overscan into account.
-		FIntPoint FullOutputSize = UMoviePipelineBlueprintLibrary::GetEffectiveOutputResolution(GetPipeline()->GetPipelineMasterConfig(), GetPipeline()->GetActiveShotList()[GetPipeline()->GetCurrentShotIndex()]);
+		const FIntPoint FullOutputSize = InSampleState.EffectiveOutputResolution;
 
-		float OutputSizeAspectRatio = FullOutputSize.X / (float)FullOutputSize.Y;
-		const FIntPoint ConstrainedFullSize = CameraCache.AspectRatio > OutputSizeAspectRatio ?
-			FIntPoint(FullOutputSize.X, FMath::CeilToInt((double)FullOutputSize.X / (double)CameraCache.AspectRatio)) :
-			FIntPoint(FMath::CeilToInt(CameraCache.AspectRatio * FullOutputSize.Y), FullOutputSize.Y);
+		const float OutputSizeAspectRatio = FullOutputSize.X / (float)FullOutputSize.Y;
+		const float CameraAspectRatio = bAllowCameraAspectRatio ? CameraCache.AspectRatio : OutputSizeAspectRatio;
+
+		const FIntPoint ConstrainedFullSize = CameraAspectRatio > OutputSizeAspectRatio ?
+			FIntPoint(FullOutputSize.X, FMath::CeilToInt((double)FullOutputSize.X / (double)CameraAspectRatio)) :
+			FIntPoint(FMath::CeilToInt(CameraAspectRatio * FullOutputSize.Y), FullOutputSize.Y);
 
 		const FIntPoint TileViewMin = InSampleState.OverlappedOffset;
 		const FIntPoint TileViewMax = TileViewMin + InSampleState.BackbufferSize;
@@ -931,13 +907,14 @@ void UMoviePipelineDeferredPassBase::PostRendererSubmission(const FMoviePipeline
 		SCOPE_CYCLE_COUNTER(STAT_MoviePipeline_WaitForAvailableAccumulator);
 		SampleAccumulator = AccumulatorPool->BlockAndGetAccumulator_GameThread(InSampleState.OutputState.OutputFrameNumber, InPassIdentifier);
 	}
-	TSharedPtr<FMoviePipelineSurfaceQueue> LocalSurfaceQueue = SurfaceQueue;
 
 	TSharedRef<FImagePixelDataPayload, ESPMode::ThreadSafe> FramePayload = MakeShared<FImagePixelDataPayload, ESPMode::ThreadSafe>();
 	FramePayload->PassIdentifier = InPassIdentifier;
 	FramePayload->SampleState = InSampleState;
 	FramePayload->SortingOrder = InSortingOrder;
 
+	TSharedPtr<FMoviePipelineSurfaceQueue, ESPMode::ThreadSafe> LocalSurfaceQueue = GetOrCreateSurfaceQueue(InSampleState.BackbufferSize, (IViewCalcPayload*)(&FramePayload.Get()));
+
 	MoviePipeline::FImageSampleAccumulationArgs AccumulationArgs;
 	{
 		AccumulationArgs.OutputMerger = GetPipeline()->OutputBuilder;
diff --git a/Engine/Plugins/MovieScene/MovieRenderPipeline/Source/MovieRenderPipelineRenderPasses/Private/MoviePipelineImagePassBase.cpp b/Engine/Plugins/MovieScene/MovieRenderPipeline/Source/MovieRenderPipelineRenderPasses/Private/MoviePipelineImagePassBase.cpp
--- a/Engine/Plugins/MovieScene/MovieRenderPipeline/Source/MovieRenderPipelineRenderPasses/Private/MoviePipelineImagePassBase.cpp
+++ b/Engine/Plugins/MovieScene/MovieRenderPipeline/Source/MovieRenderPipelineRenderPasses/Private/MoviePipelineImagePassBase.cpp
@@ -42,8 +42,37 @@ void UMoviePipelineImagePassBase::SetupImpl(const MoviePipeline::FMoviePipelineR
 	ViewState.Allocate(InPassInitSettings.FeatureLevel);
 }
 
+void UMoviePipelineImagePassBase::WaitUntilTasksComplete()
+{
+	GetPipeline()->SetPreviewTexture(nullptr);
+
+	// This may call FlushRenderingCommands if there are outstanding readbacks that need to happen.
+	for (TPair<FIntPoint, TSharedPtr<FMoviePipelineSurfaceQueue, ESPMode::ThreadSafe>> SurfaceQueueIt : SurfaceQueues)
+	{
+		if (SurfaceQueueIt.Value.IsValid())
+		{
+			SurfaceQueueIt.Value->Shutdown();
+		}
+	}
+
+	// Stall until the task graph has completed any pending accumulations.
+	FTaskGraphInterface::Get().WaitUntilTasksComplete(OutstandingTasks, ENamedThreads::GameThread);
+	OutstandingTasks.Reset();
+};
+
 void UMoviePipelineImagePassBase::TeardownImpl()
 {
+	for (TPair<FIntPoint, TWeakObjectPtr<UTextureRenderTarget2D>>& TileRenderTargetIt : TileRenderTargets)
+	{
+		if (!TileRenderTargetIt.Value.IsValid())
+		{
+			TileRenderTargetIt.Value->RemoveFromRoot();
+		}
+	}
+
+	SurfaceQueues.Empty();
+	TileRenderTargets.Empty();
+
 	FSceneViewStateInterface* Ref = ViewState.GetReference();
 	if (Ref)
 	{
@@ -66,6 +95,76 @@ void UMoviePipelineImagePassBase::AddReferencedObjects(UObject* InThis, FReferen
 	}
 }
 
+void UMoviePipelineImagePassBase::RenderSample_GameThreadImpl(const FMoviePipelineRenderPassMetrics& InSampleState)
+{
+	Super::RenderSample_GameThreadImpl(InSampleState);
+
+	// Wait for a all surfaces to be available to write to. This will stall the game thread while the RHI/Render Thread catch up.
+	SCOPE_CYCLE_COUNTER(STAT_MoviePipeline_WaitForAvailableSurface);
+	for(TPair<FIntPoint, TSharedPtr<FMoviePipelineSurfaceQueue, ESPMode::ThreadSafe>> SurfaceQueueIt : SurfaceQueues)
+	{
+		if (SurfaceQueueIt.Value.IsValid())
+		{
+			SurfaceQueueIt.Value->BlockUntilAnyAvailable();
+		}
+	}
+}
+
+TWeakObjectPtr<UTextureRenderTarget2D> UMoviePipelineImagePassBase::GetOrCreateViewRenderTarget(const FIntPoint& InSize, IViewCalcPayload* OptPayload)
+{
+	if (const TWeakObjectPtr<UTextureRenderTarget2D>* ExistViewRenderTarget = TileRenderTargets.Find(InSize))
+	{
+		return *ExistViewRenderTarget;
+	}
+
+	const TWeakObjectPtr<UTextureRenderTarget2D> NewViewRenderTarget = CreateViewRenderTargetImpl(InSize, OptPayload);
+	TileRenderTargets.Emplace(InSize, NewViewRenderTarget);
+
+	return NewViewRenderTarget;
+}
+
+TSharedPtr<FMoviePipelineSurfaceQueue, ESPMode::ThreadSafe> UMoviePipelineImagePassBase::GetOrCreateSurfaceQueue(const FIntPoint& InSize, IViewCalcPayload* OptPayload)
+{
+	if (const TSharedPtr<FMoviePipelineSurfaceQueue, ESPMode::ThreadSafe>* ExistSurfaceQueue = SurfaceQueues.Find(InSize))
+	{
+		return *ExistSurfaceQueue;
+	}
+
+	const TSharedPtr<FMoviePipelineSurfaceQueue, ESPMode::ThreadSafe> NewSurfaceQueue = CreateSurfaceQueueImpl(InSize, OptPayload);
+	SurfaceQueues.Emplace(InSize, NewSurfaceQueue);
+
+	return NewSurfaceQueue;
+}
+
+TWeakObjectPtr<UTextureRenderTarget2D> UMoviePipelineImagePassBase::CreateViewRenderTargetImpl(const FIntPoint& InSize, IViewCalcPayload* OptPayload) const
+{
+	TWeakObjectPtr<UTextureRenderTarget2D> NewTarget = NewObject<UTextureRenderTarget2D>(GetTransientPackage());
+	NewTarget->ClearColor = FLinearColor(0.0f, 0.0f, 0.0f, 0.0f);
+
+	// OCIO: Since this is a manually created Render target we don't need Gamma to be applied.
+	// We use this render target to render to via a display extension that utilizes Display Gamma
+	// which has a default value of 2.2 (DefaultDisplayGamma), therefore we need to set Gamma on this render target to 2.2 to cancel out any unwanted effects.
+	NewTarget->TargetGamma = FOpenColorIODisplayExtension::DefaultDisplayGamma;
+
+	// Initialize to the tile size (not final size) and use a 16 bit back buffer to avoid precision issues when accumulating later
+	NewTarget->InitCustomFormat(InSize.X, InSize.Y, EPixelFormat::PF_FloatRGBA, false);
+	NewTarget->AddToRoot();
+
+	if (GetPipeline()->GetPreviewTexture() == nullptr)
+	{
+		GetPipeline()->SetPreviewTexture(NewTarget.Get());
+	}
+
+	return NewTarget;
+}
+
+TSharedPtr<FMoviePipelineSurfaceQueue, ESPMode::ThreadSafe> UMoviePipelineImagePassBase::CreateSurfaceQueueImpl(const FIntPoint& InSize, IViewCalcPayload* OptPayload) const
+{
+	TSharedPtr<FMoviePipelineSurfaceQueue, ESPMode::ThreadSafe> SurfaceQueue = MakeShared<FMoviePipelineSurfaceQueue, ESPMode::ThreadSafe>(InSize, EPixelFormat::PF_FloatRGBA, 3, true);
+
+	return SurfaceQueue;
+}
+
 TSharedPtr<FSceneViewFamilyContext> UMoviePipelineImagePassBase::CalculateViewFamily(FMoviePipelineRenderPassMetrics& InOutSampleState, IViewCalcPayload* OptPayload)
 {
 	const FMoviePipelineFrameOutputState::FTimeData& TimeData = InOutSampleState.OutputState.TimeData;
@@ -74,7 +173,10 @@ TSharedPtr<FSceneViewFamilyContext> UMoviePipelineImagePassBase::CalculateViewFa
 	EViewModeIndex  ViewModeIndex;
 	GetViewShowFlags(ShowFlags, ViewModeIndex);
 	MoviePipelineRenderShowFlagOverride(ShowFlags);
-	FRenderTarget* RenderTarget = GetViewRenderTarget(OptPayload)->GameThread_GetRenderTargetResource();
+	TWeakObjectPtr<UTextureRenderTarget2D> ViewRenderTarget = GetOrCreateViewRenderTarget(InOutSampleState.BackbufferSize, OptPayload);
+	check(ViewRenderTarget.IsValid());
+
+	FRenderTarget* RenderTarget = ViewRenderTarget->GameThread_GetRenderTargetResource();
 
 	TSharedPtr<FSceneViewFamilyContext> OutViewFamily = MakeShared<FSceneViewFamilyContext>(FSceneViewFamily::ConstructionValues(
 		RenderTarget,
@@ -341,6 +443,9 @@ FSceneView* UMoviePipelineImagePassBase::GetSceneViewForSampleState(FSceneViewFa
 
 	UE::MoviePipeline::FImagePassCameraViewData CameraInfo = GetCameraInfo(InOutSampleState, OptPayload);
 
+	const float DestAspectRatio = InOutSampleState.BackbufferSize.X / (float)InOutSampleState.BackbufferSize.Y;
+	const float CameraAspectRatio = bAllowCameraAspectRatio ? CameraInfo.ViewInfo.AspectRatio : DestAspectRatio;
+
 	FSceneViewInitOptions ViewInitOptions;
 	ViewInitOptions.ViewFamily = ViewFamily;
 	ViewInitOptions.ViewOrigin = CameraInfo.ViewInfo.Location;
@@ -376,7 +481,7 @@ FSceneView* UMoviePipelineImagePassBase::GetSceneViewForSampleState(FSceneViewFa
 	{
 		if (CameraInfo.ViewInfo.ProjectionMode == ECameraProjectionMode::Orthographic)
 		{
-			const float YScale = 1.0f / CameraInfo.ViewInfo.AspectRatio;
+			const float YScale = 1.0f / CameraAspectRatio;
 			const float OverscanScale = 1.0f + InOutSampleState.OverscanPercentage;
 
 			const float HalfOrthoWidth = (CameraInfo.ViewInfo.OrthoWidth / 2.0f) * OverscanScale;
@@ -407,15 +512,13 @@ FSceneView* UMoviePipelineImagePassBase::GetSceneViewForSampleState(FSceneViewFa
 
 			if (CameraInfo.ViewInfo.bConstrainAspectRatio)
 			{
-				const float DestAspectRatio = ViewInitOptions.GetViewRect().Width() / (float)ViewInitOptions.GetViewRect().Height();
-
 				// If the camera's aspect ratio has a thinner width, then stretch the horizontal fov more than usual to 
 				// account for the extra with of (before constraining - after constraining)
-				if (CameraInfo.ViewInfo.AspectRatio < DestAspectRatio)
+				if (CameraAspectRatio < DestAspectRatio)
 				{
-					const float ConstrainedWidth = ViewInitOptions.GetViewRect().Height() * CameraInfo.ViewInfo.AspectRatio;
+					const float ConstrainedWidth = ViewInitOptions.GetViewRect().Height() * CameraAspectRatio;
 					XAxisMultiplier = ConstrainedWidth / (float)ViewInitOptions.GetViewRect().Width();
-					YAxisMultiplier = CameraInfo.ViewInfo.AspectRatio;
+					YAxisMultiplier = CameraAspectRatio;
 				}
 				// Simplified some math here but effectively functions similarly to the above, the unsimplified code would look like:
 				// const float ConstrainedHeight = ViewInitOptions.GetViewRect().Width() / CameraCache.AspectRatio;
diff --git a/Engine/Plugins/MovieScene/MovieRenderPipeline/Source/MovieRenderPipelineRenderPasses/Public/MoviePipelineDeferredPasses.h b/Engine/Plugins/MovieScene/MovieRenderPipeline/Source/MovieRenderPipelineRenderPasses/Public/MoviePipelineDeferredPasses.h
--- a/Engine/Plugins/MovieScene/MovieRenderPipeline/Source/MovieRenderPipelineRenderPasses/Public/MoviePipelineDeferredPasses.h
+++ b/Engine/Plugins/MovieScene/MovieRenderPipeline/Source/MovieRenderPipelineRenderPasses/Public/MoviePipelineDeferredPasses.h
@@ -52,7 +52,6 @@ protected:
 	virtual int32 GetOutputFileSortingOrder() const override { return 0; }
 	virtual bool IsAlphaInTonemapperRequiredImpl() const override { return bAccumulatorIncludesAlpha; }
 	virtual FSceneViewStateInterface* GetSceneViewStateInterface(IViewCalcPayload* OptPayload = nullptr) override;
-	virtual UTextureRenderTarget2D* GetViewRenderTarget(IViewCalcPayload* OptPayload = nullptr) const override;
 	virtual void AddViewExtensions(FSceneViewFamilyContext& InContext, FMoviePipelineRenderPassMetrics& InOutSampleState) override;
 	virtual bool IsAutoExposureAllowed(const FMoviePipelineRenderPassMetrics& InSampleState) const override;
 	virtual void BlendPostProcessSettings(FSceneView* InView, FMoviePipelineRenderPassMetrics& InOutSampleState, IViewCalcPayload* OptPayload = nullptr);
@@ -70,6 +69,9 @@ protected:
 	virtual FString GetCameraName(const int32 InCameraIndex) const;
 	virtual FString GetCameraNameOverride(const int32 InCameraIndex) const;
 
+	virtual FMoviePipelineRenderPassMetrics GetRenderPassMetricsForCamera(const int32 InCameraIndex, const FMoviePipelineRenderPassMetrics& InSampleState) const;
+	virtual FIntPoint GetEffectiveOutputResolutionForCamera(const int32 InCameraIndex) const;
+
 	bool IsUsingDataLayers() const;
 	int32 GetNumStencilLayers() const;
 	TArray<FString> GetStencilLayerNames() const;
@@ -154,9 +156,6 @@ protected:
 	UPROPERTY(Transient, DuplicateTransient)
 	TObjectPtr<UMaterialInterface> StencilLayerMaterial;
 
-	UPROPERTY(Transient, DuplicateTransient)
-	TArray<TObjectPtr<UTextureRenderTarget2D>> TileRenderTargets;
-
 	struct FMultiCameraViewStateData
 	{
 		struct FPerTile
diff --git a/Engine/Plugins/MovieScene/MovieRenderPipeline/Source/MovieRenderPipelineRenderPasses/Public/MoviePipelineImagePassBase.h b/Engine/Plugins/MovieScene/MovieRenderPipeline/Source/MovieRenderPipelineRenderPasses/Public/MoviePipelineImagePassBase.h
--- a/Engine/Plugins/MovieScene/MovieRenderPipeline/Source/MovieRenderPipelineRenderPasses/Public/MoviePipelineImagePassBase.h
+++ b/Engine/Plugins/MovieScene/MovieRenderPipeline/Source/MovieRenderPipelineRenderPasses/Public/MoviePipelineImagePassBase.h
@@ -92,6 +92,7 @@ class MOVIERENDERPIPELINERENDERPASSES_API UMoviePipelineImagePassBase : public U
 public:
 	UMoviePipelineImagePassBase()
 		: UMoviePipelineRenderPass()
+		, bAllowCameraAspectRatio(true)
 	{
 		PassIdentifier = FMoviePipelinePassIdentifier("ImagePassBase");
 	}
@@ -103,6 +104,8 @@ protected:
 	// UMoviePipelineRenderPass API
 	virtual void GatherOutputPassesImpl(TArray<FMoviePipelinePassIdentifier>& ExpectedRenderPasses) override;
 	virtual void SetupImpl(const MoviePipeline::FMoviePipelineRenderPassInitSettings& InPassInitSettings) override;
+	virtual void RenderSample_GameThreadImpl(const FMoviePipelineRenderPassMetrics& InSampleState) override;
+	virtual void WaitUntilTasksComplete() override;
 	virtual void TeardownImpl() override;
 	// ~UMovieRenderPassAPI
 
@@ -125,23 +128,35 @@ protected:
 	virtual bool IsAntiAliasingSupported() const { return true; }
 	virtual int32 GetOutputFileSortingOrder() const { return -1; }
 	virtual FSceneViewStateInterface* GetSceneViewStateInterface(IViewCalcPayload* OptPayload = nullptr) { return ViewState.GetReference(); }
-	virtual UTextureRenderTarget2D* GetViewRenderTarget(IViewCalcPayload* OptPayload = nullptr) const { return TileRenderTarget.Get(); }
 	virtual void AddViewExtensions(FSceneViewFamilyContext& InContext, FMoviePipelineRenderPassMetrics& InOutSampleState) { }
 	virtual bool IsAutoExposureAllowed(const FMoviePipelineRenderPassMetrics& InSampleState) const { return true; }
 	virtual FSceneView* GetSceneViewForSampleState(FSceneViewFamily* ViewFamily, FMoviePipelineRenderPassMetrics& InOutSampleState, IViewCalcPayload* OptPayload = nullptr);
 	virtual UE::MoviePipeline::FImagePassCameraViewData GetCameraInfo(FMoviePipelineRenderPassMetrics& InOutSampleState, IViewCalcPayload* OptPayload = nullptr) const;
+
+	virtual TWeakObjectPtr<UTextureRenderTarget2D> GetOrCreateViewRenderTarget(const FIntPoint& InSize, IViewCalcPayload* OptPayload = nullptr);
+	virtual TSharedPtr<FMoviePipelineSurfaceQueue, ESPMode::ThreadSafe> GetOrCreateSurfaceQueue(const FIntPoint& InSize, IViewCalcPayload* OptPayload = nullptr);
+
+	virtual TWeakObjectPtr<UTextureRenderTarget2D> CreateViewRenderTargetImpl(const FIntPoint& InSize, IViewCalcPayload* OptPayload = nullptr) const;
+	virtual TSharedPtr<FMoviePipelineSurfaceQueue, ESPMode::ThreadSafe> CreateSurfaceQueueImpl(const FIntPoint& InSize, IViewCalcPayload* OptPayload = nullptr) const;
+
+	UE_DEPRECATED(5.1, "GetViewRenderTarget is deprecated, please use GetOrCreateViewRenderTarget")
+	virtual UTextureRenderTarget2D* GetViewRenderTarget(IViewCalcPayload* OptPayload = nullptr) const { return nullptr; }
+
 public:
 	
 
 protected:
 	/** A temporary render target that we render the view to. */
-	TWeakObjectPtr<UTextureRenderTarget2D> TileRenderTarget;
+	TMap<FIntPoint, TWeakObjectPtr<UTextureRenderTarget2D>> TileRenderTargets;
 
 	/** The history for the view */
 	FSceneViewStateReference ViewState;
 
 	/** A queue of surfaces that the render targets can be copied to. If no surface is available the game thread should hold off on submitting more samples. */
-	TSharedPtr<FMoviePipelineSurfaceQueue> SurfaceQueue;
+	TMap<FIntPoint, TSharedPtr<FMoviePipelineSurfaceQueue, ESPMode::ThreadSafe>> SurfaceQueues;
+
+	/** Some render passes may ignore the aspect ratio of the camera. */
+	bool bAllowCameraAspectRatio;
 
 	FMoviePipelinePassIdentifier PassIdentifier;
 
diff --git a/Engine/Plugins/Runtime/nDisplay/Source/DisplayClusterMoviePipeline/Private/DisplayClusterMoviePipelineSettings.cpp b/Engine/Plugins/Runtime/nDisplay/Source/DisplayClusterMoviePipeline/Private/DisplayClusterMoviePipelineSettings.cpp
--- a/Engine/Plugins/Runtime/nDisplay/Source/DisplayClusterMoviePipeline/Private/DisplayClusterMoviePipelineSettings.cpp
+++ b/Engine/Plugins/Runtime/nDisplay/Source/DisplayClusterMoviePipeline/Private/DisplayClusterMoviePipelineSettings.cpp
@@ -16,7 +16,7 @@ ADisplayClusterRootActor* UDisplayClusterMoviePipelineSettings::GetRootActor(con
 		{
 			if (ADisplayClusterRootActor* RootActorPtr = RootActorRef.Get())
 			{
-				if (!DCRootActor.IsValid() || RootActorPtr->GetFName() == DCRootActor->GetFName())
+				if (!Configuration.DCRootActor.IsValid() || RootActorPtr->GetFName() == Configuration.DCRootActor->GetFName())
 				{
 					return RootActorPtr;
 				}
@@ -27,8 +27,11 @@ ADisplayClusterRootActor* UDisplayClusterMoviePipelineSettings::GetRootActor(con
 	return nullptr;
 }
 
-bool UDisplayClusterMoviePipelineSettings::GetViewports(const UWorld* InWorld, TArray<FString>& OutViewports) const
+bool UDisplayClusterMoviePipelineSettings::GetViewports(const UWorld* InWorld, TArray<FString>& OutViewports, TArray<FIntPoint>& OutViewportResolutions) const
 {
+	// When DCRA not selected we get first suitable actor from scene, but always render all viewports
+	const bool bForceRenderAllViewports = !Configuration.DCRootActor.IsValid() || Configuration.bRenderAllViewports;
+
 	if (ADisplayClusterRootActor* RootActorPtr = GetRootActor(InWorld))
 	{
 		if (const UDisplayClusterConfigurationData* InConfigurationData = RootActorPtr->GetConfigData())
@@ -47,9 +50,14 @@ bool UDisplayClusterMoviePipelineSettings::GetViewports(const UWorld* InWorld, T
 								if (InConfigurationViewport->bAllowRendering)
 								{
 									const FString& InViewportId = InConfigurationViewportIt.Key;
-									if (bRenderAllViewports || AllowedViewportNamesList.Find(InViewportId) != INDEX_NONE)
+									if (bForceRenderAllViewports || Configuration.AllowedViewportNamesList.Find(InViewportId) != INDEX_NONE)
 									{
 										OutViewports.Add(InViewportId);
+
+										if (Configuration.bUseViewportResolutions)
+										{
+											OutViewportResolutions.Add(InConfigurationViewportIt.Value->Region.ToRect().Size());
+										}
 									}
 								}
 							}
diff --git a/Engine/Plugins/Runtime/nDisplay/Source/DisplayClusterMoviePipeline/Private/DisplayClusterMoviePipelineSettings.h b/Engine/Plugins/Runtime/nDisplay/Source/DisplayClusterMoviePipeline/Private/DisplayClusterMoviePipelineSettings.h
deleted file mode 100644
diff --git a/Engine/Plugins/Runtime/nDisplay/Source/DisplayClusterMoviePipeline/Private/DisplayClusterMoviePipelineViewportPass.cpp b/Engine/Plugins/Runtime/nDisplay/Source/DisplayClusterMoviePipeline/Private/DisplayClusterMoviePipelineViewportPass.cpp
--- a/Engine/Plugins/Runtime/nDisplay/Source/DisplayClusterMoviePipeline/Private/DisplayClusterMoviePipelineViewportPass.cpp
+++ b/Engine/Plugins/Runtime/nDisplay/Source/DisplayClusterMoviePipeline/Private/DisplayClusterMoviePipelineViewportPass.cpp
@@ -20,6 +20,9 @@ namespace DisplayClusterMoviePipelineHelpers
 {
 	static void DisplayClusterWarpBlendImpl_RenderThread(FRHICommandListImmediate& RHICmdList, const FTextureRHIRef& RenderTargetRHI, IDisplayClusterViewportManagerProxy* InViewportManagerProxy, const FString& InViewportId, const FIntPoint& OffsetMin, const FIntPoint& OffsetMax)
 	{
+		check(InViewportManagerProxy);
+		check(RenderTargetRHI.IsValid());
+
 		if (InViewportManagerProxy)
 		{
 			if (IDisplayClusterViewportProxy* InViewportProxy = InViewportManagerProxy->FindViewport_RenderThread(InViewportId))
@@ -80,6 +83,15 @@ using namespace DisplayClusterMoviePipelineHelpers;
 /////////////////////////////////////////////////////////////////////////////////////////////////////////////////////////////////
 // UDisplayClusterMoviePipelineViewportPassBase
 /////////////////////////////////////////////////////////////////////////////////////////////////////////////////////////////////
+UDisplayClusterMoviePipelineViewportPassBase::UDisplayClusterMoviePipelineViewportPassBase(const FString& InRenderPassName)
+	: UMoviePipelineDeferredPassBase()
+	, RenderPassName(InRenderPassName)
+{
+	PassIdentifier = FMoviePipelinePassIdentifier(RenderPassName);
+
+	bAllowCameraAspectRatio = false;
+}
+
 void UDisplayClusterMoviePipelineViewportPassBase::SetupImpl(const MoviePipeline::FMoviePipelineRenderPassInitSettings& InPassInitSettings)
 {
 	InitializeDisplayCluster();
@@ -94,6 +106,21 @@ void UDisplayClusterMoviePipelineViewportPassBase::TeardownImpl()
 	Super::TeardownImpl();
 }
 
+FIntPoint UDisplayClusterMoviePipelineViewportPassBase::GetEffectiveOutputResolutionForCamera(const int32 InCameraIndex) const
+{
+	if (DisplayClusterViewportSizes.IsValidIndex(InCameraIndex))
+	{
+		return DisplayClusterViewportSizes[InCameraIndex];
+	}
+
+	// Use one size from OutputSettings
+	UMoviePipelineMasterConfig* MasterConfig = GetPipeline()->GetPipelineMasterConfig();
+	UMoviePipelineExecutorShot* CurrentShot = GetPipeline()->GetActiveShotList()[GetPipeline()->GetCurrentShotIndex()];
+	const FIntPoint OutputResolution = UMoviePipelineBlueprintLibrary::GetEffectiveOutputResolution(MasterConfig, CurrentShot);
+
+	return OutputResolution;
+}
+
 /////////////////////////////////////////////////////////////////////////////////////////////////////////////////////////////////
 UE::MoviePipeline::FImagePassCameraViewData UDisplayClusterMoviePipelineViewportPassBase::GetCameraInfo(FMoviePipelineRenderPassMetrics& InOutSampleState, IViewCalcPayload* OptPayload) const
 {
@@ -412,7 +439,7 @@ bool UDisplayClusterMoviePipelineViewportPassBase::InitializeDisplayCluster()
 	UDisplayClusterMoviePipelineSettings* DisplayClusterCSettings = GetPipeline()->FindOrAddSettingForShot<UDisplayClusterMoviePipelineSettings>(CurrentShot);
 
 	// Initialize
-	if (DisplayClusterCSettings && GetRootActor(*DisplayClusterCSettings) && DisplayClusterCSettings->GetViewports(GetWorld(), DisplayClusterViewports))
+	if (DisplayClusterCSettings && GetRootActor(*DisplayClusterCSettings) && DisplayClusterCSettings->GetViewports(GetWorld(), DisplayClusterViewports, DisplayClusterViewportSizes))
 	{
 		bDisplayClusterInitialized = true;
 
diff --git a/Engine/Plugins/Runtime/nDisplay/Source/DisplayClusterMoviePipeline/Private/DisplayClusterMoviePipelineViewportPass.h b/Engine/Plugins/Runtime/nDisplay/Source/DisplayClusterMoviePipeline/Private/DisplayClusterMoviePipelineViewportPass.h
deleted file mode 100644
diff --git a/Engine/Plugins/Runtime/nDisplay/Source/DisplayClusterMoviePipeline/Public/DisplayClusterMoviePipelineSettings.h b/Engine/Plugins/Runtime/nDisplay/Source/DisplayClusterMoviePipeline/Public/DisplayClusterMoviePipelineSettings.h
new file mode 100644
--- /dev/null
+++ b/Engine/Plugins/Runtime/nDisplay/Source/DisplayClusterMoviePipeline/Public/DisplayClusterMoviePipelineSettings.h
@@ -0,0 +1,71 @@
+// Copyright Epic Games, Inc. All Rights Reserved.
+
+#pragma once
+
+#include "MoviePipelineSetting.h"
+#include "DisplayClusterRootActor.h"
+
+#include "DisplayClusterMoviePipelineSettings.generated.h"
+
+USTRUCT(Blueprintable)
+struct FDisplayClusterMoviePipelineConfiguration
+{
+	GENERATED_BODY()
+
+	// Reference to Display Cluster Root Actor
+	// If not set, the first available in the scene will be set.
+	UPROPERTY(EditAnywhere, BlueprintReadWrite, Category = "nDisplay", meta = (DisplayName = "DC Root Actor"))
+	TSoftObjectPtr<ADisplayClusterRootActor> DCRootActor;
+
+	// Render with nDisplay viewport resolutions.
+	UPROPERTY(EditAnywhere, BlueprintReadWrite, Category = "nDisplay")
+	bool bUseViewportResolutions = true;
+
+	// Render all viewports
+	UPROPERTY(EditAnywhere, BlueprintReadWrite, Category = "nDisplay")
+	bool bRenderAllViewports = true;
+
+	// Render only viewports from this list.
+	UPROPERTY(EditAnywhere, BlueprintReadWrite, Category = "nDisplay", meta = (DisplayName = "Allowed Viewport Names List"))
+	TArray<FString> AllowedViewportNamesList;
+};
+
+/**
+ * nDisplay settings for MoviePipeline
+ */
+UCLASS(Blueprintable)
+class DISPLAYCLUSTERMOVIEPIPELINE_API UDisplayClusterMoviePipelineSettings : public UMoviePipelineSetting
+{
+	GENERATED_BODY()
+
+public:
+	/**
+	 * return a pointer to the specified DCRA from the current world
+	 */
+	ADisplayClusterRootActor* GetRootActor(const UWorld* InWorld) const;
+
+	/**
+	* Collect viewport names for rendering
+	* 
+	* @param InWorld       - World for render
+	* @param OutViewports  - [Out] list of viewports to render in MoviewPipeline
+	* @param OutViewportResolutions - [Out] list of viewports dimensions
+	* 
+	 * return true if is whether it populated 1 or more viewports in the given array.
+	 */
+	bool GetViewports(const UWorld* InWorld, TArray<FString>& OutViewportss, TArray<FIntPoint>& OutViewportResolutions) const;
+
+#if WITH_EDITOR
+	virtual FText GetDisplayText() const override { return NSLOCTEXT("MovieRenderPipeline", "DisplayClusterSettingDisplayName", "nDisplay"); }
+#endif
+
+protected:
+	virtual bool IsValidOnShots() const override { return true; }
+	virtual bool IsValidOnMaster() const override { return true; }
+
+public:
+	// Reference to Display Cluster Root Actor
+	// If not set, the first available in the scene will be set.
+	UPROPERTY(EditAnywhere, BlueprintReadWrite, Category = "nDisplay", meta = (DisplayName = "nDisplay"))
+	FDisplayClusterMoviePipelineConfiguration Configuration;
+};
diff --git a/Engine/Plugins/Runtime/nDisplay/Source/DisplayClusterMoviePipeline/Public/DisplayClusterMoviePipelineViewportPass.h b/Engine/Plugins/Runtime/nDisplay/Source/DisplayClusterMoviePipeline/Public/DisplayClusterMoviePipelineViewportPass.h
new file mode 100644
--- /dev/null
+++ b/Engine/Plugins/Runtime/nDisplay/Source/DisplayClusterMoviePipeline/Public/DisplayClusterMoviePipelineViewportPass.h
@@ -0,0 +1,222 @@
+// Copyright Epic Games, Inc. All Rights Reserved.
+#pragma once
+
+#include "MoviePipelineDeferredPasses.h"
+
+#include "DisplayClusterMoviePipelineViewportPass.generated.h"
+
+class ADisplayClusterRootActor;
+
+/**
+ * nDisplay viewport render pass (Lit)
+ */
+UCLASS(BlueprintType)
+class DISPLAYCLUSTERMOVIEPIPELINE_API UDisplayClusterMoviePipelineViewportPassBase
+	: public UMoviePipelineDeferredPassBase
+{
+	GENERATED_BODY()
+
+	struct FDisplayClusterViewInfo
+	{
+		FVector  ViewLocation;
+		FRotator ViewRotation;
+		FMatrix  ProjectionMatrix;
+	};
+
+public:
+	UDisplayClusterMoviePipelineViewportPassBase(const FString& InRenderPassName);
+	UDisplayClusterMoviePipelineViewportPassBase() : UDisplayClusterMoviePipelineViewportPassBase(TEXT("nDisplayLit")) { }
+
+public:
+	// UMoviePipelineDeferredPassBase
+	virtual void SetupImpl(const MoviePipeline::FMoviePipelineRenderPassInitSettings& InPassInitSettings) override;
+	virtual void TeardownImpl() override;
+
+	virtual FSceneView* GetSceneViewForSampleState(FSceneViewFamily* ViewFamily, FMoviePipelineRenderPassMetrics& InOutSampleState, IViewCalcPayload* OptPayload) override;
+	virtual void ModifyProjectionMatrixForTiling(const FMoviePipelineRenderPassMetrics& InSampleState, const bool bInOrthographic, FMatrix& InOutProjectionMatrix, float& OutDoFSensorScale) const override;
+	virtual void PostRendererSubmission(const FMoviePipelineRenderPassMetrics& InSampleState, const FMoviePipelinePassIdentifier InPassIdentifier, const int32 InSortingOrder, FCanvas& InCanvas) override;
+
+	virtual int32 GetNumCamerasToRender() const override;
+	virtual FString GetCameraName(const int32 InCameraIndex) const override;
+	virtual FString GetCameraNameOverride(const int32 InCameraIndex) const override;
+
+	virtual UE::MoviePipeline::FImagePassCameraViewData GetCameraInfo(FMoviePipelineRenderPassMetrics& InOutSampleState, IViewCalcPayload* OptPayload = nullptr) const override;
+	virtual void BlendPostProcessSettings(FSceneView* InView, FMoviePipelineRenderPassMetrics& InOutSampleState, IViewCalcPayload* OptPayload = nullptr) override;
+
+	virtual FIntPoint GetEffectiveOutputResolutionForCamera(const int32 InCameraIndex) const override;
+	// ~UMoviePipelineDeferredPassBase
+
+#if WITH_EDITOR
+	virtual FText GetDisplayText() const override { return NSLOCTEXT("MovieRenderPipeline", "DCViewportBasePassSetting_DisplayName_Lit", "nDisplay Rendering"); }
+#endif
+
+	virtual int32 GetOutputFileSortingOrder() const override { return 2; }
+	virtual bool bIsEnabledWarpBlend() const { return true; }
+
+protected:
+	bool GetViewportId(int32 InViewportIndex, FString& OutViewportId) const;
+	ADisplayClusterRootActor* GetRootActor(const class UDisplayClusterMoviePipelineSettings& InDCSettings);
+
+	bool CalculateDisplayClusterViewport(const FMoviePipelineRenderPassMetrics& InSampleState, const FString& InViewportId, const FIntPoint& InDestSize, const uint32 InContextNum, FDisplayClusterViewInfo& OutView);
+
+	bool InitializeDisplayCluster();
+	void ReleaseDisplayCluster();
+
+	void GetViewportCutOffset(const FMoviePipelineRenderPassMetrics& InSampleState, FIntPoint& OutOffsetMin, FIntPoint& OutOffsetMax) const;
+
+public:
+	// Allow warp blend for this pass
+	UPROPERTY(EditAnywhere, BlueprintReadWrite, Category = "nDisplay")
+	bool bEnabledWarpBlend = true;
+
+private:
+	// Render pass name
+	const FString RenderPassName;
+
+	// DCRA for current shot
+	ADisplayClusterRootActor* DisplayClusterRootActorPtr = nullptr;
+
+	// Names of viewports to render in the current shot
+	TArray<FString> DisplayClusterViewports;
+	TArray<FIntPoint> DisplayClusterViewportSizes;
+
+	// Cache view info for viewports
+	TMap<int32, FDisplayClusterViewInfo> DCViews;
+	TMap<int32, FDisplayClusterViewInfo> DCPrevViews;
+
+	// Runtime flags
+	bool bFrameInitialized = true;
+	bool bFrameWarpBlend = false;
+	bool bDisplayClusterInitialized = false;
+};
+
+/**
+ * nDisplay viewport render pass (Unlit)
+ */
+UCLASS(BlueprintType)
+class UDisplayClusterMoviePipelineViewportPass_Unlit
+	: public UDisplayClusterMoviePipelineViewportPassBase
+{
+	GENERATED_BODY()
+
+public:
+	UDisplayClusterMoviePipelineViewportPass_Unlit() : UDisplayClusterMoviePipelineViewportPassBase(TEXT("nDisplayUnlit"))
+	{ }
+#if WITH_EDITOR
+	virtual FText GetDisplayText() const override { return NSLOCTEXT("MovieRenderPipeline", "DCViewportBasePassSetting_DisplayName_Unlit", "nDisplay Rendering (Unlit)"); }
+#endif
+
+	virtual void GetViewShowFlags(FEngineShowFlags& OutShowFlag, EViewModeIndex& OutViewModeIndex) const override
+	{
+		OutShowFlag = FEngineShowFlags(EShowFlagInitMode::ESFIM_Game);
+		OutViewModeIndex = EViewModeIndex::VMI_Unlit;
+	}
+};
+
+/**
+ * nDisplay viewport render pass (DetailLightingOnly)
+ */
+UCLASS(BlueprintType)
+class UDisplayClusterMoviePipelineViewportPass_DetailLighting
+	: public UDisplayClusterMoviePipelineViewportPassBase
+{
+	GENERATED_BODY()
+
+public:
+	UDisplayClusterMoviePipelineViewportPass_DetailLighting() : UDisplayClusterMoviePipelineViewportPassBase(TEXT("nDisplayDetailLightingOnly"))
+	{ }
+#if WITH_EDITOR
+	virtual FText GetDisplayText() const override { return NSLOCTEXT("MovieRenderPipeline", "DCViewportBasePassSetting_DisplayName_DetailLighting", "nDisplay Rendering (Detail Lighting)"); }
+#endif
+	virtual void GetViewShowFlags(FEngineShowFlags& OutShowFlag, EViewModeIndex& OutViewModeIndex) const override
+	{
+		OutShowFlag = FEngineShowFlags(EShowFlagInitMode::ESFIM_Game);
+		OutShowFlag.SetLightingOnlyOverride(true);
+		OutViewModeIndex = EViewModeIndex::VMI_Lit_DetailLighting;
+	}
+	virtual int32 GetOutputFileSortingOrder() const override { return 2; }
+};
+
+/**
+ * nDisplay viewport render pass (LightingOnly)
+ */
+UCLASS(BlueprintType)
+class UDisplayClusterMoviePipelineViewportPass_LightingOnly
+	: public UDisplayClusterMoviePipelineViewportPassBase
+{
+	GENERATED_BODY()
+
+public:
+	UDisplayClusterMoviePipelineViewportPass_LightingOnly() : UDisplayClusterMoviePipelineViewportPassBase(TEXT("nDisplayLightingOnly"))
+	{ }
+#if WITH_EDITOR
+	virtual FText GetDisplayText() const override { return NSLOCTEXT("MovieRenderPipeline", "DCViewportBasePassSetting_DisplayName_LightingOnly", "nDisplay Rendering (Lighting Only)"); }
+#endif
+	virtual void GetViewShowFlags(FEngineShowFlags& OutShowFlag, EViewModeIndex& OutViewModeIndex) const override
+	{
+		OutShowFlag = FEngineShowFlags(EShowFlagInitMode::ESFIM_Game);
+		OutShowFlag.SetLightingOnlyOverride(true);
+		OutViewModeIndex = EViewModeIndex::VMI_LightingOnly;
+	}
+	virtual int32 GetOutputFileSortingOrder() const override { return 2; }
+};
+
+/**
+ * nDisplay viewport render pass (ReflectionsOnly)
+ */
+UCLASS(BlueprintType)
+class UDisplayClusterMoviePipelineViewportPass_ReflectionsOnly
+	: public UDisplayClusterMoviePipelineViewportPassBase
+{
+	GENERATED_BODY()
+
+public:
+	UDisplayClusterMoviePipelineViewportPass_ReflectionsOnly() : UDisplayClusterMoviePipelineViewportPassBase(TEXT("nDisplayReflectionsOnly"))
+	{ }
+#if WITH_EDITOR
+	virtual FText GetDisplayText() const override { return NSLOCTEXT("MovieRenderPipeline", "DCViewportBasePassSetting_DisplayName_ReflectionsOnly", "nDisplay Rendering (Reflections Only)"); }
+#endif
+	virtual void GetViewShowFlags(FEngineShowFlags& OutShowFlag, EViewModeIndex& OutViewModeIndex) const override
+	{
+		OutShowFlag = FEngineShowFlags(EShowFlagInitMode::ESFIM_Game);
+		OutShowFlag.SetReflectionOverride(true);
+		OutViewModeIndex = EViewModeIndex::VMI_ReflectionOverride;
+	}
+	virtual int32 GetOutputFileSortingOrder() const override { return 2; }
+};
+
+/**
+ * nDisplay viewport render pass (PathTracer)
+ */
+UCLASS(BlueprintType)
+class UDisplayClusterMoviePipelineViewportPass_PathTracer : public UDisplayClusterMoviePipelineViewportPassBase
+{
+	GENERATED_BODY()
+
+public:
+	UDisplayClusterMoviePipelineViewportPass_PathTracer() : UDisplayClusterMoviePipelineViewportPassBase(TEXT("nDisplayPathTracer"))
+	{ }
+#if WITH_EDITOR
+	virtual FText GetDisplayText() const override { return NSLOCTEXT("MovieRenderPipeline", "DCViewportBasePassSetting_DisplayName_PathTracer", "nDisplay Path Tracer"); }
+	virtual FText GetFooterText(UMoviePipelineExecutorJob* InJob) const override;
+#endif
+	virtual void GetViewShowFlags(FEngineShowFlags& OutShowFlag, EViewModeIndex& OutViewModeIndex) const override
+	{
+		OutShowFlag = FEngineShowFlags(EShowFlagInitMode::ESFIM_Game);
+		OutShowFlag.SetPathTracing(true);
+		OutShowFlag.SetMotionBlur(!bReferenceMotionBlur);
+		OutViewModeIndex = EViewModeIndex::VMI_PathTracing;
+	}
+	virtual int32 GetOutputFileSortingOrder() const override { return 2; }
+
+	virtual bool IsAntiAliasingSupported() const { return false; }
+	virtual void ValidateStateImpl() override;
+	virtual void SetupImpl(const MoviePipeline::FMoviePipelineRenderPassInitSettings& InPassInitSettings) override;
+
+	/** When enabled, the path tracer will blend all spatial and temporal samples prior to the denoising and will disable post-processed motion blur.
+	 *  In this mode it is possible to use higher temporal sample counts to improve the motion blur quality.
+	 *  When this option is disabled, the path tracer will accumulate spatial samples, but denoise them prior to accumulation of temporal samples.
+	 */
+	UPROPERTY(EditAnywhere, BlueprintReadWrite, Category = "Reference Motion Blur")
+	bool bReferenceMotionBlur;
+};
diff --git a/Engine/Plugins/Runtime/nDisplay/Source/DisplayClusterMoviePipelineEditor/DisplayClusterMoviePipelineEditor.Build.cs b/Engine/Plugins/Runtime/nDisplay/Source/DisplayClusterMoviePipelineEditor/DisplayClusterMoviePipelineEditor.Build.cs
new file mode 100644
--- /dev/null
+++ b/Engine/Plugins/Runtime/nDisplay/Source/DisplayClusterMoviePipelineEditor/DisplayClusterMoviePipelineEditor.Build.cs
@@ -0,0 +1,52 @@
+// Copyright Epic Games, Inc. All Rights Reserved.
+
+using UnrealBuildTool;
+using System.IO;
+
+public class DisplayClusterMoviePipelineEditor : ModuleRules
+{
+	public DisplayClusterMoviePipelineEditor(ReadOnlyTargetRules ROTargetRules) : base(ROTargetRules)
+	{
+		PrivateDependencyModuleNames.AddRange(
+			new string[] {
+				"DisplayCluster",
+				"DisplayClusterMoviePipeline",
+				"DisplayClusterConfiguration",
+
+				"ApplicationCore",
+				"AppFramework",
+				"AssetTools",
+				"Core",
+				"CoreUObject",
+				"DesktopPlatform",
+				"BlueprintGraph",
+				"GraphEditor",
+
+				"EditorFramework",
+				"EditorSubsystem",
+				"EditorWidgets",
+
+				"Engine",
+				"UnrealEd",
+
+				"ImageWrapper",
+				"InputCore",
+				"Kismet",
+				"KismetCompiler",
+				"MainFrame",
+				"MediaAssets",
+				"MessageLog",
+				"PinnedCommandList",
+				"Projects",
+				"PropertyEditor",
+				"Serialization",
+				"Settings",
+				"Slate",
+				"SlateCore",
+				"ToolMenus",
+				"SubobjectEditor",
+				"SubobjectDataInterface",
+				"ToolWidgets",
+			});
+	}
+}
diff --git a/Engine/Plugins/Runtime/nDisplay/Source/DisplayClusterMoviePipelineEditor/Private/Details/DisplayClusterMoviePipelineEditorBaseTypeCustomization.cpp b/Engine/Plugins/Runtime/nDisplay/Source/DisplayClusterMoviePipelineEditor/Private/Details/DisplayClusterMoviePipelineEditorBaseTypeCustomization.cpp
new file mode 100644
--- /dev/null
+++ b/Engine/Plugins/Runtime/nDisplay/Source/DisplayClusterMoviePipelineEditor/Private/Details/DisplayClusterMoviePipelineEditorBaseTypeCustomization.cpp
@@ -0,0 +1,204 @@
+// Copyright Epic Games, Inc. All Rights Reserved.
+
+#include "DisplayClusterMoviePipelineEditorBaseTypeCustomization.h"
+#include "DisplayClusterRootActor.h"
+
+#include "PropertyHandle.h"
+#include "DetailWidgetRow.h"
+#include "DetailLayoutBuilder.h"
+#include "IDetailChildrenBuilder.h"
+#include "IPropertyUtilities.h"
+
+#define LOCTEXT_NAMESPACE "FDisplayClusterMoviePipelineEditorBaseTypeCustomization"
+
+//////////////////////////////////////////////////////////////////////////////////////////////
+// FDisplayClusterMoviePipelineEditorBaseTypeCustomization
+//////////////////////////////////////////////////////////////////////////////////////////////
+
+const FName FDisplayClusterMoviePipelineEditorBaseTypeCustomization::NoHeaderMetadataKey = TEXT("NoHeader");
+const FName FDisplayClusterMoviePipelineEditorBaseTypeCustomization::HideChildrenMetadataKey = TEXT("HideChildren");
+const FName FDisplayClusterMoviePipelineEditorBaseTypeCustomization::SubstitutionsMetadataKey = TEXT("Substitutions");
+const FName FDisplayClusterMoviePipelineEditorBaseTypeCustomization::DefaultSubstitutionsMetadataKey = TEXT("DefaultSubstitutions");
+
+void FDisplayClusterMoviePipelineEditorBaseTypeCustomization::CustomizeHeader(TSharedRef<IPropertyHandle> InPropertyHandle, FDetailWidgetRow& InHeaderRow, IPropertyTypeCustomizationUtils& CustomizationUtils)
+{
+	check(InPropertyHandle->IsValidHandle());
+
+	PropertyUtilities = CustomizationUtils.GetPropertyUtilities();
+	if (PropertyUtilities == nullptr)
+	{
+		// Can be null if PropertyEditorModule.CreateSingleProperty is what triggers this customization.
+		return;
+	}
+
+	Initialize(InPropertyHandle, CustomizationUtils);
+
+	if (ShouldShowHeader(InPropertyHandle))
+	{
+		SetHeader(InPropertyHandle, InHeaderRow, CustomizationUtils);
+	}
+}
+
+void FDisplayClusterMoviePipelineEditorBaseTypeCustomization::CustomizeChildren(TSharedRef<IPropertyHandle> InPropertyHandle, IDetailChildrenBuilder& InChildBuilder, IPropertyTypeCustomizationUtils& CustomizationUtils)
+{
+	if (ShouldShowChildren(InPropertyHandle))
+	{
+		SetChildren(InPropertyHandle, InChildBuilder, CustomizationUtils);
+	}
+}
+
+TSharedPtr<IPropertyHandle> FDisplayClusterMoviePipelineEditorBaseTypeCustomization::GetChildHandleChecked(const TSharedRef<IPropertyHandle>& InPropertyHandle, const FName& ChildPropertyName) const
+{
+	TSharedPtr<IPropertyHandle> ChildHandle = InPropertyHandle->GetChildHandle(ChildPropertyName);
+	check(ChildHandle.IsValid());
+	check(ChildHandle->IsValidHandle());
+
+	return ChildHandle;
+}
+
+void FDisplayClusterMoviePipelineEditorBaseTypeCustomization::Initialize(const TSharedRef<IPropertyHandle>& InPropertyHandle, IPropertyTypeCustomizationUtils& CustomizationUtils)
+{
+	FillSubstitutionMap(InPropertyHandle);
+}
+
+void FDisplayClusterMoviePipelineEditorBaseTypeCustomization::SetHeader(const TSharedRef<IPropertyHandle>& InPropertyHandle, FDetailWidgetRow& InHeaderRow, IPropertyTypeCustomizationUtils& CustomizationUtils)
+{
+	InHeaderRow.NameContent()
+	[
+		InPropertyHandle->CreatePropertyNameWidget()
+	]
+	.ValueContent()
+	[
+		InPropertyHandle->CreatePropertyValueWidget()
+	];
+}
+
+void FDisplayClusterMoviePipelineEditorBaseTypeCustomization::SetChildren(const TSharedRef<IPropertyHandle>& InPropertyHandle, IDetailChildrenBuilder& InChildBuilder, IPropertyTypeCustomizationUtils& CustomizationUtils)
+{
+	AddAllChildren(InPropertyHandle, InChildBuilder);
+}
+
+void FDisplayClusterMoviePipelineEditorBaseTypeCustomization::AddAllChildren(const TSharedRef<IPropertyHandle>& InPropertyHandle, IDetailChildrenBuilder& InChildBuilder)
+{
+	uint32 NumChildren = 0;
+	InPropertyHandle->GetNumChildren(NumChildren);
+	for (uint32 ChildIndex = 0; ChildIndex < NumChildren; ++ChildIndex)
+	{
+		TSharedPtr<IPropertyHandle> ChildHandle = InPropertyHandle->GetChildHandle(ChildIndex);
+		if (ChildHandle && ChildHandle->IsValidHandle() && !ChildHandle->IsCustomized())
+		{
+			FText ChildTooltip = ApplySubstitutions(ChildHandle->GetToolTipText());
+			ChildHandle->SetToolTipText(ChildTooltip);
+
+			InChildBuilder.AddProperty(ChildHandle.ToSharedRef());
+		}
+	}
+}
+
+void FDisplayClusterMoviePipelineEditorBaseTypeCustomization::FillSubstitutionMap(const TSharedRef<IPropertyHandle>& InPropertyHandle)
+{
+	if (FStructProperty* StructProperty = CastField<FStructProperty>(InPropertyHandle->GetProperty()))
+	{
+		if (StructProperty->Struct->HasMetaData(DefaultSubstitutionsMetadataKey))
+		{
+			const FString DefaultSubstitutions = StructProperty->Struct->GetMetaDataText(DefaultSubstitutionsMetadataKey).ToString();
+			ParseSubstitutions(SubstitutionsMap, DefaultSubstitutions);
+		}
+	}
+
+	if (InPropertyHandle->HasMetaData(SubstitutionsMetadataKey))
+	{
+		const FString Substitutions = InPropertyHandle->GetProperty()->GetMetaDataText(SubstitutionsMetadataKey).ToString();
+		ParseSubstitutions(SubstitutionsMap, Substitutions);
+	}
+}
+
+void FDisplayClusterMoviePipelineEditorBaseTypeCustomization::ParseSubstitutions(FFormatNamedArguments& OutSubstitutionsMap, const FString& InSubstitutionsStr) const
+{
+	TArray<FString> Substitutions;
+	InSubstitutionsStr.ParseIntoArray(Substitutions, TEXT(","));
+
+	for (const FString& Substitution : Substitutions)
+	{
+		if (Substitution.Contains(TEXT("=")))
+		{
+			FString Key;
+			FString Value;
+			Substitution.Split(TEXT("="), &Key, &Value);
+			if (!OutSubstitutionsMap.Contains(Key.TrimStartAndEnd()))
+			{
+				OutSubstitutionsMap.Add(Key.TrimStartAndEnd(), FText::FromString(Value.TrimStartAndEnd()));
+			}
+			else
+			{
+				OutSubstitutionsMap[Key.TrimStartAndEnd()] = FText::FromString(Value.TrimStartAndEnd());
+			}
+		}
+	}
+}
+
+FText FDisplayClusterMoviePipelineEditorBaseTypeCustomization::ApplySubstitutions(const FText& InText) const
+{
+	return FText::Format(InText, SubstitutionsMap);
+}
+
+bool FDisplayClusterMoviePipelineEditorBaseTypeCustomization::ShouldShowHeader(const TSharedRef<IPropertyHandle>& InPropertyHandle) const
+{
+	if (InPropertyHandle->IsValidHandle())
+	{
+		if (InPropertyHandle->HasMetaData(TEXT("ShowOnlyInnerProperties")) || InPropertyHandle->HasMetaData(NoHeaderMetadataKey))
+		{
+			return false;
+		}
+
+		if (FStructProperty* StructProperty = CastField<FStructProperty>(InPropertyHandle->GetProperty()))
+		{
+			if (StructProperty->Struct->HasMetaData(NoHeaderMetadataKey))
+			{
+				return false;
+			}
+		}
+
+		return true;
+	}
+
+	return false;
+}
+
+bool FDisplayClusterMoviePipelineEditorBaseTypeCustomization::ShouldShowChildren(const TSharedRef<IPropertyHandle>& InPropertyHandle) const
+{
+	if (InPropertyHandle->IsValidHandle())
+	{
+		if (InPropertyHandle->HasMetaData(HideChildrenMetadataKey))
+		{
+			return false;
+		}
+
+		if (FStructProperty* StructProperty = CastField<FStructProperty>(InPropertyHandle->GetProperty()))
+		{
+			if (StructProperty->Struct->HasMetaData(HideChildrenMetadataKey))
+			{
+				return false;
+			}
+		}
+
+		return true;
+	}
+
+	return false;
+}
+
+const FString* FDisplayClusterMoviePipelineEditorBaseTypeCustomization::FindMetaData(TSharedPtr<IPropertyHandle> PropertyHandle, const FName& Key) const
+{
+	if (const FString* MetadataValue = PropertyHandle->GetInstanceMetaData(Key))
+	{
+		return MetadataValue;
+	}
+	else if (PropertyHandle->HasMetaData(Key))
+	{
+		return &PropertyHandle->GetMetaData(Key);
+	}
+
+	return nullptr;
+}
+#undef LOCTEXT_NAMESPACE
diff --git a/Engine/Plugins/Runtime/nDisplay/Source/DisplayClusterMoviePipelineEditor/Private/Details/DisplayClusterMoviePipelineEditorBaseTypeCustomization.h b/Engine/Plugins/Runtime/nDisplay/Source/DisplayClusterMoviePipelineEditor/Private/Details/DisplayClusterMoviePipelineEditorBaseTypeCustomization.h
new file mode 100644
--- /dev/null
+++ b/Engine/Plugins/Runtime/nDisplay/Source/DisplayClusterMoviePipelineEditor/Private/Details/DisplayClusterMoviePipelineEditorBaseTypeCustomization.h
@@ -0,0 +1,84 @@
+// Copyright Epic Games, Inc. All Rights Reserved.
+
+#pragma once
+
+#include "CoreMinimal.h"
+#include "IPropertyTypeCustomization.h"
+
+#define GET_CHILD_HANDLE(ParentType, ChildName) GetChildHandleChecked(InPropertyHandle, GET_MEMBER_NAME_CHECKED(ParentType, ChildName))
+
+/**
+ * A base class for type customization in DisplayClusterMoviePipelineEditor. Contains support for common custom metadata specifiers.
+ */
+class FDisplayClusterMoviePipelineEditorBaseTypeCustomization : public IPropertyTypeCustomization
+{
+public:
+	/**
+	 * Metadata specifier used to indicate that a type should not have a header to encapsulate its properties in a details panel; similar to 
+	 * ShowOnlyInnerProperties, but forces all child properties to be within the category the parent type belongs to, and is applied at any
+	 * depth within a details panel.
+	 */
+	static const FName NoHeaderMetadataKey;
+
+	/** Metadata specifier used to indicate that a type should not display any of its child properties in a details panel. */
+	static const FName HideChildrenMetadataKey;
+
+	/** Metadata specifier used specify substitutions to make to the display text of any of this type's child properties. */
+	static const FName SubstitutionsMetadataKey;
+
+	/** Metadata specifier used specify the default substitutions to make to the display text of any of this type's child properties when no substitutions are specified. */
+	static const FName DefaultSubstitutionsMetadataKey;
+
+
+	static TSharedRef<IPropertyTypeCustomization> MakeInstance()
+	{
+		return MakeShared<FDisplayClusterMoviePipelineEditorBaseTypeCustomization>();
+	}
+
+	//~ IPropertyTypeCustomization interface begin
+	virtual void CustomizeHeader(TSharedRef<IPropertyHandle> InPropertyHandle, FDetailWidgetRow& InHeaderRow, IPropertyTypeCustomizationUtils& CustomizationUtils) override;
+	virtual void CustomizeChildren(TSharedRef<IPropertyHandle> InPropertyHandle, IDetailChildrenBuilder& InChildBuilder, IPropertyTypeCustomizationUtils& CustomizationUtils) override;
+	//~ IPropertyTypeCustomization interface end
+
+protected:
+	/** Gets the child property handle from the specified property handle, and performs checks to ensure the child handle is valid. */
+	TSharedPtr<IPropertyHandle> GetChildHandleChecked(const TSharedRef<IPropertyHandle>& InPropertyHandle, const FName& ChildPropertyName) const;
+
+	/** Initializes the type customization just before specifying the type's header, giving the customization the opportunity to initialize anything before generating the customized UI. */
+	virtual void Initialize(const TSharedRef<IPropertyHandle>& InPropertyHandle, IPropertyTypeCustomizationUtils& CustomizationUtils);
+
+	/** Sets the header of the type customization. By default, sets the header to the property's default name and value widget. Override to customize the type's header. */
+	virtual void SetHeader(const TSharedRef<IPropertyHandle>& InPropertyHandle, FDetailWidgetRow& InHeaderRow, IPropertyTypeCustomizationUtils& CustomizationUtils);
+
+	/** Sets the children of the type customization. By default, adds all of the type's children to the customized layout. Override to customize the type's children. */
+	virtual void SetChildren(const TSharedRef<IPropertyHandle>& InPropertyHandle, IDetailChildrenBuilder& InChildBuilder, IPropertyTypeCustomizationUtils& CustomizationUtils);
+
+	/** Adds all of the type's children to the customization's child builder. */
+	void AddAllChildren(const TSharedRef<IPropertyHandle>& InPropertyHandle, IDetailChildrenBuilder& InChildBuilder);
+
+	/** Fills the substitutiom map to use for any text subsitutions using the metadata of the specified property handle */
+	void FillSubstitutionMap(const TSharedRef<IPropertyHandle>& InPropertyHandle);
+
+	/** Parses the specified comma-deliminated substitution string into a list of key-value substitution pairs */
+	void ParseSubstitutions(FFormatNamedArguments& OutSubstitutionsMap, const FString& InSubstitutionsStr) const;
+	FText ApplySubstitutions(const FText& InText) const;
+
+	/** Gets whether the type should display a header when displayed in a details panel. By default, check's the property's metadata for the NoHeader or ShowOnlyInnerProperties specifier. */
+	virtual bool ShouldShowHeader(const TSharedRef<IPropertyHandle>& InPropertyHandle) const;
+
+	/** Gets whether the type should display its children when displayed in a details panel. By default, check's the property's metadata for the HideChildren specifier. */
+	virtual bool ShouldShowChildren(const TSharedRef<IPropertyHandle>& InPropertyHandle) const;
+
+	/** Searches the specified property's instanced and default metadata maps for metadata using the specified key */
+	const FString* FindMetaData(TSharedPtr<IPropertyHandle> PropertyHandle, const FName& Key) const;
+
+protected:
+	/** A root property handle used by this type customization */
+	TWeakPtr<IPropertyHandle> RootPropertyHandle;
+
+	/** A weak reference to the property utilities used by this type customization */
+	TWeakPtr<IPropertyUtilities> PropertyUtilities;
+
+	/** A list of substitutions that will be applied to the display text of each child property of the customized type */
+	FFormatNamedArguments SubstitutionsMap;
+};
\ No newline at end of file
diff --git a/Engine/Plugins/Runtime/nDisplay/Source/DisplayClusterMoviePipelineEditor/Private/Details/DisplayClusterMoviePipelineEditorNodeSelection.cpp b/Engine/Plugins/Runtime/nDisplay/Source/DisplayClusterMoviePipelineEditor/Private/Details/DisplayClusterMoviePipelineEditorNodeSelection.cpp
new file mode 100644
--- /dev/null
+++ b/Engine/Plugins/Runtime/nDisplay/Source/DisplayClusterMoviePipelineEditor/Private/Details/DisplayClusterMoviePipelineEditorNodeSelection.cpp
@@ -0,0 +1,260 @@
+// Copyright Epic Games, Inc. All Rights Reserved.
+
+#include "DisplayClusterMoviePipelineEditorNodeSelection.h"
+
+#include "DisplayClusterRootActor.h"
+#include "DisplayClusterConfigurationTypes.h"
+#include "Widgets/SDisplayClusterMoviePipelineEditorSearchableComboBox.h"
+
+#include "DisplayClusterMoviePipelineSettings.h"
+
+#include "DetailLayoutBuilder.h"
+#include "IDetailChildrenBuilder.h"
+#include "PropertyCustomizationHelpers.h"
+#include "PropertyHandle.h"
+#include "Widgets/Text/STextBlock.h"
+
+#define LOCTEXT_NAMESPACE "FDisplayClusterMoviePipelineEditorNodeSelection"
+
+//////////////////////////////////////////////////////////////////////////////////////////////
+// FDisplayClusterMoviePipelineEditorNodeSelection
+//////////////////////////////////////////////////////////////////////////////////////////////
+const FName FDisplayClusterMoviePipelineEditorNodeSelection::NAME_ElementToolTip = TEXT("ElementToolTip");
+
+FDisplayClusterMoviePipelineEditorNodeSelection::FDisplayClusterMoviePipelineEditorNodeSelection(EOperationMode InMode, const TSharedPtr<IPropertyHandle>& InDCRAPropertyHandle, const TSharedPtr<IPropertyHandle>& InSelectedOptionsHandle)
+{
+	DCRAPropertyHandle = InDCRAPropertyHandle;
+	SelectedOptionsHandle = InSelectedOptionsHandle;
+	check(DCRAPropertyHandle.IsValid());
+
+	OperationMode = InMode;
+
+	ResetOptions();
+}
+
+UDisplayClusterConfigurationData* FDisplayClusterMoviePipelineEditorNodeSelection::GetConfigData() const
+{
+	if (DCRAPropertyHandle.IsValid())
+	{
+		UObject* CurrentValue = nullptr;
+		if (FPropertyAccess::Success == DCRAPropertyHandle->GetValue(CurrentValue))
+		{
+			if (ADisplayClusterRootActor* RootActorPtr = Cast<ADisplayClusterRootActor>(CurrentValue))
+			{
+				return RootActorPtr->GetConfigData();
+			}
+		}
+	}
+
+	return nullptr;
+}
+
+void FDisplayClusterMoviePipelineEditorNodeSelection::ResetOptionsList()
+{
+	ResetOptions();
+
+	if (OptionsComboBox.IsValid())
+	{
+		OptionsComboBox->ResetOptionsSource(&Options);
+		OptionsComboBox->SetIsOpen(false);
+	}
+}
+
+void FDisplayClusterMoviePipelineEditorNodeSelection::CreateArrayBuilder(const TSharedRef<IPropertyHandle>& InPropertyHandle,
+	IDetailChildrenBuilder& InChildBuilder)
+{
+	TSharedRef<FDetailArrayBuilder> ArrayBuilder = MakeShareable(new FDetailArrayBuilder(InPropertyHandle));
+	ArrayBuilder->OnGenerateArrayElementWidget(FOnGenerateArrayElementWidget::CreateRaw(this,
+		&FDisplayClusterMoviePipelineEditorNodeSelection::GenerateSelectionWidget));
+
+	InChildBuilder.AddCustomBuilder(ArrayBuilder);
+}
+
+FDisplayClusterMoviePipelineEditorNodeSelection::EOperationMode FDisplayClusterMoviePipelineEditorNodeSelection::GetOperationModeFromProperty(FProperty* Property)
+{
+	EOperationMode ReturnMode = ClusterNodes;
+	if(Property)
+	{
+		if (const FString* DefinedMode = Property->FindMetaData(TEXT("ConfigurationMode")))
+		{
+			FString ModeLower = (*DefinedMode).ToLower();
+			ModeLower.RemoveSpacesInline();
+			if (ModeLower == TEXT("viewports"))
+			{
+				ReturnMode = FDisplayClusterMoviePipelineEditorNodeSelection::EOperationMode::Viewports;
+			}
+			else if (ModeLower == TEXT("clusternodes"))
+			{
+				ReturnMode = FDisplayClusterMoviePipelineEditorNodeSelection::EOperationMode::ClusterNodes;
+			}
+			// Define any other modes here.
+		}
+	}
+
+	return ReturnMode;
+}
+
+void FDisplayClusterMoviePipelineEditorNodeSelection::GenerateSelectionWidget(
+	TSharedRef<IPropertyHandle> PropertyHandle, int32 ArrayIndex, IDetailChildrenBuilder& ChildrenBuilder)
+{
+	FText ElementTooltip = FText::GetEmpty();
+
+	TSharedPtr<IPropertyHandle> ParentArrayHandle = PropertyHandle->GetParentHandle();
+	if (ParentArrayHandle.IsValid() && ParentArrayHandle->IsValidHandle())
+	{
+		if (const FString* MetaData = ParentArrayHandle->GetInstanceMetaData(NAME_ElementToolTip))
+		{
+			ElementTooltip = FText::FromString(*MetaData);
+		}
+		else if (ParentArrayHandle->HasMetaData(NAME_ElementToolTip))
+		{
+			ElementTooltip = FText::FromString(ParentArrayHandle->GetMetaData(NAME_ElementToolTip));
+		}
+		else
+		{
+			ElementTooltip = ParentArrayHandle->GetPropertyDisplayName();
+		}
+	}
+
+	IDetailPropertyRow& PropertyRow = ChildrenBuilder.AddProperty(PropertyHandle);
+	PropertyRow.CustomWidget(false)
+		.NameContent()
+		[
+			PropertyHandle->CreatePropertyNameWidget(FText::GetEmpty(), ElementTooltip)
+		]
+		.IsEnabled(IsEnabledAttr)
+		.ValueContent()
+		[
+			SAssignNew(OptionsComboBox, SDisplayClusterMoviePipelineEditorSearchableComboBox)
+				.OptionsSource(&Options)
+				.OnGenerateWidget(this, &FDisplayClusterMoviePipelineEditorNodeSelection::MakeOptionComboWidget)
+				.OnSelectionChanged(this, &FDisplayClusterMoviePipelineEditorNodeSelection::OnOptionSelected, PropertyHandle)
+				.ContentPadding(2)
+				.MaxListHeight(200.0f)
+				.IsEnabled(IsEnabledAttr)
+				.Content()
+				[
+					SNew(STextBlock)
+					.Text(this, &FDisplayClusterMoviePipelineEditorNodeSelection::GetSelectedOptionText, PropertyHandle)
+					.Font(IDetailLayoutBuilder::GetDetailFont())
+				]
+		];
+}
+
+void FDisplayClusterMoviePipelineEditorNodeSelection::ResetOptions()
+{
+	if (!bResetOptionsListFunctionInPerformed)
+	{
+		// recursion not allowed through IPropertyHandle callbacks
+		bResetOptionsListFunctionInPerformed = true;
+
+		TArray<FString> UsedValues;
+		TArray<FString> ValuesFromCurrentConfig;
+
+		if (SelectedOptionsHandle.IsValid())
+		{
+			const TSharedPtr<IPropertyHandleArray> ArrayOptionsHandle = SelectedOptionsHandle->AsArray();
+			if (ArrayOptionsHandle.IsValid())
+			{
+				// Get values from DCRA
+				if (UDisplayClusterConfigurationData* ConfigData = GetConfigData())
+				{
+					for (const TTuple<FString, TObjectPtr<UDisplayClusterConfigurationClusterNode>>& Node : ConfigData->Cluster->Nodes)
+					{
+						if (OperationMode == ClusterNodes)
+						{
+							ValuesFromCurrentConfig.Add(Node.Value->GetName());
+							continue;
+						}
+						for (const TTuple<FString, TObjectPtr<UDisplayClusterConfigurationViewport>>& Viewport : Node.Value->Viewports)
+						{
+							ValuesFromCurrentConfig.Add(Viewport.Value->GetName());
+						}
+					}
+				}
+				else
+				{
+					uint32 NumElements = 0;
+					if (FPropertyAccess::Result::Success == ArrayOptionsHandle->GetNumElements(NumElements))
+					{
+						if (NumElements > 0)
+						{
+							ArrayOptionsHandle->EmptyArray();
+						}
+					}
+				}
+
+				// Get used values and remove invalid values:
+				uint32 NumElements = 0;
+				if (FPropertyAccess::Result::Success == ArrayOptionsHandle->GetNumElements(NumElements))
+				{
+					for (int32 ElementIndex = 0; ElementIndex < (int32)NumElements; ElementIndex++)
+					{
+						const TSharedRef<IPropertyHandle> ElementHandle = ArrayOptionsHandle->GetElement(ElementIndex);
+						FString StrValue;
+						if (FPropertyAccess::Result::Success == ElementHandle->GetValue(StrValue))
+						{
+							if (!StrValue.IsEmpty())
+							{
+								if (ValuesFromCurrentConfig.Find(StrValue) == INDEX_NONE)
+								{
+									// Remove invalid values
+									ArrayOptionsHandle->DeleteItem(ElementIndex);
+									NumElements--;
+								}
+								else
+								{
+									UsedValues.AddUnique(StrValue);
+								}
+							}
+						}
+					}
+				}
+			}
+		}
+
+		Options.Reset();
+
+		// Show in options only new values
+		for (const FString& ValueIt : ValuesFromCurrentConfig)
+		{
+			if (UsedValues.Find(ValueIt) == INDEX_NONE)
+			{
+				Options.Add(MakeShared<FString>(ValueIt));
+			}
+		}
+
+		bResetOptionsListFunctionInPerformed = false;
+	}
+}
+
+TSharedRef<SWidget> FDisplayClusterMoviePipelineEditorNodeSelection::MakeOptionComboWidget(
+	TSharedPtr<FString> InItem)
+{
+	return SNew(STextBlock)
+		.Text(FText::FromString(*InItem))
+		.Font(IDetailLayoutBuilder::GetDetailFont());
+}
+
+void FDisplayClusterMoviePipelineEditorNodeSelection::OnOptionSelected(TSharedPtr<FString> InValue,
+	ESelectInfo::Type SelectInfo, TSharedRef<IPropertyHandle> InPropertyHandle)
+{
+	if (InValue.IsValid())
+	{
+		check(OptionsComboBox.IsValid());
+
+		InPropertyHandle->SetValue(*InValue);
+		
+		ResetOptions();
+		OptionsComboBox->ResetOptionsSource(&Options);
+		OptionsComboBox->SetIsOpen(false);
+	}
+}
+
+FText FDisplayClusterMoviePipelineEditorNodeSelection::GetSelectedOptionText(TSharedRef<IPropertyHandle> InPropertyHandle) const
+{
+	FString Value;
+	InPropertyHandle->GetValue(Value);
+	return FText::FromString(Value);
+}
+#undef LOCTEXT_NAMESPACE
diff --git a/Engine/Plugins/Runtime/nDisplay/Source/DisplayClusterMoviePipelineEditor/Private/Details/DisplayClusterMoviePipelineEditorNodeSelection.h b/Engine/Plugins/Runtime/nDisplay/Source/DisplayClusterMoviePipelineEditor/Private/Details/DisplayClusterMoviePipelineEditorNodeSelection.h
new file mode 100644
--- /dev/null
+++ b/Engine/Plugins/Runtime/nDisplay/Source/DisplayClusterMoviePipelineEditor/Private/Details/DisplayClusterMoviePipelineEditorNodeSelection.h
@@ -0,0 +1,91 @@
+// Copyright Epic Games, Inc. All Rights Reserved.
+
+#pragma once
+
+#include "CoreMinimal.h"
+#include "DisplayClusterMoviePipelineEditorBaseTypeCustomization.h"
+
+class UDisplayClusterConfigurationData;
+class SDisplayClusterMoviePipelineEditorSearchableComboBox;
+
+/**
+ * Helper class used to construct and manage a custom array property widget that displays for each element a dropdown containing
+ * a list of the display cluster's cluster nodes or viewports to select from.
+ */
+class FDisplayClusterMoviePipelineEditorNodeSelection : public TSharedFromThis<FDisplayClusterMoviePipelineEditorNodeSelection>
+{
+public:
+	/** Custom "ElementToolTip" metadata specifier that can be used to specify the tooltip to apply to each element of the displayed array */
+	static const FName NAME_ElementToolTip;
+
+	/** Flag to indicate which cluster item is displayed in the dropdown menus of the elements of the array widget  */
+	enum EOperationMode
+	{
+		Viewports,
+		ClusterNodes
+	};
+
+	FDisplayClusterMoviePipelineEditorNodeSelection(EOperationMode InMode, const TSharedPtr<IPropertyHandle>& InDCRAPropertyHandle, const TSharedPtr<IPropertyHandle>& InSelectedOptionsHandle);
+
+	~FDisplayClusterMoviePipelineEditorNodeSelection()
+	{
+		OptionsComboBox.Reset();
+		Options.Reset();
+	}
+
+	/** Gets the cluster configuration data object of the root actor being edited by the current details panel */
+	UDisplayClusterConfigurationData* GetConfigData() const;
+
+	/** Creates the child array builder that will be used to construct the cluster item selection array elements for the details panel */
+	void CreateArrayBuilder(const TSharedRef<IPropertyHandle>& InPropertyHandle, IDetailChildrenBuilder& InChildBuilder);
+
+	/** Sets the IsEnabled attribute to use on the array property handle */
+	FDisplayClusterMoviePipelineEditorNodeSelection& IsEnabled(const TAttribute<bool>& InIsEnabled)
+	{
+		IsEnabledAttr = InIsEnabled;
+		return *this;
+	}
+
+	/** Inspects the specified property's metadata to determine if the node selection array should show cluster nodes or viewports in its dropdown lists */
+	static EOperationMode GetOperationModeFromProperty(FProperty* Property);
+
+	void ResetOptionsList();
+
+protected:
+	/** Generates a widget to display the currently selected cluster item for the specified element in the array widget and adds it to the details panel */
+	void GenerateSelectionWidget(TSharedRef<IPropertyHandle> PropertyHandle, int32 ArrayIndex, IDetailChildrenBuilder& ChildrenBuilder);
+
+	/** Resets the list of cluster items to display in the dropdown lists for each element of the array widget */
+	void ResetOptions();
+
+	/** Constructs the text widget to display in the dropdown menu for the specified item */
+	TSharedRef<SWidget> MakeOptionComboWidget(TSharedPtr<FString> InItem);
+
+	/** Raised when an item within the dropdown menu is selected */
+	void OnOptionSelected(TSharedPtr<FString> InValue, ESelectInfo::Type SelectInfo, TSharedRef<IPropertyHandle> InPropertyHandle);
+
+	/** Gets the display text for the currently selected item of the element in the array widget */
+	FText GetSelectedOptionText(TSharedRef<IPropertyHandle> InPropertyHandle) const;
+
+private:
+	/** A custom attribute to apply to the array widget to determine if it is enabled or not */
+	TAttribute<bool> IsEnabledAttr;
+
+	/** The combo box widget to display for each element in the array widget */
+	TSharedPtr<SDisplayClusterMoviePipelineEditorSearchableComboBox> OptionsComboBox;
+
+	/** The list of options to display in each combo box's dropdown menu */
+	TArray<TSharedPtr<FString>> Options;
+
+	/** A reference to the current display cluster root actor property handle that is being edited by the details panel */
+	TSharedPtr<IPropertyHandle> DCRAPropertyHandle;
+
+	/** A reference to the string list property handle that is being edited by the details panel */
+	TSharedPtr<IPropertyHandle> SelectedOptionsHandle;
+
+	/** The operation mode which indicates whether to display cluster nodes or viewports in each element's combo box */
+	EOperationMode OperationMode = ClusterNodes;
+
+	/** Protect frunction ResetOptionsList() from recursion */
+	bool bResetOptionsListFunctionInPerformed= false;
+};
diff --git a/Engine/Plugins/Runtime/nDisplay/Source/DisplayClusterMoviePipelineEditor/Private/Details/DisplayClusterMoviePipelineEditorSettingsCustomization.cpp b/Engine/Plugins/Runtime/nDisplay/Source/DisplayClusterMoviePipelineEditor/Private/Details/DisplayClusterMoviePipelineEditorSettingsCustomization.cpp
new file mode 100644
--- /dev/null
+++ b/Engine/Plugins/Runtime/nDisplay/Source/DisplayClusterMoviePipelineEditor/Private/Details/DisplayClusterMoviePipelineEditorSettingsCustomization.cpp
@@ -0,0 +1,101 @@
+// Copyright Epic Games, Inc. All Rights Reserved.
+
+#include "DisplayClusterMoviePipelineEditorSettingsCustomization.h"
+
+#include "DisplayClusterRootActor.h"
+#include "DisplayClusterConfigurationTypes.h"
+#include "Widgets/SDisplayClusterMoviePipelineEditorSearchableComboBox.h"
+
+#include "DisplayClusterMoviePipelineSettings.h"
+
+#include "DetailLayoutBuilder.h"
+#include "IDetailChildrenBuilder.h"
+#include "PropertyCustomizationHelpers.h"
+#include "PropertyHandle.h"
+#include "Widgets/Text/STextBlock.h"
+
+#define LOCTEXT_NAMESPACE "DisplayClusterMoviePipelineEditorSettingsCustomization"
+
+//////////////////////////////////////////////////////////////////////////////////////////////
+// FDisplayClusterMoviePipelineEditorSettingsCustomization
+//////////////////////////////////////////////////////////////////////////////////////////////
+void FDisplayClusterMoviePipelineEditorSettingsCustomization::Initialize(const TSharedRef<IPropertyHandle>& InPropertyHandle, IPropertyTypeCustomizationUtils& CustomizationUtils)
+{
+	FDisplayClusterMoviePipelineEditorBaseTypeCustomization::Initialize(InPropertyHandle, CustomizationUtils);
+
+	TSharedPtr<IPropertyHandle> DCRootActorHandle = GET_CHILD_HANDLE(FDisplayClusterMoviePipelineConfiguration, DCRootActor);
+	check(DCRootActorHandle.IsValid());
+
+	DCRootActorHandle->SetOnPropertyValueChanged(FSimpleDelegate::CreateSP(this, &FDisplayClusterMoviePipelineEditorSettingsCustomization::OnRootActorReferencedPropertyValueChanged));
+
+	TSharedPtr<IPropertyHandle> AllowedViewportNamesListHandle = GET_CHILD_HANDLE(FDisplayClusterMoviePipelineConfiguration, AllowedViewportNamesList);
+	check(AllowedViewportNamesListHandle);
+	AllowedViewportNamesListHandle->SetOnPropertyValueChanged(FSimpleDelegate::CreateSP(this, &FDisplayClusterMoviePipelineEditorSettingsCustomization::OnAllowedViewportNamesListHandleChanged));
+
+	NodeSelection = MakeShared<FDisplayClusterMoviePipelineEditorNodeSelection>(FDisplayClusterMoviePipelineEditorNodeSelection::EOperationMode::Viewports, DCRootActorHandle, AllowedViewportNamesListHandle);
+}
+
+void FDisplayClusterMoviePipelineEditorSettingsCustomization::OnAllowedViewportNamesListHandleChanged()
+{
+	if (NodeSelection.IsValid())
+	{
+		NodeSelection->ResetOptionsList();
+	}
+}
+void FDisplayClusterMoviePipelineEditorSettingsCustomization::OnRootActorReferencedPropertyValueChanged()
+{
+	// When a referenced property is changed, we have to trigger layout refresh
+	if (NodeSelection.IsValid())
+	{
+		NodeSelection->ResetOptionsList();
+	}
+}
+
+void FDisplayClusterMoviePipelineEditorSettingsCustomization::SetHeader(const TSharedRef<IPropertyHandle>& InPropertyHandle, FDetailWidgetRow& InHeaderRow, IPropertyTypeCustomizationUtils& CustomizationUtils)
+{
+	FText ElementTooltip = FText::GetEmpty();
+
+	TSharedPtr<IPropertyHandle> ParentArrayHandle = InPropertyHandle->GetParentHandle();
+	if (ParentArrayHandle.IsValid() && ParentArrayHandle->IsValidHandle())
+	{
+		ElementTooltip = ParentArrayHandle->GetPropertyDisplayName();
+	}
+
+	InHeaderRow.NameContent()
+		[
+			InPropertyHandle->CreatePropertyNameWidget(FText::GetEmpty(), ElementTooltip)
+		];
+}
+
+void FDisplayClusterMoviePipelineEditorSettingsCustomization::SetChildren(const TSharedRef<IPropertyHandle>& InPropertyHandle, IDetailChildrenBuilder& InChildBuilder, IPropertyTypeCustomizationUtils& CustomizationUtils)
+{
+	TSharedPtr<IPropertyHandle> DCRootActorHandle = GET_CHILD_HANDLE(FDisplayClusterMoviePipelineConfiguration, DCRootActor);
+	TSharedPtr<IPropertyHandle> UseViewportResolutionsHandle = GET_CHILD_HANDLE(FDisplayClusterMoviePipelineConfiguration, bUseViewportResolutions);
+	TSharedPtr<IPropertyHandle> IsRenderAllViewportsHandle = GET_CHILD_HANDLE(FDisplayClusterMoviePipelineConfiguration, bRenderAllViewports);
+	TSharedPtr<IPropertyHandle> ArrayHandle = GET_CHILD_HANDLE(FDisplayClusterMoviePipelineConfiguration, AllowedViewportNamesList);
+
+	InChildBuilder.AddProperty(DCRootActorHandle.ToSharedRef());
+	InChildBuilder.AddProperty(UseViewportResolutionsHandle.ToSharedRef());
+	InChildBuilder.AddProperty(IsRenderAllViewportsHandle.ToSharedRef());
+
+	if (NodeSelection.IsValid())
+	{
+		const TAttribute<bool> IsEnabledEditCondition = TAttribute<bool>::Create([this, IsRenderAllViewportsHandle]()
+			{
+				bool bCond1 = false;
+				IsRenderAllViewportsHandle->GetValue(bCond1);
+				return bCond1 == false;
+			});
+
+
+		ArrayHandle->SetInstanceMetaData(FDisplayClusterMoviePipelineEditorNodeSelection::NAME_ElementToolTip, ArrayHandle->GetPropertyDisplayName().ToString());
+
+		const TAttribute<bool> IsArrayHandleEnabledEditCondition = TAttribute<bool>::Create([this]()
+			{
+				return true;
+			});
+		NodeSelection->IsEnabled(IsArrayHandleEnabledEditCondition);
+		NodeSelection->CreateArrayBuilder(ArrayHandle.ToSharedRef(), InChildBuilder);
+	}
+}
+#undef LOCTEXT_NAMESPACE
diff --git a/Engine/Plugins/Runtime/nDisplay/Source/DisplayClusterMoviePipelineEditor/Private/Details/DisplayClusterMoviePipelineEditorSettingsCustomization.h b/Engine/Plugins/Runtime/nDisplay/Source/DisplayClusterMoviePipelineEditor/Private/Details/DisplayClusterMoviePipelineEditorSettingsCustomization.h
new file mode 100644
--- /dev/null
+++ b/Engine/Plugins/Runtime/nDisplay/Source/DisplayClusterMoviePipelineEditor/Private/Details/DisplayClusterMoviePipelineEditorSettingsCustomization.h
@@ -0,0 +1,35 @@
+// Copyright Epic Games, Inc. All Rights Reserved.
+
+#pragma once
+
+#include "CoreMinimal.h"
+#include "DisplayClusterMoviePipelineEditorNodeSelection.h"
+
+/** Type customization for the FDisplayClusterMoviePipelineConfiguration type, to customize the AllowedViewportNamesList list property */
+class FDisplayClusterMoviePipelineEditorSettingsCustomization final
+	: public FDisplayClusterMoviePipelineEditorBaseTypeCustomization
+{
+public:
+	static TSharedRef<IPropertyTypeCustomization> MakeInstance()
+	{
+		return MakeShared<FDisplayClusterMoviePipelineEditorSettingsCustomization>();
+	}
+
+	virtual ~FDisplayClusterMoviePipelineEditorSettingsCustomization() override
+	{
+		NodeSelection.Reset();
+	}
+
+protected:
+	//~ FDisplayClusterConfiguratorBaseTypeCustomization overrides begin
+	virtual void Initialize(const TSharedRef<IPropertyHandle>& InPropertyHandle, IPropertyTypeCustomizationUtils& CustomizationUtils) override;
+	virtual void SetHeader(const TSharedRef<IPropertyHandle>& InPropertyHandle, FDetailWidgetRow& InHeaderRow, IPropertyTypeCustomizationUtils& CustomizationUtils) override;
+	virtual void SetChildren(const TSharedRef<IPropertyHandle>& InPropertyHandle, IDetailChildrenBuilder& InChildBuilder, IPropertyTypeCustomizationUtils& CustomizationUtils) override;
+	//~ FDisplayClusterConfiguratorBaseTypeCustomization overrides end
+
+	void OnRootActorReferencedPropertyValueChanged();
+	void OnAllowedViewportNamesListHandleChanged();
+private:
+	/** The helper object used to generate the custom array widget to display in place of a string list for the AllowedViewportNamesList property */
+	TSharedPtr<FDisplayClusterMoviePipelineEditorNodeSelection> NodeSelection;
+};
diff --git a/Engine/Plugins/Runtime/nDisplay/Source/DisplayClusterMoviePipelineEditor/Private/DisplayClusterMoviePipelineEditorLog.cpp b/Engine/Plugins/Runtime/nDisplay/Source/DisplayClusterMoviePipelineEditor/Private/DisplayClusterMoviePipelineEditorLog.cpp
new file mode 100644
--- /dev/null
+++ b/Engine/Plugins/Runtime/nDisplay/Source/DisplayClusterMoviePipelineEditor/Private/DisplayClusterMoviePipelineEditorLog.cpp
@@ -0,0 +1,6 @@
+// Copyright Epic Games, Inc. All Rights Reserved.
+
+#include "DisplayClusterMoviePipelineEditorLog.h"
+
+// Plugin-wide log categories
+DEFINE_LOG_CATEGORY(LogDisplayClusterMoviePipelineEditor);
diff --git a/Engine/Plugins/Runtime/nDisplay/Source/DisplayClusterMoviePipelineEditor/Private/DisplayClusterMoviePipelineEditorLog.h b/Engine/Plugins/Runtime/nDisplay/Source/DisplayClusterMoviePipelineEditor/Private/DisplayClusterMoviePipelineEditorLog.h
new file mode 100644
--- /dev/null
+++ b/Engine/Plugins/Runtime/nDisplay/Source/DisplayClusterMoviePipelineEditor/Private/DisplayClusterMoviePipelineEditorLog.h
@@ -0,0 +1,8 @@
+// Copyright Epic Games, Inc. All Rights Reserved.
+
+#pragma once
+
+#include "CoreMinimal.h"
+
+// Plugin-wide log categories
+DECLARE_LOG_CATEGORY_EXTERN(LogDisplayClusterMoviePipelineEditor, Log, All);
diff --git a/Engine/Plugins/Runtime/nDisplay/Source/DisplayClusterMoviePipelineEditor/Private/DisplayClusterMoviePipelineEditorModule.cpp b/Engine/Plugins/Runtime/nDisplay/Source/DisplayClusterMoviePipelineEditor/Private/DisplayClusterMoviePipelineEditorModule.cpp
new file mode 100644
--- /dev/null
+++ b/Engine/Plugins/Runtime/nDisplay/Source/DisplayClusterMoviePipelineEditor/Private/DisplayClusterMoviePipelineEditorModule.cpp
@@ -0,0 +1,73 @@
+// Copyright Epic Games, Inc. All Rights Reserved.
+
+#include "DisplayClusterMoviePipelineEditorModule.h"
+#include "Details/DisplayClusterMoviePipelineEditorSettingsCustomization.h"
+
+#include "DisplayClusterMoviePipelineSettings.h"
+
+#include "AssetToolsModule.h"
+#include "IAssetTools.h"
+#include "ISettingsModule.h"
+#include "AssetRegistry/AssetRegistryModule.h"
+#include "AssetTypeCategories.h"
+#include "HAL/IConsoleManager.h"
+#include "Modules/ModuleManager.h"
+#include "ActorFactories/ActorFactoryBlueprint.h"
+
+#define REGISTER_PROPERTY_LAYOUT(PropertyType, CustomizationType) { \
+	const FName LayoutName = PropertyType::StaticStruct()->GetFName(); \
+	RegisteredPropertyLayoutNames.Add(LayoutName); \
+	PropertyModule.RegisterCustomPropertyTypeLayout(LayoutName, \
+		FOnGetPropertyTypeCustomizationInstance::CreateStatic(&CustomizationType::MakeInstance)); \
+}
+
+#define REGISTER_OBJECT_LAYOUT(ObjectType, CustomizationType) { \
+	const FName LayoutName = ObjectType::StaticClass()->GetFName(); \
+	RegisteredClassLayoutNames.Add(LayoutName); \
+	PropertyModule.RegisterCustomClassLayout(LayoutName, \
+		FOnGetDetailCustomizationInstance::CreateStatic(&CustomizationType::MakeInstance)); \
+}
+
+#define LOCTEXT_NAMESPACE "DisplayClusterMoviePipelineEditor"
+
+//////////////////////////////////////////////////////////////////////////////////////////
+// FDisplayClusterMoviePipelineEditorModule
+//////////////////////////////////////////////////////////////////////////////////////////
+void FDisplayClusterMoviePipelineEditorModule::StartupModule()
+{
+	RegisterCustomLayouts();
+}
+
+void FDisplayClusterMoviePipelineEditorModule::ShutdownModule()
+{
+	UnregisterCustomLayouts();
+}
+
+void FDisplayClusterMoviePipelineEditorModule::RegisterCustomLayouts()
+{
+	FPropertyEditorModule& PropertyModule = FModuleManager::LoadModuleChecked<FPropertyEditorModule>("PropertyEditor");
+
+	REGISTER_PROPERTY_LAYOUT(FDisplayClusterMoviePipelineConfiguration, FDisplayClusterMoviePipelineEditorSettingsCustomization);
+}
+
+void FDisplayClusterMoviePipelineEditorModule::UnregisterCustomLayouts()
+{
+	FPropertyEditorModule& PropertyModule = FModuleManager::LoadModuleChecked<FPropertyEditorModule>("PropertyEditor");
+
+	for (const FName& LayoutName : RegisteredClassLayoutNames)
+	{
+		PropertyModule.UnregisterCustomClassLayout(LayoutName);
+	}
+
+	for (const FName& LayoutName : RegisteredPropertyLayoutNames)
+	{
+		PropertyModule.UnregisterCustomPropertyTypeLayout(LayoutName);
+	}
+
+	RegisteredClassLayoutNames.Empty();
+	RegisteredPropertyLayoutNames.Empty();
+}
+
+IMPLEMENT_MODULE(FDisplayClusterMoviePipelineEditorModule, DisplayClusterMoviePipelineEditor);
+
+#undef LOCTEXT_NAMESPACE
diff --git a/Engine/Plugins/Runtime/nDisplay/Source/DisplayClusterMoviePipelineEditor/Private/DisplayClusterMoviePipelineEditorModule.h b/Engine/Plugins/Runtime/nDisplay/Source/DisplayClusterMoviePipelineEditor/Private/DisplayClusterMoviePipelineEditorModule.h
new file mode 100644
--- /dev/null
+++ b/Engine/Plugins/Runtime/nDisplay/Source/DisplayClusterMoviePipelineEditor/Private/DisplayClusterMoviePipelineEditorModule.h
@@ -0,0 +1,25 @@
+// Copyright Epic Games, Inc. All Rights Reserved.
+
+#pragma once
+
+#include "IDisplayClusterMoviePipelineEditor.h"
+
+/**
+ * Display Cluster Movie Pipeline Editor module
+ */
+class FDisplayClusterMoviePipelineEditorModule :
+	public IDisplayClusterMoviePipelineEditor
+{
+public:
+	//~ IModuleInterface interface
+	virtual void StartupModule() override;
+	virtual void ShutdownModule() override;
+
+private:
+	void RegisterCustomLayouts();
+	void UnregisterCustomLayouts();
+
+private:
+	TArray<FName> RegisteredClassLayoutNames;
+	TArray<FName> RegisteredPropertyLayoutNames;
+};
diff --git a/Engine/Plugins/Runtime/nDisplay/Source/DisplayClusterMoviePipelineEditor/Private/Widgets/SDisplayClusterMoviePipelineEditorSearchableComboBox.cpp b/Engine/Plugins/Runtime/nDisplay/Source/DisplayClusterMoviePipelineEditor/Private/Widgets/SDisplayClusterMoviePipelineEditorSearchableComboBox.cpp
new file mode 100644
--- /dev/null
+++ b/Engine/Plugins/Runtime/nDisplay/Source/DisplayClusterMoviePipelineEditor/Private/Widgets/SDisplayClusterMoviePipelineEditorSearchableComboBox.cpp
@@ -0,0 +1,67 @@
+// Copyright Epic Games, Inc. All Rights Reservekd.
+
+#include "SDisplayClusterMoviePipelineEditorSearchableComboBox.h"
+
+////////////////////////////////////////////////////////////////////////////////////////////////
+// SDisplayClusterMoviePipelineEditorSearchableComboBox
+////////////////////////////////////////////////////////////////////////////////////////////////
+void SDisplayClusterMoviePipelineEditorSearchableComboBox::Construct(const FArguments& InArgs)
+{
+	// Set up content
+	TSharedPtr<SWidget> ButtonContent = InArgs._Content.Widget;
+	if (InArgs._Content.Widget == SNullWidget::NullWidget)
+	{
+		SAssignNew(ButtonContent, STextBlock)
+			.Text(NSLOCTEXT("SDisplayClusterMoviePipelineEditorSearchableComboBox", "ContentWarning", "No Content Provided"))
+			.ColorAndOpacity(FLinearColor::Red);
+	}
+
+	DisplayClusterOptionsSource = InArgs._OptionsSource;
+
+	SSearchableComboBox::Construct(SSearchableComboBox::FArguments()
+			.ComboBoxStyle(InArgs._ComboBoxStyle)
+			.ButtonStyle(InArgs._ButtonStyle)
+			.ItemStyle(InArgs._ItemStyle)
+			.ContentPadding(InArgs._ContentPadding)
+			.ForegroundColor(InArgs._ForegroundColor)
+			.OptionsSource(InArgs._OptionsSource)
+			.OnSelectionChanged(InArgs._OnSelectionChanged)
+			.OnGenerateWidget(InArgs._OnGenerateWidget)
+			.OnComboBoxOpening(InArgs._OnComboBoxOpening)
+			.CustomScrollbar(InArgs._CustomScrollbar)
+			.InitiallySelectedItem(InArgs._InitiallySelectedItem)
+			.Method(InArgs._Method)
+			.MaxListHeight(InArgs._MaxListHeight)
+			.HasDownArrow(InArgs._HasDownArrow)
+			.SearchVisibility(InArgs._SearchVisibility)
+			.Content()
+			[
+				ButtonContent.ToSharedRef()
+			]
+		);
+}
+
+void SDisplayClusterMoviePipelineEditorSearchableComboBox::ResetOptionsSource(const TArray<TSharedPtr<FString>>* NewOptionsSource)
+{
+	if (NewOptionsSource && NewOptionsSource != DisplayClusterOptionsSource)
+	{
+		DisplayClusterOptionsSource = NewOptionsSource;
+		SetOptionsSource(NewOptionsSource);
+	}
+	
+	RefreshOptions();
+	
+	{
+		// Without this the selection is cleared and dropdown auto closed when clicking on the combo button next. Combo box bug?
+		
+		const TSharedPtr<FString> CurrentSelectedItem = GetSelectedItem();
+		if (CurrentSelectedItem.IsValid() && DisplayClusterOptionsSource && DisplayClusterOptionsSource->Contains(CurrentSelectedItem))
+		{
+			SetSelectedItem(CurrentSelectedItem);
+		}
+		else
+		{
+			ClearSelection();
+		}
+	}
+}
diff --git a/Engine/Plugins/Runtime/nDisplay/Source/DisplayClusterMoviePipelineEditor/Private/Widgets/SDisplayClusterMoviePipelineEditorSearchableComboBox.h b/Engine/Plugins/Runtime/nDisplay/Source/DisplayClusterMoviePipelineEditor/Private/Widgets/SDisplayClusterMoviePipelineEditorSearchableComboBox.h
new file mode 100644
--- /dev/null
+++ b/Engine/Plugins/Runtime/nDisplay/Source/DisplayClusterMoviePipelineEditor/Private/Widgets/SDisplayClusterMoviePipelineEditorSearchableComboBox.h
@@ -0,0 +1,95 @@
+// Copyright Epic Games, Inc. All Rights Reservekd.
+
+#pragma once
+
+#include "CoreMinimal.h"
+#include "SSearchableComboBox.h"
+
+/**
+ * A searchable text combo box
+ * Mostly copied from SearchableComboBox, but with extension for dynamically update the Option source list
+ */
+class SDisplayClusterMoviePipelineEditorSearchableComboBox
+	: public SSearchableComboBox
+{
+public:
+
+	SLATE_BEGIN_ARGS(SDisplayClusterMoviePipelineEditorSearchableComboBox)
+		: _Content()
+		, _ComboBoxStyle(&FAppStyle::Get().GetWidgetStyle<FComboBoxStyle>("ComboBox"))
+		, _ButtonStyle(nullptr)
+		, _ItemStyle(&FAppStyle::Get().GetWidgetStyle<FTableRowStyle>("ComboBox.Row"))
+		, _ContentPadding(_ComboBoxStyle->ContentPadding)
+		, _ForegroundColor(FSlateColor::UseStyle())
+		, _OptionsSource()
+		, _OnSelectionChanged()
+		, _OnGenerateWidget()
+		, _InitiallySelectedItem(nullptr)
+		, _Method()
+		, _MaxListHeight(450.0f)
+		, _HasDownArrow(true)
+		, _SearchVisibility()
+	{}
+
+	/** Slot for this button's content (optional) */
+	SLATE_DEFAULT_SLOT(FArguments, Content)
+
+		SLATE_STYLE_ARGUMENT(FComboBoxStyle, ComboBoxStyle)
+
+		/** The visual style of the button part of the combo box (overrides ComboBoxStyle) */
+		SLATE_STYLE_ARGUMENT(FButtonStyle, ButtonStyle)
+
+		SLATE_STYLE_ARGUMENT(FTableRowStyle, ItemStyle)
+
+		SLATE_ATTRIBUTE(FMargin, ContentPadding)
+		SLATE_ATTRIBUTE(FSlateColor, ForegroundColor)
+
+		SLATE_ARGUMENT(const TArray< TSharedPtr<FString> >*, OptionsSource)
+		SLATE_EVENT(FOnSelectionChanged, OnSelectionChanged)
+		SLATE_EVENT(FOnGenerateWidget, OnGenerateWidget)
+
+		/** Called when combo box is opened, before list is actually created */
+		SLATE_EVENT(FOnComboBoxOpening, OnComboBoxOpening)
+
+		/** The custom scrollbar to use in the ListView */
+		SLATE_ARGUMENT(TSharedPtr<SScrollBar>, CustomScrollbar)
+
+		/** The option that should be selected when the combo box is first created */
+		SLATE_ARGUMENT(TSharedPtr<FString>, InitiallySelectedItem)
+
+		SLATE_ARGUMENT(TOptional<EPopupMethod>, Method)
+
+		/** The max height of the combo box menu */
+		SLATE_ARGUMENT(float, MaxListHeight)
+
+		/**
+		 * When false, the down arrow is not generated and it is up to the API consumer
+		 * to make their own visual hint that this is a drop down.
+		 */
+		SLATE_ARGUMENT(bool, HasDownArrow)
+
+		/** Allow setting the visibility of the search box dynamically */
+		SLATE_ATTRIBUTE(EVisibility, SearchVisibility)
+
+	SLATE_END_ARGS()
+
+	/**
+	 * Construct the widget from a declaration
+	 *
+	 * @param InArgs   Declaration from which to construct the combo box
+	 */
+	void Construct(const FArguments& InArgs);
+
+	/**
+	 * Reset the source in the Combo Box List. The options source is intelligently refreshed based on if the Source is the same or not.
+	 *
+	 * @param NewOptionsSource The options data to use. If null the original will be used.
+	 */
+	void ResetOptionsSource(const TArray<TSharedPtr<FString>>* NewOptionsSource = nullptr);
+
+protected:
+	/** Local version of parent class private variable `OptionsSource` */
+	const TArray< TSharedPtr<FString> >* DisplayClusterOptionsSource;
+};
+
+
diff --git a/Engine/Plugins/Runtime/nDisplay/Source/DisplayClusterMoviePipelineEditor/Public/IDisplayClusterMoviePipelineEditor.h b/Engine/Plugins/Runtime/nDisplay/Source/DisplayClusterMoviePipelineEditor/Public/IDisplayClusterMoviePipelineEditor.h
new file mode 100644
--- /dev/null
+++ b/Engine/Plugins/Runtime/nDisplay/Source/DisplayClusterMoviePipelineEditor/Public/IDisplayClusterMoviePipelineEditor.h
@@ -0,0 +1,17 @@
+// Copyright Epic Games, Inc. All Rights Reserved.
+
+#pragma once
+
+#include "Modules/ModuleInterface.h"
+#include "Modules/ModuleManager.h"
+
+
+/**
+ * Display Cluster editor module interface
+ */
+class IDisplayClusterMoviePipelineEditor :
+	public IModuleInterface
+{
+public:
+	virtual ~IDisplayClusterMoviePipelineEditor() = default;
+};
diff --git a/Engine/Plugins/Runtime/nDisplay/nDisplay.uplugin b/Engine/Plugins/Runtime/nDisplay/nDisplay.uplugin
--- a/Engine/Plugins/Runtime/nDisplay/nDisplay.uplugin
+++ b/Engine/Plugins/Runtime/nDisplay/nDisplay.uplugin
@@ -177,6 +177,15 @@
 				"Win64",
 				"Linux"
 			]
+		},
+		{
+			"Name": "DisplayClusterMoviePipelineEditor",
+			"Type": "Editor",
+			"LoadingPhase": "Default",
+			"PlatformAllowList": [
+				"Win64",
+				"Linux"
+			]
 		}
 	],
 	"Plugins": [
