MoviePipeline: Fixed anti-aliasing not being editable via the Movie Render Queue UI, the code was incorrectly calling FSceneView::SetupAntiAliasingMethod despite the comment right before it specifically saying we shouldn't call it.

diff --git a/Engine/Plugins/MovieScene/MovieRenderPipeline/Source/MovieRenderPipelineRenderPasses/Private/MoviePipelineImagePassBase.cpp b/Engine/Plugins/MovieScene/MovieRenderPipeline/Source/MovieRenderPipelineRenderPasses/Private/MoviePipelineImagePassBase.cpp
--- a/Engine/Plugins/MovieScene/MovieRenderPipeline/Source/MovieRenderPipelineRenderPasses/Private/MoviePipelineImagePassBase.cpp
+++ b/Engine/Plugins/MovieScene/MovieRenderPipeline/Source/MovieRenderPipelineRenderPasses/Private/MoviePipelineImagePassBase.cpp
@@ -318,7 +318,6 @@ TSharedPtr<FSceneViewFamilyContext> UMoviePipelineImagePassBase::CalculateViewFa
 			View->AntiAliasingMethod = AAM_None;
 		}
 	}
-	View->SetupAntiAliasingMethod();
 
 	// Anti Aliasing
 	{
