MoviePipeline: Clarified the pop-up dialog text when trying to render a Queue with an unsaved map.

diff --git a/Engine/Plugins/MovieScene/MovieRenderPipeline/Source/MovieRenderPipelineEditor/Private/MoviePipelinePIEExecutor.cpp b/Engine/Plugins/MovieScene/MovieRenderPipeline/Source/MovieRenderPipelineEditor/Private/MoviePipelinePIEExecutor.cpp
--- a/Engine/Plugins/MovieScene/MovieRenderPipeline/Source/MovieRenderPipelineEditor/Private/MoviePipelinePIEExecutor.cpp
+++ b/Engine/Plugins/MovieScene/MovieRenderPipeline/Source/MovieRenderPipelineEditor/Private/MoviePipelinePIEExecutor.cpp
@@ -100,7 +100,7 @@ void UMoviePipelinePIEExecutor::Start(const UMoviePipelineExecutorJob* InJob)
 	{
 		if(!IsRenderingOffscreen())
 		{
-			FText FailureReason = LOCTEXT("UnsavedMapFailureDialog", "One or more jobs in the queue have an unsaved map as their target map. Maps must be saved at least once before rendering.");
+			FText FailureReason = LOCTEXT("UnsavedMapFailureDialog", "One or more jobs in the queue have an unsaved map as their target map. Maps must be saved at least once before rendering, and then the job must be manually updated to point to the newly saved map.");
 			FMessageDialog::Open(EAppMsgType::Ok, FailureReason);
 		}
 
