MoviePipeline: Fixed motion blur sometimes not being applied on the first frame after a camera cut for some objects. This appears to be a threading issue combined with an optimization which checks if the position hasn't moved and MRQ's "evaluate frame 0, then frame 0.99, then frame 0" can occasionally trip the early out as it sees the two frame 0's and thus generates no motion blur.

This disables the early out optimization.

diff --git a/Engine/Plugins/MovieScene/MovieRenderPipeline/Source/MovieRenderPipelineCore/Private/MoviePipelineGameOverrideSetting.cpp b/Engine/Plugins/MovieScene/MovieRenderPipeline/Source/MovieRenderPipelineCore/Private/MoviePipelineGameOverrideSetting.cpp
--- a/Engine/Plugins/MovieScene/MovieRenderPipeline/Source/MovieRenderPipelineCore/Private/MoviePipelineGameOverrideSetting.cpp
+++ b/Engine/Plugins/MovieScene/MovieRenderPipeline/Source/MovieRenderPipelineCore/Private/MoviePipelineGameOverrideSetting.cpp
@@ -124,6 +124,11 @@ void UMoviePipelineGameOverrideSetting::ApplyCVarSettings(const bool bOverrideVa
 
 	// Remove any minimum delta time requirements from Chaos Physics to ensure accuracy at high Temporal Sample counts
 	MOVIEPIPELINE_STORE_AND_OVERRIDE_CVAR_FLOAT(PreviousChaosImmPhysicsMinStepTime, TEXT("p.Chaos.ImmPhys.MinStepTime"), 0, bOverrideValues);
+
+	// MRQ's 0 -> 0.99 -> 0 evaluation for motion blur emulation can occasionally cause it to be detected as a redundant update and thus never updated
+	// which causes objects to render in the wrong position on the first frame (and without motion blur). This disables an optimization that detects
+	// the redundant updates so the update will get sent through anyways even though it thinks it's a duplicate (but it's not).
+	MOVIEPIPELINE_STORE_AND_OVERRIDE_CVAR_INT(PreviousSkipRedundantTransformUpdate, TEXT("r.SkipRedundantTransformUpdate"), 0, bOverrideValues);
 }
 
 void UMoviePipelineGameOverrideSetting::BuildNewProcessCommandLineArgsImpl(TArray<FString>& InOutUnrealURLParams, TArray<FString>& InOutCommandLineArgs, TArray<FString>& InOutDeviceProfileCvars, TArray<FString>& InOutExecCmds) const
@@ -216,4 +221,5 @@ void UMoviePipelineGameOverrideSetting::BuildNewProcessCommandLineArgsImpl(TArra
 	InOutDeviceProfileCvars.Add(FString::Printf(TEXT("r.VolumetricRenderTarget=%d"), 0));
 	InOutDeviceProfileCvars.Add(FString::Printf(TEXT("wp.Runtime.BlockOnSlowStreaming=%d"), 0));
 	InOutDeviceProfileCvars.Add(FString::Printf(TEXT("p.Chaos.ImmPhys.MinStepTime=%d"), 0));
+	InOutDeviceProfileCvars.Add(FString::Printf(TEXT("r.SkipRedundantTransformUpdate=%d"), 0));
 }
diff --git a/Engine/Plugins/MovieScene/MovieRenderPipeline/Source/MovieRenderPipelineCore/Public/MoviePipelineGameOverrideSetting.h b/Engine/Plugins/MovieScene/MovieRenderPipeline/Source/MovieRenderPipelineCore/Public/MoviePipelineGameOverrideSetting.h
--- a/Engine/Plugins/MovieScene/MovieRenderPipeline/Source/MovieRenderPipelineCore/Public/MoviePipelineGameOverrideSetting.h
+++ b/Engine/Plugins/MovieScene/MovieRenderPipeline/Source/MovieRenderPipelineCore/Public/MoviePipelineGameOverrideSetting.h
@@ -145,4 +145,5 @@ private:
 	int32 PreviousGeoCacheStreamerBlockTillFinish;
 #endif
 	float PreviousChaosImmPhysicsMinStepTime;
+	int32 PreviousSkipRedundantTransformUpdate;
 };
\ No newline at end of file
