MoviePipeline: Fixed cloth simulation not working if a render started long after the world was first loaded. This was caused by a mismatch between the render world time and the game world time, and some systems compare the last rendered time to the current game time as an optimization

diff --git a/Engine/Plugins/MovieScene/MovieRenderPipeline/Source/MovieRenderPipelineCore/Private/MoviePipeline.cpp b/Engine/Plugins/MovieScene/MovieRenderPipeline/Source/MovieRenderPipelineCore/Private/MoviePipeline.cpp
--- a/Engine/Plugins/MovieScene/MovieRenderPipeline/Source/MovieRenderPipelineCore/Private/MoviePipeline.cpp
+++ b/Engine/Plugins/MovieScene/MovieRenderPipeline/Source/MovieRenderPipelineCore/Private/MoviePipeline.cpp
@@ -314,6 +314,11 @@ void UMoviePipeline::Initialize(UMoviePipelineExecutorJob* InJob)
 	// delta time not generated by us) so we'll wait until the next engine frame to start rendering.
 	InitializationTime = FDateTime::UtcNow();
 
+	// We need to start our CachedOutputState at the time the Game World is currently. This way
+	// as MRQ counts up time, the render world stays in sync allowing comparisons between last rendered time
+	// and game world time.
+	CachedOutputState.TimeData.WorldSeconds = GetWorld()->GetTimeSeconds();
+
 	// Resolve the version number for each shot. If the version number comes before the shot name in
 	// the format string, all shots will land on the same 'global' version number (as this is done
 	// before any file writing, so all shots within this job will calculate the same value). If the
