MoviePipeline: Fixed a crash when calling FindOrGetDefaultSettingForShot when InMasterConfig was null. Fixed an issue where if you called Job->SetConfiguration(Job->GetConfiguration()) it would wipe out all of the settings from the job.

diff --git a/Engine/Plugins/MovieScene/MovieRenderPipeline/Source/MovieRenderPipelineCore/Private/MoviePipelineBlueprintLibrary.cpp b/Engine/Plugins/MovieScene/MovieRenderPipeline/Source/MovieRenderPipelineCore/Private/MoviePipelineBlueprintLibrary.cpp
--- a/Engine/Plugins/MovieScene/MovieRenderPipeline/Source/MovieRenderPipelineCore/Private/MoviePipelineBlueprintLibrary.cpp
+++ b/Engine/Plugins/MovieScene/MovieRenderPipeline/Source/MovieRenderPipelineCore/Private/MoviePipelineBlueprintLibrary.cpp
@@ -937,10 +937,13 @@ UMoviePipelineSetting* UMoviePipelineBlueprintLibrary::FindOrGetDefaultSettingFo
 	}
 
 	// If they didn't have a shot override, or the setting wasn't enabled, we'll check the master config.
-	UMoviePipelineSetting* Setting = InMasterConfig->FindSettingByClass(InSettingType);
-	if (Setting)
+	if (InMasterConfig)
 	{
-		return Setting->IsEnabled() ? Setting : InSettingType->GetDefaultObject<UMoviePipelineSetting>();
+		UMoviePipelineSetting* Setting = InMasterConfig->FindSettingByClass(InSettingType);
+		if (Setting && Setting->IsEnabled())
+		{
+			return Setting;
+		}
 	}
 
 	// If no one overrode it, then we return the default.
diff --git a/Engine/Plugins/MovieScene/MovieRenderPipeline/Source/MovieRenderPipelineCore/Private/MoviePipelineConfigBase.cpp b/Engine/Plugins/MovieScene/MovieRenderPipeline/Source/MovieRenderPipelineCore/Private/MoviePipelineConfigBase.cpp
--- a/Engine/Plugins/MovieScene/MovieRenderPipeline/Source/MovieRenderPipelineCore/Private/MoviePipelineConfigBase.cpp
+++ b/Engine/Plugins/MovieScene/MovieRenderPipeline/Source/MovieRenderPipelineCore/Private/MoviePipelineConfigBase.cpp
@@ -19,6 +19,13 @@ void UMoviePipelineConfigBase::RemoveSetting(UMoviePipelineSetting* InSetting)
 
 void UMoviePipelineConfigBase::CopyFrom(UMoviePipelineConfigBase* InConfig)
 {
+	// Don't allow copying ourself to ourself as we first wipe out our settings array
+	// and then copy from the other config's settings array which will now be empty.
+	if (InConfig == this)
+	{
+		return;
+	}
+
 #if WITH_EDITOR
 	Modify();
 #endif
diff --git a/Engine/Plugins/MovieScene/MovieRenderPipeline/Source/MovieRenderPipelineCore/Private/MoviePipelineMasterConfig.cpp b/Engine/Plugins/MovieScene/MovieRenderPipeline/Source/MovieRenderPipelineCore/Private/MoviePipelineMasterConfig.cpp
--- a/Engine/Plugins/MovieScene/MovieRenderPipeline/Source/MovieRenderPipelineCore/Private/MoviePipelineMasterConfig.cpp
+++ b/Engine/Plugins/MovieScene/MovieRenderPipeline/Source/MovieRenderPipelineCore/Private/MoviePipelineMasterConfig.cpp
@@ -255,6 +255,11 @@ void UMoviePipelineMasterConfig::CopyFrom(UMoviePipelineConfigBase* InConfig)
 {
 	Super::CopyFrom(InConfig);
 
+	if (InConfig == this)
+	{
+		return;
+	}
+
 	if (InConfig->IsA<UMoviePipelineMasterConfig>())
 	{
 		UMoviePipelineMasterConfig* MasterConfig = CastChecked<UMoviePipelineMasterConfig>(InConfig);
