MoviePipeline: Fixed an ensure that was tripping in scenarios where it did not need to.

diff --git a/Engine/Plugins/MovieScene/MovieRenderPipeline/Source/MovieRenderPipelineCore/Private/MoviePipelineTiming.cpp b/Engine/Plugins/MovieScene/MovieRenderPipeline/Source/MovieRenderPipelineCore/Private/MoviePipelineTiming.cpp
--- a/Engine/Plugins/MovieScene/MovieRenderPipeline/Source/MovieRenderPipelineCore/Private/MoviePipelineTiming.cpp
+++ b/Engine/Plugins/MovieScene/MovieRenderPipeline/Source/MovieRenderPipelineCore/Private/MoviePipelineTiming.cpp
@@ -217,7 +217,13 @@ void UMoviePipeline::TickProducingFrames()
 			FrameInfo.PrevSidecarViewLocations = FrameInfo.CurrSidecarViewLocations;
 			FrameInfo.PrevSidecarViewRotations = FrameInfo.CurrSidecarViewRotations;
 
-			ensureMsgf(bSuccess, TEXT("Failed to evaluate camera to create cur/prev camera locations. No camera actor found on Eval Tick: %d"), CurrentMasterSeqTick.Value);
+			// This warning is only relevant if they actually have a camera in Sequencer. If they don't,
+			// it'll be falling back to the LocalPlayerController anyways (above) so it's not actually
+			// the error we're trying to catch.
+			if (CurrentCameraCut->ShotInfo.SubSectionHierarchy->CameraCutSection.IsValid())
+			{
+				ensureMsgf(bSuccess, TEXT("Failed to evaluate camera to create cur/prev camera locations. No camera actor found on Eval Tick: %d"), CurrentMasterSeqTick.Value);
+			}
 		}
 
 		// We can safely fall through to the below states as they're OK to process the same frame we set up.
