VulkanRHI: Replace the readback texture ref with a FRHIGPUTextureReadback which handles proper fencing.

#rnx

diff --git a/Engine/Plugins/MovieScene/MovieRenderPipeline/Source/MovieRenderPipelineCore/Private/MoviePipelineSurfaceReader.cpp b/Engine/Plugins/MovieScene/MovieRenderPipeline/Source/MovieRenderPipelineCore/Private/MoviePipelineSurfaceReader.cpp
--- a/Engine/Plugins/MovieScene/MovieRenderPipeline/Source/MovieRenderPipelineCore/Private/MoviePipelineSurfaceReader.cpp
+++ b/Engine/Plugins/MovieScene/MovieRenderPipeline/Source/MovieRenderPipelineCore/Private/MoviePipelineSurfaceReader.cpp
@@ -28,7 +28,7 @@ FMoviePipelineSurfaceReader::~FMoviePipelineSurfaceReader()
 {
 	BlockUntilAvailable();
 
-	ReadbackTexture.SafeRelease();
+	ReadbackTexture = nullptr;
 }
 
 void FMoviePipelineSurfaceReader::Initialize()
@@ -40,17 +40,14 @@ void FMoviePipelineSurfaceReader::Initialize()
 
 void FMoviePipelineSurfaceReader::ResizeImpl(uint32 Width, uint32 Height)
 {
-	ReadbackTexture.SafeRelease();
+	ReadbackTexture = nullptr;
 	Size = FIntPoint((int32)Width, (int32)Height);
 
 	// TexCreate_CPUReadback is important to make this texture available for reading later.
 	ENQUEUE_RENDER_COMMAND(CreateCaptureFrameTexture)(
-		[Width, Height, SurfaceReader = SharedThis(this)](FRHICommandListImmediate& RHICmdList)
+		[SurfaceReader = SharedThis(this)](FRHICommandListImmediate& RHICmdList)
 		{
-			const FRHITextureCreateDesc Desc =
-				FRHITextureCreateDesc::Create2D(TEXT("FMoviePipelineSurfaceReader"), Width, Height, SurfaceReader->PixelFormat)
-				.SetFlags(ETextureCreateFlags::CPUReadback);
-			SurfaceReader->ReadbackTexture = RHICreateTexture(Desc);
+			SurfaceReader->ReadbackTexture = MakeUnique<FRHIGPUTextureReadback>(TEXT("FMoviePipelineSurfaceReader"));
 		});
 }
 
@@ -87,11 +84,11 @@ void FMoviePipelineSurfaceReader::ResolveSampleToReadbackTexture_RenderThread(co
 	FRHICommandListImmediate& RHICmdList = GetImmediateCommandList_ForRenderCommand();
 
 	// Retrieve a temporary render target that matches the destination surface type/size.
-	const FIntPoint TargetSize(ReadbackTexture->GetSizeX(), ReadbackTexture->GetSizeY());
+	const FIntPoint TargetSize(Size);
 
 	FPooledRenderTargetDesc OutputDesc = FPooledRenderTargetDesc::Create2DDesc(
 		TargetSize,
-		ReadbackTexture->GetFormat(),
+		PixelFormat,
 		FClearValueBinding::None,
 		TexCreate_None,
 		TexCreate_RenderTargetable,
@@ -159,7 +156,8 @@ void FMoviePipelineSurfaceReader::ResolveSampleToReadbackTexture_RenderThread(co
 	}
 	RHICmdList.EndRenderPass();
 
-	TransitionAndCopyTexture(RHICmdList, ResampleTexturePooledRenderTarget->GetRHI(), ReadbackTexture, {});
+	RHICmdList.Transition(FRHITransitionInfo(ResampleTexturePooledRenderTarget->GetRHI(), ERHIAccess::RTV, ERHIAccess::CopySrc));
+	ReadbackTexture->EnqueueCopy(RHICmdList, ResampleTexturePooledRenderTarget->GetRHI());
 }
 
 void FMoviePipelineSurfaceReader::CopyReadbackTexture_RenderThread(TUniqueFunction<void(TUniquePtr<FImagePixelData>&&)>&& InFunctionCallback, TSharedPtr<FImagePixelDataPayload, ESPMode::ThreadSafe> InFramePayload)
@@ -182,10 +180,9 @@ void FMoviePipelineSurfaceReader::CopyReadbackTexture_RenderThread(TUniqueFuncti
 		SCOPED_GPU_MASK(RHICmdList, GPUMask);
 #endif
 
-		void* ColorDataBuffer = nullptr;
-
 		int32 ActualSizeX = 0, ActualSizeY = 0;
-		RHICmdList.MapStagingSurface(ReadbackTexture, ColorDataBuffer, ActualSizeX, ActualSizeY);
+		void* ColorDataBuffer = ReadbackTexture->Lock(ActualSizeX, &ActualSizeY);
+
 		int32 ExpectedSizeX = Size.X;
 		int32 ExpectedSizeY = Size.Y;
 
@@ -249,7 +246,7 @@ void FMoviePipelineSurfaceReader::CopyReadbackTexture_RenderThread(TUniqueFuncti
 		}
 
 		// Enqueue the Unmap before we broadcast the resulting pixels, though the broadcast shouldn't do anything blocking.
-		RHICmdList.UnmapStagingSurface(ReadbackTexture);
+		ReadbackTexture->Unlock();
 
 		InFunctionCallback(MoveTemp(PixelData));
 
diff --git a/Engine/Plugins/MovieScene/MovieRenderPipeline/Source/MovieRenderPipelineCore/Public/MoviePipelineSurfaceReader.h b/Engine/Plugins/MovieScene/MovieRenderPipeline/Source/MovieRenderPipelineCore/Public/MoviePipelineSurfaceReader.h
--- a/Engine/Plugins/MovieScene/MovieRenderPipeline/Source/MovieRenderPipelineCore/Public/MoviePipelineSurfaceReader.h
+++ b/Engine/Plugins/MovieScene/MovieRenderPipeline/Source/MovieRenderPipelineCore/Public/MoviePipelineSurfaceReader.h
@@ -5,6 +5,7 @@
 #include "PixelFormat.h"
 #include "RHI.h"
 #include "MovieRenderPipelineDataTypes.h"
+#include "RHIGPUReadback.h"
 
 /**
 * When the GPU finishes rendering a view we can copy the resulting data back to the CPU. Because the
@@ -58,7 +59,7 @@ protected:
 	FEvent* AvailableEvent;
 
 	/** Texture used to store the resolved render target */
-	FTexture2DRHIRef ReadbackTexture;
+	TUniquePtr<FRHIGPUTextureReadback> ReadbackTexture;
 
 	/** The desired pixel format of the resolved textures */
 	EPixelFormat PixelFormat;
