MoviePipeline: Fixed an issue where the -MoviePipelineConfig argument couldn't be used with assets that existed outside of the "/Game/" directory (ie: project plugins).

diff --git a/Engine/Plugins/MovieScene/MovieRenderPipeline/Source/MovieRenderPipelineCore/Private/MovieRenderPipelineCommandLine.cpp b/Engine/Plugins/MovieScene/MovieRenderPipeline/Source/MovieRenderPipelineCore/Private/MovieRenderPipelineCommandLine.cpp
--- a/Engine/Plugins/MovieScene/MovieRenderPipeline/Source/MovieRenderPipelineCore/Private/MovieRenderPipelineCommandLine.cpp
+++ b/Engine/Plugins/MovieScene/MovieRenderPipeline/Source/MovieRenderPipelineCore/Private/MovieRenderPipelineCommandLine.cpp
@@ -202,9 +202,11 @@ bool FMovieRenderPipelineCoreModule::IsTryingToRenderMovieFromCommandLine(FStrin
 *		- Preset or Queue to use (A preset is required if using a Level Sequence above, passed via -MoviePipelineConfig="/Game/...")
 *		- Will render on current map
 *		ie: 
-*			"E:\SubwaySequencer\SubwaySequencer.uproject" subwaySequencer_P -game -LevelSequence="/Game/Sequencer/SubwaySequencerMASTER.SubwaySequencerMASTER" -MoviePipelineConfig="/Game/Cinematics/MoviePipeline/Presets/SmallTestPreset.SmallTestPreset" -windowed -resx=1280 -resy=720 -log -notexturestreaming
+*			"E:\SubwaySequencer\SubwaySequencer.uproject" subwaySequencer_P -game -LevelSequence="/Game/Sequencer/SubwaySequencerMASTER.SubwaySequencerMASTER"
+				-MoviePipelineConfig="/Game/Cinematics/MoviePipeline/Presets/SmallTestPreset.SmallTestPreset" -windowed -resx=1280 -resy=720 -log -notexturestreaming
 *		or:
-*			ie: "E:\SubwaySequencer\SubwaySequencer.uproject" subwaySequencer_P -game -MoviePipelineConfig="/Game/Cinematics/MoviePipeline/Presets/BigTestQueue.BigTestQueue" -windowed -resx=1280 -resy=720 -log -notexturestreaming
+*			ie: "E:\SubwaySequencer\SubwaySequencer.uproject" subwaySequencer_P -game -MoviePipelineConfig="/Game/Cinematics/MoviePipeline/Presets/BigTestQueue.BigTestQueue" 
+				-windowed -resx=1280 -resy=720 -log -notexturestreaming
 *
 *	Option 2: Advanced Custom Executor. 
 *		- Executor Class (Required, pass via -MoviePipelineLocalExecutorClass=/Script/MovieRenderPipelineCore.MoviePipelinePythonHostExecutor)
@@ -250,42 +252,13 @@ uint8 FMovieRenderPipelineCoreModule::ParseMovieRenderData(const FString& InSequ
 	}
 
 	// Now look for the configuration file to see how to render it.
-	if (InConfigAssetPath.StartsWith("/Game/"))
+	if (InConfigAssetPath.Len() > 0)
 	{
-		// Convert it to a soft object path and use that load to ensure it follows redirectors, etc.
-		FSoftObjectPath AssetPath = FSoftObjectPath(InConfigAssetPath);
-		if (UMoviePipelineQueue* AssetAsQueue = Cast<UMoviePipelineQueue>(AssetPath.TryLoad()))
-		{
-			OutQueue = AssetAsQueue;
-		}
-		else if (UMoviePipelineMasterConfig* AssetAsConfig = Cast<UMoviePipelineMasterConfig>(AssetPath.TryLoad()))
-		{
-			OutQueue = NewObject<UMoviePipelineQueue>();
-			UMoviePipelineExecutorJob* NewJob = OutQueue->AllocateNewJob(UMoviePipelineExecutorJob::StaticClass()); // Only the default job type is supported right now.
-			NewJob->Sequence = FSoftObjectPath(InSequenceAssetPath);
-			NewJob->SetConfiguration(AssetAsConfig);
-			UWorld* CurrentWorld = MoviePipeline::FindCurrentWorld();
-			if (CurrentWorld)
-			{
-				NewJob->Map = FSoftObjectPath(CurrentWorld);
-			}
-		}
-
-		if (!OutQueue)
-		{
-			UE_LOG(LogMovieRenderPipeline, Fatal, TEXT("Failed to find Pipeline Configuration asset to render. Please note that the /Content/ part of the on-disk structure is omitted. Looked for: %s"), *InConfigAssetPath);
-			UE_LOG(LogMovieRenderPipeline, Warning, TEXT("%s"), *ConfigAssetFormatString);
-			return MoviePipelineErrorCodes::NoConfig;
-		}
-	}
-	else if (InConfigAssetPath.Len() > 0)
-	{
-		// If they didn't pass a path that started with /Game/, we'll try to see if it is a manifest file.
 		if (InConfigAssetPath.EndsWith(FPackageName::GetTextAssetPackageExtension()))
 		{
 			OutQueue = UMoviePipelineBlueprintLibrary::LoadManifestFileFromString(InConfigAssetPath);
-			
-			if(!OutQueue)
+
+			if (!OutQueue)
 			{
 				UE_LOG(LogMovieRenderPipeline, Fatal, TEXT("Could not parse text asset package."));
 				return MoviePipelineErrorCodes::NoConfig;
@@ -293,7 +266,30 @@ uint8 FMovieRenderPipelineCoreModule::ParseMovieRenderData(const FString& InSequ
 		}
 		else
 		{
-			UE_LOG(LogMovieRenderPipeline, Fatal, TEXT("Unknown Config Asset Path. Path: %s"), *InConfigAssetPath);
+			// Convert it to a soft object path and use that load to ensure it follows redirectors, etc.
+			FSoftObjectPath AssetPath = FSoftObjectPath(InConfigAssetPath);
+			if (UMoviePipelineQueue* AssetAsQueue = Cast<UMoviePipelineQueue>(AssetPath.TryLoad()))
+			{
+				OutQueue = AssetAsQueue;
+			}
+			else if (UMoviePipelineMasterConfig* AssetAsConfig = Cast<UMoviePipelineMasterConfig>(AssetPath.TryLoad()))
+			{
+				OutQueue = NewObject<UMoviePipelineQueue>();
+				UMoviePipelineExecutorJob* NewJob = OutQueue->AllocateNewJob(UMoviePipelineExecutorJob::StaticClass()); // Only the default job type is supported right now.
+				NewJob->Sequence = FSoftObjectPath(InSequenceAssetPath);
+				NewJob->SetConfiguration(AssetAsConfig);
+				UWorld* CurrentWorld = MoviePipeline::FindCurrentWorld();
+				if (CurrentWorld)
+				{
+					NewJob->Map = FSoftObjectPath(CurrentWorld);
+				}
+			}
+		}
+
+		if (!OutQueue)
+		{
+			UE_LOG(LogMovieRenderPipeline, Fatal, TEXT("Failed to find Pipeline Configuration asset to render. Please note that the /Content/ part of the on-disk structure should be omitted. Looked for: %s"), *InConfigAssetPath);
+			UE_LOG(LogMovieRenderPipeline, Warning, TEXT("%s"), *ConfigAssetFormatString);
 			return MoviePipelineErrorCodes::NoConfig;
 		}
 	}
