MoviePipeline: Fixed an issue where overlapped offset was calculated based on uninitialized memory which resulted in broken ObjectId passes.

diff --git a/Engine/Plugins/MovieScene/MovieRenderPipeline/Source/MovieRenderPipelineCore/Private/MoviePipelineUtils.cpp b/Engine/Plugins/MovieScene/MovieRenderPipeline/Source/MovieRenderPipelineCore/Private/MoviePipelineUtils.cpp
--- a/Engine/Plugins/MovieScene/MovieRenderPipeline/Source/MovieRenderPipelineCore/Private/MoviePipelineUtils.cpp
+++ b/Engine/Plugins/MovieScene/MovieRenderPipeline/Source/MovieRenderPipelineCore/Private/MoviePipelineUtils.cpp
@@ -915,7 +915,7 @@ namespace UE
 				const UMoviePipelineHighResSetting* HighResSettings = Cast<const UMoviePipelineHighResSetting>(UMoviePipelineBlueprintLibrary::FindOrGetDefaultSettingForShot(UMoviePipelineHighResSetting::StaticClass(), InMasterConfig, InPipelineExecutorShot));
 				check(HighResSettings);
 
-				FIntPoint BackbufferResolution = FIntPoint(FMath::CeilToInt((float)InEffectiveOutputResolution.X / (float)InRenderPassMetrics.OriginalTileCounts.X), FMath::CeilToInt((float)InEffectiveOutputResolution.Y / (float)InRenderPassMetrics.OriginalTileCounts.Y));
+				FIntPoint BackbufferResolution = FIntPoint(FMath::CeilToInt((float)InEffectiveOutputResolution.X / (float)OutRenderPassMetrics.OriginalTileCounts.X), FMath::CeilToInt((float)InEffectiveOutputResolution.Y / (float)OutRenderPassMetrics.OriginalTileCounts.Y));
 				const FIntPoint TileResolution = BackbufferResolution;
 
 				// Apply size padding.
@@ -923,7 +923,7 @@ namespace UE
 
 
 				OutRenderPassMetrics.TileSize = TileResolution;
-				OutRenderPassMetrics.ProjectionMatrixJitterAmount = FVector2D((float)(InRenderPassMetrics.SpatialShiftX) * 2.0f / BackbufferResolution.X, (float)InRenderPassMetrics.SpatialShiftY * -2.0f / BackbufferResolution.Y);
+				OutRenderPassMetrics.ProjectionMatrixJitterAmount = FVector2D((float)(OutRenderPassMetrics.SpatialShiftX) * 2.0f / BackbufferResolution.X, (float)OutRenderPassMetrics.SpatialShiftY * -2.0f / BackbufferResolution.Y);
 
 				OutRenderPassMetrics.BackbufferSize = BackbufferResolution;
 				OutRenderPassMetrics.EffectiveOutputResolution = InEffectiveOutputResolution;
@@ -931,12 +931,12 @@ namespace UE
 				{
 					OutRenderPassMetrics.OverlappedPad = FIntPoint(FMath::CeilToInt(TileResolution.X * HighResSettings->OverlapRatio),
 						FMath::CeilToInt(TileResolution.Y * HighResSettings->OverlapRatio));
-					OutRenderPassMetrics.OverlappedOffset = FIntPoint(InRenderPassMetrics.TileIndexes.X * TileResolution.X - InRenderPassMetrics.OverlappedPad.X,
-						InRenderPassMetrics.TileIndexes.Y * TileResolution.Y - InRenderPassMetrics.OverlappedPad.Y);
+					OutRenderPassMetrics.OverlappedOffset = FIntPoint(OutRenderPassMetrics.TileIndexes.X * TileResolution.X - OutRenderPassMetrics.OverlappedPad.X,
+						OutRenderPassMetrics.TileIndexes.Y * TileResolution.Y - OutRenderPassMetrics.OverlappedPad.Y);
 
 					// Move the final render by this much in the accumulator to counteract the offset put into the view matrix.
 					// Note that when bAllowSpatialJitter is false, SpatialShiftX/Y will always be zero.
-					OutRenderPassMetrics.OverlappedSubpixelShift = FVector2D(0.5f - InRenderPassMetrics.SpatialShiftX, 0.5f - InRenderPassMetrics.SpatialShiftY);
+					OutRenderPassMetrics.OverlappedSubpixelShift = FVector2D(0.5f - OutRenderPassMetrics.SpatialShiftX, 0.5f - OutRenderPassMetrics.SpatialShiftY);
 				}
 			}
 
