Sequencer: Remove Master from Sequencer

diff --git a/Engine/Plugins/MovieScene/LevelSequenceEditor/Source/LevelSequenceEditor/Private/Misc/LevelSequenceEditorHelpers.cpp b/Engine/Plugins/MovieScene/LevelSequenceEditor/Source/LevelSequenceEditor/Private/Misc/LevelSequenceEditorHelpers.cpp
--- a/Engine/Plugins/MovieScene/LevelSequenceEditor/Source/LevelSequenceEditor/Private/Misc/LevelSequenceEditorHelpers.cpp
+++ b/Engine/Plugins/MovieScene/LevelSequenceEditor/Source/LevelSequenceEditor/Private/Misc/LevelSequenceEditorHelpers.cpp
@@ -42,18 +42,18 @@
 
 #define LOCTEXT_NAMESPACE "LevelSequenceEditorHelpers"
 
-TWeakPtr<SWindow> MasterSequenceSettingsWindow;
+TWeakPtr<SWindow> LevelSequenceWithShotsSettingsWindow;
 
-class SMasterSequenceSettings : public SCompoundWidget, public FGCObject
+class SLevelSequenceWithShotsSettings : public SCompoundWidget, public FGCObject
 {
-	SLATE_BEGIN_ARGS(SMasterSequenceSettings) {}
-		SLATE_ARGUMENT(ULevelSequenceMasterSequenceSettings*, MasterSequenceSettings)
+	SLATE_BEGIN_ARGS(SLevelSequenceWithShotsSettings) {}
+		SLATE_ARGUMENT(ULevelSequenceWithShotsSettings*, LevelSequenceWithShotsSettings)
 		SLATE_ARGUMENT(UMovieSceneToolsProjectSettings*, ToolsProjectSettings)
 	SLATE_END_ARGS()
 
 	void Construct(const FArguments& InArgs)
 	{
-		const ULevelSequenceMasterSequenceSettings* LevelSequenceSettings = GetDefault<ULevelSequenceMasterSequenceSettings>();
+		const ULevelSequenceWithShotsSettings* LevelSequenceSettings = GetDefault<ULevelSequenceWithShotsSettings>();
 
 		FPropertyEditorModule& PropertyEditor = FModuleManager::LoadModuleChecked<FPropertyEditorModule>("PropertyEditor");
 
@@ -64,7 +64,7 @@ class SMasterSequenceSettings : public SCompoundWidget, public FGCObject
 		Details1ViewArgs.bShowOptions = false;
 		Details1ViewArgs.bAllowFavoriteSystem = false;
 		Details1ViewArgs.NameAreaSettings = FDetailsViewArgs::HideNameArea;
-		Details1ViewArgs.ViewIdentifier = "MasterSequenceSettings";
+		Details1ViewArgs.ViewIdentifier = "LevelSequenceWithShotsSettings";
 
 		Details1View = PropertyEditor.CreateDetailView(Details1ViewArgs);
 
@@ -105,7 +105,7 @@ class SMasterSequenceSettings : public SCompoundWidget, public FGCObject
 			.Padding(10.f)
 			[
 				SNew(STextBlock)
-				.Text(this, &SMasterSequenceSettings::GetMasterSequenceFullPath)
+				.Text(this, &SLevelSequenceWithShotsSettings::GetFullPath)
 			]
 
 			+ SVerticalBox::Slot()
@@ -114,7 +114,7 @@ class SMasterSequenceSettings : public SCompoundWidget, public FGCObject
 			.Padding(5.f)
 			[
 				SAssignNew(ErrorText, STextBlock)
-				.Text(this, &SMasterSequenceSettings::GetErrorText)
+				.Text(this, &SLevelSequenceWithShotsSettings::GetErrorText)
 				.TextStyle(FAppStyle::Get(), TEXT("Log.Warning"))
 			]
 
@@ -124,14 +124,14 @@ class SMasterSequenceSettings : public SCompoundWidget, public FGCObject
 			.Padding(5.f)
 			[
 				SNew(SPrimaryButton)
-				.Text(LOCTEXT("CreateMasterSequence", "Create Master Sequence"))
-				.OnClicked(this, &SMasterSequenceSettings::OnCreateMasterSequence)
+				.Text(LOCTEXT("CreateLevelSequenceWithShots", "Create Level Sequence with Shots"))
+				.OnClicked(this, &SLevelSequenceWithShotsSettings::OnLevelCreateSequenceWithShots)
 			]
 		];
 
-		if (InArgs._MasterSequenceSettings)
+		if (InArgs._LevelSequenceWithShotsSettings)
 		{
-			SetMasterSequenceSettings(InArgs._MasterSequenceSettings);
+			SetLevelSequenceWithShotsSettings(InArgs._LevelSequenceWithShotsSettings);
 		}
 		if (InArgs._ToolsProjectSettings)
 		{
@@ -139,11 +139,11 @@ class SMasterSequenceSettings : public SCompoundWidget, public FGCObject
 		}
 	}
 
-	void SetMasterSequenceSettings(ULevelSequenceMasterSequenceSettings* InMasterSequenceSettings)
+	void SetLevelSequenceWithShotsSettings(ULevelSequenceWithShotsSettings* InLevelSequenceWithShotsSettings)
 	{
-		MasterSequenceSettings = InMasterSequenceSettings;
+		LevelSequenceWithShotsSettings = InLevelSequenceWithShotsSettings;
 
-		Details1View->SetObject(InMasterSequenceSettings);
+		Details1View->SetObject(InLevelSequenceWithShotsSettings);
 	}
 
 	void SetToolsProjectSettings(UMovieSceneToolsProjectSettings* InToolsProjectSettings)
@@ -155,23 +155,23 @@ class SMasterSequenceSettings : public SCompoundWidget, public FGCObject
 
 	virtual void AddReferencedObjects( FReferenceCollector& Collector ) override
 	{
-		Collector.AddReferencedObject(MasterSequenceSettings);
+		Collector.AddReferencedObject(LevelSequenceWithShotsSettings);
 		Collector.AddReferencedObject(ToolsProjectSettings);
 	}
 	virtual FString GetReferencerName() const override
 	{
-		return TEXT("SMasterSequenceSettings");
+		return TEXT("SLevelSequenceWithShotsSettings");
 	}
 
 private:
 
-	FText GetMasterSequenceFullPath() const
+	FText GetFullPath() const
 	{
-		const ULevelSequenceMasterSequenceSettings* LevelSequenceSettings = GetDefault<ULevelSequenceMasterSequenceSettings>();
-		FString FullPath = LevelSequenceSettings->MasterSequenceBasePath.Path;
-		FullPath /= LevelSequenceSettings->MasterSequenceName;
-		FullPath /= LevelSequenceSettings->MasterSequenceName;
-		FullPath += LevelSequenceSettings->MasterSequenceSuffix;
+		const ULevelSequenceWithShotsSettings* LevelSequenceSettings = GetDefault<ULevelSequenceWithShotsSettings>();
+		FString FullPath = LevelSequenceSettings->BasePath.Path;
+		FullPath /= LevelSequenceSettings->Name;
+		FullPath /= LevelSequenceSettings->Name;
+		FullPath += LevelSequenceSettings->Suffix;
 		FullPath += TEXT(".uasset");
 		return FText::FromString(FullPath);
 	}
@@ -180,40 +180,40 @@ private:
 	{			
 		FAssetRegistryModule& AssetRegistryModule = FModuleManager::LoadModuleChecked<FAssetRegistryModule>(TEXT("AssetRegistry"));
 
-		const ULevelSequenceMasterSequenceSettings* LevelSequenceSettings = GetDefault<ULevelSequenceMasterSequenceSettings>();
-		FString FullPath = LevelSequenceSettings->MasterSequenceBasePath.Path;
-		FullPath /= LevelSequenceSettings->MasterSequenceName;
-		FullPath /= LevelSequenceSettings->MasterSequenceName;
-		FullPath += LevelSequenceSettings->MasterSequenceSuffix;
+		const ULevelSequenceWithShotsSettings* LevelSequenceSettings = GetDefault<ULevelSequenceWithShotsSettings>();
+		FString FullPath = LevelSequenceSettings->BasePath.Path;
+		FullPath /= LevelSequenceSettings->Name;
+		FullPath /= LevelSequenceSettings->Name;
+		FullPath += LevelSequenceSettings->Suffix;
 		FAssetData AssetData = AssetRegistryModule.Get().GetAssetByObjectPath(FSoftObjectPath(FullPath));
 		if (AssetData.IsValid())
 		{
-			return LOCTEXT("MasterSequenceExists", "Warning: Master Sequence Exists");
+			return LOCTEXT("LevelSequenceWithShotsExists", "Warning: Level Sequence with Shots Exists");
 		}
 
 		return FText::GetEmpty();
 	}
 
 
-	FReply OnCreateMasterSequence()
+	FReply OnLevelCreateSequenceWithShots()
 	{
-		const ULevelSequenceMasterSequenceSettings* LevelSequenceSettings = GetDefault<ULevelSequenceMasterSequenceSettings>();
+		const ULevelSequenceWithShotsSettings* LevelSequenceSettings = GetDefault<ULevelSequenceWithShotsSettings>();
 		
-		// Create a master sequence
-		FString MasterSequenceAssetName = LevelSequenceSettings->MasterSequenceName;
-		FString MasterSequencePackagePath = LevelSequenceSettings->MasterSequenceBasePath.Path;
-		MasterSequencePackagePath /= MasterSequenceAssetName;
-		MasterSequenceAssetName += LevelSequenceSettings->MasterSequenceSuffix;
-
-		UObject* MasterSequenceAsset = LevelSequenceEditorHelpers::CreateLevelSequenceAsset(MasterSequenceAssetName, MasterSequencePackagePath);
-		if (MasterSequenceAsset)
+		// Create a sequence with shots
+		FString AssetName = LevelSequenceSettings->Name;
+		FString PackagePath = LevelSequenceSettings->BasePath.Path;
+		PackagePath /= AssetName;
+		AssetName += LevelSequenceSettings->Suffix;
+
+		UObject* LevelSequenceWithShotsAsset = LevelSequenceEditorHelpers::CreateLevelSequenceAsset(AssetName, PackagePath);
+		if (LevelSequenceWithShotsAsset)
 		{
-			GEditor->GetEditorSubsystem<UAssetEditorSubsystem>()->OpenEditorForAsset(MasterSequenceAsset);
+			GEditor->GetEditorSubsystem<UAssetEditorSubsystem>()->OpenEditorForAsset(LevelSequenceWithShotsAsset);
 		
 			ILevelSequenceEditorModule& LevelSequenceEditorModule = FModuleManager::LoadModuleChecked<ILevelSequenceEditorModule>("LevelSequenceEditor");		
-			LevelSequenceEditorModule.OnMasterSequenceCreated().Broadcast(MasterSequenceAsset);
+			LevelSequenceEditorModule.OnLevelSequenceWithShotsCreated().Broadcast(LevelSequenceWithShotsAsset);
 			
-			MasterSequenceSettingsWindow.Pin()->RequestDestroyWindow();
+			LevelSequenceWithShotsSettingsWindow.Pin()->RequestDestroyWindow();
 		}
 
 		return FReply::Handled();
@@ -221,16 +221,16 @@ private:
 
 	TSharedPtr<IDetailsView> Details1View;
 	TSharedPtr<IDetailsView> Details2View;
-	TSharedPtr<SEditableTextBox> MasterSequencePathText;
+	TSharedPtr<SEditableTextBox> SequenceWithShotsPathText;
 	TSharedPtr<STextBlock> ErrorText;
-	ULevelSequenceMasterSequenceSettings* MasterSequenceSettings;
+	ULevelSequenceWithShotsSettings* LevelSequenceWithShotsSettings;
 	UMovieSceneToolsProjectSettings* ToolsProjectSettings;
 };
 	
 
-void LevelSequenceEditorHelpers::OpenMasterSequenceDialog(const TSharedRef<FTabManager>& TabManager)
+void LevelSequenceEditorHelpers::OpenLevelSequenceWithShotsDialog(const TSharedRef<FTabManager>& TabManager)
 {
-	TSharedPtr<SWindow> ExistingWindow = MasterSequenceSettingsWindow.Pin();
+	TSharedPtr<SWindow> ExistingWindow = LevelSequenceWithShotsSettingsWindow.Pin();
 	if (ExistingWindow.IsValid())
 	{
 		ExistingWindow->BringToFront();
@@ -238,7 +238,7 @@ void LevelSequenceEditorHelpers::OpenMasterSequenceDialog(const TSharedRef<FTabM
 	else
 	{
 		ExistingWindow = SNew(SWindow)
-			.Title( LOCTEXT("MasterSequenceSettingsTitle", "Master Sequence Settings") )
+			.Title( LOCTEXT("LevelSequenceWithShotsSettingsTitle", "Level Sequence with Shots Settings") )
 			.HasCloseButton(true)
 			.SupportsMaximize(false)
 			.SupportsMinimize(false)
@@ -256,16 +256,16 @@ void LevelSequenceEditorHelpers::OpenMasterSequenceDialog(const TSharedRef<FTabM
 		}
 	}
 
-	ULevelSequenceMasterSequenceSettings* MasterSequenceSettings = GetMutableDefault<ULevelSequenceMasterSequenceSettings>();
+	ULevelSequenceWithShotsSettings* LevelSequenceWithShotsSettings = GetMutableDefault<ULevelSequenceWithShotsSettings>();
 	UMovieSceneToolsProjectSettings* ToolsProjectSettings = GetMutableDefault<UMovieSceneToolsProjectSettings>();
 
 	ExistingWindow->SetContent(
-		SNew(SMasterSequenceSettings)
-		.MasterSequenceSettings(MasterSequenceSettings)
+		SNew(SLevelSequenceWithShotsSettings)
+		.LevelSequenceWithShotsSettings(LevelSequenceWithShotsSettings)
 		.ToolsProjectSettings(ToolsProjectSettings)
 	);
 
-	MasterSequenceSettingsWindow = ExistingWindow;
+	LevelSequenceWithShotsSettingsWindow = ExistingWindow;
 }
 
 UObject* LevelSequenceEditorHelpers::CreateLevelSequenceAsset(const FString& AssetName, const FString& PackagePath, UObject* AssetToDuplicate)
diff --git a/Engine/Plugins/MovieScene/LevelSequenceEditor/Source/LevelSequenceEditor/Private/Misc/LevelSequenceEditorHelpers.h b/Engine/Plugins/MovieScene/LevelSequenceEditor/Source/LevelSequenceEditor/Private/Misc/LevelSequenceEditorHelpers.h
--- a/Engine/Plugins/MovieScene/LevelSequenceEditor/Source/LevelSequenceEditor/Private/Misc/LevelSequenceEditorHelpers.h
+++ b/Engine/Plugins/MovieScene/LevelSequenceEditor/Source/LevelSequenceEditor/Private/Misc/LevelSequenceEditorHelpers.h
@@ -10,8 +10,8 @@ class LevelSequenceEditorHelpers
 {
 public:
 
-	/** Open dialog for creating a master sequence */
-	static void OpenMasterSequenceDialog(const TSharedRef<FTabManager>& TabManager);
+	/** Open dialog for creating a level sequence with shots */
+	static void OpenLevelSequenceWithShotsDialog(const TSharedRef<FTabManager>& TabManager);
 	
 	/** Create a level sequence asset given an asset name and package path */
 	static UObject* CreateLevelSequenceAsset(const FString& AssetName, const FString& PackagePath, UObject* AssetToDuplicate = nullptr);
diff --git a/Engine/Plugins/MovieScene/LevelSequenceEditor/Source/LevelSequenceEditor/Private/Misc/LevelSequenceEditorSettings.cpp b/Engine/Plugins/MovieScene/LevelSequenceEditor/Source/LevelSequenceEditor/Private/Misc/LevelSequenceEditorSettings.cpp
--- a/Engine/Plugins/MovieScene/LevelSequenceEditor/Source/LevelSequenceEditor/Private/Misc/LevelSequenceEditorSettings.cpp
+++ b/Engine/Plugins/MovieScene/LevelSequenceEditor/Source/LevelSequenceEditor/Private/Misc/LevelSequenceEditorSettings.cpp
@@ -11,12 +11,12 @@ ULevelSequenceEditorSettings::ULevelSequenceEditorSettings(const FObjectInitiali
 	bAutoBindToPIE      = true;
 }
 
-ULevelSequenceMasterSequenceSettings::ULevelSequenceMasterSequenceSettings(const FObjectInitializer& ObjectInitializer)
+ULevelSequenceWithShotsSettings::ULevelSequenceWithShotsSettings(const FObjectInitializer& ObjectInitializer)
 	: Super(ObjectInitializer)
-	, MasterSequenceName(TEXT("Sequence"))
-	, MasterSequenceSuffix(TEXT("Master"))
-	, MasterSequenceNumShots(5)
+	, Name(TEXT("Sequence"))
+	, Suffix(TEXT("Root"))
+	, NumShots(5)
 {
-	MasterSequenceBasePath.Path = TEXT("/Game/Cinematics/Sequences");
+	BasePath.Path = TEXT("/Game/Cinematics/Sequences");
 }
 
diff --git a/Engine/Plugins/MovieScene/LevelSequenceEditor/Source/LevelSequenceEditor/Private/Misc/LevelSequenceEditorSettings.h b/Engine/Plugins/MovieScene/LevelSequenceEditor/Source/LevelSequenceEditor/Private/Misc/LevelSequenceEditorSettings.h
--- a/Engine/Plugins/MovieScene/LevelSequenceEditor/Source/LevelSequenceEditor/Private/Misc/LevelSequenceEditorSettings.h
+++ b/Engine/Plugins/MovieScene/LevelSequenceEditor/Source/LevelSequenceEditor/Private/Misc/LevelSequenceEditorSettings.h
@@ -79,41 +79,41 @@ public:
 };
 
 /**
- * Level Sequence Master Sequence settings.
+ * Level Sequence With Shots Settings.
  */
 UCLASS(config=EditorPerProjectUserSettings)
-class ULevelSequenceMasterSequenceSettings
+class ULevelSequenceWithShotsSettings
 	: public UObject
 {
 	GENERATED_UCLASS_BODY()
 
 public:
-	/** Master sequence name. */
-	UPROPERTY(config, DisplayName="Name", EditAnywhere, Category=MasterSequence)
-	FString MasterSequenceName;
+	/** Sequence With Shots name. */
+	UPROPERTY(config, DisplayName="Name", EditAnywhere, Category=SequenceWithShots)
+	FString Name;
 
-	/** Master sequence suffix. */
-	UPROPERTY(config, DisplayName="Suffix", EditAnywhere, Category=MasterSequence)
-	FString MasterSequenceSuffix;
+	/** Sequence With Shots suffix. */
+	UPROPERTY(config, DisplayName="Suffix", EditAnywhere, Category=SequenceWithShots)
+	FString Suffix;
 
-	/** Master sequence path. */
-	UPROPERTY(config, DisplayName="Base Path", EditAnywhere, Category=MasterSequence, meta=(ContentDir))
-	FDirectoryPath MasterSequenceBasePath;
+	/** Sequence With Shots path. */
+	UPROPERTY(config, DisplayName="Base Path", EditAnywhere, Category=SequenceWithShots, meta=(ContentDir))
+	FDirectoryPath BasePath;
 
-	/** Master sequence number of shots. */
-	UPROPERTY(config, DisplayName="Number of Shots", EditAnywhere, Category=MasterSequence, meta = (UIMin = "1", UIMax = "100"))
-	uint32 MasterSequenceNumShots;
+	/** Sequence With Shots number of shots. */
+	UPROPERTY(config, DisplayName="Number of Shots", EditAnywhere, Category=SequenceWithShots, meta = (UIMin = "1", UIMax = "100"))
+	uint32 NumShots;
 
-	/** Master sequence level sequence to duplicate when creating shots. */
-	UPROPERTY(Transient, DisplayName="Sequence to Duplicate", EditAnywhere, Category=MasterSequence)
-	TLazyObjectPtr<class ULevelSequence> MasterSequenceLevelSequenceToDuplicate;
+	/** Sequence With Shots level sequence to duplicate when creating shots. */
+	UPROPERTY(Transient, DisplayName="Sequence to Duplicate", EditAnywhere, Category=SequenceWithShots)
+	TLazyObjectPtr<class ULevelSequence> SequenceToDuplicate;
 
 	/** Array of sub sequence names, each will result in a level sequence asset in the shot. */
-	UPROPERTY(EditAnywhere, BlueprintReadWrite, Category=MasterSequence)
+	UPROPERTY(EditAnywhere, BlueprintReadWrite, Category= SequenceWithShots)
 	TArray<FName> SubSequenceNames;
 
 	/** Whether to instance sub sequences based on the first created sub sequences. */
-	UPROPERTY(config, DisplayName="Instance Sub Sequences", EditAnywhere, Category=MasterSequence)
+	UPROPERTY(config, DisplayName="Instance Sub Sequences", EditAnywhere, Category=SequenceWithShots)
 	bool bInstanceSubSequences;
 };
 
diff --git a/Engine/Plugins/MovieScene/LevelSequenceEditor/Source/LevelSequenceEditor/Private/Styles/LevelSequenceEditorStyle.h b/Engine/Plugins/MovieScene/LevelSequenceEditor/Source/LevelSequenceEditor/Private/Styles/LevelSequenceEditorStyle.h
--- a/Engine/Plugins/MovieScene/LevelSequenceEditor/Source/LevelSequenceEditor/Private/Styles/LevelSequenceEditorStyle.h
+++ b/Engine/Plugins/MovieScene/LevelSequenceEditor/Source/LevelSequenceEditor/Private/Styles/LevelSequenceEditorStyle.h
@@ -41,8 +41,8 @@ public:
 		Set("LevelSequenceEditor.CreateNewLevelSequenceInLevel", new IMAGE_BRUSH_SVG("LevelSequence", Icon16x16));
 		Set("LevelSequenceEditor.CreateNewLevelSequenceInLevel.Small", new IMAGE_BRUSH_SVG("LevelSequence", Icon16x16));
 
-		Set("LevelSequenceEditor.CreateNewMasterSequenceInLevel", new IMAGE_BRUSH_SVG("MasterSequence", Icon16x16));
-		Set("LevelSequenceEditor.CreateNewMasterSequenceInLevel.Small", new IMAGE_BRUSH_SVG("MasterSequence", Icon16x16));
+		Set("LevelSequenceEditor.CreateNewLevelSequenceWithShotsInLevel", new IMAGE_BRUSH_SVG("LevelSequenceWithShots", Icon16x16));
+		Set("LevelSequenceEditor.CreateNewLevelSequenceWithShotsInLevel.Small", new IMAGE_BRUSH_SVG("LevelSequenceWithShots", Icon16x16));
 		
 		Set("LevelSequenceEditor.CinematicViewportPlayMarker", new IMAGE_BRUSH("CinematicViewportPlayMarker", FVector2D(11, 6)));
 		Set("LevelSequenceEditor.CinematicViewportRangeStart", new BORDER_BRUSH("CinematicViewportRangeStart", FMargin(1.f,.3f,0.f,.6f)));
diff --git a/Engine/Plugins/MovieScene/LevelSequenceEditor/Source/LevelSequenceEditor/Public/LevelSequenceEditorBlueprintLibrary.h b/Engine/Plugins/MovieScene/LevelSequenceEditor/Source/LevelSequenceEditor/Public/LevelSequenceEditorBlueprintLibrary.h
--- a/Engine/Plugins/MovieScene/LevelSequenceEditor/Source/LevelSequenceEditor/Public/LevelSequenceEditorBlueprintLibrary.h
+++ b/Engine/Plugins/MovieScene/LevelSequenceEditor/Source/LevelSequenceEditor/Public/LevelSequenceEditorBlueprintLibrary.h
@@ -52,7 +52,7 @@ public:
 	static bool OpenLevelSequence(ULevelSequence* LevelSequence);
 
 	/*
-	 * Get the currently opened root/master level sequence asset
+	 * Get the currently opened root level sequence asset
 	 */
 	UFUNCTION(BlueprintCallable, Category = "Level Sequence Editor")
 	static ULevelSequence* GetCurrentLevelSequence();
diff --git a/Engine/Plugins/MovieScene/LevelSequenceEditor/Source/LevelSequenceEditor/Public/LevelSequenceEditorModule.h b/Engine/Plugins/MovieScene/LevelSequenceEditor/Source/LevelSequenceEditor/Public/LevelSequenceEditorModule.h
--- a/Engine/Plugins/MovieScene/LevelSequenceEditor/Source/LevelSequenceEditor/Public/LevelSequenceEditorModule.h
+++ b/Engine/Plugins/MovieScene/LevelSequenceEditor/Source/LevelSequenceEditor/Public/LevelSequenceEditorModule.h
@@ -16,8 +16,8 @@ class ILevelSequenceEditorModule
 	: public IModuleInterface
 {
 public:
-	DECLARE_EVENT_OneParam(ILevelSequenceEditorModule, FOnMasterSequenceCreated, UObject*);
-	virtual FOnMasterSequenceCreated& OnMasterSequenceCreated() = 0;
+	DECLARE_EVENT_OneParam(ILevelSequenceEditorModule, FOnLevelSequenceWithShotsCreated, UObject*);
+	virtual FOnLevelSequenceWithShotsCreated& OnLevelSequenceWithShotsCreated() = 0;
 
 	DECLARE_MULTICAST_DELEGATE_OneParam(FAllowPlaybackContext, bool&);
 	virtual FAllowPlaybackContext& OnComputePlaybackContext() = 0;
diff --git a/Engine/Plugins/MovieScene/MovieRenderPipeline/Source/MovieRenderPipelineCore/Private/MoviePipeline.cpp b/Engine/Plugins/MovieScene/MovieRenderPipeline/Source/MovieRenderPipelineCore/Private/MoviePipeline.cpp
--- a/Engine/Plugins/MovieScene/MovieRenderPipeline/Source/MovieRenderPipelineCore/Private/MoviePipeline.cpp
+++ b/Engine/Plugins/MovieScene/MovieRenderPipeline/Source/MovieRenderPipelineCore/Private/MoviePipeline.cpp
@@ -31,7 +31,7 @@
 #include "DistanceFieldAtlas.h"
 #include "UObject/SoftObjectPath.h"
 #include "MoviePipelineOutputSetting.h"
-#include "MoviePipelineMasterConfig.h"
+#include "MoviePipelinePrimaryConfig.h"
 #include "MoviePipelineOutputSetting.h"
 #include "MoviePipelineBlueprintLibrary.h"
 #include "ImageWriteQueue.h"
@@ -108,7 +108,7 @@ void UMoviePipeline::ValidateSequenceAndSettings() const
 		IConsoleVariable* TonemapAlphaCVar = IConsoleManager::Get().FindConsoleVariable(TEXT("r.PostProcessing.PropagateAlpha"));
 		check(TonemapAlphaCVar);
 
-		TArray<UMoviePipelineRenderPass*> OutputSettings = GetPipelineMasterConfig()->FindSettings<UMoviePipelineRenderPass>();
+		TArray<UMoviePipelineRenderPass*> OutputSettings = GetPipelinePrimaryConfig()->FindSettings<UMoviePipelineRenderPass>();
 		bool bAnyOutputWantsAlpha = false;
 
 		for (const UMoviePipelineRenderPass* Output : OutputSettings)
@@ -153,7 +153,7 @@ void UMoviePipeline::Initialize(UMoviePipelineExecutorJob* InJob)
 		//  the UI is the one who updates the configuration from the preset.
 		if (InJob->GetPresetOrigin())
 		{
-			UE_LOG(LogMovieRenderPipeline, Log, TEXT("Job has a master preset specified, updating local master configuration from preset."));
+			UE_LOG(LogMovieRenderPipeline, Log, TEXT("Job has a primary preset specified, updating local primary configuration from preset."));
 			InJob->GetConfiguration()->CopyFrom(InJob->GetPresetOrigin());
 		}
 
@@ -206,7 +206,7 @@ void UMoviePipeline::Initialize(UMoviePipelineExecutorJob* InJob)
 
 	// Override the frame range on the target sequence if needed first before anyone has a chance to modify it.
 	{
-		UMoviePipelineOutputSetting* OutputSetting = GetPipelineMasterConfig()->FindSetting<UMoviePipelineOutputSetting>();
+		UMoviePipelineOutputSetting* OutputSetting = GetPipelinePrimaryConfig()->FindSetting<UMoviePipelineOutputSetting>();
 		if (OutputSetting->bUseCustomPlaybackRange)
 		{
 			FFrameNumber StartFrameTickResolution = FFrameRate::TransformTime(FFrameTime(FFrameNumber(OutputSetting->CustomStartFrame)), TargetSequence->GetMovieScene()->GetDisplayRate(), TargetSequence->GetMovieScene()->GetTickResolution()).FloorToFrame();
@@ -229,8 +229,8 @@ void UMoviePipeline::Initialize(UMoviePipelineExecutorJob* InJob)
 		}
 	}
 	
-	// Initialize all of our master config settings. Shot specific ones will be called for their appropriate shot.
-	for (UMoviePipelineSetting* Setting : GetPipelineMasterConfig()->GetAllSettings())
+	// Initialize all of our primary config settings. Shot specific ones will be called for their appropriate shot.
+	for (UMoviePipelineSetting* Setting : GetPipelinePrimaryConfig()->GetAllSettings())
 	{
 		Setting->OnMoviePipelineInitialized(this);
 	}
@@ -527,8 +527,8 @@ void UMoviePipeline::TransitionToState(const EMovieRenderPipelineState InNewStat
 			bInvalidTransition = false;
 			PipelineState = InNewState;
 
-			// Uninitialize our master config settings.
-			for (UMoviePipelineSetting* Setting : GetPipelineMasterConfig()->GetAllSettings())
+			// Uninitialize our primary config settings.
+			for (UMoviePipelineSetting* Setting : GetPipelinePrimaryConfig()->GetAllSettings())
 			{
 				Setting->OnMoviePipelineShutdown(this);
 			}
@@ -546,7 +546,7 @@ void UMoviePipeline::TransitionToState(const EMovieRenderPipelineState InNewStat
 				DebugWidget = nullptr;
 			}
 
-			for (UMoviePipelineOutputBase* Setting : GetPipelineMasterConfig()->GetOutputContainers())
+			for (UMoviePipelineOutputBase* Setting : GetPipelinePrimaryConfig()->GetOutputContainers())
 			{
 				Setting->OnPipelineFinished();
 			}
@@ -687,7 +687,7 @@ void UMoviePipeline::BeginFinalize()
 {
 	// Notify all of our output containers that we have finished producing and
 	// submitting all frames to them and that they should start any async flushes.
-	for (UMoviePipelineOutputBase* Container : GetPipelineMasterConfig()->GetOutputContainers())
+	for (UMoviePipelineOutputBase* Container : GetPipelinePrimaryConfig()->GetOutputContainers())
 	{
 		Container->BeginFinalize();
 	}
@@ -695,7 +695,7 @@ void UMoviePipeline::BeginFinalize()
 
 void UMoviePipeline::BeginExport()
 {
-	for (UMoviePipelineSetting* Setting : GetPipelineMasterConfig()->GetAllSettings())
+	for (UMoviePipelineSetting* Setting : GetPipelinePrimaryConfig()->GetAllSettings())
 	{
 		Setting->BeginExport();
 	}
@@ -711,7 +711,7 @@ void UMoviePipeline::TickFinalizeOutputContainers(const bool bInForceFinish)
 		bAllContainsFinishedProcessing = true;
 
 		// Ask the containers if they're all done processing.
-		for (UMoviePipelineOutputBase* Container : GetPipelineMasterConfig()->GetOutputContainers())
+		for (UMoviePipelineOutputBase* Container : GetPipelinePrimaryConfig()->GetOutputContainers())
 		{
 			bAllContainsFinishedProcessing &= Container->HasFinishedProcessing();
 		}
@@ -736,7 +736,7 @@ void UMoviePipeline::TickFinalizeOutputContainers(const bool bInForceFinish)
 		return;
 	}
 
-	for (UMoviePipelineOutputBase* Container : GetPipelineMasterConfig()->GetOutputContainers())
+	for (UMoviePipelineOutputBase* Container : GetPipelinePrimaryConfig()->GetOutputContainers())
 	{
 		// All containers have finished processing, final shutdown.
 		Container->Finalize();
@@ -761,7 +761,7 @@ void UMoviePipeline::TickPostFinalizeExport(const bool bInForceFinish)
 		bAllContainsFinishedProcessing = true;
 
 		// Ask the containers if they're all done processing.
-		for (UMoviePipelineSetting* Setting : GetPipelineMasterConfig()->GetAllSettings())
+		for (UMoviePipelineSetting* Setting : GetPipelinePrimaryConfig()->GetAllSettings())
 		{
 			bAllContainsFinishedProcessing &= Setting->HasFinishedExporting();
 		}
@@ -897,7 +897,7 @@ void UMoviePipeline::BuildShotListFromSequence()
 		// This info is read in ExpandShot so needs to be set first
 		Shot->ShotInfo.NumTemporalSamples = AntiAliasingSettings->TemporalSampleCount;
 		Shot->ShotInfo.NumSpatialSamples = AntiAliasingSettings->SpatialSampleCount;
-		Shot->ShotInfo.CachedFrameRate = GetPipelineMasterConfig()->GetEffectiveFrameRate(TargetSequence);
+		Shot->ShotInfo.CachedFrameRate = GetPipelinePrimaryConfig()->GetEffectiveFrameRate(TargetSequence);
 		Shot->ShotInfo.CachedTickResolution = TargetSequence->GetMovieScene()->GetTickResolution();
 		Shot->ShotInfo.CachedShotTickResolution = Shot->ShotInfo.CachedTickResolution;
 		if (Shot->ShotInfo.SubSectionHierarchy.IsValid() && Shot->ShotInfo.SubSectionHierarchy->MovieScene.IsValid())
@@ -928,7 +928,7 @@ void UMoviePipeline::BuildShotListFromSequence()
 
 		// When we expanded the shot above, it pushed the first/last camera cuts ranges to account for Handle Frames.
 		// We want to start rendering from the first handle frame. Shutter Timing is a fixed offset from this number.
-		Shot->ShotInfo.CurrentTickInMaster = Shot->ShotInfo.TotalOutputRangeMaster.GetLowerBoundValue();
+		Shot->ShotInfo.CurrentTickInRoot = Shot->ShotInfo.TotalOutputRangeRoot.GetLowerBoundValue();
 	}
 
 	// The active shot-list is a subset of the whole shot-list; The ShotInfo contains information about every range it detected to render
@@ -950,8 +950,8 @@ void UMoviePipeline::InitializeShot(UMoviePipelineExecutorShot* InShot)
 	// Set the new shot as the active shot. This enables the specified shot section and disables all other shot sections.
 	SetSoloShot(InShot);
 
-	// Loop through just our master settings and let them know which shot we're about to start.
-	for (UMoviePipelineSetting* Setting : GetPipelineMasterConfig()->GetAllSettings())
+	// Loop through just our primary settings and let them know which shot we're about to start.
+	for (UMoviePipelineSetting* Setting : GetPipelinePrimaryConfig()->GetAllSettings())
 	{
 		Setting->OnSetupForShot(InShot);
 	}
@@ -986,7 +986,7 @@ void UMoviePipeline::TeardownShot(UMoviePipelineExecutorShot* InShot)
 	}
 
 	// Notify our containers that the current shot has ended.
-	for (UMoviePipelineOutputBase* Container : GetPipelineMasterConfig()->GetOutputContainers())
+	for (UMoviePipelineOutputBase* Container : GetPipelinePrimaryConfig()->GetOutputContainers())
 	{
 		Container->OnShotFinished(InShot, IsFlushDiskWritesPerShot());
 	}
@@ -1000,8 +1000,8 @@ void UMoviePipeline::TeardownShot(UMoviePipelineExecutorShot* InShot)
 		}
 	}
 
-	// Loop through just our master settings and let them know which shot we're about to end.
-	for (UMoviePipelineSetting* Setting : GetPipelineMasterConfig()->GetAllSettings())
+	// Loop through just our primary settings and let them know which shot we're about to end.
+	for (UMoviePipelineSetting* Setting : GetPipelinePrimaryConfig()->GetAllSettings())
 	{
 		Setting->OnTeardownForShot(InShot);
 	}
@@ -1034,7 +1034,7 @@ void UMoviePipeline::TeardownShot(UMoviePipelineExecutorShot* InShot)
 		// We call the command line encoder as a special case here because it may want to modify the file list
 		// ie: if it deletes the file after use we probably don't want scripting looking for those files.
 		const bool bIncludeDisabledSettings = false;
-		UMoviePipelineCommandLineEncoder* Encoder = GetPipelineMasterConfig()->FindSetting<UMoviePipelineCommandLineEncoder>(bIncludeDisabledSettings);
+		UMoviePipelineCommandLineEncoder* Encoder = GetPipelinePrimaryConfig()->FindSetting<UMoviePipelineCommandLineEncoder>(bIncludeDisabledSettings);
 		if (Encoder)
 		{
 			const bool bInIsShotEncode = true;
@@ -1128,14 +1128,14 @@ void UMoviePipeline::ExpandShot(UMoviePipelineExecutorShot* InShot, const int32
 		RightDeltaFrames += 1;
 	}
 
-	// Check to see if the detected range was not aligned to a whole frame on the master.
-	const FFrameRate MasterDisplayRate = GetPipelineMasterConfig()->GetEffectiveFrameRate(TargetSequence);
-	FFrameTime StartTimeInMaster = FFrameRate::TransformTime(InShot->ShotInfo.TotalOutputRangeMaster.GetLowerBoundValue(), InShot->ShotInfo.CachedTickResolution, MasterDisplayRate);
-	if(bIsPrePass && StartTimeInMaster.GetSubFrame() != 0.f)
+	// Check to see if the detected range was not aligned to a whole frame on the root.
+	const FFrameRate RootDisplayRate = GetPipelinePrimaryConfig()->GetEffectiveFrameRate(TargetSequence);
+	FFrameTime StartTimeInRoot = FFrameRate::TransformTime(InShot->ShotInfo.TotalOutputRangeRoot.GetLowerBoundValue(), InShot->ShotInfo.CachedTickResolution, RootDisplayRate);
+	if(bIsPrePass && StartTimeInRoot.GetSubFrame() != 0.f)
 	{
 		UE_LOG(LogMovieRenderPipeline, Warning, TEXT("Shot/Camera \"%s\" starts on a sub-frame. Rendered range has been rounded to the previous frame to match Sequencer."), *InShot->OuterName, *InShot->InnerName);
-		FFrameNumber NewStartFrame = FFrameRate::TransformTime(FFrameTime(StartTimeInMaster.GetFrame()), MasterDisplayRate, InShot->ShotInfo.CachedTickResolution).FloorToFrame();
-		InShot->ShotInfo.TotalOutputRangeMaster.SetLowerBoundValue(NewStartFrame);
+		FFrameNumber NewStartFrame = FFrameRate::TransformTime(FFrameTime(StartTimeInRoot.GetFrame()), RootDisplayRate, InShot->ShotInfo.CachedTickResolution).FloorToFrame();
+		InShot->ShotInfo.TotalOutputRangeRoot.SetLowerBoundValue(NewStartFrame);
 	}
 
 	FFrameNumber LeftDeltaTicks = FFrameRate::TransformTime(FFrameTime(LeftDeltaFrames), FrameMetrics.FrameRate, FrameMetrics.TickResolution).CeilToFrame().Value;
@@ -1147,13 +1147,13 @@ void UMoviePipeline::ExpandShot(UMoviePipelineExecutorShot* InShot, const int32
 	FFrameNumber LeftDeltaTicksUserPoV = FFrameRate::TransformTime(FFrameTime(LeftDeltaFramesUserPoV), FrameMetrics.FrameRate, FrameMetrics.TickResolution).CeilToFrame().Value;
 
 	// We generate this from excess camera cut length. If they don't want real warmup then this is overriden later.
-	// Needs to happen after we expand the TotalOutputRangeMaster for handle frames - the MoviePipelineTiming re-calculates
+	// Needs to happen after we expand the TotalOutputRangeRoot for handle frames - the MoviePipelineTiming re-calculates
 	// the required offset when jumping based on NumEngineWarmUpFramesRemaining.
-	if (!InShot->ShotInfo.WarmupRangeMaster.IsEmpty())
+	if (!InShot->ShotInfo.WarmupRangeRoot.IsEmpty())
 	{
 		if (bIsPrePass)
 		{
-			FFrameNumber TicksForWarmUp = InShot->ShotInfo.WarmupRangeMaster.Size<FFrameNumber>();
+			FFrameNumber TicksForWarmUp = InShot->ShotInfo.WarmupRangeRoot.Size<FFrameNumber>();
 			InShot->ShotInfo.NumEngineWarmUpFramesRemaining = FFrameRate::TransformTime(FFrameTime(TicksForWarmUp), FrameMetrics.TickResolution, FrameMetrics.FrameRate).CeilToFrame().Value;
 
 			// Handle frames weren't accounted for when we calculated the warm up range, so just reduce the amount of warmup by that.
@@ -1201,8 +1201,8 @@ void UMoviePipeline::ExpandShot(UMoviePipelineExecutorShot* InShot, const int32
 				Node->MovieScene->MarkAsChanged();
 			}
 
-			FFrameNumber LowerCheckBound = InShot->ShotInfo.TotalOutputRangeMaster.GetLowerBoundValue() - LeftDeltaTicksUserPoV;
-			FFrameNumber UpperCheckBound = InShot->ShotInfo.TotalOutputRangeMaster.GetLowerBoundValue();
+			FFrameNumber LowerCheckBound = InShot->ShotInfo.TotalOutputRangeRoot.GetLowerBoundValue() - LeftDeltaTicksUserPoV;
+			FFrameNumber UpperCheckBound = InShot->ShotInfo.TotalOutputRangeRoot.GetLowerBoundValue();
 
 			TRange<FFrameNumber> CheckRange = TRange<FFrameNumber>(TRangeBound<FFrameNumber>::Exclusive(LowerCheckBound), TRangeBound<FFrameNumber>::Inclusive(UpperCheckBound));
 
@@ -1240,13 +1240,13 @@ void UMoviePipeline::ExpandShot(UMoviePipelineExecutorShot* InShot, const int32
 
 			// We only do our warnings during the pre-pass
 			// Check for sections that start in the expanded evaluation range and warn user. Only check the frames user expects to (handle + temporal, no need for warm up frames to get checked as well)
-			MoviePipeline::CheckPartialSectionEvaluationAndWarn(LeftDeltaTicksUserPoV, Node, InShot, MasterDisplayRate);
+			MoviePipeline::CheckPartialSectionEvaluationAndWarn(LeftDeltaTicksUserPoV, Node, InShot, RootDisplayRate);
 		}
 
 		Node = Node->GetParent();
 	}
 
-	// Expand the Total Output Range Master by Handle Frames. The expansion of TotalOutputRangeMaster has to come after we do partial evaluation checks,
+	// Expand the Total Output Range Root by Handle Frames. The expansion of TotalOutputRangeRoot has to come after we do partial evaluation checks,
 	// otherwise the expanded range makes it check the wrong area for partial evaluations.
 	if (bIsPrePass)
 	{
@@ -1255,7 +1255,7 @@ void UMoviePipeline::ExpandShot(UMoviePipelineExecutorShot* InShot, const int32
 		FFrameNumber LeftHandleTicks = FFrameRate::TransformTime(FFrameTime(InNumHandleFrames), FrameMetrics.FrameRate, FrameMetrics.TickResolution).CeilToFrame().Value;
 		FFrameNumber RightHandleTicks = FFrameRate::TransformTime(FFrameTime(InNumHandleFrames), FrameMetrics.FrameRate, FrameMetrics.TickResolution).CeilToFrame().Value;
 
-		InShot->ShotInfo.TotalOutputRangeMaster = UE::MovieScene::DilateRange(InShot->ShotInfo.TotalOutputRangeMaster, -LeftHandleTicks, RightHandleTicks);
+		InShot->ShotInfo.TotalOutputRangeRoot = UE::MovieScene::DilateRange(InShot->ShotInfo.TotalOutputRangeRoot, -LeftHandleTicks, RightHandleTicks);
 	}
 }
 
@@ -1334,7 +1334,7 @@ MoviePipeline::FFrameConstantMetrics UMoviePipeline::CalculateShotFrameMetrics(c
 	MoviePipeline::FFrameConstantMetrics Output;
 	Output.TickResolution = TargetSequence->GetMovieScene()->GetTickResolution();
 	Output.ShotTickResolution = InShot->ShotInfo.CachedShotTickResolution;
-	Output.FrameRate = GetPipelineMasterConfig()->GetEffectiveFrameRate(TargetSequence);
+	Output.FrameRate = GetPipelinePrimaryConfig()->GetEffectiveFrameRate(TargetSequence);
 	Output.TicksPerOutputFrame = FFrameRate::TransformTime(FFrameTime(FFrameNumber(1)), Output.FrameRate, Output.TickResolution);
 
 	UMoviePipelineCameraSetting* CameraSettings = FindOrAddSettingForShot<UMoviePipelineCameraSetting>(InShot);
@@ -1465,7 +1465,7 @@ MoviePipeline::FFrameConstantMetrics UMoviePipeline::CalculateShotFrameMetrics(c
 }
 
 
-UMoviePipelineMasterConfig* UMoviePipeline::GetPipelineMasterConfig() const
+UMoviePipelinePrimaryConfig* UMoviePipeline::GetPipelinePrimaryConfig() const
 { 
 	return CurrentJob->GetConfiguration(); 
 }
@@ -1487,7 +1487,7 @@ TArray<UMoviePipelineSetting*> UMoviePipeline::FindSettingsForShot(TSubclassOf<U
 	}
 
 	// Add all enabled settings of given subclass not overridden by shot override
-	for (UMoviePipelineSetting* Setting : GetPipelineMasterConfig()->FindSettingsByClass(InSetting))
+	for (UMoviePipelineSetting* Setting : GetPipelinePrimaryConfig()->FindSettingsByClass(InSetting))
 	{
 		if (Setting && Setting->IsEnabled())
 		{
@@ -1518,7 +1518,7 @@ void UMoviePipeline::ResolveFilenameFormatArguments(const FString& InFormatStrin
 		Params.CameraNameOverride = InOutputState->CameraNameOverride;
 	}
 
-	UMoviePipelineOutputSetting* OutputSetting = GetPipelineMasterConfig()->FindSetting<UMoviePipelineOutputSetting>();
+	UMoviePipelineOutputSetting* OutputSetting = GetPipelinePrimaryConfig()->FindSetting<UMoviePipelineOutputSetting>();
 	check(OutputSetting);
 
 	Params.ZeroPadFrameNumberCount = OutputSetting->ZeroPadFrameNumbers;
@@ -1621,7 +1621,7 @@ void UMoviePipeline::PrintVerboseLogForFiles(const TArray<FMoviePipelineShotOutp
 void UMoviePipeline::StartUnrealInsightsCapture()
 {
 	// Generate a filepath to attempt to store the trace file in.
-	UMoviePipelineOutputSetting* OutputSetting = GetPipelineMasterConfig()->FindSetting<UMoviePipelineOutputSetting>();
+	UMoviePipelineOutputSetting* OutputSetting = GetPipelinePrimaryConfig()->FindSetting<UMoviePipelineOutputSetting>();
 	FString FileName = OutputSetting->FileNameFormat + TEXT("_UnrealInsights");
 	FString FileNameFormatString = OutputSetting->OutputDirectory.Path / FileName;
 
diff --git a/Engine/Plugins/MovieScene/MovieRenderPipeline/Source/MovieRenderPipelineCore/Private/MoviePipelineAudioRendering.cpp b/Engine/Plugins/MovieScene/MovieRenderPipeline/Source/MovieRenderPipelineCore/Private/MoviePipelineAudioRendering.cpp
--- a/Engine/Plugins/MovieScene/MovieRenderPipeline/Source/MovieRenderPipelineCore/Private/MoviePipelineAudioRendering.cpp
+++ b/Engine/Plugins/MovieScene/MovieRenderPipeline/Source/MovieRenderPipelineCore/Private/MoviePipelineAudioRendering.cpp
@@ -9,7 +9,7 @@
 #include "LevelSequence.h"
 #include "AudioDevice.h"
 #include "AudioMixerDevice.h"
-#include "MoviePipelineMasterConfig.h"
+#include "MoviePipelinePrimaryConfig.h"
 #include "MovieRenderPipelineDataTypes.h"
 #include "AudioMixerSubmix.h"
 #include "DSP/BufferVectorOperations.h"
@@ -174,7 +174,7 @@ void UMoviePipeline::ProcessAudioTick()
 		bCanRenderAudio = CachedOutputState.IsLastTemporalSample();
 		
 		// Process work that has been submitted from the game thread to the audio thread over the temporal samples of this frame.
-		AudioDeltaTime = GetPipelineMasterConfig()->GetEffectiveFrameRate(TargetSequence).AsInterval();
+		AudioDeltaTime = GetPipelinePrimaryConfig()->GetEffectiveFrameRate(TargetSequence).AsInterval();
 	}
 
 	{
@@ -184,7 +184,7 @@ void UMoviePipeline::ProcessAudioTick()
 
 	{
 		// Process work that has been submitted from the game thread to the audio thread over the temporal samples of this frame.
-		FFrameRate OutputFrameRate = GetPipelineMasterConfig()->GetEffectiveFrameRate(TargetSequence);
+		FFrameRate OutputFrameRate = GetPipelinePrimaryConfig()->GetEffectiveFrameRate(TargetSequence);
 		NRTPlatform->RenderAudio(OutputFrameRate.AsInterval());
 	}
 }
diff --git a/Engine/Plugins/MovieScene/MovieRenderPipeline/Source/MovieRenderPipelineCore/Private/MoviePipelineBlueprintLibrary.cpp b/Engine/Plugins/MovieScene/MovieRenderPipeline/Source/MovieRenderPipelineCore/Private/MoviePipelineBlueprintLibrary.cpp
--- a/Engine/Plugins/MovieScene/MovieRenderPipeline/Source/MovieRenderPipelineCore/Private/MoviePipelineBlueprintLibrary.cpp
+++ b/Engine/Plugins/MovieScene/MovieRenderPipeline/Source/MovieRenderPipelineCore/Private/MoviePipelineBlueprintLibrary.cpp
@@ -162,13 +162,13 @@ UMovieSceneSequence* UMoviePipelineBlueprintLibrary::DuplicateSequence(UObject*
 	// Now go through looking for Shot and SubSequence tracks. These currently point to the same (shared) sequence as the InSequence.
 	TArray<UMovieSceneSection*> AllSubSequenceSections;
 
-	UMovieSceneCinematicShotTrack* ShotTrack = DuplicatedSequence->GetMovieScene()->FindMasterTrack<UMovieSceneCinematicShotTrack>();
+	UMovieSceneCinematicShotTrack* ShotTrack = DuplicatedSequence->GetMovieScene()->FindTrack<UMovieSceneCinematicShotTrack>();
 	if (ShotTrack)
 	{
 		AllSubSequenceSections.Append(ShotTrack->GetAllSections());
 	}
 
-	UMovieSceneSubTrack* SubTrack = DuplicatedSequence->GetMovieScene()->FindMasterTrack<UMovieSceneSubTrack>();
+	UMovieSceneSubTrack* SubTrack = DuplicatedSequence->GetMovieScene()->FindTrack<UMovieSceneSubTrack>();
 	if (SubTrack)
 	{
 		AllSubSequenceSections.Append(SubTrack->GetAllSections());
@@ -233,7 +233,7 @@ float UMoviePipelineBlueprintLibrary::GetCompletionPercentage(const UMoviePipeli
 	return CompletionPercentage;
 }
 
-FTimecode UMoviePipelineBlueprintLibrary::GetMasterTimecode(const UMoviePipeline* InMoviePipeline)
+FTimecode UMoviePipelineBlueprintLibrary::GetRootTimecode(const UMoviePipeline* InMoviePipeline)
 {
 	if (InMoviePipeline)
 	{
@@ -243,14 +243,14 @@ FTimecode UMoviePipelineBlueprintLibrary::GetMasterTimecode(const UMoviePipeline
 	return FTimecode();
 }
 
-FFrameNumber UMoviePipelineBlueprintLibrary::GetMasterFrameNumber(const UMoviePipeline* InMoviePipeline)
+FFrameNumber UMoviePipelineBlueprintLibrary::GetRootFrameNumber(const UMoviePipeline* InMoviePipeline)
 {
 	if (InMoviePipeline)
 	{
 		if (InMoviePipeline->GetTargetSequence() && InMoviePipeline->GetTargetSequence()->GetMovieScene())
 		{
-			FFrameRate EffectiveFrameRate = InMoviePipeline->GetPipelineMasterConfig()->GetEffectiveFrameRate(InMoviePipeline->GetTargetSequence());
-			return GetMasterTimecode(InMoviePipeline).ToFrameNumber(EffectiveFrameRate);
+			FFrameRate EffectiveFrameRate = InMoviePipeline->GetPipelinePrimaryConfig()->GetEffectiveFrameRate(InMoviePipeline->GetTargetSequence());
+			return GetRootTimecode(InMoviePipeline).ToFrameNumber(EffectiveFrameRate);
 		}
 	}
 
@@ -271,7 +271,7 @@ FFrameNumber UMoviePipelineBlueprintLibrary::GetCurrentShotFrameNumber(const UMo
 {
 	if (InMoviePipeline)
 	{
-		FFrameRate EffectiveFrameRate = InMoviePipeline->GetPipelineMasterConfig()->GetEffectiveFrameRate(InMoviePipeline->GetTargetSequence());
+		FFrameRate EffectiveFrameRate = InMoviePipeline->GetPipelinePrimaryConfig()->GetEffectiveFrameRate(InMoviePipeline->GetTargetSequence());
 		return GetCurrentShotTimecode(InMoviePipeline).ToFrameNumber(EffectiveFrameRate);
 	}
 
@@ -483,7 +483,7 @@ void UMoviePipelineBlueprintLibrary::UpdateJobShotListFromSequence(ULevelSequenc
 
 	{
 		UE::MovieScene::FSequenceVisitParams Params;
-		Params.bVisitMasterTracks = true;
+		Params.bVisitRootTracks = true;
 		Params.bVisitSubSequences = true;
 		FSequenceRangeVisitor RangeVisitor;
 
@@ -679,21 +679,21 @@ void UMoviePipelineBlueprintLibrary::UpdateJobShotListFromSequence(ULevelSequenc
 
 			// To make the camera cut range detection more consistent, we'll convert the start of this Entity into root sequence space, and then we'll convert the start of the
 			// camera cut into root sequence space, and the difference between them is the range to use. This is a more reliable way than looking at the actual Playback Range.
-			TRange<FFrameNumber> EntityRangeInMaster = Entity.Range;
-			TRange<FFrameNumber> CameraCutRangeInMaster = EntityRangeInMaster;
+			TRange<FFrameNumber> EntityRangeInRoot = Entity.Range;
+			TRange<FFrameNumber> CameraCutRangeInRoot = EntityRangeInRoot;
 			if (LeafNode->CameraCutSection.IsValid())
 			{
-				CameraCutRangeInMaster = LeafNode->CameraCutSection->GetRange() * InnerToOuterTransform;
+				CameraCutRangeInRoot = LeafNode->CameraCutSection->GetRange() * InnerToOuterTransform;
 			}
 
 			TRange<FFrameNumber> CameraCutWarmUpRange = TRange<FFrameNumber>::Empty();
-			if (EntityRangeInMaster.GetLowerBound().IsClosed() && CameraCutRangeInMaster.GetLowerBound().IsClosed())
+			if (EntityRangeInRoot.GetLowerBound().IsClosed() && CameraCutRangeInRoot.GetLowerBound().IsClosed())
 			{
-				CameraCutWarmUpRange = TRange<FFrameNumber>(CameraCutRangeInMaster.GetLowerBound(), EntityRangeInMaster.GetLowerBound());
+				CameraCutWarmUpRange = TRange<FFrameNumber>(CameraCutRangeInRoot.GetLowerBound(), EntityRangeInRoot.GetLowerBound());
 			}
-			if (CameraCutRangeInMaster.IsDegenerate())
+			if (CameraCutRangeInRoot.IsDegenerate())
 			{
-				CameraCutRangeInMaster = TRange<FFrameNumber>::Empty();
+				CameraCutRangeInRoot = TRange<FFrameNumber>::Empty();
 			}
 
 			Entity.CameraCutWarmUpRange = CameraCutWarmUpRange;
@@ -784,7 +784,7 @@ void UMoviePipelineBlueprintLibrary::UpdateJobShotListFromSequence(ULevelSequenc
 				}
 			}
 
-			// We can run into certain situations where the shots look the same ie: a master with a sub-scene and a camera being used multiple times
+			// We can run into certain situations where the shots look the same ie: a root with a sub-scene and a camera being used multiple times
 			// have the same inner and outer name. We don't want to re-use the same shot multiple times
 
 			if (!NewShot)
@@ -804,11 +804,11 @@ void UMoviePipelineBlueprintLibrary::UpdateJobShotListFromSequence(ULevelSequenc
 			NewShot->InnerName = Entity.Name.Get<0>();
 			NewShot->OuterName = Entity.Name.Get<1>();
 			NewShot->ShotInfo.SubSectionHierarchy = Entity.LeafNode;
-			NewShot->ShotInfo.TotalOutputRangeMaster = Entity.Range;
-			NewShot->ShotInfo.WarmupRangeMaster = Entity.CameraCutWarmUpRange;
+			NewShot->ShotInfo.TotalOutputRangeRoot = Entity.Range;
+			NewShot->ShotInfo.WarmupRangeRoot = Entity.CameraCutWarmUpRange;
 			NewShot->ShotInfo.OuterToInnerTransform = Entity.InnerToOuterTransform.Inverse();
 			NewShot->SidecarCameras = Entity.SidecarCameras;
-			UE_LOG(LogMovieRenderPipeline, Log, TEXT("Registering range: %s (InnerName: %s OuterName: %s)"), *LexToString(NewShot->ShotInfo.TotalOutputRangeMaster), *NewShot->InnerName, *NewShot->OuterName);
+			UE_LOG(LogMovieRenderPipeline, Log, TEXT("Registering range: %s (InnerName: %s OuterName: %s)"), *LexToString(NewShot->ShotInfo.TotalOutputRangeRoot), *NewShot->InnerName, *NewShot->OuterName);
 		}
 
 		// Now that we've read the job's shot mask we will clear it and replace it with only things still valid.
@@ -898,12 +898,12 @@ int32 UMoviePipelineBlueprintLibrary::ResolveVersionNumber(FMoviePipelineFilenam
 	return 0;
 }
 
-FIntPoint UMoviePipelineBlueprintLibrary::GetEffectiveOutputResolution(UMoviePipelineMasterConfig* InMasterConfig, UMoviePipelineExecutorShot* InPipelineExecutorShot)
+FIntPoint UMoviePipelineBlueprintLibrary::GetEffectiveOutputResolution(UMoviePipelinePrimaryConfig* InPrimaryConfig, UMoviePipelineExecutorShot* InPipelineExecutorShot)
 {
-	if (InMasterConfig && InPipelineExecutorShot)
+	if (InPrimaryConfig && InPipelineExecutorShot)
 	{
-		UMoviePipelineOutputSetting* OutputSetting = InMasterConfig->FindSetting<UMoviePipelineOutputSetting>();
-		const UMoviePipelineCameraSetting* CameraSetting = Cast<const UMoviePipelineCameraSetting>(UMoviePipelineBlueprintLibrary::FindOrGetDefaultSettingForShot(UMoviePipelineCameraSetting::StaticClass(), InMasterConfig, InPipelineExecutorShot));
+		UMoviePipelineOutputSetting* OutputSetting = InPrimaryConfig->FindSetting<UMoviePipelineOutputSetting>();
+		const UMoviePipelineCameraSetting* CameraSetting = Cast<const UMoviePipelineCameraSetting>(UMoviePipelineBlueprintLibrary::FindOrGetDefaultSettingForShot(UMoviePipelineCameraSetting::StaticClass(), InPrimaryConfig, InPipelineExecutorShot));
 		check(OutputSetting);
 		check(CameraSetting);
 		FIntPoint EffectiveResolution(OutputSetting->OutputResolution);
@@ -922,7 +922,7 @@ FIntPoint UMoviePipelineBlueprintLibrary::GetEffectiveOutputResolution(UMoviePip
 }
 
 
-UMoviePipelineSetting* UMoviePipelineBlueprintLibrary::FindOrGetDefaultSettingForShot(TSubclassOf<UMoviePipelineSetting> InSettingType, const UMoviePipelineMasterConfig* InMasterConfig, const UMoviePipelineExecutorShot* InShot)
+UMoviePipelineSetting* UMoviePipelineBlueprintLibrary::FindOrGetDefaultSettingForShot(TSubclassOf<UMoviePipelineSetting> InSettingType, const UMoviePipelinePrimaryConfig* InPrimaryConfig, const UMoviePipelineExecutorShot* InShot)
 {
 	// Check to see if this setting is in the shot override, if it is we'll use the shot version of that.
 	if (InShot && InShot->GetShotOverrideConfiguration())
@@ -936,10 +936,10 @@ UMoviePipelineSetting* UMoviePipelineBlueprintLibrary::FindOrGetDefaultSettingFo
 		}
 	}
 
-	// If they didn't have a shot override, or the setting wasn't enabled, we'll check the master config.
-	if (InMasterConfig)
+	// If they didn't have a shot override, or the setting wasn't enabled, we'll check the primary config.
+	if (InPrimaryConfig)
 	{
-		UMoviePipelineSetting* Setting = InMasterConfig->FindSettingByClass(InSettingType);
+		UMoviePipelineSetting* Setting = InPrimaryConfig->FindSettingByClass(InSettingType);
 		if (Setting && Setting->IsEnabled())
 		{
 			return Setting;
@@ -1035,7 +1035,7 @@ void UMoviePipelineBlueprintLibrary::ResolveFilenameFormatArguments(const FStrin
 
 	// And from ourself
 	{
-		// Use the shared function as UMoviePipelineMasterConfig::GetFormatArguments to ensure all the time variables get overwritten.
+		// Use the shared function as UMoviePipelinePrimaryConfig::GetFormatArguments to ensure all the time variables get overwritten.
 		UE::MoviePipeline::GetSharedFormatArguments(OutMergedFormatArgs.FilenameArguments, OutMergedFormatArgs.FileMetadata, InParams.InitializationTime, InParams.InitializationVersion, InParams.Job);
 		
 		// By default, we don't want to show frame duplication numbers. If we need to start writing them,
diff --git a/Engine/Plugins/MovieScene/MovieRenderPipeline/Source/MovieRenderPipelineCore/Private/MoviePipelineCommandLineEncoder.cpp b/Engine/Plugins/MovieScene/MovieRenderPipeline/Source/MovieRenderPipelineCore/Private/MoviePipelineCommandLineEncoder.cpp
--- a/Engine/Plugins/MovieScene/MovieRenderPipeline/Source/MovieRenderPipelineCore/Private/MoviePipelineCommandLineEncoder.cpp
+++ b/Engine/Plugins/MovieScene/MovieRenderPipeline/Source/MovieRenderPipelineCore/Private/MoviePipelineCommandLineEncoder.cpp
@@ -6,7 +6,7 @@
 #include "MovieRenderPipelineCoreModule.h"
 #include "MoviePipeline.h"
 #include "MoviePipelineBlueprintLibrary.h"
-#include "MoviePipelineMasterConfig.h"	
+#include "MoviePipelinePrimaryConfig.h"	
 #include "Misc/CoreDelegates.h"
 #include "Misc/Paths.h"
 #include "Misc/FileHelper.h"
@@ -93,7 +93,7 @@ void UMoviePipelineCommandLineEncoder::StartEncodingProcess(TArray<FMoviePipelin
 		}
 	}
 
-	UMoviePipelineOutputSetting* OutputSetting = GetPipeline()->GetPipelineMasterConfig()->FindSetting<UMoviePipelineOutputSetting>();
+	UMoviePipelineOutputSetting* OutputSetting = GetPipeline()->GetPipelinePrimaryConfig()->FindSetting<UMoviePipelineOutputSetting>();
 	FString OutputFilename = FileNameFormatOverride.Len() > 0 ? FileNameFormatOverride : OutputSetting->FileNameFormat;
 	FString FileNameFormatString = OutputSetting->OutputDirectory.Path / OutputFilename;
 	
@@ -117,7 +117,7 @@ void UMoviePipelineCommandLineEncoder::StartEncodingProcess(TArray<FMoviePipelin
 	SharedArguments.Add(TEXT("Executable"), ExecutablePathNoQuotes);
 	SharedArguments.Add(TEXT("AudioCodec"), EncoderSettings->AudioCodec);
 	SharedArguments.Add(TEXT("VideoCodec"), EncoderSettings->VideoCodec);
-	FFrameRate RenderFrameRate = GetPipeline()->GetPipelineMasterConfig()->GetEffectiveFrameRate(GetPipeline()->GetTargetSequence());
+	FFrameRate RenderFrameRate = GetPipeline()->GetPipelinePrimaryConfig()->GetEffectiveFrameRate(GetPipeline()->GetTargetSequence());
 	SharedArguments.Add(TEXT("FrameRate"), RenderFrameRate.AsDecimal());
 	SharedArguments.Add(TEXT("AdditionalLocalArgs"), AdditionalCommandLineArgs);
 	SharedArguments.Add(TEXT("Quality"), GetQualitySettingString());
@@ -233,7 +233,7 @@ void UMoviePipelineCommandLineEncoder::StartEncodingProcess(TArray<FMoviePipelin
 
 void UMoviePipelineCommandLineEncoder::LaunchEncoder(const FEncoderParams& InParams)
 {
-	UMoviePipelineOutputSetting* OutputSetting = GetPipeline()->GetPipelineMasterConfig()->FindSetting<UMoviePipelineOutputSetting>();
+	UMoviePipelineOutputSetting* OutputSetting = GetPipeline()->GetPipelinePrimaryConfig()->FindSetting<UMoviePipelineOutputSetting>();
 
 	// Generate a text file for each input type which lists the files for that input type. We generate a FGuid in case there are
 	// multiple encode jobs going at once.
@@ -335,7 +335,7 @@ void UMoviePipelineCommandLineEncoder::LaunchEncoder(const FEncoderParams& InPar
 
 		// Automatically delete the input files we generated when the job is done
 		bool bDeleteInputTexts = true;
-		UMoviePipelineDebugSettings* DebugSettings = GetPipeline()->GetPipelineMasterConfig()->FindSetting<UMoviePipelineDebugSettings>();
+		UMoviePipelineDebugSettings* DebugSettings = GetPipeline()->GetPipelinePrimaryConfig()->FindSetting<UMoviePipelineDebugSettings>();
 		if (DebugSettings)
 		{
 			bDeleteInputTexts = !DebugSettings->bWriteAllSamples;
@@ -412,7 +412,7 @@ void UMoviePipelineCommandLineEncoder::OnTick()
 
 bool UMoviePipelineCommandLineEncoder::NeedsPerShotFlushing() const
 {
-	UMoviePipelineOutputSetting* OutputSetting = GetPipeline()->GetPipelineMasterConfig()->FindSetting<UMoviePipelineOutputSetting>();
+	UMoviePipelineOutputSetting* OutputSetting = GetPipeline()->GetPipelinePrimaryConfig()->FindSetting<UMoviePipelineOutputSetting>();
 	FString OutputFilename = FileNameFormatOverride.Len() > 0 ? FileNameFormatOverride : OutputSetting->FileNameFormat;
 	FString FullPath = OutputSetting->OutputDirectory.Path / OutputFilename;
 
diff --git a/Engine/Plugins/MovieScene/MovieRenderPipeline/Source/MovieRenderPipelineCore/Private/MoviePipelineFCPXMLExporterSetting.cpp b/Engine/Plugins/MovieScene/MovieRenderPipeline/Source/MovieRenderPipelineCore/Private/MoviePipelineFCPXMLExporterSetting.cpp
--- a/Engine/Plugins/MovieScene/MovieRenderPipeline/Source/MovieRenderPipelineCore/Private/MoviePipelineFCPXMLExporterSetting.cpp
+++ b/Engine/Plugins/MovieScene/MovieRenderPipeline/Source/MovieRenderPipelineCore/Private/MoviePipelineFCPXMLExporterSetting.cpp
@@ -3,7 +3,7 @@
 #include "MoviePipelineFCPXMLExporterSetting.h"
 #include "MoviePipeline.h"
 #include "MoviePipelineOutputSetting.h"
-#include "MoviePipelineMasterConfig.h"
+#include "MoviePipelinePrimaryConfig.h"
 #include "Misc/FileHelper.h"
 #include "HAL/PlatformFileManager.h"
 
@@ -26,7 +26,7 @@ void UMoviePipelineFCPXMLExporter::BeginExportImpl()
 	bHasFinishedExporting = true;
 
 #if WITH_EDITOR
-	UMoviePipelineOutputSetting* OutputSetting = GetPipeline()->GetPipelineMasterConfig()->FindSetting<UMoviePipelineOutputSetting>();
+	UMoviePipelineOutputSetting* OutputSetting = GetPipeline()->GetPipelinePrimaryConfig()->FindSetting<UMoviePipelineOutputSetting>();
 
 	// Use our file name format on the end of the shared common directory.
 	FString FileNameFormat = FileNameFormatOverride.Len() > 0 ? FileNameFormatOverride : OutputSetting->FileNameFormat;
@@ -59,7 +59,7 @@ void UMoviePipelineFCPXMLExporter::BeginExportImpl()
 			return;
 		}
 
-		UMovieSceneCinematicShotTrack* ShotTrack = MovieScene->FindMasterTrack<UMovieSceneCinematicShotTrack>();
+		UMovieSceneCinematicShotTrack* ShotTrack = MovieScene->FindTrack<UMovieSceneCinematicShotTrack>();
 		if (!ShotTrack)
 		{
 			UE_LOG(LogMovieRenderPipeline, Error, TEXT("FCPXML Export only works with a Cinematic Shot track. No FCPXML file will be written."));
@@ -68,7 +68,7 @@ void UMoviePipelineFCPXMLExporter::BeginExportImpl()
 		
 		FString FilenameFormat = OutputSetting->FileNameFormat;
 		int32 HandleFrames = OutputSetting->HandleFrameCount;
-		FFrameRate FrameRate = GetPipeline()->GetPipelineMasterConfig()->GetEffectiveFrameRate(Sequence);
+		FFrameRate FrameRate = GetPipeline()->GetPipelinePrimaryConfig()->GetEffectiveFrameRate(Sequence);
 		uint32 ResX = OutputSetting->OutputResolution.X;
 		uint32 ResY = OutputSetting->OutputResolution.Y;
 		FString MovieExtension = ".mxf";
@@ -117,7 +117,7 @@ bool UMoviePipelineFCPXMLExporter::EnsureWritableFile()
 		IFileManager::Get().MakeDirectory(*Directory);
 	}
 
-	UMoviePipelineOutputSetting* OutputSetting = GetPipeline()->GetPipelineMasterConfig()->FindSetting<UMoviePipelineOutputSetting>();
+	UMoviePipelineOutputSetting* OutputSetting = GetPipeline()->GetPipelinePrimaryConfig()->FindSetting<UMoviePipelineOutputSetting>();
 
 	// If the file doesn't exist, we're ok to continue
 	if (IFileManager::Get().FileSize(*FilePath) == -1)
diff --git a/Engine/Plugins/MovieScene/MovieRenderPipeline/Source/MovieRenderPipelineCore/Private/MoviePipelineOutputSetting.cpp b/Engine/Plugins/MovieScene/MovieRenderPipeline/Source/MovieRenderPipelineCore/Private/MoviePipelineOutputSetting.cpp
--- a/Engine/Plugins/MovieScene/MovieRenderPipeline/Source/MovieRenderPipelineCore/Private/MoviePipelineOutputSetting.cpp
+++ b/Engine/Plugins/MovieScene/MovieRenderPipeline/Source/MovieRenderPipelineCore/Private/MoviePipelineOutputSetting.cpp
@@ -2,7 +2,7 @@
 
 #include "MoviePipelineOutputSetting.h"
 #include "Misc/Paths.h"
-#include "MoviePipelineMasterConfig.h"
+#include "MoviePipelinePrimaryConfig.h"
 #include "MoviePipelineQueue.h"
 #include "MovieRenderPipelineDataTypes.h"
 #include "MoviePipeline.h"
@@ -53,11 +53,11 @@ FText UMoviePipelineOutputSetting::GetFooterText(UMoviePipelineExecutorJob* InJo
 	FMoviePipelineFormatArgs FormatArgs;
 	FormatArgs.InJob = InJob;
 	
-	// Find the master configuration that owns us
-	UMoviePipelineMasterConfig* MasterConfig = GetTypedOuter<UMoviePipelineMasterConfig>();
-	if (MasterConfig)
+	// Find the primary configuration that owns us
+	UMoviePipelinePrimaryConfig* PrimaryConfig = GetTypedOuter<UMoviePipelinePrimaryConfig>();
+	if (PrimaryConfig)
 	{
-		MasterConfig->GetFormatArguments(FormatArgs);
+		PrimaryConfig->GetFormatArguments(FormatArgs);
 	}
 
 	for (const TPair<FString, FString>& KVP : FormatArgs.FilenameArguments)
diff --git a/Engine/Plugins/MovieScene/MovieRenderPipeline/Source/MovieRenderPipelineCore/Private/MoviePipelineQueue.cpp b/Engine/Plugins/MovieScene/MovieRenderPipeline/Source/MovieRenderPipelineCore/Private/MoviePipelineQueue.cpp
--- a/Engine/Plugins/MovieScene/MovieRenderPipeline/Source/MovieRenderPipelineCore/Private/MoviePipelineQueue.cpp
+++ b/Engine/Plugins/MovieScene/MovieRenderPipeline/Source/MovieRenderPipelineCore/Private/MoviePipelineQueue.cpp
@@ -176,7 +176,7 @@ void UMoviePipelineExecutorJob::SetSequence(FSoftObjectPath InSequence)
 	}
 }
 
-void UMoviePipelineExecutorJob::SetConfiguration(UMoviePipelineMasterConfig* InPreset)
+void UMoviePipelineExecutorJob::SetConfiguration(UMoviePipelinePrimaryConfig* InPreset)
 {
 	if (InPreset)
 	{
@@ -185,12 +185,12 @@ void UMoviePipelineExecutorJob::SetConfiguration(UMoviePipelineMasterConfig* InP
 	}
 }
 
-void UMoviePipelineExecutorJob::SetPresetOrigin(UMoviePipelineMasterConfig* InPreset)
+void UMoviePipelineExecutorJob::SetPresetOrigin(UMoviePipelinePrimaryConfig* InPreset)
 {
 	if (InPreset)
 	{
 		Configuration->CopyFrom(InPreset);
-		PresetOrigin = TSoftObjectPtr<UMoviePipelineMasterConfig>(InPreset);
+		PresetOrigin = TSoftObjectPtr<UMoviePipelinePrimaryConfig>(InPreset);
 	}
 }
 
diff --git a/Engine/Plugins/MovieScene/MovieRenderPipeline/Source/MovieRenderPipelineCore/Private/MoviePipelineRendering.cpp b/Engine/Plugins/MovieScene/MovieRenderPipeline/Source/MovieRenderPipelineCore/Private/MoviePipelineRendering.cpp
--- a/Engine/Plugins/MovieScene/MovieRenderPipeline/Source/MovieRenderPipelineCore/Private/MoviePipelineRendering.cpp
+++ b/Engine/Plugins/MovieScene/MovieRenderPipeline/Source/MovieRenderPipelineCore/Private/MoviePipelineRendering.cpp
@@ -12,7 +12,7 @@
 #include "MoviePipelineDebugSettings.h"
 #include "MoviePipelineOutputSetting.h"
 #include "MoviePipelineConfigBase.h"
-#include "MoviePipelineMasterConfig.h"
+#include "MoviePipelinePrimaryConfig.h"
 #include "MoviePipelineBlueprintLibrary.h"
 #include "Math/Halton.h"
 #include "ImageWriteTask.h"
@@ -59,12 +59,12 @@ void UMoviePipeline::SetupRenderingPipelineForShot(UMoviePipelineExecutorShot* I
 	*/
 	UMoviePipelineAntiAliasingSetting* AccumulationSettings = FindOrAddSettingForShot<UMoviePipelineAntiAliasingSetting>(InShot);
 	UMoviePipelineHighResSetting* HighResSettings = FindOrAddSettingForShot<UMoviePipelineHighResSetting>(InShot);
-	UMoviePipelineOutputSetting* OutputSettings = GetPipelineMasterConfig()->FindSetting<UMoviePipelineOutputSetting>();
+	UMoviePipelineOutputSetting* OutputSettings = GetPipelinePrimaryConfig()->FindSetting<UMoviePipelineOutputSetting>();
 	check(OutputSettings);
 
 
 	FIntPoint BackbufferTileCount = FIntPoint(HighResSettings->TileCount, HighResSettings->TileCount);
-	FIntPoint OutputResolution = UMoviePipelineBlueprintLibrary::GetEffectiveOutputResolution(GetPipelineMasterConfig(), InShot);
+	FIntPoint OutputResolution = UMoviePipelineBlueprintLibrary::GetEffectiveOutputResolution(GetPipelinePrimaryConfig(), InShot);
 
 	// Figure out how big each sub-region (tile) is.
 	FIntPoint BackbufferResolution = FIntPoint(
@@ -163,11 +163,11 @@ void UMoviePipeline::RenderFrame()
 	UMoviePipelineAntiAliasingSetting* AntiAliasingSettings = FindOrAddSettingForShot<UMoviePipelineAntiAliasingSetting>(ActiveShotList[CurrentShotIndex]);
 	UMoviePipelineCameraSetting* CameraSettings = FindOrAddSettingForShot<UMoviePipelineCameraSetting>(ActiveShotList[CurrentShotIndex]);
 	UMoviePipelineHighResSetting* HighResSettings = FindOrAddSettingForShot<UMoviePipelineHighResSetting>(ActiveShotList[CurrentShotIndex]);
-	UMoviePipelineOutputSetting* OutputSettings = GetPipelineMasterConfig()->FindSetting<UMoviePipelineOutputSetting>();
+	UMoviePipelineOutputSetting* OutputSettings = GetPipelinePrimaryConfig()->FindSetting<UMoviePipelineOutputSetting>();
 	UMoviePipelineDebugSettings* DebugSettings = FindOrAddSettingForShot<UMoviePipelineDebugSettings>(ActiveShotList[CurrentShotIndex]);
 	
 	// Color settings are optional, so we don't need to do any assertion checks.
-	UMoviePipelineColorSetting* ColorSettings = GetPipelineMasterConfig()->FindSetting<UMoviePipelineColorSetting>();
+	UMoviePipelineColorSetting* ColorSettings = GetPipelinePrimaryConfig()->FindSetting<UMoviePipelineColorSetting>();
 	check(AntiAliasingSettings);
 	check(CameraSettings);
 	check(HighResSettings);
@@ -175,7 +175,7 @@ void UMoviePipeline::RenderFrame()
 
 	FIntPoint TileCount = FIntPoint(HighResSettings->TileCount, HighResSettings->TileCount);
 	FIntPoint OriginalTileCount = TileCount;
-	FIntPoint OutputResolution = UMoviePipelineBlueprintLibrary::GetEffectiveOutputResolution(GetPipelineMasterConfig(), ActiveShotList[CurrentShotIndex]);
+	FIntPoint OutputResolution = UMoviePipelineBlueprintLibrary::GetEffectiveOutputResolution(GetPipelinePrimaryConfig(), ActiveShotList[CurrentShotIndex]);
 
 	int32 NumSpatialSamples = AntiAliasingSettings->SpatialSampleCount;
 	int32 NumTemporalSamples = AntiAliasingSettings->TemporalSampleCount;
@@ -390,7 +390,7 @@ void UMoviePipeline::RenderFrame()
 				SampleState.OverscanPercentage = FMath::Clamp(CameraSettings->OverscanPercentage, 0.0f, 1.0f);
 
 				// Render each output pass
-				FMoviePipelineRenderPassMetrics SampleStateForCurrentResolution = UE::MoviePipeline::GetRenderPassMetrics(GetPipelineMasterConfig(), ActiveShotList[CurrentShotIndex], SampleState, OutputResolution);
+				FMoviePipelineRenderPassMetrics SampleStateForCurrentResolution = UE::MoviePipeline::GetRenderPassMetrics(GetPipelinePrimaryConfig(), ActiveShotList[CurrentShotIndex], SampleState, OutputResolution);
 				for (UMoviePipelineRenderPass* RenderPass : InputBuffers)
 				{
 					RenderPass->RenderSample_GameThread(SampleStateForCurrentResolution);
@@ -450,7 +450,7 @@ void UMoviePipeline::ProcessOutstandingFinishedFrames()
 		FRenderTimeStatistics& TimeStats = RenderTimeFrameStatistics.FindOrAdd(OutputFrame.FrameOutputState.OutputFrameNumber);
 		TimeStats.EndTime = FDateTime::UtcNow();
 	
-		for (UMoviePipelineOutputBase* OutputContainer : GetPipelineMasterConfig()->GetOutputContainers())
+		for (UMoviePipelineOutputBase* OutputContainer : GetPipelinePrimaryConfig()->GetOutputContainers())
 		{
 			OutputContainer->OnReceiveImageData(&OutputFrame);
 		}
@@ -462,7 +462,7 @@ void UMoviePipeline::OnSampleRendered(TUniquePtr<FImagePixelData>&& OutputSample
 {
 	// This function handles the "Write all Samples" feature which lets you inspect data
 	// pre-accumulation.
-	UMoviePipelineOutputSetting* OutputSettings = GetPipelineMasterConfig()->FindSetting<UMoviePipelineOutputSetting>();
+	UMoviePipelineOutputSetting* OutputSettings = GetPipelinePrimaryConfig()->FindSetting<UMoviePipelineOutputSetting>();
 	check(OutputSettings);
 
 	// This is for debug output, writing every individual sample to disk that comes off of the GPU (that isn't discarded).
diff --git a/Engine/Plugins/MovieScene/MovieRenderPipeline/Source/MovieRenderPipelineCore/Private/MoviePipelineTiming.cpp b/Engine/Plugins/MovieScene/MovieRenderPipeline/Source/MovieRenderPipelineCore/Private/MoviePipelineTiming.cpp
--- a/Engine/Plugins/MovieScene/MovieRenderPipeline/Source/MovieRenderPipelineCore/Private/MoviePipelineTiming.cpp
+++ b/Engine/Plugins/MovieScene/MovieRenderPipeline/Source/MovieRenderPipelineCore/Private/MoviePipelineTiming.cpp
@@ -15,7 +15,7 @@
 #include "Engine/World.h"
 #include "GameFramework/PlayerController.h"
 #include "GameFramework/WorldSettings.h"
-#include "MoviePipelineMasterConfig.h"
+#include "MoviePipelinePrimaryConfig.h"
 #include "MoviePipelineCameraSetting.h"
 #include "MoviePipelineQueue.h"
 #include "MoviePipelineOutputSetting.h"
@@ -181,18 +181,18 @@ void UMoviePipeline::TickProducingFrames()
 			// to render the first frame we will be on the correct offset.
 			TicksToEndOfPreviousFrame += FFrameRate::TransformTime(FFrameTime(FFrameNumber(CurrentCameraCut->ShotInfo.NumEngineWarmUpFramesRemaining)), TargetSequence->GetMovieScene()->GetDisplayRate(), TargetSequence->GetMovieScene()->GetTickResolution());
 			AccumulatedTickSubFrameDeltas -= TicksToEndOfPreviousFrame.GetSubFrame();
-			CurrentCameraCut->ShotInfo.CurrentTickInMaster = CurrentCameraCut->ShotInfo.CurrentTickInMaster - TicksToEndOfPreviousFrame.FloorToFrame();
+			CurrentCameraCut->ShotInfo.CurrentTickInRoot = CurrentCameraCut->ShotInfo.CurrentTickInRoot - TicksToEndOfPreviousFrame.FloorToFrame();
 			SequencerEvalOffset = TicksToEndOfPreviousFrame.FloorToFrame();
 		}
 
 		// Jump to the first frame of the sequence that we will be playing from. This doesn't take into account camera timing offset or temporal sampling, but that is
 		// proably okay as it means the first frame (frame 0) is exactly evaluated which is easier to preview effects in the editor instead of -.25 of a frame.
-		FFrameNumber CurrentMasterSeqTick = CurrentCameraCut->ShotInfo.CurrentTickInMaster + SequencerEvalOffset;
-		FFrameTime TimeInPlayRate = FFrameRate::TransformTime(CurrentMasterSeqTick, TargetSequence->GetMovieScene()->GetTickResolution(), TargetSequence->GetMovieScene()->GetDisplayRate());
+		FFrameNumber CurrentRootSeqTick = CurrentCameraCut->ShotInfo.CurrentTickInRoot + SequencerEvalOffset;
+		FFrameTime TimeInPlayRate = FFrameRate::TransformTime(CurrentRootSeqTick, TargetSequence->GetMovieScene()->GetTickResolution(), TargetSequence->GetMovieScene()->GetDisplayRate());
 
 		// We tell it to jump so that things responding to different scrub-types work correctly.
 		LevelSequenceActor->GetSequencePlayer()->SetPlaybackPosition(FMovieSceneSequencePlaybackParams(TimeInPlayRate, EUpdatePositionMethod::Jump));
-		CustomSequenceTimeController->SetCachedFrameTiming(FQualifiedFrameTime(CurrentMasterSeqTick, TargetSequence->GetMovieScene()->GetTickResolution()));
+		CustomSequenceTimeController->SetCachedFrameTiming(FQualifiedFrameTime(CurrentRootSeqTick, TargetSequence->GetMovieScene()->GetTickResolution()));
 
 		if (!CurrentCameraCut->ShotInfo.bEmulateFirstFrameMotionBlur)
 		{
@@ -223,7 +223,7 @@ void UMoviePipeline::TickProducingFrames()
 			// the error we're trying to catch.
 			if (CurrentCameraCut->ShotInfo.SubSectionHierarchy->CameraCutSection.IsValid())
 			{
-				ensureMsgf(bSuccess, TEXT("Failed to evaluate camera to create cur/prev camera locations. No camera actor found on Eval Tick: %d"), CurrentMasterSeqTick.Value);
+				ensureMsgf(bSuccess, TEXT("Failed to evaluate camera to create cur/prev camera locations. No camera actor found on Eval Tick: %d"), CurrentRootSeqTick.Value);
 			}
 		}
 
@@ -263,7 +263,7 @@ void UMoviePipeline::TickProducingFrames()
 		CachedOutputState.bSkipRendering = !bRenderFrame;
 
 		// Render the next frame at a reasonable frame rate to pass time in the world.
-		double FrameDeltaTime = GetPipelineMasterConfig()->GetEffectiveFrameRate(TargetSequence).AsInterval();
+		double FrameDeltaTime = GetPipelinePrimaryConfig()->GetEffectiveFrameRate(TargetSequence).AsInterval();
 
 		UE_LOG(LogMovieRenderPipeline, VeryVerbose, TEXT("[%d] Shot WarmUp set engine DeltaTime to %f seconds."), GFrameCounter, FrameDeltaTime);
 		CustomTimeStep->SetCachedFrameTiming(MoviePipeline::FFrameTimeStepCache(FrameDeltaTime));
@@ -271,10 +271,10 @@ void UMoviePipeline::TickProducingFrames()
 		if (!CurrentCameraCut->ShotInfo.bEmulateFirstFrameMotionBlur)
 		{
 			// If we're using real warmup frames we need to advance the sequence instead of just sitting there like normal.
-			CurrentCameraCut->ShotInfo.CurrentTickInMaster = CurrentCameraCut->ShotInfo.CurrentTickInMaster + FrameMetrics.TicksPerOutputFrame.FloorToFrame();
+			CurrentCameraCut->ShotInfo.CurrentTickInRoot = CurrentCameraCut->ShotInfo.CurrentTickInRoot + FrameMetrics.TicksPerOutputFrame.FloorToFrame();
 
 			// The sequence player will already be playing (due to the Uninitialized block) so we just tell it which frame it should be playing at.
-			FFrameTime FinalEvalTime = FrameMetrics.GetFinalEvalTime(CurrentCameraCut->ShotInfo.CurrentTickInMaster);
+			FFrameTime FinalEvalTime = FrameMetrics.GetFinalEvalTime(CurrentCameraCut->ShotInfo.CurrentTickInRoot);
 			CustomSequenceTimeController->SetCachedFrameTiming(FQualifiedFrameTime(FinalEvalTime, FrameMetrics.TickResolution));
 		}
 
@@ -329,7 +329,7 @@ void UMoviePipeline::TickProducingFrames()
 			TicksToEndOfPreviousFrame = TicksToEndOfPreviousFrame * WorldTimeDilation;
 		}
 		AccumulatedTickSubFrameDeltas -= TicksToEndOfPreviousFrame.GetSubFrame();
-		CurrentCameraCut->ShotInfo.CurrentTickInMaster = CurrentCameraCut->ShotInfo.CurrentTickInMaster - TicksToEndOfPreviousFrame.FloorToFrame();
+		CurrentCameraCut->ShotInfo.CurrentTickInRoot = CurrentCameraCut->ShotInfo.CurrentTickInRoot - TicksToEndOfPreviousFrame.FloorToFrame();
 
 		// Skip to rendering which will skip the next block.
 		CurrentCameraCut->ShotInfo.State = EMovieRenderShotState::Rendering;
@@ -358,7 +358,7 @@ void UMoviePipeline::TickProducingFrames()
 
 		// CurrentShot.CurrentTick should be at the first frame they want to evaluate, without the shutter timing/motion blur
 		// offsets taken into account. When we evaluate for this frame, take those into account.
-		FFrameTime FinalEvalTime = FrameMetrics.GetFinalEvalTime(CurrentCameraCut->ShotInfo.CurrentTickInMaster + MotionBlurDurationTicks);
+		FFrameTime FinalEvalTime = FrameMetrics.GetFinalEvalTime(CurrentCameraCut->ShotInfo.CurrentTickInRoot + MotionBlurDurationTicks);
 
 		// Jump to the motion blur frame
 		FFrameTime TimeInPlayRate = FFrameRate::TransformTime(FinalEvalTime, FrameMetrics.TickResolution, TargetSequence->GetMovieScene()->GetDisplayRate());
@@ -639,12 +639,12 @@ void UMoviePipeline::TickProducingFrames()
 		AccumulatedTickSubFrameDeltas += DeltaFrameTime.GetSubFrame();
 
 		// Increment where we should evaluate on the current camera cut.
-		CurrentCameraCut->ShotInfo.CurrentTickInMaster += DeltaFrameTime.GetFrame();
+		CurrentCameraCut->ShotInfo.CurrentTickInRoot += DeltaFrameTime.GetFrame();
 
-		FFrameTime FinalEvalTime = FrameMetrics.GetFinalEvalTime(CurrentCameraCut->ShotInfo.CurrentTickInMaster);
+		FFrameTime FinalEvalTime = FrameMetrics.GetFinalEvalTime(CurrentCameraCut->ShotInfo.CurrentTickInRoot);
 		UE_LOG(LogMovieRenderPipeline, VeryVerbose, TEXT("[%d] FinalEvalTime: %s (Tick: %d SOffset: %d MOffset: %d)"), 
 			GFrameCounter, *LexToString(FinalEvalTime.FloorToFrame()),
-			CurrentCameraCut->ShotInfo.CurrentTickInMaster.Value, FrameMetrics.ShutterOffsetTicks.GetFrame().Value, FrameMetrics.MotionBlurCenteringOffsetTicks.GetFrame().Value);
+			CurrentCameraCut->ShotInfo.CurrentTickInRoot.Value, FrameMetrics.ShutterOffsetTicks.GetFrame().Value, FrameMetrics.MotionBlurCenteringOffsetTicks.GetFrame().Value);
 
 
 		if (!LevelSequenceActor->GetSequencePlayer()->IsPlaying())
@@ -680,7 +680,7 @@ void UMoviePipeline::TickProducingFrames()
 
 		// Now that we know the delta time for the upcoming frame, we can see if this time would extend past the end of our
 		// camera cut. If it does, we don't want to produce this frame as that would produce an extra frame at the end!
-		if (CurrentCameraCut->ShotInfo.CurrentTickInMaster >= CurrentCameraCut->ShotInfo.TotalOutputRangeMaster.GetUpperBoundValue())
+		if (CurrentCameraCut->ShotInfo.CurrentTickInRoot >= CurrentCameraCut->ShotInfo.TotalOutputRangeRoot.GetUpperBoundValue())
 		{
 			// If this isn't the last shot, we'll immediately call this function again to just
 			// determine a new setup/seek/etc. on the same engine tick. Otherwise we would have
@@ -709,7 +709,7 @@ void UMoviePipeline::TickProducingFrames()
 	
 #if WITH_EDITOR && !UE_BUILD_SHIPPING
 		{
-			UMoviePipelineDebugSettings* DebugSettings = GetPipelineMasterConfig()->FindSetting<UMoviePipelineDebugSettings>();
+			UMoviePipelineDebugSettings* DebugSettings = GetPipelinePrimaryConfig()->FindSetting<UMoviePipelineDebugSettings>();
 			if (IsValid(DebugSettings) && DebugSettings->IsRenderDocEnabled())
 			{
 				if (CachedOutputState.SourceFrameNumber == DebugSettings->CaptureFrame)
@@ -723,7 +723,7 @@ void UMoviePipeline::TickProducingFrames()
 		// Check to see if we should be rendering this frame. If they are frame stepping (rendering every Nth frame) then we
 		// just disable rendering but otherwise let the game logic remain the same for evaluation consistency.
 		{
-			UMoviePipelineOutputSetting* OutputSettings = GetPipelineMasterConfig()->FindSetting<UMoviePipelineOutputSetting>();
+			UMoviePipelineOutputSetting* OutputSettings = GetPipelinePrimaryConfig()->FindSetting<UMoviePipelineOutputSetting>();
 			int32 OutputFrameStep = FMath::Max(1, OutputSettings->OutputFrameStep);
 			CachedOutputState.bSkipRendering = ((CachedOutputState.OutputFrameNumber + OutputSettings->DEBUG_OutputFrameStepOffset) % OutputFrameStep) != 0;
 		}
@@ -748,13 +748,13 @@ void UMoviePipeline::CalculateFrameNumbersForOutputState(const MoviePipeline::FF
 {
 	// Find the closest number on the source Sequence. We use the effective framerate and not the original Sequence framerate because had they actually resampled their sequence
 	// to the given display rate then the frame numbers would match what they see. This may produce duplicates when using slow-mo tracks, in that case the user should use the _rel version.
-	FFrameRate SourceDisplayRate = GetPipelineMasterConfig()->GetEffectiveFrameRate(TargetSequence);
+	FFrameRate SourceDisplayRate = GetPipelinePrimaryConfig()->GetEffectiveFrameRate(TargetSequence);
 
 	// "Closest" isn't straightforward when using temporal sub-sampling, ie: A large enough shutter angle pushes a sample over the half way point and it rounds to
 	// the wrong one. Because temporal sub-sampling isn't centered around a frame (the centering is done via the final eval time) we can just subtract TS*TPS to get our centered value.
 	FFrameNumber CenteringOffset = FFrameNumber(InOutOutputState.TemporalSampleIndex * InFrameMetrics.TicksPerSample.RoundToFrame().Value);
 	{
-		FFrameNumber CenteredTick = InCameraCut->ShotInfo.CurrentTickInMaster - CenteringOffset;
+		FFrameNumber CenteredTick = InCameraCut->ShotInfo.CurrentTickInRoot - CenteringOffset;
 
 		InOutOutputState.SourceFrameNumber = FFrameRate::TransformTime(CenteredTick, InFrameMetrics.TickResolution, SourceDisplayRate).RoundToFrame().Value;
 		InOutOutputState.SourceTimeCode = FTimecode::FromFrameNumber(InOutOutputState.SourceFrameNumber, InFrameMetrics.FrameRate, false);
@@ -763,9 +763,9 @@ void UMoviePipeline::CalculateFrameNumbersForOutputState(const MoviePipeline::FF
 	}
 
 	{
-		FFrameNumber CenteredTick = InCameraCut->ShotInfo.CurrentTickInMaster - CenteringOffset;
+		FFrameNumber CenteredTick = InCameraCut->ShotInfo.CurrentTickInRoot - CenteringOffset;
 
-		// Convert from master space back into shot space - based on the inner most detected shot.
+		// Convert from Root space back into shot space - based on the inner most detected shot.
 		CenteredTick = (CenteredTick * InCameraCut->ShotInfo.OuterToInnerTransform).FloorToFrame();
 
 		InOutOutputState.CurrentShotSourceFrameNumber = FFrameRate::TransformTime(CenteredTick, InFrameMetrics.ShotTickResolution, SourceDisplayRate).RoundToFrame().Value;
@@ -774,6 +774,6 @@ void UMoviePipeline::CalculateFrameNumbersForOutputState(const MoviePipeline::FF
 
 	InOutOutputState.ShotName = InCameraCut->OuterName;
 
-	// Continue to use the InnerName (which is the master camera for the shot), individual render passes can override it later if rendering more than one camera.
+	// Continue to use the InnerName (which is the Root camera for the shot), individual render passes can override it later if rendering more than one camera.
 	InOutOutputState.CameraName = InCameraCut->InnerName;
 }
\ No newline at end of file
diff --git a/Engine/Plugins/MovieScene/MovieRenderPipeline/Source/MovieRenderPipelineCore/Private/MoviePipelineUtils.cpp b/Engine/Plugins/MovieScene/MovieRenderPipeline/Source/MovieRenderPipelineCore/Private/MoviePipelineUtils.cpp
--- a/Engine/Plugins/MovieScene/MovieRenderPipeline/Source/MovieRenderPipelineCore/Private/MoviePipelineUtils.cpp
+++ b/Engine/Plugins/MovieScene/MovieRenderPipeline/Source/MovieRenderPipelineCore/Private/MoviePipelineUtils.cpp
@@ -204,7 +204,7 @@ namespace MoviePipeline
 		// The evaluation tree only contains the active bits of the hierarchy (which is what we want). However we disable non-active sections while
 		// soloing, and we can't restore them because they weren't part of the shot hierarchy. To resolve this, we build a complete copy of the
 		// original for restoration at the end. We'll build our own tree, kept separate from the per-shot ones.
-		for (UMovieSceneTrack* Track : InSequence->GetMovieScene()->GetMasterTracks())
+		for (UMovieSceneTrack* Track : InSequence->GetMovieScene()->GetTracks())
 		{
 			UMovieSceneSubTrack* SubTrack = Cast<UMovieSceneSubTrack>(Track);
 			if (!SubTrack)
@@ -382,7 +382,7 @@ namespace MoviePipeline
 		}
 	}
 
-	void CheckPartialSectionEvaluationAndWarn(const FFrameNumber& LeftDeltaTicks, TSharedPtr<FCameraCutSubSectionHierarchyNode> Node, UMoviePipelineExecutorShot* InShot, const FFrameRate& InMasterDisplayRate)
+	void CheckPartialSectionEvaluationAndWarn(const FFrameNumber& LeftDeltaTicks, TSharedPtr<FCameraCutSubSectionHierarchyNode> Node, UMoviePipelineExecutorShot* InShot, const FFrameRate& InRootDisplayRate)
 	{
 		// For the given movie scene, we want to produce a warning if there is no data to evaluate once we've expanded for 
 		// handle frames or temporal sub-sampling. To do our best guess at which tracks are relevant, we can look at the
@@ -393,8 +393,8 @@ namespace MoviePipeline
 		// and offsets the output frame number (when rounded back to whole numbers) for output
 		if (LeftDeltaTicks > 0 && Node->MovieScene.IsValid())
 		{
-			FFrameNumber LowerCheckBound = InShot->ShotInfo.TotalOutputRangeMaster.GetLowerBoundValue() - LeftDeltaTicks;
-			FFrameNumber UpperCheckBound = InShot->ShotInfo.TotalOutputRangeMaster.GetLowerBoundValue();
+			FFrameNumber LowerCheckBound = InShot->ShotInfo.TotalOutputRangeRoot.GetLowerBoundValue() - LeftDeltaTicks;
+			FFrameNumber UpperCheckBound = InShot->ShotInfo.TotalOutputRangeRoot.GetLowerBoundValue();
 
 			TRange<FFrameNumber> CheckRange = TRange<FFrameNumber>(TRangeBound<FFrameNumber>::Exclusive(LowerCheckBound), TRangeBound<FFrameNumber>::Inclusive(UpperCheckBound));
 
@@ -439,10 +439,10 @@ namespace MoviePipeline
 						}
 
 						// Convert ticks back to frames for human consumption
-						FFrameNumber LowerCheckBoundFrame = FFrameRate::TransformTime(LowerCheckBound, InShot->ShotInfo.CachedTickResolution, InMasterDisplayRate).FloorToFrame();
-						FFrameNumber UpperCheckBoundFrame = FFrameRate::TransformTime(UpperCheckBound, InShot->ShotInfo.CachedTickResolution, InMasterDisplayRate).FloorToFrame();
+						FFrameNumber LowerCheckBoundFrame = FFrameRate::TransformTime(LowerCheckBound, InShot->ShotInfo.CachedTickResolution, InRootDisplayRate).FloorToFrame();
+						FFrameNumber UpperCheckBoundFrame = FFrameRate::TransformTime(UpperCheckBound, InShot->ShotInfo.CachedTickResolution, InRootDisplayRate).FloorToFrame();
 
-						UE_LOG(LogMovieRenderPipeline, Warning, TEXT("[%s %s] Due to Temporal sub-sampling or handle frames, evaluation will occur outside of shot boundaries (from frame %d to %d). Section %s (Binding: %s) starts during this time period and cannot be auto-expanded. Please extend this section to start on frame %d. (All times listed are relative to the master sequence)"),
+						UE_LOG(LogMovieRenderPipeline, Warning, TEXT("[%s %s] Due to Temporal sub-sampling or handle frames, evaluation will occur outside of shot boundaries (from frame %d to %d). Section %s (Binding: %s) starts during this time period and cannot be auto-expanded. Please extend this section to start on frame %d. (All times listed are relative to the root sequence)"),
 							*InShot->OuterName, *InShot->InnerName, LowerCheckBoundFrame.Value, UpperCheckBoundFrame.Value,
 							*SectionName, *BindingName, LowerCheckBoundFrame.Value);
 					}
@@ -591,8 +591,8 @@ namespace MoviePipeline
 
 		if (ChildSubSectionData)
 		{
-			const TArray<UMovieSceneTrack*> MasterTracks = Sequence->GetMovieScene()->GetMasterTracks();
-			for (const UMovieSceneTrack* Track : MasterTracks)
+			const TArray<UMovieSceneTrack*> Tracks = Sequence->GetMovieScene()->GetTracks();
+			for (const UMovieSceneTrack* Track : Tracks)
 			{
 				// SubTracks encompass both Cinematic Shot sections and Sub-Sequences
 				if (Track && Track->IsA<UMovieSceneSubTrack>())
@@ -906,13 +906,13 @@ namespace UE
 {
 	namespace MoviePipeline
 	{
-		FMoviePipelineRenderPassMetrics GetRenderPassMetrics(UMoviePipelineMasterConfig* InMasterConfig, UMoviePipelineExecutorShot* InPipelineExecutorShot, const FMoviePipelineRenderPassMetrics& InRenderPassMetrics, const FIntPoint& InEffectiveOutputResolution)
+		FMoviePipelineRenderPassMetrics GetRenderPassMetrics(UMoviePipelinePrimaryConfig* InPrimaryConfig, UMoviePipelineExecutorShot* InPipelineExecutorShot, const FMoviePipelineRenderPassMetrics& InRenderPassMetrics, const FIntPoint& InEffectiveOutputResolution)
 		{
 			FMoviePipelineRenderPassMetrics OutRenderPassMetrics = InRenderPassMetrics;
 
-			if (InMasterConfig && InPipelineExecutorShot)
+			if (InPrimaryConfig && InPipelineExecutorShot)
 			{
-				const UMoviePipelineHighResSetting* HighResSettings = Cast<const UMoviePipelineHighResSetting>(UMoviePipelineBlueprintLibrary::FindOrGetDefaultSettingForShot(UMoviePipelineHighResSetting::StaticClass(), InMasterConfig, InPipelineExecutorShot));
+				const UMoviePipelineHighResSetting* HighResSettings = Cast<const UMoviePipelineHighResSetting>(UMoviePipelineBlueprintLibrary::FindOrGetDefaultSettingForShot(UMoviePipelineHighResSetting::StaticClass(), InPrimaryConfig, InPipelineExecutorShot));
 				check(HighResSettings);
 
 				FIntPoint BackbufferResolution = FIntPoint(FMath::CeilToInt((float)InEffectiveOutputResolution.X / (float)OutRenderPassMetrics.OriginalTileCounts.X), FMath::CeilToInt((float)InEffectiveOutputResolution.Y / (float)OutRenderPassMetrics.OriginalTileCounts.Y));
diff --git a/Engine/Plugins/MovieScene/MovieRenderPipeline/Source/MovieRenderPipelineCore/Private/MoviePipelineVideoOutputBase.cpp b/Engine/Plugins/MovieScene/MovieRenderPipeline/Source/MovieRenderPipelineCore/Private/MoviePipelineVideoOutputBase.cpp
--- a/Engine/Plugins/MovieScene/MovieRenderPipeline/Source/MovieRenderPipelineCore/Private/MoviePipelineVideoOutputBase.cpp
+++ b/Engine/Plugins/MovieScene/MovieRenderPipeline/Source/MovieRenderPipelineCore/Private/MoviePipelineVideoOutputBase.cpp
@@ -3,7 +3,7 @@
 #include "MoviePipelineVideoOutputBase.h"
 #include "MoviePipeline.h"
 #include "MoviePipelineOutputSetting.h"
-#include "MoviePipelineMasterConfig.h"
+#include "MoviePipelinePrimaryConfig.h"
 #include "MovieRenderPipelineCoreModule.h"
 #include "MoviePipelineUtils.h"
 #include "ImageWriteTask.h"
@@ -26,7 +26,7 @@ void UMoviePipelineVideoOutputBase::OnShotFinishedImpl(const UMoviePipelineExecu
 	{
 		// If they don't have {shot_name} or {camera_name} in their output path, this probably doesn't do what they expect
 		// as it will finalize and then overwrite itself.
-		UMoviePipelineOutputSetting* OutputSettings = GetPipeline()->GetPipelineMasterConfig()->FindSetting<UMoviePipelineOutputSetting>();
+		UMoviePipelineOutputSetting* OutputSettings = GetPipeline()->GetPipelinePrimaryConfig()->FindSetting<UMoviePipelineOutputSetting>();
 		check(OutputSettings);
 		FString FullPath = OutputSettings->OutputDirectory.Path / OutputSettings->FileNameFormat;
 		if (!(FullPath.Contains(TEXT("{shot_name}")) || FullPath.Contains(TEXT("{camera_name}"))))
@@ -56,7 +56,7 @@ void UMoviePipelineVideoOutputBase::OnReceiveImageDataImpl(FMoviePipelineMergerO
 		return;
 	}
 
-	UMoviePipelineOutputSetting* OutputSettings = GetPipeline()->GetPipelineMasterConfig()->FindSetting<UMoviePipelineOutputSetting>();
+	UMoviePipelineOutputSetting* OutputSettings = GetPipeline()->GetPipelinePrimaryConfig()->FindSetting<UMoviePipelineOutputSetting>();
 	check(OutputSettings);
 
 	FString OutputDirectory = OutputSettings->OutputDirectory.Path;
@@ -117,7 +117,7 @@ void UMoviePipelineVideoOutputBase::OnReceiveImageDataImpl(FMoviePipelineMergerO
 			// off so we can generate a "stable" filename for the incoming image sample, then compare it against existing writers, and
 			// if we don't find a writer, then generate a new one (respecting Overwrite Existing) to add the sample to.
 			{
-				UMoviePipelineOutputSetting* OutputSetting = GetPipeline()->GetPipelineMasterConfig()->FindSetting<UMoviePipelineOutputSetting>();
+				UMoviePipelineOutputSetting* OutputSetting = GetPipeline()->GetPipelinePrimaryConfig()->FindSetting<UMoviePipelineOutputSetting>();
 				const bool bPreviousOverwriteExisting = OutputSetting->bOverrideExistingOutput;
 				OutputSetting->bOverrideExistingOutput = true;
 				
@@ -216,7 +216,7 @@ void UMoviePipelineVideoOutputBase::OnReceiveImageDataImpl(FMoviePipelineMergerO
 		FImagePixelData* RawRenderPassData = RenderPassData.Value.Get();
 
 		// Making sure that if OCIO is enabled the Quantization won't do additional color conversion.
-		UMoviePipelineColorSetting* ColorSetting = GetPipeline()->GetPipelineMasterConfig()->FindSetting<UMoviePipelineColorSetting>();
+		UMoviePipelineColorSetting* ColorSetting = GetPipeline()->GetPipelinePrimaryConfig()->FindSetting<UMoviePipelineColorSetting>();
 		OutputWriter->Get<0>()->bConvertToSrgb = !(ColorSetting && ColorSetting->OCIOConfiguration.bIsEnabled);
 
 		//FGraphEventRef Event = Task.Execute([this, OutputWriter, RawRenderPassData]
diff --git a/Engine/Plugins/MovieScene/MovieRenderPipeline/Source/MovieRenderPipelineCore/Private/MovieRenderPipelineCommandLine.cpp b/Engine/Plugins/MovieScene/MovieRenderPipeline/Source/MovieRenderPipelineCore/Private/MovieRenderPipelineCommandLine.cpp
--- a/Engine/Plugins/MovieScene/MovieRenderPipeline/Source/MovieRenderPipelineCore/Private/MovieRenderPipelineCommandLine.cpp
+++ b/Engine/Plugins/MovieScene/MovieRenderPipeline/Source/MovieRenderPipelineCore/Private/MovieRenderPipelineCommandLine.cpp
@@ -8,7 +8,7 @@
 #include "MoviePipeline.h"
 #include "MoviePipelineInProcessExecutor.h"
 #include "LevelSequence.h"
-#include "MoviePipelineMasterConfig.h"
+#include "MoviePipelinePrimaryConfig.h"
 #include "MovieRenderPipelineDataTypes.h"
 #include "MoviePipelineQueue.h"
 #include "Misc/PackageName.h"
@@ -202,7 +202,7 @@ bool FMovieRenderPipelineCoreModule::IsTryingToRenderMovieFromCommandLine(FStrin
 *		- Preset or Queue to use (A preset is required if using a Level Sequence above, passed via -MoviePipelineConfig="/Game/...")
 *		- Will render on current map
 *		ie: 
-*			"E:\SubwaySequencer\SubwaySequencer.uproject" subwaySequencer_P -game -LevelSequence="/Game/Sequencer/SubwaySequencerMASTER.SubwaySequencerMASTER"
+*			"E:\SubwaySequencer\SubwaySequencer.uproject" subwaySequencer_P -game -LevelSequence="/Game/Sequencer/SubwaySequencer.SubwaySequencer"
 				-MoviePipelineConfig="/Game/Cinematics/MoviePipeline/Presets/SmallTestPreset.SmallTestPreset" -windowed -resx=1280 -resy=720 -log -notexturestreaming
 *		or:
 *			ie: "E:\SubwaySequencer\SubwaySequencer.uproject" subwaySequencer_P -game -MoviePipelineConfig="/Game/Cinematics/MoviePipeline/Presets/BigTestQueue.BigTestQueue" 
@@ -272,7 +272,7 @@ uint8 FMovieRenderPipelineCoreModule::ParseMovieRenderData(const FString& InSequ
 			{
 				OutQueue = AssetAsQueue;
 			}
-			else if (UMoviePipelineMasterConfig* AssetAsConfig = Cast<UMoviePipelineMasterConfig>(AssetPath.TryLoad()))
+			else if (UMoviePipelinePrimaryConfig* AssetAsConfig = Cast<UMoviePipelinePrimaryConfig>(AssetPath.TryLoad()))
 			{
 				OutQueue = NewObject<UMoviePipelineQueue>();
 				UMoviePipelineExecutorJob* NewJob = OutQueue->AllocateNewJob(UMoviePipelineExecutorJob::StaticClass()); // Only the default job type is supported right now.
diff --git a/Engine/Plugins/MovieScene/MovieRenderPipeline/Source/MovieRenderPipelineCore/Public/MoviePipeline.h b/Engine/Plugins/MovieScene/MovieRenderPipeline/Source/MovieRenderPipelineCore/Public/MoviePipeline.h
--- a/Engine/Plugins/MovieScene/MovieRenderPipeline/Source/MovieRenderPipelineCore/Public/MoviePipeline.h
+++ b/Engine/Plugins/MovieScene/MovieRenderPipeline/Source/MovieRenderPipelineCore/Public/MoviePipeline.h
@@ -13,7 +13,7 @@
 #include "MoviePipeline.generated.h"
 
 // Forward Declares
-class UMoviePipelineMasterConfig;
+class UMoviePipelinePrimaryConfig;
 class ULevelSequence;
 class UMovieSceneSequencePlayer;
 class UMoviePipelineCustomTimeStep;
@@ -160,11 +160,15 @@ public:
 	FMoviePipelineWorkFinished OnMoviePipelineShotWorkFinishedDelegate;
 
 	/**
-	* Get the Master Configuration used to render this shot. This contains the global settings for the shot, as well as per-shot
+	* Get the Primary Configuration used to render this shot. This contains the global settings for the shot, as well as per-shot
 	* configurations which can contain their own settings.
 	*/
 	UFUNCTION(BlueprintPure, Category = "Movie Render Pipeline")
-	UMoviePipelineMasterConfig* GetPipelineMasterConfig() const;
+	UMoviePipelinePrimaryConfig* GetPipelinePrimaryConfig() const;
+	
+	UE_DEPRECATED(5.2, "GetPipelineMasterConfig is deprecated. Please use GetPipelinePrimaryConfig instead")
+	UFUNCTION(BlueprintPure, Category = "Movie Render Pipeline", meta=(DeprecatedFunction, DeprecationMessage = "Use GetPipelinePrimaryConfig"))
+	UMoviePipelinePrimaryConfig* GetPipelineMasterConfig() const { return GetPipelinePrimaryConfig(); }
 
 public:
 	ULevelSequence* GetTargetSequence() const { return TargetSequence; }
@@ -208,7 +212,7 @@ public:
 	template<typename SettingType>
 	SettingType* FindOrAddSettingForShot(const UMoviePipelineExecutorShot* InShot) const
 	{
-		return (SettingType*)UMoviePipelineBlueprintLibrary::FindOrGetDefaultSettingForShot(SettingType::StaticClass(), GetPipelineMasterConfig(), InShot);
+		return (SettingType*)UMoviePipelineBlueprintLibrary::FindOrGetDefaultSettingForShot(SettingType::StaticClass(), GetPipelinePrimaryConfig(), InShot);
 	}
 
 	template<typename SettingType>
@@ -227,7 +231,7 @@ public:
 	TArray<UMoviePipelineSetting*> FindSettingsForShot(TSubclassOf<UMoviePipelineSetting> InSetting, const UMoviePipelineExecutorShot* InShot) const;
 
 	/**
-	* Resolves the provided InFormatString by converting {format_strings} into settings provided by the master config.
+	* Resolves the provided InFormatString by converting {format_strings} into settings provided by the primary config.
 	* @param	InFormatString		A format string (in the form of "{format_key1}_{format_key2}") to resolve.
 	* @param	InFormatOverrides	A series of Key/Value pairs to override particular format keys. Useful for things that
 	*								change based on the caller such as filename extensions.
diff --git a/Engine/Plugins/MovieScene/MovieRenderPipeline/Source/MovieRenderPipelineCore/Public/MoviePipelineAntiAliasingSetting.h b/Engine/Plugins/MovieScene/MovieRenderPipeline/Source/MovieRenderPipelineCore/Public/MoviePipelineAntiAliasingSetting.h
--- a/Engine/Plugins/MovieScene/MovieRenderPipeline/Source/MovieRenderPipelineCore/Public/MoviePipelineAntiAliasingSetting.h
+++ b/Engine/Plugins/MovieScene/MovieRenderPipeline/Source/MovieRenderPipelineCore/Public/MoviePipelineAntiAliasingSetting.h
@@ -31,7 +31,7 @@ public:
 #endif
 protected:
 	virtual bool IsValidOnShots() const override { return true; }
-	virtual bool IsValidOnMaster() const override { return true; }
+	virtual bool IsValidOnPrimary() const override { return true; }
 	virtual void ValidateStateImpl() override;
 	virtual void GetFormatArguments(FMoviePipelineFormatArgs& InOutFormatArgs) const override;
 
diff --git a/Engine/Plugins/MovieScene/MovieRenderPipeline/Source/MovieRenderPipelineCore/Public/MoviePipelineBlueprintLibrary.h b/Engine/Plugins/MovieScene/MovieRenderPipeline/Source/MovieRenderPipelineCore/Public/MoviePipelineBlueprintLibrary.h
--- a/Engine/Plugins/MovieScene/MovieRenderPipeline/Source/MovieRenderPipelineCore/Public/MoviePipelineBlueprintLibrary.h
+++ b/Engine/Plugins/MovieScene/MovieRenderPipeline/Source/MovieRenderPipelineCore/Public/MoviePipelineBlueprintLibrary.h
@@ -38,7 +38,7 @@ public:
 	static UMovieSceneSequence* DuplicateSequence(UObject* Outer, UMovieSceneSequence* InSequence);
 
 	/**
-	* Resolves the provided InFormatString by converting {format_strings} into settings provided by the master config.
+	* Resolves the provided InFormatString by converting {format_strings} into settings provided by the primary config.
 	* @param	InFormatString		A format string (in the form of "{format_key1}_{format_key2}") to resolve.
 	* @param	InParams			The parameters to resolve the format string with. See FMoviePipelineFilenameResolveParams properties for details. 
 	*								Expected that you fill out all of the parameters so that they can be used to resolve strings, otherwise default
@@ -101,10 +101,18 @@ public:
 	static float GetCompletionPercentage(const UMoviePipeline* InPipeline);
 
 	UFUNCTION(BlueprintPure, Category = "Movie Render Pipeline")
-	static FTimecode GetMasterTimecode(const UMoviePipeline* InMoviePipeline);
+	static FTimecode GetRootTimecode(const UMoviePipeline* InMoviePipeline);
+
+	UE_DEPRECATED(5.2, "GetMasterTimecode is deprecated. Please use GetRootTimecode instead")
+	UFUNCTION(BlueprintPure, Category = "Movie Render Pipeline", meta = (DeprecatedFunction, DeprecationMessage = "GetMasterTimecode is deprecated. Please use GetRootTimecode instead"))
+	static FTimecode GetMasterTimecode(const UMoviePipeline* InMoviePipeline) { return GetRootTimecode(InMoviePipeline); }
 
 	UFUNCTION(BlueprintPure, Category = "Movie Render Pipeline")
-	static FFrameNumber GetMasterFrameNumber(const UMoviePipeline* InMoviePipeline);
+	static FFrameNumber GetRootFrameNumber(const UMoviePipeline* InMoviePipeline);
+
+	UE_DEPRECATED(5.2, "GetMasterFrameNumber is deprecated. Please use GetRootFrameNumber instead")
+	UFUNCTION(BlueprintPure, Category = "Movie Render Pipeline", meta = (DeprecatedFunction, DeprecationMessage = "GetMasterFrameNumber is deprecated. Please GetRootFrameNumber instead"))
+	static FFrameNumber GetMasterFrameNumber(const UMoviePipeline* InMoviePipeline) { return GetRootFrameNumber(InMoviePipeline); }
 
 	UFUNCTION(BlueprintPure, Category = "Movie Render Pipeline")
 	static FTimecode GetCurrentShotTimecode(const UMoviePipeline* InMoviePipeline);
@@ -139,11 +147,11 @@ public:
 
 	/** In case of Overscan percentage being higher than 0 we render additional pixels. This function returns the resolution with overscan taken into account. */
 	UFUNCTION(BlueprintCallable, Category = "Movie Render Pipeline")
-	static FIntPoint GetEffectiveOutputResolution(UMoviePipelineMasterConfig* InMasterConfig, UMoviePipelineExecutorShot* InPipelineExecutorShot);
+	static FIntPoint GetEffectiveOutputResolution(UMoviePipelinePrimaryConfig* InPrimaryConfig, UMoviePipelineExecutorShot* InPipelineExecutorShot);
 
 	/** Allows access to a setting of provided type for specific shot. */
 	UFUNCTION(BlueprintCallable, Category = "Movie Render Pipeline", meta = (DeterminesOutputType = "InSettingType"))
-	static UMoviePipelineSetting* FindOrGetDefaultSettingForShot(TSubclassOf<UMoviePipelineSetting> InSettingType, const UMoviePipelineMasterConfig* InMasterConfig, const UMoviePipelineExecutorShot* InShot);
+	static UMoviePipelineSetting* FindOrGetDefaultSettingForShot(TSubclassOf<UMoviePipelineSetting> InSettingType, const UMoviePipelinePrimaryConfig* InPrimaryConfig, const UMoviePipelineExecutorShot* InShot);
 
 	UFUNCTION(BlueprintPure, Category = "Movie Render Pipeline")
 	static ULevelSequence* GetCurrentSequence(const UMoviePipeline* InMoviePipeline);
diff --git a/Engine/Plugins/MovieScene/MovieRenderPipeline/Source/MovieRenderPipelineCore/Public/MoviePipelineCameraSetting.h b/Engine/Plugins/MovieScene/MovieRenderPipeline/Source/MovieRenderPipelineCore/Public/MoviePipelineCameraSetting.h
--- a/Engine/Plugins/MovieScene/MovieRenderPipeline/Source/MovieRenderPipelineCore/Public/MoviePipelineCameraSetting.h
+++ b/Engine/Plugins/MovieScene/MovieRenderPipeline/Source/MovieRenderPipelineCore/Public/MoviePipelineCameraSetting.h
@@ -21,7 +21,7 @@ public:
 #endif
 protected:
 	virtual bool IsValidOnShots() const override { return true; }
-	virtual bool IsValidOnMaster() const override { return true; }
+	virtual bool IsValidOnPrimary() const override { return true; }
 	
 	virtual void GetFormatArguments(FMoviePipelineFormatArgs& InOutFormatArgs) const override
 	{
diff --git a/Engine/Plugins/MovieScene/MovieRenderPipeline/Source/MovieRenderPipelineCore/Public/MoviePipelineExecutor.h b/Engine/Plugins/MovieScene/MovieRenderPipeline/Source/MovieRenderPipelineCore/Public/MoviePipelineExecutor.h
--- a/Engine/Plugins/MovieScene/MovieRenderPipeline/Source/MovieRenderPipelineCore/Public/MoviePipelineExecutor.h
+++ b/Engine/Plugins/MovieScene/MovieRenderPipeline/Source/MovieRenderPipelineCore/Public/MoviePipelineExecutor.h
@@ -10,7 +10,7 @@
 #include "MovieRenderDebugWidget.h"
 #include "MoviePipelineExecutor.generated.h"
 
-class UMoviePipelineMasterConfig;
+class UMoviePipelinePrimaryConfig;
 class UMoviePipelineExecutorBase;
 class UMoviePipelineExecutorJob;
 class UMoviePipeline;
