Sequencer: Remove Master from Sequencer

Added missing file.

#rnx

diff --git a/Engine/Plugins/MovieScene/MovieRenderPipeline/Source/MovieRenderPipelineCore/Private/MoviePipelineSettingBlueprintBase.cpp b/Engine/Plugins/MovieScene/MovieRenderPipeline/Source/MovieRenderPipelineCore/Private/MoviePipelineSettingBlueprintBase.cpp
new file mode 100644
--- /dev/null
+++ b/Engine/Plugins/MovieScene/MovieRenderPipeline/Source/MovieRenderPipelineCore/Private/MoviePipelineSettingBlueprintBase.cpp
@@ -0,0 +1,14 @@
+// Copyright Epic Games, Inc. All Rights Reserved.
+#include "MoviePipelineSettingBlueprintBase.h"
+
+void UMoviePipelineSetting_BlueprintBase::PostLoad()
+{
+	Super::PostLoad();
+
+#if WITH_EDITOR
+	if (!bIsValidOnMaster_DEPRECATED)
+	{
+		bIsValidOnPrimary = bIsValidOnMaster_DEPRECATED;
+	}
+#endif
+}
