Sequencer: Remove Master from Sequencer

Fixed editor only

#rnx

diff --git a/Engine/Plugins/MovieScene/MovieRenderPipeline/Source/MovieRenderPipelineCore/Public/MoviePipelineSettingBlueprintBase.h b/Engine/Plugins/MovieScene/MovieRenderPipeline/Source/MovieRenderPipelineCore/Public/MoviePipelineSettingBlueprintBase.h
--- a/Engine/Plugins/MovieScene/MovieRenderPipeline/Source/MovieRenderPipelineCore/Public/MoviePipelineSettingBlueprintBase.h
+++ b/Engine/Plugins/MovieScene/MovieRenderPipeline/Source/MovieRenderPipelineCore/Public/MoviePipelineSettingBlueprintBase.h
@@ -21,7 +21,9 @@ public:
 	{
 		CategoryText = NSLOCTEXT("MovieRenderPipelineBP", "DefaultCategoryName_Text", "Custom Settings");
 		bIsValidOnPrimary = true;
+#if WITH_EDITORONLY_DATA
 		bIsValidOnMaster_DEPRECATED = true;
+#endif
 		bIsValidOnShots = true;
 		bCanBeDisabled = true;
 	}
