Movie Pipeline: Added reference to jira for the config bug

#rnx

diff --git a/Engine/Plugins/MovieScene/MovieRenderPipeline/Source/MovieRenderPipelineEditor/Private/Widgets/SMoviePipelineConfigPanel.cpp b/Engine/Plugins/MovieScene/MovieRenderPipeline/Source/MovieRenderPipelineEditor/Private/Widgets/SMoviePipelineConfigPanel.cpp
--- a/Engine/Plugins/MovieScene/MovieRenderPipeline/Source/MovieRenderPipelineEditor/Private/Widgets/SMoviePipelineConfigPanel.cpp
+++ b/Engine/Plugins/MovieScene/MovieRenderPipeline/Source/MovieRenderPipelineEditor/Private/Widgets/SMoviePipelineConfigPanel.cpp
@@ -371,7 +371,7 @@ TSharedRef<SWidget> SMoviePipelineConfigPanel::OnGeneratePresetsMenu()
 
 		// The redirected masterconfig assets aren't found when the new primaryconfig asset path is specified. This is a temporary workaround 
 		// for the AssetRegistry not applying CoreRedirects to the ClassPath. 
-		// This can be removed once that bug is fixed and the KnownCase in CheckUnacceptableWords can also be removed.
+		// This can be removed once that bug (UE-168245) is fixed and the KnownCase in CheckUnacceptableWords can also be removed.
 		if (ConfigAssetType->GetClassPathName() == FTopLevelAssetPath(TEXT("/Script/MovieRenderPipelineCore.MoviePipelinePrimaryConfig")))
 		{
 			AssetPickerConfig.Filter.ClassPaths.Add(FTopLevelAssetPath(TEXT("/Script/MovieRenderPipelineCore.MoviePipelineMasterConfig")));
