MoviePipeline: Fixed an issue where the Deferred Render Pass couldn't be added to a job after being deleted if other render passes existed.

diff --git a/Engine/Plugins/MovieScene/MovieRenderPipeline/Source/MovieRenderPipelineCore/Public/MoviePipelineConfigBase.h b/Engine/Plugins/MovieScene/MovieRenderPipeline/Source/MovieRenderPipelineCore/Public/MoviePipelineConfigBase.h
--- a/Engine/Plugins/MovieScene/MovieRenderPipeline/Source/MovieRenderPipelineCore/Public/MoviePipelineConfigBase.h
+++ b/Engine/Plugins/MovieScene/MovieRenderPipeline/Source/MovieRenderPipelineCore/Public/MoviePipelineConfigBase.h
@@ -83,7 +83,7 @@ public:
 		TArray<UMoviePipelineSetting*> MatchingSettings;
 		for (UMoviePipelineSetting* Setting : AllSettings)
 		{
-			if ((Setting && Setting->IsA(InClass.Get())) && (Setting->IsEnabled() || bIncludeDisabledSettings))
+			if ((Setting && Setting->GetClass() == InClass) && (Setting->IsEnabled() || bIncludeDisabledSettings))
 			{
 				MatchingSettings.Add(Setting);
 			}
