EXR: Minimum support for writing the working color space chromaticities attribute.

diff --git a/Engine/Plugins/MovieScene/MovieRenderPipeline/Source/MovieRenderPipelineRenderPasses/MovieRenderPipelineRenderPasses.Build.cs b/Engine/Plugins/MovieScene/MovieRenderPipeline/Source/MovieRenderPipelineRenderPasses/MovieRenderPipelineRenderPasses.Build.cs
--- a/Engine/Plugins/MovieScene/MovieRenderPipeline/Source/MovieRenderPipelineRenderPasses/MovieRenderPipelineRenderPasses.Build.cs
+++ b/Engine/Plugins/MovieScene/MovieRenderPipeline/Source/MovieRenderPipelineRenderPasses/MovieRenderPipelineRenderPasses.Build.cs
@@ -22,6 +22,7 @@ public class MovieRenderPipelineRenderPasses : ModuleRules
 				"ImageWrapper",				
 				"CinematicCamera", // For metadata
 				"MovieRenderPipelineSettings", // For settings
+				"ColorManagement"
 			}
 		);
 
diff --git a/Engine/Plugins/MovieScene/MovieRenderPipeline/Source/MovieRenderPipelineRenderPasses/Private/MoviePipelineEXROutput.cpp b/Engine/Plugins/MovieScene/MovieRenderPipeline/Source/MovieRenderPipelineRenderPasses/Private/MoviePipelineEXROutput.cpp
--- a/Engine/Plugins/MovieScene/MovieRenderPipeline/Source/MovieRenderPipelineRenderPasses/Private/MoviePipelineEXROutput.cpp
+++ b/Engine/Plugins/MovieScene/MovieRenderPipeline/Source/MovieRenderPipelineRenderPasses/Private/MoviePipelineEXROutput.cpp
@@ -18,9 +18,11 @@
 #include "IOpenExrRTTIModule.h"
 #include "Modules/ModuleManager.h"
 #include "MoviePipelineUtils.h"
+#include "ColorSpace.h"
 
 THIRD_PARTY_INCLUDES_START
 #include "OpenEXR/ImfChannelList.h"
+#include "OpenEXR/ImfStandardAttributes.h"
 
 #include UE_INLINE_GENERATED_CPP_BY_NAME(MoviePipelineEXROutput)
 
@@ -164,6 +166,16 @@ bool FEXRImageWriteTask::WriteToDisk()
 		// Insert our key-value pair metadata (if any, can be an arbitrary set of key/value pairs)
 		AddFileMetadata(Header);
 
+		// Write the working color space into EXR chromaticities
+		const UE::Color::FColorSpace& WCS = UE::Color::FColorSpace::GetWorking();
+		Imf::Chromaticities Chromaticities = {
+			IMATH_NAMESPACE::V2f((float)WCS.GetRedChromaticity().X, (float)WCS.GetRedChromaticity().Y),
+			IMATH_NAMESPACE::V2f((float)WCS.GetGreenChromaticity().X, (float)WCS.GetGreenChromaticity().Y),
+			IMATH_NAMESPACE::V2f((float)WCS.GetBlueChromaticity().X, (float)WCS.GetBlueChromaticity().Y),
+			IMATH_NAMESPACE::V2f((float)WCS.GetWhiteChromaticity().X, (float)WCS.GetWhiteChromaticity().Y),
+		};
+		Imf::addChromaticities(Header, Chromaticities);
+
 		FExrMemStreamOut OutputFile; 
 		
 		{
diff --git a/Engine/Source/Runtime/ImageWrapper/ImageWrapper.Build.cs b/Engine/Source/Runtime/ImageWrapper/ImageWrapper.Build.cs
--- a/Engine/Source/Runtime/ImageWrapper/ImageWrapper.Build.cs
+++ b/Engine/Source/Runtime/ImageWrapper/ImageWrapper.Build.cs
@@ -16,7 +16,12 @@ public class ImageWrapper : ModuleRules
 		PublicDefinitions.Add("WITH_UNREALPNG=1");
 		PublicDefinitions.Add("WITH_UNREALJPEG=1");
 
-		PrivateDependencyModuleNames.Add("Core");
+		PrivateDependencyModuleNames.AddRange(
+			new string[] {
+				"Core",
+				"ColorManagement"
+			}
+		);
 
 		PublicDependencyModuleNames.Add("LibTiff");
 		PublicDependencyModuleNames.Add("ImageCore");
diff --git a/Engine/Source/Runtime/ImageWrapper/Private/Formats/ExrImageWrapper.cpp b/Engine/Source/Runtime/ImageWrapper/Private/Formats/ExrImageWrapper.cpp
--- a/Engine/Source/Runtime/ImageWrapper/Private/Formats/ExrImageWrapper.cpp
+++ b/Engine/Source/Runtime/ImageWrapper/Private/Formats/ExrImageWrapper.cpp
@@ -3,6 +3,7 @@
 #include "Formats/ExrImageWrapper.h"
 #include "ImageWrapperPrivate.h"
 
+#include "ColorSpace.h"
 #include "Containers/StringConv.h"
 #include "HAL/PlatformTime.h"
 #include "Math/Float16.h"
@@ -413,6 +414,16 @@ void FExrImageWrapper::Compress(int32 Quality)
 		ImfFrameBuffer.insert(ChannelNames[c], Imf::Slice(ImfPixelType, (char*)ChannelData[c].GetData(), BytesPerChannelPixel, (size_t)BytesPerChannelPixel * Width));
 	}
 
+	// Write the working color space into EXR chromaticities
+	const UE::Color::FColorSpace& WCS = UE::Color::FColorSpace::GetWorking();
+	Imf::Chromaticities Chromaticities = {
+		IMATH_NAMESPACE::V2f((float)WCS.GetRedChromaticity().X, (float)WCS.GetRedChromaticity().Y),
+		IMATH_NAMESPACE::V2f((float)WCS.GetGreenChromaticity().X, (float)WCS.GetGreenChromaticity().Y),
+		IMATH_NAMESPACE::V2f((float)WCS.GetBlueChromaticity().X, (float)WCS.GetBlueChromaticity().Y),
+		IMATH_NAMESPACE::V2f((float)WCS.GetWhiteChromaticity().X, (float)WCS.GetWhiteChromaticity().Y),
+	};
+	Imf::addChromaticities(ImfHeader, Chromaticities);
+
 	FMemFileOut MemFile("");
 	int64 MemFileLength;
 
diff --git a/Engine/Source/Runtime/ImageWrapper/Private/Formats/ExrImageWrapper.h b/Engine/Source/Runtime/ImageWrapper/Private/Formats/ExrImageWrapper.h
--- a/Engine/Source/Runtime/ImageWrapper/Private/Formats/ExrImageWrapper.h
+++ b/Engine/Source/Runtime/ImageWrapper/Private/Formats/ExrImageWrapper.h
@@ -18,6 +18,7 @@ THIRD_PARTY_INCLUDES_START
 	#include "OpenEXR/ImfInputFile.h"
 	#include "OpenEXR/ImfOutputFile.h"
 	#include "OpenEXR/ImfRgbaFile.h"
+	#include "OpenEXR/ImfStandardAttributes.h"
 	#include "OpenEXR/ImfStdIO.h"
 	#include "OpenEXR/ImfVersion.h"
 THIRD_PARTY_INCLUDES_END
