MoviePipeline: Added a simple auto-completion widget to the Config Editor UI that will helpfully suggest possible format tokens when you type a "{" in the File Name Format field. This doesn't change which tokens are supported and this does not work in the Output Directory path (as it already has a customization). This also enables support for double-click selection in the footer text and allowing copy/paste from the footer as another way to get the format tokens.

diff --git a/Engine/Plugins/MovieScene/MovieRenderPipeline/Source/MovieRenderPipelineEditor/Private/Widgets/SMoviePipelineConfigEditor.cpp b/Engine/Plugins/MovieScene/MovieRenderPipeline/Source/MovieRenderPipelineEditor/Private/Widgets/SMoviePipelineConfigEditor.cpp
--- a/Engine/Plugins/MovieScene/MovieRenderPipeline/Source/MovieRenderPipelineEditor/Private/Widgets/SMoviePipelineConfigEditor.cpp
+++ b/Engine/Plugins/MovieScene/MovieRenderPipeline/Source/MovieRenderPipelineEditor/Private/Widgets/SMoviePipelineConfigEditor.cpp
@@ -27,15 +27,19 @@
 // Slate includes
 #include "Widgets/SBoxPanel.h"
 #include "Widgets/Text/STextBlock.h"
+#include "Widgets/Text/SMultiLineEditableText.h"
+#include "Widgets/Input/SMultiLineEditableTextBox.h"
 #include "Widgets/Input/SButton.h"
 #include "Widgets/Input/SComboButton.h"
 #include "Widgets/Layout/SSplitter.h"
 #include "Widgets/Layout/SScrollBox.h"
 #include "Widgets/Layout/SBorder.h"
+#include "Widgets/Views/SListView.h"
 #include "Framework/MultiBox/MultiBoxBuilder.h"
 #include "Styling/SlateIconFinder.h"
 #include "SPositiveActionButton.h"
 
+
 // EditorStyle includes
 #include "Styling/AppStyle.h"
 #include "EditorFontGlyphs.h"
@@ -44,13 +48,408 @@
 #include "IDetailsView.h"
 #include "MovieRenderPipelineStyle.h"
 #include "FrameNumberDetailsCustomization.h"
+#include "IDetailCustomization.h"
 #include "Editor.h"
 
+#include "DetailLayoutBuilder.h"
+#include "DetailWidgetRow.h"
+#include "IPropertyUtilities.h"
+
 #include "MoviePipelineOutputSetting.h"
 
 #define LOCTEXT_NAMESPACE "SMoviePipelineEditor"
 
 
+
+
+/** This is a specific auto-complete widget made for the MoviePipeline that isn't very flexible. It's
+* similar to SSuggestionTextBox, but SSuggestionTextBox doesn't handle more than one word/suggestions
+* mid string. This widget is hardcoded to look for '{' characters (for {format_tokens}) and then auto
+* completes them from a list and fixes up the {} braces.
+*/
+class SMoviePipelineAutoCompleteTextBox : public SCompoundWidget
+{
+public:
+
+	SLATE_BEGIN_ARGS(SMoviePipelineAutoCompleteTextBox)
+		: _Text()
+		, _Suggestions()
+	{}
+
+	SLATE_ATTRIBUTE(FText, Text)
+	SLATE_ATTRIBUTE(TArray<FString>, Suggestions)
+	/** Called whenever the text is changed programmatically or interactively by the user. */
+	SLATE_EVENT(FOnTextChanged, OnTextChanged)
+	SLATE_END_ARGS()
+
+	void Construct(const FArguments& InArgs)
+	{
+		ChildSlot
+		[
+			SAssignNew(MenuAnchor, SMenuAnchor)
+			.Placement(MenuPlacement_ComboBox)
+			[
+				SAssignNew(TextBox, SMultiLineEditableTextBox)
+				.Text(InArgs._Text)
+				.OnKeyDownHandler(this, &SMoviePipelineAutoCompleteTextBox::OnKeyDown)
+				.OnTextChanged(this, &SMoviePipelineAutoCompleteTextBox::HandleTextBoxTextChanged)
+				.SelectWordOnMouseDoubleClick(true)
+				.AllowMultiLine(false)
+			]
+			.MenuContent
+			(
+				SNew(SBorder)
+				.Padding(FMargin(2))
+				[
+					SAssignNew(VerticalBox, SVerticalBox)
+					+ SVerticalBox::Slot()
+					.AutoHeight()
+					[
+						SAssignNew(SuggestionListView, SListView<TSharedPtr<FString>>)
+						.ItemHeight(18.f)
+						.ListItemsSource(&Suggestions)
+						.SelectionMode(ESelectionMode::Single)
+						.OnGenerateRow(this, &SMoviePipelineAutoCompleteTextBox::HandleSuggestionListViewGenerateRow)
+						.OnSelectionChanged(this, &SMoviePipelineAutoCompleteTextBox::HandleSuggestionListViewSelectionChanged)
+					]
+				]
+			)
+		];
+
+		// We just call it once and cache it for now as the selection code isn't tested against
+		// the amount of suggestions changing.
+		AllSuggestions.Append(InArgs._Suggestions.Get());
+		OnTextChanged = InArgs._OnTextChanged;
+	}
+
+	FReply OnKeyDown(const FGeometry& MyGeometry, const FKeyEvent& KeyEvent)
+	{
+		if (MenuAnchor->IsOpen())
+		{
+			if (KeyEvent.GetKey() == EKeys::Up)
+			{
+				// Because the pop-up dialog is below the text, 'up' actually goes to an earlier item in the list.
+				int32 NewSuggestionIndex = CurrentSuggestionIndex - 1;
+				if (NewSuggestionIndex < 0)
+				{
+					NewSuggestionIndex = Suggestions.Num() - 1;
+				}
+
+				SetActiveSuggestionIndex(NewSuggestionIndex);
+				return FReply::Handled();
+			}
+			else if(KeyEvent.GetKey() == EKeys::Down)
+			{
+				int32 NewSuggestionIndex = CurrentSuggestionIndex + 1;
+				if (NewSuggestionIndex > Suggestions.Num() - 1)
+				{
+					NewSuggestionIndex = 0;
+				}
+
+				SetActiveSuggestionIndex(NewSuggestionIndex);
+				return FReply::Handled();
+			}
+			else if (KeyEvent.GetKey() == EKeys::Tab || KeyEvent.GetKey() == EKeys::Enter)
+			{
+				if(CurrentSuggestionIndex >= 0 && CurrentSuggestionIndex <= Suggestions.Num() - 1)
+				{
+					// Trigger the auto-complete for the highlighted suggestion
+					FString SuggestionText = *Suggestions[CurrentSuggestionIndex];
+					ReplaceRelevantTextWithSuggestion(SuggestionText);
+					return FReply::Handled();
+				}
+			}
+		}
+		return FReply::Unhandled();
+	}
+
+
+	void FindAutoCompleteableTextAtPos(const FString& InWholeString, int32 InCursorPos, FString& OutStr, bool& bShowAutoComplete)
+	{
+		OutStr = FString();
+		bShowAutoComplete = false;
+
+		// We want to find a { brace on or to the left of InCursorPos, but if we find a } we
+		// stop looking, because that's a brace for another text. (+1 for ::FromEnd off by one)
+		int32 StartingBracePos = InWholeString.Find(TEXT("{"), ESearchCase::Type::IgnoreCase, ESearchDir::Type::FromEnd, InCursorPos + 1);
+		int32 PreviousEndBracePos = InWholeString.Find(TEXT("}"), ESearchCase::Type::IgnoreCase, ESearchDir::Type::FromEnd, InCursorPos);
+
+		if (StartingBracePos < PreviousEndBracePos)
+		{
+			return;
+		}
+
+		FString AutoCompleteText;
+
+		// Now that we found a {, take the substring between it and either the next }, or the end of the string.
+		int32 NextEndBracePos = InWholeString.Find(TEXT("}"), ESearchCase::Type::IgnoreCase, ESearchDir::Type::FromStart, InCursorPos);
+		if (StartingBracePos >= 0)
+		{
+			int32 Count = InWholeString.Len() - StartingBracePos;
+			if (NextEndBracePos >= 0)
+			{
+				Count = NextEndBracePos - StartingBracePos;
+			}
+
+			AutoCompleteText = InWholeString.Mid(StartingBracePos + 1, Count - 1);
+		}
+
+		OutStr = AutoCompleteText;
+		bShowAutoComplete = StartingBracePos >= 0 && OutStr.Len() == 0;
+	}
+
+	void ReplaceRelevantTextWithSuggestion(const FString& InSuggestionText)
+	{
+		FString TextBoxText = TextBox->GetText().ToString();
+		int32 CursorPos = TextBoxText.Len();
+		FTextLocation CursorLoc = TextBox->GetCursorLocation();
+		if (CursorLoc.IsValid())
+		{
+			CursorPos = FMath::Clamp(CursorLoc.GetOffset(), 0, CursorPos);
+		}
+		
+		// Look for the { to the left of the cursor. We search StrPositionIndex from +1 here due to a bug in ::FromEnd being off by one.
+		int32 StartingBracePos = TextBoxText.Find(TEXT("{"), ESearchCase::Type::IgnoreCase, ESearchDir::Type::FromEnd, CursorPos + 1);
+
+		// Now that we found a {, take the substring between it and either the next }, or the end of the string.
+		int32 NextEndBracePos = TextBoxText.Find(TEXT("}"), ESearchCase::Type::IgnoreCase, ESearchDir::Type::FromStart, CursorPos);
+		int32 NewCursorPos = 0;
+		if (StartingBracePos >= 0)	
+		{
+			// +1 to keep the left { brace
+			FString Left = TextBoxText.Left(StartingBracePos+1);
+			FString Right;
+
+			if (NextEndBracePos >= 0)
+			{
+				Right = TextBoxText.RightChop(NextEndBracePos);
+			}
+
+			// Since the user chose the suggestion ensure there's already a } brace to close off the pair.
+			if (!Right.StartsWith(TEXT("}")))
+			{
+				Right = FString::Printf(TEXT("}%s"), *Right);
+			}
+
+			TextBoxText = Left + InSuggestionText + Right;
+			// We subtract 1 from the Right as we want to put the cursor after the automatically generated "}" token.
+			NewCursorPos = TextBoxText.Len() - (Right.Len() - 1);
+		}
+
+		TextBox->SetText(FText::FromString(TextBoxText));
+		TextBox->GoTo(FTextLocation(0, NewCursorPos));
+	}
+
+	void HandleTextBoxTextChanged(const FText& InText)
+	{
+		OnTextChanged.ExecuteIfBound(InText);
+
+		FString TextAsStr = InText.ToString();
+		if (TextAsStr.Len() > 0)
+		{
+			FString OutStr;
+			bool bShowAutoComplete;
+
+			int32 CursorPos = TextAsStr.Len();
+			FTextLocation CursorLoc = TextBox->GetCursorLocation();
+			if (CursorLoc.IsValid())
+			{
+				CursorPos = FMath::Clamp(CursorLoc.GetOffset(), 0, CursorPos);
+			}
+
+			FindAutoCompleteableTextAtPos(TextAsStr, CursorPos, OutStr, bShowAutoComplete);
+			FilterVisibleSuggestions(OutStr, bShowAutoComplete);
+		}	
+		else
+		{
+			// If they have no text, suggest all possible solutions
+			FilterVisibleSuggestions(FString(), false);
+		}
+	}
+
+	void FilterVisibleSuggestions(const FString& StrToMatch, const bool bForceShowAll)
+	{
+		Suggestions.Reset();
+		for (const FString& Suggestion : AllSuggestions)
+		{
+			if (Suggestion.Contains(StrToMatch) || bForceShowAll)
+			{
+				Suggestions.Add(MakeShared<FString>(Suggestion));
+			}
+		}
+
+		if (Suggestions.Num() > 0)
+		{
+			// We don't focus the menu (because then you can't type on the keyboard) and instead
+			// keep the focus on the text field and bubble the keyboard commands to it.
+			const bool bIsOpen = true;
+			const bool bFocusMenu = false;
+			MenuAnchor->SetIsOpen(bIsOpen, bFocusMenu);
+			SuggestionListView->RequestScrollIntoView(Suggestions[0]);
+		}
+		else
+		{
+			CloseMenuAndReset();
+		}
+	}
+
+	void CloseMenuAndReset()
+	{
+		const bool bIsOpen = false;
+		MenuAnchor->SetIsOpen(bIsOpen);
+
+		// Reset their index when the drawer closes so that the first item is always selected when we re-open.
+		CurrentSuggestionIndex = -1;
+	}
+
+	void SetActiveSuggestionIndex(int32 InIndex)
+	{
+		if (InIndex < 0 || InIndex >= Suggestions.Num())
+		{
+			return;
+		}
+
+		TSharedPtr<FString> Suggestion = Suggestions[InIndex];
+		SuggestionListView->SetSelection(Suggestion);
+		if (!SuggestionListView->IsItemVisible(Suggestion))
+		{
+			SuggestionListView->RequestScrollIntoView(Suggestion);
+		}
+		CurrentSuggestionIndex = InIndex;
+	}
+
+	TSharedRef<ITableRow> HandleSuggestionListViewGenerateRow(TSharedPtr<FString> Text, const TSharedRef<STableViewBase>& OwnerTable)
+	{
+		FString SuggestionText = *Text;
+
+		return SNew(STableRow<TSharedPtr<FString> >, OwnerTable)
+			[
+				SNew(SBox)
+				[
+					SNew(STextBlock)
+					.Text(FText::FromString(SuggestionText))
+				]
+			];
+	}
+
+	void HandleSuggestionListViewSelectionChanged(TSharedPtr<FString> NewValue, ESelectInfo::Type SelectInfo)
+	{
+		// This is called when clicking on an item and when navigating the menu via arrow keys, but we already
+		// handle arrow keys in OnKeyDown so we only want to handle mouse click here.
+		if (SelectInfo == ESelectInfo::Type::OnMouseClick)
+		{
+			// Trigger the auto-complete for the highlighted suggestion
+			ReplaceRelevantTextWithSuggestion(*NewValue);
+			CloseMenuAndReset();
+		}
+	}
+
+private:
+	TSharedPtr<SListView<TSharedPtr<FString>>> SuggestionListView;
+	TSharedPtr<SMultiLineEditableTextBox> TextBox;
+	TSharedPtr<SMenuAnchor> MenuAnchor;
+	TSharedPtr<SVerticalBox> VerticalBox;
+	// Holds a delegate that is executed when the text has changed.
+	FOnTextChanged OnTextChanged;
+
+	// The pool of suggestions to show
+	TArray<FString> AllSuggestions;
+	// The currently filtered suggestion list
+	TArray<TSharedPtr<FString>> Suggestions;
+	int32 CurrentSuggestionIndex = -1;
+};
+
+class FOutputFormatDetailsCustomization : public IDetailCustomization
+{
+public:
+	/** Creates a detail customization instance */
+	static TSharedRef<IDetailCustomization> MakeInstance(TSharedRef<SMoviePipelineConfigEditor> InEditor)
+	{
+		return MakeShared<FOutputFormatDetailsCustomization>(InEditor);
+	}
+
+	FOutputFormatDetailsCustomization(TSharedRef<SMoviePipelineConfigEditor> InEditor)
+	{
+		OwningEditor = InEditor;
+	}
+
+protected:
+	//~ Begin IDetailCustomization interface
+	virtual void CustomizeDetails(IDetailLayoutBuilder& DetailBuilder) override
+	{
+		OutputFormatPropertyHandle = DetailBuilder.GetProperty(GET_MEMBER_NAME_CHECKED(UMoviePipelineOutputSetting, FileNameFormat));
+		const bool bShowChildren = false;
+		FText InitialText;
+		OutputFormatPropertyHandle->GetValueAsDisplayText(InitialText);
+
+		DetailBuilder.EditDefaultProperty(OutputFormatPropertyHandle)->CustomWidget(bShowChildren)
+			.NameContent()
+			[
+				OutputFormatPropertyHandle->CreatePropertyNameWidget()
+			]
+			.ValueContent()
+			.MinDesiredWidth(200.0f)
+			[
+				SNew(SMoviePipelineAutoCompleteTextBox)
+				.Text(this, &FOutputFormatDetailsCustomization::GetText)
+				.Suggestions(this, &FOutputFormatDetailsCustomization::GetSuggestions)
+				.OnTextChanged(this, &FOutputFormatDetailsCustomization::OnTextChanged)
+			];
+
+
+	}
+	//~ End IDetailCustomization interface
+
+	TArray<FString> GetSuggestions() const
+	{
+		TSharedPtr<SMoviePipelineConfigEditor> ConfigEditor = OwningEditor.Pin();
+		if (ConfigEditor.IsValid())
+		{
+			TArray<FString> Suggestions;
+			if (ConfigEditor->CachedPipelineConfig.IsValid())
+			{
+				FMoviePipelineFormatArgs FormatArgs;
+				FormatArgs.InJob = ConfigEditor->CachedOwningJob.Get();
+
+				// Find the primary configuration that owns us (we won't necessairly have a job, but
+				// we might be able to find the primary config that this setting is outered to.
+				UMoviePipelinePrimaryConfig* PrimaryConfig = Cast<UMoviePipelinePrimaryConfig>(ConfigEditor->CachedPipelineConfig.Get());
+				if (PrimaryConfig)
+				{
+					PrimaryConfig->GetFormatArguments(FormatArgs);
+				}
+
+				for (const TPair<FString, FString>& KVP : FormatArgs.FilenameArguments)
+				{
+					Suggestions.Add(KVP.Key);
+				}
+			}
+
+			return Suggestions;
+		}
+
+		return TArray<FString>();
+	}
+
+	void OnTextChanged(const FText& InValue)
+	{
+		OutputFormatPropertyHandle->SetValue(InValue.ToString());
+	}
+
+	FText GetText() const
+	{
+		FText DisplayText;
+		OutputFormatPropertyHandle->GetValueAsDisplayText(DisplayText);
+		return DisplayText;
+	}
+
+	TSharedPtr<IPropertyHandle> OutputFormatPropertyHandle;
+	TWeakPtr<SMoviePipelineConfigEditor> OwningEditor;
+};
+
+/**
+ * Widget used to edit a Movie Render Pipeline Shot Config.
+ */
 PRAGMA_DISABLE_OPTIMIZATION
 void SMoviePipelineConfigEditor::Construct(const FArguments& InArgs)
 {
@@ -110,7 +509,6 @@ void SMoviePipelineConfigEditor::Construct(const FArguments& InArgs)
 			[
 				DetailsBox.ToSharedRef()
 			]
-
 			+ SVerticalBox::Slot()
 			.Padding(FMargin(0, 4, 0, 0))
 			.AutoHeight()
@@ -125,8 +523,10 @@ void SMoviePipelineConfigEditor::Construct(const FArguments& InArgs)
 						SNew(SScrollBox)
 						+ SScrollBox::Slot()
 						[
-							SNew(STextBlock)
+							SNew(SMultiLineEditableText)
+							.IsReadOnly(true)
 							.AutoWrapText(true)
+							.SelectWordOnMouseDoubleClick(true)
 							.Text(this, &SMoviePipelineConfigEditor::GetSettingsFooterText)
 						]
 					]	
@@ -333,6 +733,8 @@ void SMoviePipelineConfigEditor::UpdateDetails()
 		{
 			TSharedRef<IDetailsView> Details = PropertyEditorModule.CreateDetailView(DetailsViewArgs);
 			Details->RegisterInstancedCustomPropertyTypeLayout("FrameNumber", FOnGetPropertyTypeCustomizationInstance::CreateStatic(&FFrameNumberDetailsCustomization::MakeInstance, NumericTypeInterface));
+			Details->RegisterInstancedCustomPropertyLayout(UMoviePipelineOutputSetting::StaticClass(), FOnGetDetailCustomizationInstance::CreateStatic(FOutputFormatDetailsCustomization::MakeInstance, SharedThis(this)));
+
 
 			Details->SetObjects(Pair.Value);
 			DetailsBox->AddSlot()
diff --git a/Engine/Plugins/MovieScene/MovieRenderPipeline/Source/MovieRenderPipelineEditor/Private/Widgets/SMoviePipelineConfigEditor.h b/Engine/Plugins/MovieScene/MovieRenderPipeline/Source/MovieRenderPipelineEditor/Private/Widgets/SMoviePipelineConfigEditor.h
--- a/Engine/Plugins/MovieScene/MovieRenderPipeline/Source/MovieRenderPipelineEditor/Private/Widgets/SMoviePipelineConfigEditor.h
+++ b/Engine/Plugins/MovieScene/MovieRenderPipeline/Source/MovieRenderPipelineEditor/Private/Widgets/SMoviePipelineConfigEditor.h
@@ -77,6 +77,9 @@ private:
 
 	bool bRequestDetailsRefresh;
 
+	// Let the customization access these so we can get the options for auto-complete.
+	friend class FOutputFormatDetailsCustomization; 
+
 	// Edited Pipeline
 	TAttribute<UMoviePipelineConfigBase*> PipelineConfigAttribute;
 	TWeakObjectPtr<UMoviePipelineConfigBase> CachedPipelineConfig;
