MoviePipeline: When setting a cvar (via the "Console Variables" setting), MRQ should set it using the cvar's current priority level. This will ensure it gets set correctly, and also not block other UIs from editing it.

diff --git a/Engine/Plugins/MovieScene/MovieRenderPipeline/Source/MovieRenderPipelineSettings/Private/MoviePipelineConsoleVariableSetting.cpp b/Engine/Plugins/MovieScene/MovieRenderPipeline/Source/MovieRenderPipelineSettings/Private/MoviePipelineConsoleVariableSetting.cpp
--- a/Engine/Plugins/MovieScene/MovieRenderPipeline/Source/MovieRenderPipelineSettings/Private/MoviePipelineConsoleVariableSetting.cpp
+++ b/Engine/Plugins/MovieScene/MovieRenderPipeline/Source/MovieRenderPipelineSettings/Private/MoviePipelineConsoleVariableSetting.cpp
@@ -22,15 +22,15 @@ namespace UE
 			// cast to integer first (to avoid scientific notation) if we know the cvar is an integer.
 			if (InCVar->IsVariableInt())
 			{
-				InCVar->Set(static_cast<int32>(InValue), EConsoleVariableFlags::ECVF_SetByConsole);
+				InCVar->SetWithCurrentPriority(static_cast<int32>(InValue));
 			}
 			else if (InCVar->IsVariableBool())
 			{
-				InCVar->Set(InValue != 0.f ? true : false, EConsoleVariableFlags::ECVF_SetByConsole);
+				InCVar->SetWithCurrentPriority(InValue != 0.f ? true : false);
 			}
 			else
 			{
-				InCVar->Set(InValue, EConsoleVariableFlags::ECVF_SetByConsole);
+				InCVar->SetWithCurrentPriority(InValue);
 			}
 		}
 	}
