MoviePipeline: Allow adding a "Deferred Rendering" setting if one of the derived settings (eg, "Deferred Rendering - Lighting Only") is already added.

diff --git a/Engine/Plugins/MovieScene/MovieRenderPipeline/Source/MovieRenderPipelineCore/Public/MoviePipelineConfigBase.h b/Engine/Plugins/MovieScene/MovieRenderPipeline/Source/MovieRenderPipelineCore/Public/MoviePipelineConfigBase.h
--- a/Engine/Plugins/MovieScene/MovieRenderPipeline/Source/MovieRenderPipelineCore/Public/MoviePipelineConfigBase.h
+++ b/Engine/Plugins/MovieScene/MovieRenderPipeline/Source/MovieRenderPipelineCore/Public/MoviePipelineConfigBase.h
@@ -74,16 +74,23 @@ public:
 	* Find all settings of a particular type for this config.
 	* @param InClass - Class that you wish to find the setting object for.
 	* @param bIncludeDisabledSettings - if true, disabled settings will be included in the search
+	* @param bExactMatch - if true, only exact matches of the specified class will be found, otherwise subclasses of the specified class will also be found
 	* @return An array of instances of this class if it already exists as a setting on this config
 	*/
 	UFUNCTION(BlueprintPure, meta = (DeterminesOutputType = "InClass"), Category = "Movie Render Pipeline")
-	TArray<UMoviePipelineSetting*> FindSettingsByClass(TSubclassOf<UMoviePipelineSetting> InClass, const bool bIncludeDisabledSettings = false) const
+	TArray<UMoviePipelineSetting*> FindSettingsByClass(TSubclassOf<UMoviePipelineSetting> InClass, const bool bIncludeDisabledSettings = false, const bool bExactMatch = false) const
 	{
 		TArray<UMoviePipelineSetting*> AllSettings = GetUserSettings();
 		TArray<UMoviePipelineSetting*> MatchingSettings;
 		for (UMoviePipelineSetting* Setting : AllSettings)
 		{
-			if ((Setting && Setting->IsA(InClass.Get())) && (Setting->IsEnabled() || bIncludeDisabledSettings))
+			if (!Setting)
+			{
+				continue;
+			}
+
+			const bool IsTypeMatch = bExactMatch ? (InClass.Get() == Setting->GetClass()) : Setting->IsA(InClass.Get());
+			if (IsTypeMatch && (Setting->IsEnabled() || bIncludeDisabledSettings))
 			{
 				MatchingSettings.Add(Setting);
 			}
@@ -96,12 +103,13 @@ public:
 	* Find a setting of a particular type for this config.
 	* @param InClass - Class that you wish to find the setting object for.
 	* @param bIncludeDisabledSettings - if true, disabled settings will be included in the search
+	* @param bExactMatch - if true, only exact matches of the specified class will be found, otherwise subclasses of the specified class will also be found
 	* @return An instance of this class if it already exists as a setting on this config, otherwise null.
 	*/
 	UFUNCTION(BlueprintPure, meta = (DeterminesOutputType = "InClass"), Category = "Movie Render Pipeline")
-	UMoviePipelineSetting* FindSettingByClass(TSubclassOf<UMoviePipelineSetting> InClass, const bool bIncludeDisabledSettings = false) const
+	UMoviePipelineSetting* FindSettingByClass(TSubclassOf<UMoviePipelineSetting> InClass, const bool bIncludeDisabledSettings = false, const bool bExactMatch = false) const
 	{
-		TArray<UMoviePipelineSetting*> AllInstances = FindSettingsByClass(InClass, bIncludeDisabledSettings);
+		TArray<UMoviePipelineSetting*> AllInstances = FindSettingsByClass(InClass, bIncludeDisabledSettings, bExactMatch);
 		if (AllInstances.Num() > 0)
 		{
 			return AllInstances[0];
@@ -115,12 +123,13 @@ public:
 	* Finds a setting of a particular type for this pipeline config, adding it if it doesn't already exist.
 	* @param InClass - Class you wish to find or create the setting object for.
 	* @param bIncludeDisabledSettings - if true, disabled settings will be included in the search
+	* @param bExactMatch - if true, only exact matches of the specified class will be found, otherwise subclasses of the specified class will also be found
 	* @return An instance of this class as a setting on this config.
 	*/
 	UFUNCTION(BlueprintCallable, meta = (DeterminesOutputType = "InClass"), Category = "Movie Render Pipeline")
-	UMoviePipelineSetting* FindOrAddSettingByClass(TSubclassOf<UMoviePipelineSetting> InClass, const bool bIncludeDisabledSettings = false)
+	UMoviePipelineSetting* FindOrAddSettingByClass(TSubclassOf<UMoviePipelineSetting> InClass, const bool bIncludeDisabledSettings = false, const bool bExactMatch = false)
 	{
-		UMoviePipelineSetting* Found = FindSettingByClass(InClass, bIncludeDisabledSettings);
+		UMoviePipelineSetting* Found = FindSettingByClass(InClass, bIncludeDisabledSettings, bExactMatch);
 		if (!Found)
 		{
 			Modify();
diff --git a/Engine/Plugins/MovieScene/MovieRenderPipeline/Source/MovieRenderPipelineEditor/Private/Widgets/SMoviePipelineConfigEditor.cpp b/Engine/Plugins/MovieScene/MovieRenderPipeline/Source/MovieRenderPipelineEditor/Private/Widgets/SMoviePipelineConfigEditor.cpp
--- a/Engine/Plugins/MovieScene/MovieRenderPipeline/Source/MovieRenderPipelineEditor/Private/Widgets/SMoviePipelineConfigEditor.cpp
+++ b/Engine/Plugins/MovieScene/MovieRenderPipeline/Source/MovieRenderPipelineEditor/Private/Widgets/SMoviePipelineConfigEditor.cpp
@@ -761,7 +761,10 @@ void SMoviePipelineConfigEditor::AddSettingFromClass(TSubclassOf<UMoviePipelineS
 	if (*SettingClass && Pipeline)
 	{
 		FScopedTransaction Transaction(FText::Format(LOCTEXT("AddNewSetting", "Add New {0} Setting"), SettingClass->GetDisplayNameText()));
-		UMoviePipelineSetting* NewSetting = Pipeline->FindOrAddSettingByClass(SettingClass);
+
+		const bool bIncludeDisabledSettings = false;
+		const bool bExactMatch = true;
+		UMoviePipelineSetting* NewSetting = Pipeline->FindOrAddSettingByClass(SettingClass, bIncludeDisabledSettings, bExactMatch);
 
 		SettingsWidget->SetSelectedSettings({ NewSetting });
 	}
