MoviePipeline: Don't create an empty v000 directory when clicking on the output directory, and navigate to the most recent version rather than the first.

diff --git a/Engine/Plugins/MovieScene/MovieRenderPipeline/Source/MovieRenderPipelineCore/Private/MoviePipelineBlueprintLibrary.cpp b/Engine/Plugins/MovieScene/MovieRenderPipeline/Source/MovieRenderPipelineCore/Private/MoviePipelineBlueprintLibrary.cpp
--- a/Engine/Plugins/MovieScene/MovieRenderPipeline/Source/MovieRenderPipelineCore/Private/MoviePipelineBlueprintLibrary.cpp
+++ b/Engine/Plugins/MovieScene/MovieRenderPipeline/Source/MovieRenderPipelineCore/Private/MoviePipelineBlueprintLibrary.cpp
@@ -821,7 +821,7 @@ void UMoviePipelineBlueprintLibrary::UpdateJobShotListFromSequence(ULevelSequenc
 	}
 }
 
-int32 UMoviePipelineBlueprintLibrary::ResolveVersionNumber(FMoviePipelineFilenameResolveParams InParams)
+int32 UMoviePipelineBlueprintLibrary::ResolveVersionNumber(FMoviePipelineFilenameResolveParams InParams, bool bGetNextVersion)
 {
 	if (!InParams.Job)
 	{
@@ -892,7 +892,7 @@ int32 UMoviePipelineBlueprintLibrary::ResolveVersionNumber(FMoviePipelineFilenam
 
 		}
 
-		return  HighestVersion + 1;
+		return HighestVersion + (bGetNextVersion ? 1 : 0);
 	}
 
 	return 0;
diff --git a/Engine/Plugins/MovieScene/MovieRenderPipeline/Source/MovieRenderPipelineCore/Public/MoviePipelineBlueprintLibrary.h b/Engine/Plugins/MovieScene/MovieRenderPipeline/Source/MovieRenderPipelineCore/Public/MoviePipelineBlueprintLibrary.h
--- a/Engine/Plugins/MovieScene/MovieRenderPipeline/Source/MovieRenderPipelineCore/Public/MoviePipelineBlueprintLibrary.h
+++ b/Engine/Plugins/MovieScene/MovieRenderPipeline/Source/MovieRenderPipelineCore/Public/MoviePipelineBlueprintLibrary.h
@@ -141,9 +141,9 @@ public:
 	UFUNCTION(BlueprintCallable, Category = "Movie Render Pipeline")
 	static void UpdateJobShotListFromSequence(ULevelSequence* InSequence, UMoviePipelineExecutorJob* InJob, bool& bShotsChanged);
 	
-	/**  If version number is manually specified by the Job, returns that. Otherwise search the Output Directory for the highest version already existing an increments it by one. */
+	/**  If version number is manually specified by the Job, returns that. Otherwise search the Output Directory for the highest version already existing (and increment it by one if bGetNextVersion is true). */
 	UFUNCTION(BlueprintCallable, Category = "Movie Render Pipeline")
-	static int32 ResolveVersionNumber(FMoviePipelineFilenameResolveParams InParams);
+	static int32 ResolveVersionNumber(FMoviePipelineFilenameResolveParams InParams, bool bGetNextVersion = true);
 
 	/** In case of Overscan percentage being higher than 0 we render additional pixels. This function returns the resolution with overscan taken into account. */
 	UFUNCTION(BlueprintCallable, Category = "Movie Render Pipeline")
diff --git a/Engine/Plugins/MovieScene/MovieRenderPipeline/Source/MovieRenderPipelineEditor/Private/MoviePipelineEditorBlueprintLibrary.cpp b/Engine/Plugins/MovieScene/MovieRenderPipeline/Source/MovieRenderPipelineEditor/Private/MoviePipelineEditorBlueprintLibrary.cpp
--- a/Engine/Plugins/MovieScene/MovieRenderPipeline/Source/MovieRenderPipelineEditor/Private/MoviePipelineEditorBlueprintLibrary.cpp
+++ b/Engine/Plugins/MovieScene/MovieRenderPipeline/Source/MovieRenderPipelineEditor/Private/MoviePipelineEditorBlueprintLibrary.cpp
@@ -233,11 +233,13 @@ FString UMoviePipelineEditorBlueprintLibrary::ResolveOutputDirectoryFromJob(UMov
 	// By having it swap {camera_name} and {shot_name} with an unresolvable tag, it will
 	// stay in the resolved path and can be removed using the code below.
 	static const FString DummyTag = TEXT("{dontresolvethis}");
+	const bool bGetNextVersion = false;
 	FMoviePipelineFilenameResolveParams Params;
 	Params.Job = InJob;
 	Params.ShotNameOverride = DummyTag;
 	Params.CameraNameOverride = DummyTag;
 	Params.InitializationTime = FDateTime::UtcNow();
+	Params.InitializationVersion = UMoviePipelineBlueprintLibrary::ResolveVersionNumber(Params, bGetNextVersion);
 
 	FString OutResolvedPath;
 	FMoviePipelineFormatArgs Dummy;
