MoviePipeline: De-duplicate stencil layer names (for both actor and data layers) before a render begins. Duplicate layer names can cause the render to fail.

diff --git a/Engine/Plugins/MovieScene/MovieRenderPipeline/Source/MovieRenderPipelineCore/Private/MoviePipelineBlueprintLibrary.cpp b/Engine/Plugins/MovieScene/MovieRenderPipeline/Source/MovieRenderPipelineCore/Private/MoviePipelineBlueprintLibrary.cpp
--- a/Engine/Plugins/MovieScene/MovieRenderPipeline/Source/MovieRenderPipelineCore/Private/MoviePipelineBlueprintLibrary.cpp
+++ b/Engine/Plugins/MovieScene/MovieRenderPipeline/Source/MovieRenderPipelineCore/Private/MoviePipelineBlueprintLibrary.cpp
@@ -380,41 +380,6 @@ UMoviePipelineQueue* UMoviePipelineBlueprintLibrary::LoadManifestFileFromString(
 	return OutQueue;
 }
 
-namespace UE
-{
-namespace MoviePipeline
-{
-	void DeduplicateNameArray(TArray<FString>& InOutNames)
-	{
-		TMap<FString, int32> NameUseCount;
-		for (FString& Name : InOutNames)
-		{
-			int32& Count = NameUseCount.FindOrAdd(Name, 0);
-			if (++Count > 1)
-			{
-				Name.Append(FString::Format(TEXT("({0})"), { NameUseCount[Name] }));
-			}
-		}
-
-		// For any shot names we found duplicates, append 1 to the first to keep naming consistent
-		for (TPair<FString, int32>& Pair : NameUseCount)
-		{
-			if (Pair.Value > 1)
-			{
-				for (FString& Name : InOutNames)
-				{
-					if (Name.Equals(Pair.Key))
-					{
-						Name.Append(TEXT("(1)"));
-						break;
-					}
-				}
-			}
-		}
-	}
-}
-}
-
 void UMoviePipelineBlueprintLibrary::UpdateJobShotListFromSequence(ULevelSequence* InSequence, UMoviePipelineExecutorJob* InJob, bool& bShotsChanged)
 {
 	if (!ensureMsgf(InSequence && InJob, TEXT("Cannot generate shot list for null sequence/job.")))
diff --git a/Engine/Plugins/MovieScene/MovieRenderPipeline/Source/MovieRenderPipelineCore/Private/MoviePipelineUtils.cpp b/Engine/Plugins/MovieScene/MovieRenderPipeline/Source/MovieRenderPipelineCore/Private/MoviePipelineUtils.cpp
--- a/Engine/Plugins/MovieScene/MovieRenderPipeline/Source/MovieRenderPipelineCore/Private/MoviePipelineUtils.cpp
+++ b/Engine/Plugins/MovieScene/MovieRenderPipeline/Source/MovieRenderPipelineCore/Private/MoviePipelineUtils.cpp
@@ -746,6 +746,35 @@ namespace UE
 			}
 		}
 
+		void DeduplicateNameArray(TArray<FString>& InOutNames)
+		{
+			TMap<FString, int32> NameUseCount;
+			for (FString& Name : InOutNames)
+			{
+				int32& Count = NameUseCount.FindOrAdd(Name, 0);
+				if (++Count > 1)
+				{
+					Name.Append(FString::Format(TEXT("({0})"), { NameUseCount[Name] }));
+				}
+			}
+
+			// For any names we found duplicates of, append (1) to the first to keep naming consistent
+			for (TPair<FString, int32>& Pair : NameUseCount)
+			{
+				if (Pair.Value > 1)
+				{
+					for (FString& Name : InOutNames)
+					{
+						if (Name.Equals(Pair.Key))
+						{
+							Name.Append(TEXT("(1)"));
+							break;
+						}
+					}
+				}
+			}
+		}
+
 		void RemoveFrameNumberFormatStrings(FString& InOutFilenameFormatString, const bool bIncludeShots)
 		{
 			// Strip {frame_number} related separators from their file name, otherwise it will create one output file per frame.
diff --git a/Engine/Plugins/MovieScene/MovieRenderPipeline/Source/MovieRenderPipelineCore/Public/MoviePipelineUtils.h b/Engine/Plugins/MovieScene/MovieRenderPipeline/Source/MovieRenderPipelineCore/Public/MoviePipelineUtils.h
--- a/Engine/Plugins/MovieScene/MovieRenderPipeline/Source/MovieRenderPipelineCore/Public/MoviePipelineUtils.h
+++ b/Engine/Plugins/MovieScene/MovieRenderPipeline/Source/MovieRenderPipelineCore/Public/MoviePipelineUtils.h
@@ -130,6 +130,8 @@ namespace UE
 	{
 		MOVIERENDERPIPELINECORE_API void ValidateOutputFormatString(FString& InOutFilenameFormatString, const bool bTestRenderPass, const bool bTestFrameNumber, const bool bIncludeCameraName = false);
 		MOVIERENDERPIPELINECORE_API void RemoveFrameNumberFormatStrings(FString& InOutFilenameFormatString, const bool bIncludeShots);
+		/** De-duplicates the provided array of strings by appending (1), (2), etc to the end of duplicates. */
+		MOVIERENDERPIPELINECORE_API void DeduplicateNameArray(TArray<FString>& InOutNames);
 
 		MOVIERENDERPIPELINECORE_API FString GetJobAuthor(const UMoviePipelineExecutorJob* InJob);
 		MOVIERENDERPIPELINECORE_API void GetSharedFormatArguments(TMap<FString, FString>& InFilenameArguments, TMap<FString, FString>& InFileMetadata, const FDateTime& InDateTime, const int32 InVersionNumber, const UMoviePipelineExecutorJob* InJob);
diff --git a/Engine/Plugins/MovieScene/MovieRenderPipeline/Source/MovieRenderPipelineRenderPasses/Private/MoviePipelineDeferredPasses.cpp b/Engine/Plugins/MovieScene/MovieRenderPipeline/Source/MovieRenderPipelineRenderPasses/Private/MoviePipelineDeferredPasses.cpp
--- a/Engine/Plugins/MovieScene/MovieRenderPipeline/Source/MovieRenderPipelineRenderPasses/Private/MoviePipelineDeferredPasses.cpp
+++ b/Engine/Plugins/MovieScene/MovieRenderPipeline/Source/MovieRenderPipelineRenderPasses/Private/MoviePipelineDeferredPasses.cpp
@@ -265,11 +265,35 @@ void UMoviePipelineDeferredPassBase::SetupImpl(const MoviePipeline::FMoviePipeli
 			ColorFormatCVar->Set(1, EConsoleVariableFlags::ECVF_SetByConsole);
 		}
 	}
+
+	// Cache out the stencil layer names (from data layers or actor layers) and de-duplicate. If layers with the same name
+	// are provided, renders may fail, which is why the the names need to be de-duplicated.
+	if (IsUsingDataLayers())
+	{
+		for (FSoftObjectPath DataLayerAssetPath : DataLayers)
+		{
+			UDataLayerAsset* DataLayerAsset = Cast<UDataLayerAsset>(DataLayerAssetPath.TryLoad());
+			if (DataLayerAsset)
+			{
+				UniqueStencilLayerNames.Add(DataLayerAsset->GetName());
+			}
+		}
+	}
+	else
+	{
+		for (const FActorLayer& Layer : ActorLayers)
+		{
+			UniqueStencilLayerNames.Add(Layer.Name.ToString());
+		}
+	}
+	
+	UE::MoviePipeline::DeduplicateNameArray(UniqueStencilLayerNames);
 }
 
 void UMoviePipelineDeferredPassBase::TeardownImpl()
 {
 	ActivePostProcessMaterials.Reset();
+	UniqueStencilLayerNames.Reset();
 
 	for (FMultiCameraViewStateData& CameraData : CameraViewStateData)
 	{
@@ -534,16 +558,16 @@ void UMoviePipelineDeferredPassBase::RenderSample_GameThreadImpl(const FMoviePip
 			};
 
 			// Now for each stencil layer we reconfigure all the actors custom depth/stencil 
-			TArray<FString> AllStencilLayers = GetStencilLayerNames();
+			TArray<FString> AllStencilLayerNames = GetStencilLayerNames();
 			if (bAddDefaultLayer)
 			{
-				AllStencilLayers.Add(TEXT("DefaultLayer"));
+				AllStencilLayerNames.Add(TEXT("DefaultLayer"));
 			}
 
 			// If we're going to be using stencil layers, we need to cache all of the users
 			// custom stencil/depth settings since we're changing them to do the mask.
 			TMap<UPrimitiveComponent*, FStencilValues> PreviousValues;
-			if (AllStencilLayers.Num() > 0)
+			if (AllStencilLayerNames.Num() > 0)
 			{
 				for (TActorIterator<AActor> ActorItr(GetWorld()); ActorItr; ++ActorItr)
 				{
@@ -566,9 +590,9 @@ void UMoviePipelineDeferredPassBase::RenderSample_GameThreadImpl(const FMoviePip
 			}
 
 
-			for (int32 StencilLayerIndex = 0; StencilLayerIndex < AllStencilLayers.Num(); StencilLayerIndex++)
+			for (int32 StencilLayerIndex = 0; StencilLayerIndex < AllStencilLayerNames.Num(); StencilLayerIndex++)
 			{
-				const FString& LayerName = AllStencilLayers[StencilLayerIndex];
+				const FString& LayerName = AllStencilLayerNames[StencilLayerIndex];
 				FMoviePipelinePassIdentifier LayerPassIdentifier = FMoviePipelinePassIdentifier(PassIdentifierForCurrentCamera.Name + LayerName);
 				LayerPassIdentifier.CameraName = PassIdentifierForCurrentCamera.CameraName;
 
@@ -1057,27 +1081,7 @@ int32 UMoviePipelineDeferredPassBase::GetNumStencilLayers() const
 
 TArray<FString> UMoviePipelineDeferredPassBase::GetStencilLayerNames() const
 {
-	TArray<FString> LayerNames;
-	if (IsUsingDataLayers())
-	{
-		for (FSoftObjectPath DataLayerAssetPath : DataLayers)
-		{
-			UDataLayerAsset* DataLayerAsset = Cast<UDataLayerAsset>(DataLayerAssetPath.TryLoad());
-			if (DataLayerAsset)
-			{
-				LayerNames.Add(DataLayerAsset->GetName());
-			}
-		}
-	}
-	else
-	{
-		for (const FActorLayer& Layer : ActorLayers)
-		{
-			LayerNames.Add(Layer.Name.ToString());
-		}
-	}
-
-	return LayerNames;
+	return UniqueStencilLayerNames;
 }
 
 FSoftObjectPath UMoviePipelineDeferredPassBase::GetValidDataLayerByIndex(const int32 InIndex) const
diff --git a/Engine/Plugins/MovieScene/MovieRenderPipeline/Source/MovieRenderPipelineRenderPasses/Public/MoviePipelineDeferredPasses.h b/Engine/Plugins/MovieScene/MovieRenderPipeline/Source/MovieRenderPipelineRenderPasses/Public/MoviePipelineDeferredPasses.h
--- a/Engine/Plugins/MovieScene/MovieRenderPipeline/Source/MovieRenderPipelineRenderPasses/Public/MoviePipelineDeferredPasses.h
+++ b/Engine/Plugins/MovieScene/MovieRenderPipeline/Source/MovieRenderPipelineRenderPasses/Public/MoviePipelineDeferredPasses.h
@@ -182,6 +182,10 @@ public:
 	static FString StencilLayerMaterialAsset;
 	static FString DefaultDepthAsset;
 	static FString DefaultMotionVectorsAsset;
+
+private:
+	/** Cached de-duplicated stencil layer names. */
+	TArray<FString> UniqueStencilLayerNames;
 };
 
 
