MoviePipeline: Added Console Variables Editor presets into MRQ's Console Variables setting. Any number of presets can be added, and cvars specified under "Console Variables" will override the values coming from the presets.

diff --git a/Engine/Plugins/MovieScene/MovieRenderPipeline/Source/MovieRenderPipelineSettings/Private/MoviePipelineConsoleVariableSetting.cpp b/Engine/Plugins/MovieScene/MovieRenderPipeline/Source/MovieRenderPipelineSettings/Private/MoviePipelineConsoleVariableSetting.cpp
--- a/Engine/Plugins/MovieScene/MovieRenderPipeline/Source/MovieRenderPipelineSettings/Private/MoviePipelineConsoleVariableSetting.cpp
+++ b/Engine/Plugins/MovieScene/MovieRenderPipeline/Source/MovieRenderPipelineSettings/Private/MoviePipelineConsoleVariableSetting.cpp
@@ -1,10 +1,13 @@
 // Copyright Epic Games, Inc. All Rights Reserved.
 
 #include "MoviePipelineConsoleVariableSetting.h"
+
 #include "MovieRenderPipelineCoreModule.h"
 #include "HAL/IConsoleManager.h"
 #include "Engine/World.h"
 #include "Kismet/KismetSystemLibrary.h"
+#include "Misc/DefaultValueHelper.h"
+#include "Sections/MovieSceneConsoleVariableTrackInterface.h"
 
 #include UE_INLINE_GENERATED_CPP_BY_NAME(MoviePipelineConsoleVariableSetting)
 
@@ -50,12 +53,13 @@ void UMoviePipelineConsoleVariableSetting::ApplyCVarSettings(const bool bOverrid
 {
 	if (bOverrideValues)
 	{
+		MergeConsoleVariables();
 		PreviousConsoleVariableValues.Reset();
-		PreviousConsoleVariableValues.SetNumZeroed(ConsoleVariables.Num());
+		PreviousConsoleVariableValues.SetNumZeroed(MergedConsoleVariables.Num());
 	}
 
 	int32 Index = 0;
-	for(const TPair<FString, float>& KVP : ConsoleVariables)
+	for(const TPair<FString, float>& KVP : MergedConsoleVariables)
 	{
 		// We don't use the shared macro here because we want to soft-warn the user instead of tripping an ensure over missing cvar values.
 		const FString TrimmedCvar = KVP.Key.TrimStartAndEnd();
@@ -100,3 +104,42 @@ void UMoviePipelineConsoleVariableSetting::ApplyCVarSettings(const bool bOverrid
 	}
 }
 
+void UMoviePipelineConsoleVariableSetting::MergeConsoleVariables()
+{
+	MergedConsoleVariables.Reset();
+	
+	// Merge in the presets
+	for (const TScriptInterface<IMovieSceneConsoleVariableTrackInterface>& Preset : ConsoleVariablePresets)
+	{
+		if (!Preset)
+		{
+			UE_LOG(LogMovieRenderPipeline, Warning, TEXT("Invalid CVar preset specified. Ignoring."));
+			continue;
+		}
+		
+		const bool bOnlyIncludeChecked = true;
+		TArray<TTuple<FString, FString>> PresetCVars;
+		Preset->GetConsoleVariablesForTrack(bOnlyIncludeChecked, PresetCVars);
+		
+		for (const TTuple<FString, FString>& CVarPair : PresetCVars)
+		{
+			float CVarFloatValue = 0.0f;
+			if (FDefaultValueHelper::ParseFloat(CVarPair.Value, CVarFloatValue))
+			{
+				MergedConsoleVariables.Add(CVarPair.Key, CVarFloatValue);
+			}
+			else
+			{
+				UE_LOG(LogMovieRenderPipeline, Warning, TEXT("Failed to apply CVar \"%s\" (from preset \"%s\") because value could not be parsed into a float. Ignoring."),
+					*CVarPair.Key, *Preset.GetObject()->GetName());
+			}
+		}
+	}
+	
+	// Merge in the overrides
+	for (const TPair<FString, float>& KVP : ConsoleVariables)
+	{
+		MergedConsoleVariables.Add(KVP);
+	}
+}
+
diff --git a/Engine/Plugins/MovieScene/MovieRenderPipeline/Source/MovieRenderPipelineSettings/Public/MoviePipelineConsoleVariableSetting.h b/Engine/Plugins/MovieScene/MovieRenderPipeline/Source/MovieRenderPipelineSettings/Public/MoviePipelineConsoleVariableSetting.h
--- a/Engine/Plugins/MovieScene/MovieRenderPipeline/Source/MovieRenderPipelineSettings/Public/MoviePipelineConsoleVariableSetting.h
+++ b/Engine/Plugins/MovieScene/MovieRenderPipeline/Source/MovieRenderPipelineSettings/Public/MoviePipelineConsoleVariableSetting.h
@@ -4,6 +4,8 @@
 #include "MoviePipelineSetting.h"
 #include "MoviePipelineConsoleVariableSetting.generated.h"
 
+class IMovieSceneConsoleVariableTrackInterface;
+
 UCLASS(BlueprintType)
 class MOVIERENDERPIPELINESETTINGS_API UMoviePipelineConsoleVariableSetting : public UMoviePipelineSetting
 {
@@ -22,6 +24,16 @@ protected:
 	void ApplyCVarSettings(const bool bOverrideValues);
 
 public:
+	// Note that the interface is used here instead of directly using UConsoleVariablesAsset in order to not
+	// depend on the Console Variables Editor.
+	/**
+	 * An array of presets from the Console Variables Editor. The preset cvars will be applied (in the order they are
+	 * specified) before any of the cvars in "Console Variables". In other words, cvars in "Console Variables" will
+	 * take precedence over the cvars coming from these presets.
+	 */
+	UPROPERTY(EditAnywhere, BlueprintReadWrite, Category = "Settings")
+	TArray<TScriptInterface<IMovieSceneConsoleVariableTrackInterface>> ConsoleVariablePresets;
+	
 	/** 
 	* An array of key/value pairs for console variable name and the value you wish to set for that cvar.
 	* The existing value will automatically be cached and restored afterwards.
@@ -44,6 +56,13 @@ public:
 	UPROPERTY(EditAnywhere, BlueprintReadWrite, Category = "Settings")
 	TArray<FString> EndConsoleCommands;
 
+private:
+	/** Merge together preset and override cvars into MergedConsoleVariables. Discards result of a prior merge (if any). */
+	void MergeConsoleVariables();
+
 private:
 	TArray<float> PreviousConsoleVariableValues;
+
+	/** Merged result of preset cvars and override cvars. */
+	TMap<FString, float> MergedConsoleVariables;
 };
