Making the new volumetric cloud render target mode default for cinematics and it supports it all: full resolution, high quality light shaft as well as cloud on transparent.

diff --git a/Engine/Plugins/MovieScene/MovieRenderPipeline/Source/MovieRenderPipelineCore/Private/MoviePipelineGameOverrideSetting.cpp b/Engine/Plugins/MovieScene/MovieRenderPipeline/Source/MovieRenderPipelineCore/Private/MoviePipelineGameOverrideSetting.cpp
--- a/Engine/Plugins/MovieScene/MovieRenderPipeline/Source/MovieRenderPipelineCore/Private/MoviePipelineGameOverrideSetting.cpp
+++ b/Engine/Plugins/MovieScene/MovieRenderPipeline/Source/MovieRenderPipelineCore/Private/MoviePipelineGameOverrideSetting.cpp
@@ -116,8 +116,9 @@ void UMoviePipelineGameOverrideSetting::ApplyCVarSettings(const bool bOverrideVa
 	// To make sure that the skylight is always valid and consistent accross capture sessions, we enforce a full capture each frame, accepting a small GPU cost.
 	MOVIEPIPELINE_STORE_AND_OVERRIDE_CVAR_INT(PreviousSkyLightRealTimeReflectionCaptureTimeSlice, TEXT("r.SkyLight.RealTimeReflectionCapture.TimeSlice"), 0, bOverrideValues);
 
-	// To make sure that the skylight is always valid and consistent accross capture sessions, we enforce a full capture each frame, accepting a small GPU cost.
-	MOVIEPIPELINE_STORE_AND_OVERRIDE_CVAR_INT(PreviousVolumetricRenderTarget, TEXT("r.VolumetricRenderTarget"), 0, bOverrideValues);
+	// Cloud are rendered using high quality volumetric render target mode 3: per pixel tracing and composition on screen, while supporting cloud on translucent.
+	MOVIEPIPELINE_STORE_AND_OVERRIDE_CVAR_INT(PreviousVolumetricRenderTarget, TEXT("r.VolumetricRenderTarget"), 1, bOverrideValues);
+	MOVIEPIPELINE_STORE_AND_OVERRIDE_CVAR_INT(PreviousVolumetricRenderTargetMode, TEXT("r.VolumetricRenderTarget.Mode"), 3, bOverrideValues);
 
 	// To make sure that the world partition streaming doesn't end up in critical streaming performances and stops streaming low priority cells.
 	MOVIEPIPELINE_STORE_AND_OVERRIDE_CVAR_INT(PreviousIgnoreStreamingPerformance, TEXT("wp.Runtime.BlockOnSlowStreaming"), 0, bOverrideValues);
@@ -218,7 +219,8 @@ void UMoviePipelineGameOverrideSetting::BuildNewProcessCommandLineArgsImpl(TArra
 
 	InOutDeviceProfileCvars.Add(FString::Printf(TEXT("au.NeverMuteNonRealtimeAudioDevices=%d"), 1));
 	InOutDeviceProfileCvars.Add(FString::Printf(TEXT("r.SkyLight.RealTimeReflectionCapture.TimeSlice=%d"), 0));
-	InOutDeviceProfileCvars.Add(FString::Printf(TEXT("r.VolumetricRenderTarget=%d"), 0));
+	InOutDeviceProfileCvars.Add(FString::Printf(TEXT("r.VolumetricRenderTarget=%d"), 1));
+	InOutDeviceProfileCvars.Add(FString::Printf(TEXT("r.VolumetricRenderTarget.Mode=%d"), 3));
 	InOutDeviceProfileCvars.Add(FString::Printf(TEXT("wp.Runtime.BlockOnSlowStreaming=%d"), 0));
 	InOutDeviceProfileCvars.Add(FString::Printf(TEXT("p.Chaos.ImmPhys.MinStepTime=%d"), 0));
 	InOutDeviceProfileCvars.Add(FString::Printf(TEXT("r.SkipRedundantTransformUpdate=%d"), 0));
diff --git a/Engine/Plugins/MovieScene/MovieRenderPipeline/Source/MovieRenderPipelineCore/Public/MoviePipelineGameOverrideSetting.h b/Engine/Plugins/MovieScene/MovieRenderPipeline/Source/MovieRenderPipelineCore/Public/MoviePipelineGameOverrideSetting.h
--- a/Engine/Plugins/MovieScene/MovieRenderPipeline/Source/MovieRenderPipelineCore/Public/MoviePipelineGameOverrideSetting.h
+++ b/Engine/Plugins/MovieScene/MovieRenderPipeline/Source/MovieRenderPipelineCore/Public/MoviePipelineGameOverrideSetting.h
@@ -138,6 +138,7 @@ private:
 	int32 PreviousNeverMuteNonRealtimeAudio;
 	int32 PreviousSkyLightRealTimeReflectionCaptureTimeSlice;
 	int32 PreviousVolumetricRenderTarget;
+	int32 PreviousVolumetricRenderTargetMode;
 	int32 PreviousIgnoreStreamingPerformance;
 	int32 PreviousStreamingManagerSyncState;
 #if WITH_EDITOR
