MoviePipeline: Fixed the cloth simulation undulating jitter when Temporal Samples are enabled, caused by the cloth simulation time step smoothing being enabled during renders.

diff --git a/Engine/Plugins/MovieScene/MovieRenderPipeline/Source/MovieRenderPipelineCore/Private/MoviePipelineGameOverrideSetting.cpp b/Engine/Plugins/MovieScene/MovieRenderPipeline/Source/MovieRenderPipelineCore/Private/MoviePipelineGameOverrideSetting.cpp
--- a/Engine/Plugins/MovieScene/MovieRenderPipeline/Source/MovieRenderPipelineCore/Private/MoviePipelineGameOverrideSetting.cpp
+++ b/Engine/Plugins/MovieScene/MovieRenderPipeline/Source/MovieRenderPipelineCore/Private/MoviePipelineGameOverrideSetting.cpp
@@ -130,6 +130,9 @@ void UMoviePipelineGameOverrideSetting::ApplyCVarSettings(const bool bOverrideVa
 	// which causes objects to render in the wrong position on the first frame (and without motion blur). This disables an optimization that detects
 	// the redundant updates so the update will get sent through anyways even though it thinks it's a duplicate (but it's not).
 	MOVIEPIPELINE_STORE_AND_OVERRIDE_CVAR_INT(PreviousSkipRedundantTransformUpdate, TEXT("r.SkipRedundantTransformUpdate"), 0, bOverrideValues);
+
+	// Cloth's time step smoothing messes up the change in number of simulation substeps that fixes the cloth simulation behavior when using Temporal Samples.
+	MOVIEPIPELINE_STORE_AND_OVERRIDE_CVAR_INT(PreviousChaosClothUseTimeStepSmoothing, TEXT("p.ChaosCloth.UseTimeStepSmoothing"), 0, bOverrideValues);
 }
 
 void UMoviePipelineGameOverrideSetting::BuildNewProcessCommandLineArgsImpl(TArray<FString>& InOutUnrealURLParams, TArray<FString>& InOutCommandLineArgs, TArray<FString>& InOutDeviceProfileCvars, TArray<FString>& InOutExecCmds) const
@@ -224,4 +227,5 @@ void UMoviePipelineGameOverrideSetting::BuildNewProcessCommandLineArgsImpl(TArra
 	InOutDeviceProfileCvars.Add(FString::Printf(TEXT("wp.Runtime.BlockOnSlowStreaming=%d"), 0));
 	InOutDeviceProfileCvars.Add(FString::Printf(TEXT("p.Chaos.ImmPhys.MinStepTime=%d"), 0));
 	InOutDeviceProfileCvars.Add(FString::Printf(TEXT("r.SkipRedundantTransformUpdate=%d"), 0));
+	InOutDeviceProfileCvars.Add(FString::Printf(TEXT("p.ChaosCloth.UseTimeStepSmoothing=%d"), 0));
 }
diff --git a/Engine/Plugins/MovieScene/MovieRenderPipeline/Source/MovieRenderPipelineCore/Public/MoviePipelineGameOverrideSetting.h b/Engine/Plugins/MovieScene/MovieRenderPipeline/Source/MovieRenderPipelineCore/Public/MoviePipelineGameOverrideSetting.h
--- a/Engine/Plugins/MovieScene/MovieRenderPipeline/Source/MovieRenderPipelineCore/Public/MoviePipelineGameOverrideSetting.h
+++ b/Engine/Plugins/MovieScene/MovieRenderPipeline/Source/MovieRenderPipelineCore/Public/MoviePipelineGameOverrideSetting.h
@@ -147,4 +147,5 @@ private:
 #endif
 	float PreviousChaosImmPhysicsMinStepTime;
 	int32 PreviousSkipRedundantTransformUpdate;
+	int32 PreviousChaosClothUseTimeStepSmoothing;
 };
\ No newline at end of file
