MoviePipeline: In the Console Variables setting, warn if the sequence having its settings modified has at least one CVar track in the Sequencer. Sub-sequences will also be searched for CVar tracks. For the warning to show up, the CVar track must not be muted, and contain an active section w/ CVars set in it.

diff --git a/Engine/Plugins/MovieScene/MovieRenderPipeline/Source/MovieRenderPipelineSettings/Private/MoviePipelineConsoleVariableSetting.cpp b/Engine/Plugins/MovieScene/MovieRenderPipeline/Source/MovieRenderPipelineSettings/Private/MoviePipelineConsoleVariableSetting.cpp
--- a/Engine/Plugins/MovieScene/MovieRenderPipeline/Source/MovieRenderPipelineSettings/Private/MoviePipelineConsoleVariableSetting.cpp
+++ b/Engine/Plugins/MovieScene/MovieRenderPipeline/Source/MovieRenderPipelineSettings/Private/MoviePipelineConsoleVariableSetting.cpp
@@ -2,15 +2,23 @@
 
 #include "MoviePipelineConsoleVariableSetting.h"
 
+#include "MoviePipelineQueue.h"
 #include "MovieRenderPipelineCoreModule.h"
+#include "MovieScene.h"
+#include "MovieSceneCommonHelpers.h"
 #include "HAL/IConsoleManager.h"
 #include "Engine/World.h"
 #include "Kismet/KismetSystemLibrary.h"
 #include "Misc/DefaultValueHelper.h"
 #include "Sections/MovieSceneConsoleVariableTrackInterface.h"
+#include "Sections/MovieSceneCVarSection.h"
+#include "Tracks/MovieSceneCVarTrack.h"
+#include "Tracks/MovieSceneSubTrack.h"
 
 #include UE_INLINE_GENERATED_CPP_BY_NAME(MoviePipelineConsoleVariableSetting)
 
+#define LOCTEXT_NAMESPACE "MoviePipelineConsoleVariableSetting"
+
 namespace UE
 {
 	namespace MoviePipeline
@@ -36,9 +44,97 @@ namespace UE
 				InCVar->SetWithCurrentPriority(InValue);
 			}
 		}
+
+		/**
+		 * Determine if the given UMovieScene contains any CVar tracks that are not muted, and also have an
+		 * active section with CVars that are set. Sub-sequences will be searched for CVar tracks as well.
+		 */
+		static bool IsCVarTrackPresent(const UMovieScene* InMovieScene)
+		{
+			if (!InMovieScene)
+			{
+				return false;
+			}
+
+			for (UMovieSceneTrack* Track : InMovieScene->GetTracks())
+			{
+				// Process CVar tracks. Return immediately if any of the CVar tracks contain CVars that are set.
+				// If this is the case, sub tracks don't need to be searched.
+				if (Track->IsA<UMovieSceneCVarTrack>())
+				{
+					const UMovieSceneCVarTrack* CVarTrack = Cast<UMovieSceneCVarTrack>(Track);
+					for (const UMovieSceneSection* Section : CVarTrack->GetAllSections())
+					{
+						const UMovieSceneCVarSection* CVarSection = Cast<UMovieSceneCVarSection>(Section);
+						if (!CVarSection || !MovieSceneHelpers::IsSectionKeyable(CVarSection))
+						{
+							continue;
+						}
+						
+						// Does this CVar track have anything in it?
+						if (!CVarSection->ConsoleVariableCollections.IsEmpty() || !CVarSection->ConsoleVariables.ValuesByCVar.IsEmpty())
+						{
+							return true;
+						}
+					}
+				}
+				
+				// Process sub tracks (which could potentially contain other sequences with CVar tracks)
+				if (Track->IsA<UMovieSceneSubTrack>())
+				{
+					const UMovieSceneSubTrack* SubTrack = Cast<UMovieSceneSubTrack>(Track);
+					for (const UMovieSceneSection* Section : SubTrack->GetAllSections())
+					{
+						const UMovieSceneSubSection* SubSection = Cast<UMovieSceneSubSection>(Section);
+						if (!SubSection || !MovieSceneHelpers::IsSectionKeyable(SubSection))
+						{
+							continue;
+						}
+
+						// Recurse into sub-sequences
+						if (const UMovieSceneSequence* SubSequence = SubSection->GetSequence())
+						{
+							if (IsCVarTrackPresent(SubSequence->GetMovieScene()))
+							{
+								return true;
+							}
+						}
+					}
+				}
+			}
+			
+			return false;
+		}
+	}
+}
+
+#if WITH_EDITOR
+
+FText UMoviePipelineConsoleVariableSetting::GetFooterText(UMoviePipelineExecutorJob* InJob) const
+{
+	if (!InJob)
+	{
+		return FText();
+	}
+	
+	const ULevelSequence* LoadedSequence = Cast<ULevelSequence>(InJob->Sequence.TryLoad());
+	if (!LoadedSequence)
+	{
+		return FText();
+	}
+	
+	if (!UE::MoviePipeline::IsCVarTrackPresent(LoadedSequence->MovieScene))
+	{
+		return FText();
 	}
+	
+	return FText(LOCTEXT(
+		"SequencerCvarWarning",
+		"The current job contains a Level Sequence with a Console Variables Track, additional settings are configured in Sequencer."));
 }
 
+#endif // WITH_EDITOR
+
 void UMoviePipelineConsoleVariableSetting::SetupForPipelineImpl(UMoviePipeline* InPipeline)
 {
 	ApplyCVarSettings(true);
@@ -143,3 +239,4 @@ void UMoviePipelineConsoleVariableSetting::MergeConsoleVariables()
 	}
 }
 
+#undef LOCTEXT_NAMESPACE
\ No newline at end of file
diff --git a/Engine/Plugins/MovieScene/MovieRenderPipeline/Source/MovieRenderPipelineSettings/Public/MoviePipelineConsoleVariableSetting.h b/Engine/Plugins/MovieScene/MovieRenderPipeline/Source/MovieRenderPipelineSettings/Public/MoviePipelineConsoleVariableSetting.h
--- a/Engine/Plugins/MovieScene/MovieRenderPipeline/Source/MovieRenderPipelineSettings/Public/MoviePipelineConsoleVariableSetting.h
+++ b/Engine/Plugins/MovieScene/MovieRenderPipeline/Source/MovieRenderPipelineSettings/Public/MoviePipelineConsoleVariableSetting.h
@@ -15,6 +15,7 @@ public:
 public:
 #if WITH_EDITOR
 	virtual FText GetDisplayText() const override { return NSLOCTEXT("MovieRenderPipeline", "ConsoleVariableSettingDisplayName", "Console Variables"); }
+	virtual FText GetFooterText(UMoviePipelineExecutorJob* InJob) const override;
 #endif
 	virtual bool IsValidOnShots() const override { return true; }
 	virtual bool IsValidOnPrimary() const override { return true; }
