MoviePipeline: Made a number of changes to the save/load preset experience in the preset editing dialog.
* The "Load/Save Preset" button text has been updated to reflect the preset being edited, with an "*" if it has unsaved changes.
* Provide "Save Preset" (in addition to "Save As"). Transient presets always have "Save" available.
* Provide the name of the preset being saved in the save drop-down, or "None" if the preset has never been saved.
* Warn the user if changes may be lost when loading another preset.
* Removed an unneeded scoped transaction that had the potential to cause crashes in the future.

In order to make some of these updates possible, changes to how the origin preset was tracked were made such that the dialog
always knows if it is editing a preset which has been saved in the past.

diff --git a/Engine/Plugins/MovieScene/MovieRenderPipeline/Source/MovieRenderPipelineCore/Private/MoviePipelineConfigBase.cpp b/Engine/Plugins/MovieScene/MovieRenderPipeline/Source/MovieRenderPipelineCore/Private/MoviePipelineConfigBase.cpp
--- a/Engine/Plugins/MovieScene/MovieRenderPipeline/Source/MovieRenderPipelineCore/Private/MoviePipelineConfigBase.cpp
+++ b/Engine/Plugins/MovieScene/MovieRenderPipeline/Source/MovieRenderPipelineCore/Private/MoviePipelineConfigBase.cpp
@@ -19,6 +19,11 @@ void UMoviePipelineConfigBase::RemoveSetting(UMoviePipelineSetting* InSetting)
 
 void UMoviePipelineConfigBase::CopyFrom(UMoviePipelineConfigBase* InConfig)
 {
+	if (!InConfig)
+	{
+		return;
+	}
+	
 	// Don't allow copying ourself to ourself as we first wipe out our settings array
 	// and then copy from the other config's settings array which will now be empty.
 	if (InConfig == this)
@@ -26,6 +31,8 @@ void UMoviePipelineConfigBase::CopyFrom(UMoviePipelineConfigBase* InConfig)
 		return;
 	}
 
+	ConfigOrigin = InConfig->ConfigOrigin;
+
 #if WITH_EDITOR
 	Modify();
 #endif
diff --git a/Engine/Plugins/MovieScene/MovieRenderPipeline/Source/MovieRenderPipelineCore/Public/MoviePipelineConfigBase.h b/Engine/Plugins/MovieScene/MovieRenderPipeline/Source/MovieRenderPipelineCore/Public/MoviePipelineConfigBase.h
--- a/Engine/Plugins/MovieScene/MovieRenderPipeline/Source/MovieRenderPipelineCore/Public/MoviePipelineConfigBase.h
+++ b/Engine/Plugins/MovieScene/MovieRenderPipeline/Source/MovieRenderPipelineCore/Public/MoviePipelineConfigBase.h
@@ -46,6 +46,19 @@ public:
 	UFUNCTION(BlueprintCallable, Category = "Movie Render Pipeline")
 	virtual TArray<UMoviePipelineSetting*> GetUserSettings() const { return Settings; }
 
+	/**
+	 * Gets the config that this config was originally based on (if any). Helpful for determining where a transient
+	 * config came from.
+	 */
+	UFUNCTION(BlueprintCallable, Category = "Movie Render Pipeline")
+	UMoviePipelineConfigBase* GetConfigOrigin() const { return ConfigOrigin.LoadSynchronous(); }
+
+	/**
+	 * Sets the config that this config originated from (if any).
+	 */
+	UFUNCTION(BlueprintCallable, Category = "Movie Render Pipeline")
+	void SetConfigOrigin(UMoviePipelineConfigBase* InConfig) { ConfigOrigin = InConfig; }
+
 public:
 	template<typename SettingType>
 	TArray<SettingType*> FindSettings(const bool bIncludeDisabledSettings = false) const
@@ -170,5 +183,9 @@ protected:
 	TArray<TObjectPtr<UMoviePipelineSetting>> Settings;
 
 private:
+	/** The config that this config originated from. Helpful for transient configs. */
+	UPROPERTY()
+	TSoftObjectPtr<UMoviePipelineConfigBase> ConfigOrigin;
+	
 	int32 SettingsSerialNumber;
 };
\ No newline at end of file
diff --git a/Engine/Plugins/MovieScene/MovieRenderPipeline/Source/MovieRenderPipelineCore/Public/MoviePipelineQueue.h b/Engine/Plugins/MovieScene/MovieRenderPipeline/Source/MovieRenderPipelineCore/Public/MoviePipelineQueue.h
--- a/Engine/Plugins/MovieScene/MovieRenderPipeline/Source/MovieRenderPipelineCore/Public/MoviePipelineQueue.h
+++ b/Engine/Plugins/MovieScene/MovieRenderPipeline/Source/MovieRenderPipelineCore/Public/MoviePipelineQueue.h
@@ -338,12 +338,23 @@ public:
 	UFUNCTION(BlueprintCallable, Category = "Movie Render Pipeline")
 	void SetPresetOrigin(UMoviePipelinePrimaryConfig* InPreset);
 
+	/**
+	 * Gets the preset for this job, but only if the preset has not been modified. If it has been modified, or the preset
+	 * no longer exists, returns nullptr.
+	 * @see GetConfiguration()
+	 */
 	UFUNCTION(BlueprintPure, Category = "Movie Render Pipeline")
 	UMoviePipelinePrimaryConfig* GetPresetOrigin() const 
 	{
 		return PresetOrigin.LoadSynchronous();
 	}
 
+	/**
+	 * Gets the configuration for the job. This configuration is either standalone (not associated with any preset), or
+	 * contains a copy of the preset origin plus any modifications made on top of it. If the preset that this
+	 * configuration was originally based on no longer exists, this configuration will still be valid.
+	 * @see GetPresetOrigin()
+	 */
 	UFUNCTION(BlueprintPure, Category = "Movie Render Pipeline")
 	UMoviePipelinePrimaryConfig* GetConfiguration() const
 	{
diff --git a/Engine/Plugins/MovieScene/MovieRenderPipeline/Source/MovieRenderPipelineEditor/Private/Widgets/SMoviePipelineConfigPanel.cpp b/Engine/Plugins/MovieScene/MovieRenderPipeline/Source/MovieRenderPipelineEditor/Private/Widgets/SMoviePipelineConfigPanel.cpp
--- a/Engine/Plugins/MovieScene/MovieRenderPipeline/Source/MovieRenderPipelineEditor/Private/Widgets/SMoviePipelineConfigPanel.cpp
+++ b/Engine/Plugins/MovieScene/MovieRenderPipeline/Source/MovieRenderPipelineEditor/Private/Widgets/SMoviePipelineConfigPanel.cpp
@@ -60,11 +60,12 @@ void SMoviePipelineConfigPanel::Construct(const FArguments& InArgs, TSubclassOf<
 		if (InArgs._BasePreset)
 		{
 			TransientPreset->CopyFrom(InArgs._BasePreset);
-			PresetUsedIfNotModified = InArgs._BasePreset;
+			OriginPreset = InArgs._BasePreset;
 		}
 		else if (InArgs._BaseConfig)
 		{
 			TransientPreset->CopyFrom(InArgs._BaseConfig);
+			OriginPreset = nullptr;
 		}
 	}
 	else
@@ -154,7 +155,7 @@ void SMoviePipelineConfigPanel::Construct(const FArguments& InArgs, TSubclassOf<
 						.Padding(0, 1, 0, 0)
 						[
 							SNew(STextBlock)
-							.Text(LOCTEXT("PresetsToolbarButton", "Load/Save Preset"))
+							.Text(this, &SMoviePipelineConfigPanel::GetPresetsMenuButtonText)
 						]
 					]
 				]
@@ -239,6 +240,19 @@ void SMoviePipelineConfigPanel::Construct(const FArguments& InArgs, TSubclassOf<
 
 UE_ENABLE_OPTIMIZATION_SHIP
 
+FText SMoviePipelineConfigPanel::GetPresetsMenuButtonText() const
+{
+	const FString PresetName = GetConfigOriginName();
+	if (!PresetName.IsEmpty())
+	{
+		return FText::Format(FText::FromString(TEXT("{0}{1}")),
+			FText::FromString(PresetName),
+			FText::FromString(IsConfigDirty() ? TEXT(" *") : TEXT("")));
+	}
+	
+	return LOCTEXT("PresetsSaveMenuUnsavedConfig_Text", "Unsaved Config *");
+}
+
 FText SMoviePipelineConfigPanel::GetConfigTypeLabel() const
 {
 	return LOCTEXT("PrimaryConfigType_Label", "PRIMARY");
@@ -251,14 +265,21 @@ FReply SMoviePipelineConfigPanel::OnConfirmChanges()
 	// accept will simply assign that preset to the resulting job. Because the UI
 	// works on a copy it is hard to tell so we separately track the last preset-based
 	// action and then clear it if they've modified 'it' (which is actually our UI copy).
-	if (PresetUsedIfNotModified.IsValid())
+	if (OriginPreset.IsValid() && !bHasModifiedLoadedPreset)
 	{
 		// They just want to use the preset. We'll call a different callback so it's easier for
 		// the caller to determine the intent.
-		OnConfigurationSetToPreset.ExecuteIfBound(WeakJob, WeakShot, PresetUsedIfNotModified.Get());
+		OnConfigurationSetToPreset.ExecuteIfBound(WeakJob, WeakShot, OriginPreset.Get());
 	}
 	else
 	{
+		// If the panel was originally loaded up with an origin preset, but changes have been
+		// made, update the transient preset to indicate where it came from
+		if (OriginPreset.IsValid())
+		{
+			TransientPreset->SetConfigOrigin(OriginPreset.Get());
+		}
+		
 		// They modified the object after the last preset action, then we'll just use 
 		// the transient ui copy and the owning job will know it has been modified.
 		OnConfigurationModified.ExecuteIfBound(WeakJob, WeakShot, TransientPreset);
@@ -343,10 +364,39 @@ TSharedRef<SWidget> SMoviePipelineConfigPanel::OnGeneratePresetsMenu()
 
 	IContentBrowserSingleton& ContentBrowser = FModuleManager::LoadModuleChecked<FContentBrowserModule>("ContentBrowser").Get();
 
+	static const FText NoLoadedPresetText = LOCTEXT("NoLoadedPreset", "None");
+	static const FText LoadedPresetFormatText = LOCTEXT("LoadedPresetFormat", "Current Preset: {0}");
+	const FString PresetName = GetConfigOriginName();
+	const FText LoadedPresetName = FText::Format(
+		LoadedPresetFormatText, !PresetName.IsEmpty() ? FText::FromString(PresetName) : NoLoadedPresetText);
+
+	MenuBuilder.AddMenuEntry(
+		LoadedPresetName,
+		LoadedPresetName,
+		FSlateIcon(),
+		FUIAction(
+			FExecuteAction(),
+			FCanExecuteAction::CreateLambda([]() { return false; })),
+		NAME_None,
+		EUserInterfaceActionType::None
+	);
+
+	MenuBuilder.AddMenuSeparator();
+
+	MenuBuilder.AddMenuEntry(
+		LOCTEXT("SavePreset_Text", "Save Preset"),
+		LOCTEXT("SavePreset_Tip", "Save the current preset if one has been loaded."),
+		FSlateIcon(FAppStyle::Get().GetStyleSetName(), "AssetEditor.SaveAsset"),
+		FUIAction(
+			FExecuteAction::CreateSP(this, &SMoviePipelineConfigPanel::OnSavePreset),
+			FCanExecuteAction::CreateSP(this, &SMoviePipelineConfigPanel::CanSavePreset)
+		)
+	);
+
 	MenuBuilder.AddMenuEntry(
 		LOCTEXT("SaveAsPreset_Text", "Save As Preset"),
 		LOCTEXT("SaveAsPreset_Tip", "Save the current configuration as a new preset that can be shared between multiple jobs, or imported later as the base of a new configuration."),
-		FSlateIcon(FAppStyle::Get().GetStyleSetName(), "AssetEditor.SaveAsset"),
+		FSlateIcon(FAppStyle::Get().GetStyleSetName(), "AssetEditor.SaveAssetAs"),
 		FUIAction(FExecuteAction::CreateSP(this, &SMoviePipelineConfigPanel::OnSaveAsPreset))
 	);
 
@@ -470,6 +520,20 @@ bool SMoviePipelineConfigPanel::GetSavePresetPackageName(const FString& InExisti
 	return true;
 }
 
+void SMoviePipelineConfigPanel::OnSavePreset()
+{
+	UMoviePipelineConfigBase* ConfigOrigin = GetConfigOrigin();
+	
+	// If the preset origin is not known, then treat the "Save" as a "Save As"
+	if (!ConfigOrigin)
+	{
+		OnSaveAsPreset();
+		return;
+	}
+	
+	SaveTransientPresetToAsset(ConfigOrigin);
+}
+
 void SMoviePipelineConfigPanel::OnSaveAsPreset()
 {
 	FString PackageName;
@@ -478,23 +542,27 @@ void SMoviePipelineConfigPanel::OnSaveAsPreset()
 		return;
 	}
 
-	FScopedTransaction Transaction(LOCTEXT("SaveAsPreset", "Save As Preset"));
-
 	// Saving into a new package
 	const FString NewAssetName = FPackageName::GetLongPackageAssetName(PackageName);
-	UPackage*     NewPackage = CreatePackage(*PackageName);
+	UPackage* NewPackage = CreatePackage(*PackageName);
 	NewPackage->MarkAsFullyLoaded();
-	UMoviePipelineConfigBase*  NewPreset = NewObject<UMoviePipelineConfigBase>(NewPackage, ConfigAssetType, *NewAssetName, RF_Public | RF_Standalone | RF_Transactional);
-	
-	if (NewPreset)
+	UMoviePipelineConfigBase* NewPreset = NewObject<UMoviePipelineConfigBase>(NewPackage, ConfigAssetType, *NewAssetName, RF_Public | RF_Standalone | RF_Transactional);
+
+	SaveTransientPresetToAsset(NewPreset);
+}
+
+void SMoviePipelineConfigPanel::SaveTransientPresetToAsset(UMoviePipelineConfigBase* DestinationPreset)
+{
+	if (TransientPreset && DestinationPreset)
 	{
-		NewPreset->CopyFrom(TransientPreset);
-		NewPreset->MarkPackageDirty();
+		DestinationPreset->CopyFrom(TransientPreset);
+		DestinationPreset->SetConfigOrigin(DestinationPreset);
+		DestinationPreset->MarkPackageDirty();
 
 		// We need to be able to save an empty string into the file. When we duplicate the Output Setting in CopyFrom, the PostLoad is fired which fixes it back up
 		// to point to the local output directory. A bit unfortunate but we need to do this after the CopyFrom since it needs to only happen on asset save.
 		UMoviePipelineOutputSetting* ExistingOutputSetting = TransientPreset->FindSetting<UMoviePipelineOutputSetting>();
-		UMoviePipelineOutputSetting* NewOutputSetting = NewPreset->FindSetting<UMoviePipelineOutputSetting>();
+		UMoviePipelineOutputSetting* NewOutputSetting = DestinationPreset->FindSetting<UMoviePipelineOutputSetting>();
 		if (ExistingOutputSetting && NewOutputSetting)
 		{
 			if (ExistingOutputSetting->OutputDirectory.Path.Len() == 0)
@@ -505,23 +573,26 @@ void SMoviePipelineConfigPanel::OnSaveAsPreset()
 		}
 		
 		// Make the Display Name match the preset name so when imported later they look like the asset.
-		NewPreset->DisplayName = NewAssetName;
+		DestinationPreset->DisplayName = DestinationPreset->GetName();
 
-		FAssetRegistryModule::AssetCreated(NewPreset);
+		FAssetRegistryModule::AssetCreated(DestinationPreset);
 
-		FEditorFileUtils::EPromptReturnCode PromptReturnCode = FEditorFileUtils::PromptForCheckoutAndSave({ NewPackage }, false, false);
+		const bool bCheckDirty = false;
+		const bool bPromptToSave = false;
+		FEditorFileUtils::EPromptReturnCode PromptReturnCode = FEditorFileUtils::PromptForCheckoutAndSave({ DestinationPreset->GetPackage() }, bCheckDirty, bPromptToSave);
 		if (PromptReturnCode == FEditorFileUtils::EPromptReturnCode::PR_Success)
 		{
 			// If the user has exported a new preset, they expect that hitting accept will
 			// just use that preset directly, unless they further modify the ui. Because this
 			// ui uses a copy of the object there's no good way to compare them after closing
 			// the dialog.
-			PresetUsedIfNotModified = NewPreset;
+			bHasModifiedLoadedPreset = false;
+			OriginPreset = DestinationPreset;
 
 			// We also want to rename our transient preset. The name is used to indicate what preset
 			// it originally came from, and now the closest preset to it is not the original but the
 			// newly exported preset.
-			// @todo
+			TransientPreset->DisplayName = DestinationPreset->DisplayName;
 		}
 
 		// CockpitWidget->GetMetaData()->SetPresetOrigin(NewPreset);
@@ -536,13 +607,24 @@ void SMoviePipelineConfigPanel::OnAnyObjectModified(UObject* InModifiedObject)
 	{
 		if (InModifiedObject->IsIn(TransientPreset) || (InModifiedObject == TransientPreset))
 		{
-			PresetUsedIfNotModified = nullptr;
+			bHasModifiedLoadedPreset = true;
 		}
 	}
 }
 
 void SMoviePipelineConfigPanel::OnImportPreset(const FAssetData& InPresetAsset)
 {
+	if (IsConfigDirty())
+	{
+		const FText TitleText = LOCTEXT("UnsavedPresetWarningTitle", "Unsaved Changes to Preset");
+		const FText MessageText = LOCTEXT("UnsavedPresetWarningMessage", "The changes made to the current preset will be lost by importing another preset. Do you want to continue with this import?");
+
+		if (FMessageDialog::Open(EAppMsgType::YesNo, MessageText, &TitleText) == EAppReturnType::No)
+		{
+			return;
+		}
+	}
+	
 	FSlateApplication::Get().DismissAllMenus();
 
 	UMoviePipelineConfigBase* Preset = CastChecked<UMoviePipelineConfigBase>(InPresetAsset.GetAsset());
@@ -553,11 +635,25 @@ void SMoviePipelineConfigPanel::OnImportPreset(const FAssetData& InPresetAsset)
 		TransientPreset->Modify();
 		TransientPreset->CopyFrom(Preset);
 
-		PresetUsedIfNotModified = Preset;
+		bHasModifiedLoadedPreset = false;
+		OriginPreset = Preset;
 		// CockpitWidget->GetMetaData()->SetPresetOrigin(Take);
 	}
 }
 
+bool SMoviePipelineConfigPanel::CanSavePreset() const
+{
+	// If a transient preset is loaded, "Save" will always be available (although it may trigger a "Save As" if the
+	// transient's origin preset is not known)
+	if (!OriginPreset.IsValid() && TransientPreset)
+	{
+		return true;
+	}
+
+	// If a non-transient preset is loaded, it can be saved as long as it has been modified
+	return bHasModifiedLoadedPreset;
+}
+
 void SMoviePipelineConfigPanel::AddReferencedObjects(FReferenceCollector& Collector)
 {
 	Collector.AddReferencedObject(TransientPreset);
@@ -582,4 +678,38 @@ TSharedRef<SWidget> SMoviePipelineConfigPanel::MakeSettingsWidget()
 {
 	return MoviePipelineEditorWidget->MakeAddSettingButton();
 }
+
+bool SMoviePipelineConfigPanel::IsConfigDirty() const
+{
+	// The preset is dirty if...
+	// 1) There was never an origin preset specified, meaning this is a purely transient config, or
+	// 2) There *was* an origin preset, but it has been modified after being loaded in this dialog
+	return !OriginPreset.IsValid() || (OriginPreset.IsValid() && bHasModifiedLoadedPreset);
+}
+
+UMoviePipelineConfigBase* SMoviePipelineConfigPanel::GetConfigOrigin() const
+{
+	if (OriginPreset.IsValid())
+	{
+		return OriginPreset.Get();
+	}
+	
+	if (TransientPreset->GetConfigOrigin())
+	{
+		return TransientPreset->GetConfigOrigin();
+	}
+
+	return nullptr;
+}
+
+FString SMoviePipelineConfigPanel::GetConfigOriginName() const
+{
+	if (const UMoviePipelineConfigBase* ConfigOrigin = GetConfigOrigin())
+	{
+		return ConfigOrigin->DisplayName;
+	}
+
+	return FString();
+}
+
 #undef LOCTEXT_NAMESPACE // SMoviePipelinePanel
\ No newline at end of file
diff --git a/Engine/Plugins/MovieScene/MovieRenderPipeline/Source/MovieRenderPipelineEditor/Private/Widgets/SMoviePipelineConfigPanel.h b/Engine/Plugins/MovieScene/MovieRenderPipeline/Source/MovieRenderPipelineEditor/Private/Widgets/SMoviePipelineConfigPanel.h
--- a/Engine/Plugins/MovieScene/MovieRenderPipeline/Source/MovieRenderPipelineEditor/Private/Widgets/SMoviePipelineConfigPanel.h
+++ b/Engine/Plugins/MovieScene/MovieRenderPipeline/Source/MovieRenderPipelineEditor/Private/Widgets/SMoviePipelineConfigPanel.h
@@ -71,6 +71,7 @@ private:
 	/** Generate the widget that is visible in the Choose Preset dropdown. */
 	TSharedRef<SWidget> OnGeneratePresetsMenu();
 	
+	FText GetPresetsMenuButtonText() const;
 	FText GetConfigTypeLabel() const;
 
 	FReply OnCancelChanges();
@@ -86,20 +87,38 @@ private:
 	
 	/** When a user wants to import an existing preset asset over the current config. */
 	void OnImportPreset(const FAssetData& InPresetAsset);
-	/** Save the current configuration out to an asset. */
+
+	/** Returns true if the user can save the current preset, else false. */
+	bool CanSavePreset() const;
+	
+	/** Saves the current preset to its existing asset, or triggers a "Save As" if a transient preset is being saved. */
+	void OnSavePreset();
+	
+	/** Saves the current preset out to a new or existing asset. */
 	void OnSaveAsPreset();
 
+	/** Saves the transient preset to DestinationPreset. */
+	void SaveTransientPresetToAsset(UMoviePipelineConfigBase* DestinationPreset);
+
 	FText GetValidationWarningText() const;
 
 public:
 	TSharedRef<SWidget> MakeSettingsWidget();
 
+private:
+	bool IsConfigDirty() const;
+	UMoviePipelineConfigBase* GetConfigOrigin() const;
+	FString GetConfigOriginName() const;
+
 private:
 	/** The transient preset that we use - kept alive by AddReferencedObjects */
-	UMoviePipelineConfigBase* TransientPreset;
+	UMoviePipelineConfigBase* TransientPreset = nullptr;
+
+	/** The preset that is in use, if any. Will be null if the preset was never saved and is only transient. */
+	TSoftObjectPtr<UMoviePipelineConfigBase> OriginPreset;
 
-	/** This is set each time the user performs an action that makes them feel like they've used a specific preset in this UI. */
-	TSoftObjectPtr<UMoviePipelineConfigBase> PresetUsedIfNotModified;
+	/** Whether the preset (transient or otherwise) has been modified while loaded in this panel. */
+	bool bHasModifiedLoadedPreset = false;
 
 	/** The job this editing panel is for. Kept alive externally. */
 	TWeakObjectPtr<UMoviePipelineExecutorJob> WeakJob;
