MoviePipeline: Fixed a bug which caused a job in the queue to sometimes have an incorrect display name for settings (when the configuration is based on a modified preset).

Upon duplication of a config object, the display name is reset (which happens during save of a queue). Instead of relying on the config's
display name, use the display name from the config's preset origin.

diff --git a/Engine/Plugins/MovieScene/MovieRenderPipeline/Source/MovieRenderPipelineEditor/Private/Widgets/SMoviePipelineQueueEditor.cpp b/Engine/Plugins/MovieScene/MovieRenderPipeline/Source/MovieRenderPipelineEditor/Private/Widgets/SMoviePipelineQueueEditor.cpp
--- a/Engine/Plugins/MovieScene/MovieRenderPipeline/Source/MovieRenderPipelineEditor/Private/Widgets/SMoviePipelineQueueEditor.cpp
+++ b/Engine/Plugins/MovieScene/MovieRenderPipeline/Source/MovieRenderPipelineEditor/Private/Widgets/SMoviePipelineQueueEditor.cpp
@@ -203,10 +203,18 @@ public:
 		UMoviePipelineExecutorJob* Job = WeakJob.Get();
 		if (Job)
 		{
+			// If the job has a preset origin (ie, its config is based off a preset w/o any modifications), use its
+			// display name. If the config has a preset origin (ie, it's based off a preset, but has modifications), use
+			// that display name. Otherwise, fall back to the config's display name.
 			UMoviePipelineConfigBase* Config = Job->GetPresetOrigin();
 			if (!Config)
 			{
 				Config = Job->GetConfiguration();
+
+				if (Config && Config->GetConfigOrigin())
+				{
+					Config = Config->GetConfigOrigin();
+				}
 			}
 
 			if (Config)
@@ -794,10 +802,18 @@ struct FMoviePipelineShotItem : IMoviePipelineQueueTreeItem
 		UMoviePipelineExecutorShot* Shot = WeakShot.Get();
 		if (Shot)
 		{
+			// If the shot has a preset origin (ie, its config is based off a preset w/o any modifications), use its
+			// display name. If the config has a preset origin (ie, it's based off a preset, but has modifications), use
+			// that display name. Otherwise, fall back to the config's display name.
 			UMoviePipelineShotConfig* Config = Shot->GetShotOverridePresetOrigin();
 			if (!Config)
 			{
 				Config = Shot->GetShotOverrideConfiguration();
+
+				if (Config && Config->GetConfigOrigin())
+				{
+					Config = Cast<UMoviePipelineShotConfig>(Config->GetConfigOrigin());
+				}
 			}
 
 			if (Config)
