MoviePipeline: Fix PVS Studio static analysis warning from CL 23499286. Addresses part of Horde 244907.

Marking a preset's origin as itself is not an error. However, in this particular scenario, it was determined to not be needed.

diff --git a/Engine/Plugins/MovieScene/MovieRenderPipeline/Source/MovieRenderPipelineEditor/Private/Widgets/SMoviePipelineConfigPanel.cpp b/Engine/Plugins/MovieScene/MovieRenderPipeline/Source/MovieRenderPipelineEditor/Private/Widgets/SMoviePipelineConfigPanel.cpp
--- a/Engine/Plugins/MovieScene/MovieRenderPipeline/Source/MovieRenderPipelineEditor/Private/Widgets/SMoviePipelineConfigPanel.cpp
+++ b/Engine/Plugins/MovieScene/MovieRenderPipeline/Source/MovieRenderPipelineEditor/Private/Widgets/SMoviePipelineConfigPanel.cpp
@@ -556,7 +556,6 @@ void SMoviePipelineConfigPanel::SaveTransientPresetToAsset(UMoviePipelineConfigB
 	if (TransientPreset && DestinationPreset)
 	{
 		DestinationPreset->CopyFrom(TransientPreset);
-		DestinationPreset->SetConfigOrigin(DestinationPreset);
 		DestinationPreset->MarkPackageDirty();
 
 		// We need to be able to save an empty string into the file. When we duplicate the Output Setting in CopyFrom, the PostLoad is fired which fixes it back up
