MoviePipeline: Added a priority integer to settings classes to work around a scenario where the UMoviePipelineGameOverride setting would apply after the UMoviePipelineConsoleVariable setting which could then flip the cvar values away from what the user explicitly tried to set in the ConsoleVariable setting.

diff --git a/Engine/Plugins/MovieScene/MovieRenderPipeline/Source/MovieRenderPipelineCore/Private/MoviePipeline.cpp b/Engine/Plugins/MovieScene/MovieRenderPipeline/Source/MovieRenderPipelineCore/Private/MoviePipeline.cpp
--- a/Engine/Plugins/MovieScene/MovieRenderPipeline/Source/MovieRenderPipelineCore/Private/MoviePipeline.cpp
+++ b/Engine/Plugins/MovieScene/MovieRenderPipeline/Source/MovieRenderPipelineCore/Private/MoviePipeline.cpp
@@ -231,7 +231,9 @@ void UMoviePipeline::Initialize(UMoviePipelineExecutorJob* InJob)
 	}
 	
 	// Initialize all of our primary config settings. Shot specific ones will be called for their appropriate shot.
-	for (UMoviePipelineSetting* Setting : GetPipelinePrimaryConfig()->GetAllSettings())
+	TArray<UMoviePipelineSetting*> Settings = GetPipelinePrimaryConfig()->GetAllSettings();
+	Algo::SortBy(Settings, [](const UMoviePipelineSetting* Setting) { return Setting->GetPriority(); });
+	for (UMoviePipelineSetting* Setting : Settings)
 	{
 		Setting->OnMoviePipelineInitialized(this);
 	}
@@ -528,8 +530,10 @@ void UMoviePipeline::TransitionToState(const EMovieRenderPipelineState InNewStat
 			bInvalidTransition = false;
 			PipelineState = InNewState;
 
-			// Uninitialize our primary config settings.
-			for (UMoviePipelineSetting* Setting : GetPipelinePrimaryConfig()->GetAllSettings())
+			// Uninitialize our primary config settings. Reverse sorted so settings that cached values restore correctly.
+			TArray<UMoviePipelineSetting*> Settings = GetPipelinePrimaryConfig()->GetAllSettings();
+			Algo::SortBy(Settings, [](const UMoviePipelineSetting* Setting) { return Setting->GetPriority(); },  TLess<int32>());
+			for (UMoviePipelineSetting* Setting : Settings)
 			{
 				Setting->OnMoviePipelineShutdown(this);
 			}
@@ -547,7 +551,9 @@ void UMoviePipeline::TransitionToState(const EMovieRenderPipelineState InNewStat
 				DebugWidget = nullptr;
 			}
 
-			for (UMoviePipelineOutputBase* Setting : GetPipelinePrimaryConfig()->GetOutputContainers())
+			TArray<UMoviePipelineOutputBase*> ContainerSettings = GetPipelinePrimaryConfig()->GetOutputContainers();
+			Algo::SortBy(ContainerSettings, [](const UMoviePipelineOutputBase* Setting) { return Setting->GetPriority(); });
+			for (UMoviePipelineOutputBase* Setting : ContainerSettings)
 			{
 				Setting->OnPipelineFinished();
 			}
@@ -688,7 +694,10 @@ void UMoviePipeline::BeginFinalize()
 {
 	// Notify all of our output containers that we have finished producing and
 	// submitting all frames to them and that they should start any async flushes.
-	for (UMoviePipelineOutputBase* Container : GetPipelinePrimaryConfig()->GetOutputContainers())
+	TArray<UMoviePipelineOutputBase*> Settings = GetPipelinePrimaryConfig()->GetOutputContainers();
+	Algo::SortBy(Settings, [](const UMoviePipelineOutputBase* Setting) { return Setting->GetPriority(); });
+
+	for (UMoviePipelineOutputBase* Container : Settings)
 	{
 		Container->BeginFinalize();
 	}
@@ -696,7 +705,9 @@ void UMoviePipeline::BeginFinalize()
 
 void UMoviePipeline::BeginExport()
 {
-	for (UMoviePipelineSetting* Setting : GetPipelinePrimaryConfig()->GetAllSettings())
+	TArray<UMoviePipelineSetting*> Settings = GetPipelinePrimaryConfig()->GetAllSettings();
+	Algo::SortBy(Settings, [](const UMoviePipelineSetting* Setting) { return Setting->GetPriority(); });
+	for (UMoviePipelineSetting* Setting : Settings)
 	{
 		Setting->BeginExport();
 	}
@@ -737,7 +748,9 @@ void UMoviePipeline::TickFinalizeOutputContainers(const bool bInForceFinish)
 		return;
 	}
 
-	for (UMoviePipelineOutputBase* Container : GetPipelinePrimaryConfig()->GetOutputContainers())
+	TArray<UMoviePipelineOutputBase*> Settings = GetPipelinePrimaryConfig()->GetOutputContainers();
+	Algo::SortBy(Settings, [](const UMoviePipelineOutputBase* Setting) { return Setting->GetPriority(); });
+	for (UMoviePipelineOutputBase* Container : Settings)
 	{
 		// All containers have finished processing, final shutdown.
 		Container->Finalize();
@@ -952,7 +965,9 @@ void UMoviePipeline::InitializeShot(UMoviePipelineExecutorShot* InShot)
 	SetSoloShot(InShot);
 
 	// Loop through just our primary settings and let them know which shot we're about to start.
-	for (UMoviePipelineSetting* Setting : GetPipelinePrimaryConfig()->GetAllSettings())
+	TArray<UMoviePipelineSetting*> Settings = GetPipelinePrimaryConfig()->GetAllSettings();
+	Algo::SortBy(Settings, [](const UMoviePipelineSetting* Setting) { return Setting->GetPriority(); });
+	for (UMoviePipelineSetting* Setting : Settings)
 	{
 		Setting->OnSetupForShot(InShot);
 	}
@@ -960,7 +975,9 @@ void UMoviePipeline::InitializeShot(UMoviePipelineExecutorShot* InShot)
 	if (InShot->GetShotOverrideConfiguration() != nullptr)
 	{
 		// Any shot-specific overrides haven't had first time initialization. So we'll do that now.
-		for (UMoviePipelineSetting* Setting : InShot->GetShotOverrideConfiguration()->GetUserSettings())
+		TArray<UMoviePipelineSetting*> ShotSettings = InShot->GetShotOverrideConfiguration()->GetUserSettings();
+		Algo::SortBy(ShotSettings, [](const UMoviePipelineSetting* Setting) { return Setting->GetPriority(); });
+		for (UMoviePipelineSetting* Setting : ShotSettings)
 		{
 			Setting->OnMoviePipelineInitialized(this);
 		}
@@ -987,7 +1004,9 @@ void UMoviePipeline::TeardownShot(UMoviePipelineExecutorShot* InShot)
 	}
 
 	// Notify our containers that the current shot has ended.
-	for (UMoviePipelineOutputBase* Container : GetPipelinePrimaryConfig()->GetOutputContainers())
+	TArray<UMoviePipelineOutputBase*> OutputSettings = GetPipelinePrimaryConfig()->GetOutputContainers();
+	Algo::SortBy(OutputSettings, [](const UMoviePipelineOutputBase* Setting) { return Setting->GetPriority(); },  TLess<int32>());
+	for (UMoviePipelineOutputBase* Container : OutputSettings)
 	{
 		Container->OnShotFinished(InShot, IsFlushDiskWritesPerShot());
 	}
@@ -995,14 +1014,18 @@ void UMoviePipeline::TeardownShot(UMoviePipelineExecutorShot* InShot)
 	if (InShot->GetShotOverrideConfiguration() != nullptr)
 	{
 		// Any shot-specific overrides should get shutdown now.
-		for (UMoviePipelineSetting* Setting : InShot->GetShotOverrideConfiguration()->GetUserSettings())
+		TArray<UMoviePipelineOutputBase*> ShotSettings = GetPipelinePrimaryConfig()->GetOutputContainers();
+		Algo::SortBy(ShotSettings, [](const UMoviePipelineOutputBase* Setting) { return Setting->GetPriority(); },  TLess<int32>());
+		for (UMoviePipelineSetting* Setting : ShotSettings)
 		{
 			Setting->OnMoviePipelineShutdown(this);
 		}
 	}
 
 	// Loop through just our primary settings and let them know which shot we're about to end.
-	for (UMoviePipelineSetting* Setting : GetPipelinePrimaryConfig()->GetAllSettings())
+	TArray<UMoviePipelineSetting*> AllSettings = GetPipelinePrimaryConfig()->GetAllSettings();
+	Algo::SortBy(AllSettings, [](const UMoviePipelineSetting* Setting) { return Setting->GetPriority(); },  TLess<int32>());
+	for (UMoviePipelineSetting* Setting : AllSettings)
 	{
 		Setting->OnTeardownForShot(InShot);
 	}
diff --git a/Engine/Plugins/MovieScene/MovieRenderPipeline/Source/MovieRenderPipelineCore/Private/MoviePipelineGameOverrideSetting.cpp b/Engine/Plugins/MovieScene/MovieRenderPipeline/Source/MovieRenderPipelineCore/Private/MoviePipelineGameOverrideSetting.cpp
--- a/Engine/Plugins/MovieScene/MovieRenderPipeline/Source/MovieRenderPipelineCore/Private/MoviePipelineGameOverrideSetting.cpp
+++ b/Engine/Plugins/MovieScene/MovieRenderPipeline/Source/MovieRenderPipelineCore/Private/MoviePipelineGameOverrideSetting.cpp
@@ -23,6 +23,7 @@ void UMoviePipelineGameOverrideSetting::TeardownForPipelineImpl(UMoviePipeline*
 	ApplyCVarSettings(false);
 }
 
+
 void UMoviePipelineGameOverrideSetting::ApplyCVarSettings(const bool bOverrideValues)
 {
 	if (bCinematicQualitySettings)
@@ -41,7 +42,9 @@ void UMoviePipelineGameOverrideSetting::ApplyCVarSettings(const bool bOverrideVa
 		}
 		else
 		{
-			Scalability::SetQualityLevels(PreviousQualityLevels);
+			// We re-apply old scalability settings at the end of the function during teardown
+			// so that any values that are also specified in Scalability don't get overwritten
+			// with the wrong values from the ones below restoring.
 		}
 	}
 
@@ -133,6 +136,16 @@ void UMoviePipelineGameOverrideSetting::ApplyCVarSettings(const bool bOverrideVa
 
 	// Cloth's time step smoothing messes up the change in number of simulation substeps that fixes the cloth simulation behavior when using Temporal Samples.
 	MOVIEPIPELINE_STORE_AND_OVERRIDE_CVAR_INT(PreviousChaosClothUseTimeStepSmoothing, TEXT("p.ChaosCloth.UseTimeStepSmoothing"), 0, bOverrideValues);
+
+	// Must come after the above cvars so that if one of those cvars is also specified by the Scalability level, then we restore to the value in the original scalability level
+	// not the value we cached in the Cinematic level (if applied).
+	if (bCinematicQualitySettings)
+	{
+		if (!bOverrideValues)
+		{
+			Scalability::SetQualityLevels(PreviousQualityLevels);
+		}
+	}
 }
 
 void UMoviePipelineGameOverrideSetting::BuildNewProcessCommandLineArgsImpl(TArray<FString>& InOutUnrealURLParams, TArray<FString>& InOutCommandLineArgs, TArray<FString>& InOutDeviceProfileCvars, TArray<FString>& InOutExecCmds) const
diff --git a/Engine/Plugins/MovieScene/MovieRenderPipeline/Source/MovieRenderPipelineCore/Public/MoviePipelineSetting.h b/Engine/Plugins/MovieScene/MovieRenderPipeline/Source/MovieRenderPipelineCore/Public/MoviePipelineSetting.h
--- a/Engine/Plugins/MovieScene/MovieRenderPipeline/Source/MovieRenderPipelineCore/Public/MoviePipelineSetting.h
+++ b/Engine/Plugins/MovieScene/MovieRenderPipeline/Source/MovieRenderPipelineCore/Public/MoviePipelineSetting.h
@@ -129,6 +129,11 @@ public:
 	*/
 	virtual bool IgnoreTransientFilters() const { return false; }
 
+	/**
+	* Higher priority settings are run after lower priority settings when setting up shots during rendering. This is run in reverse during teardown!
+	*/
+	virtual int32 GetPriority() const { return 0; }
+
 	// Validation
 	/** What is the result of the last validation? Only valid if the setting has had ValidateState() called on it. */
 	virtual EMoviePipelineValidationState GetValidationState() const { return ValidationState; }
diff --git a/Engine/Plugins/MovieScene/MovieRenderPipeline/Source/MovieRenderPipelineCore/Public/MoviePipelineUtils.h b/Engine/Plugins/MovieScene/MovieRenderPipeline/Source/MovieRenderPipelineCore/Public/MoviePipelineUtils.h
--- a/Engine/Plugins/MovieScene/MovieRenderPipeline/Source/MovieRenderPipelineCore/Public/MoviePipelineUtils.h
+++ b/Engine/Plugins/MovieScene/MovieRenderPipeline/Source/MovieRenderPipelineCore/Public/MoviePipelineUtils.h
@@ -58,11 +58,11 @@ MOVIERENDERPIPELINECORE_API void GetPassCompositeData(FMoviePipelineMergerOutput
 		if(bUseOverride) \
 		{ \
 			InOutVariable = CVar->GetInt(); \
-			CVar->Set(OverrideValue, EConsoleVariableFlags::ECVF_SetByConsole); \
+			CVar->SetWithCurrentPriority(OverrideValue); \
 		} \
 		else \
 		{ \
-			CVar->Set(InOutVariable, EConsoleVariableFlags::ECVF_SetByConsole); \
+			CVar->SetWithCurrentPriority(InOutVariable); \
 		} \
 	} \
 }
@@ -75,11 +75,11 @@ MOVIERENDERPIPELINECORE_API void GetPassCompositeData(FMoviePipelineMergerOutput
 		if(bUseOverride) \
 		{ \
 			InOutVariable = CVar->GetFloat(); \
-			CVar->Set(OverrideValue, EConsoleVariableFlags::ECVF_SetByConsole); \
+			CVar->SetWithCurrentPriority(OverrideValue); \
 		} \
 		else \
 		{ \
-			CVar->Set(InOutVariable, EConsoleVariableFlags::ECVF_SetByConsole); \
+			CVar->SetWithCurrentPriority(InOutVariable); \
 		} \
 	} \
 }
@@ -92,11 +92,11 @@ MOVIERENDERPIPELINECORE_API void GetPassCompositeData(FMoviePipelineMergerOutput
 		if(bUseOverride) \
 		{ \
 			InOutVariable = CVar->GetInt(); \
-			CVar->Set(OverrideValue, EConsoleVariableFlags::ECVF_SetByConsole); \
+			CVar->SetWithCurrentPriority(OverrideValue); \
 		} \
 		else \
 		{ \
-			CVar->Set(InOutVariable, EConsoleVariableFlags::ECVF_SetByConsole); \
+			CVar->SetWithCurrentPriority(InOutVariable); \
 		} \
 	} \
 }
@@ -109,11 +109,11 @@ MOVIERENDERPIPELINECORE_API void GetPassCompositeData(FMoviePipelineMergerOutput
 		if(bUseOverride) \
 		{ \
 			InOutVariable = CVar->GetFloat(); \
-			CVar->Set(OverrideValue, EConsoleVariableFlags::ECVF_SetByConsole); \
+			CVar->SetWithCurrentPriority(OverrideValue); \
 		} \
 		else \
 		{ \
-			CVar->Set(InOutVariable, EConsoleVariableFlags::ECVF_SetByConsole); \
+			CVar->SetWithCurrentPriority(InOutVariablee); \
 		} \
 	} \
 }
diff --git a/Engine/Plugins/MovieScene/MovieRenderPipeline/Source/MovieRenderPipelineSettings/Public/MoviePipelineConsoleVariableSetting.h b/Engine/Plugins/MovieScene/MovieRenderPipeline/Source/MovieRenderPipelineSettings/Public/MoviePipelineConsoleVariableSetting.h
--- a/Engine/Plugins/MovieScene/MovieRenderPipeline/Source/MovieRenderPipelineSettings/Public/MoviePipelineConsoleVariableSetting.h
+++ b/Engine/Plugins/MovieScene/MovieRenderPipeline/Source/MovieRenderPipelineSettings/Public/MoviePipelineConsoleVariableSetting.h
@@ -21,6 +21,10 @@ public:
 	virtual bool IsValidOnPrimary() const override { return true; }
 	virtual void SetupForPipelineImpl(UMoviePipeline* InPipeline) override;
 	virtual void TeardownForPipelineImpl(UMoviePipeline* InPipeline) override;
+
+	// This needs to be higher priority than the Game Override setting so that the values the user specifies for cvars here are the ones actually applied during renders
+	// otherwise the Scalability Settings of the Game Override setting can change values away from what the user expects.
+	virtual int32 GetPriority() const override { return 1; }
 protected:
 	void ApplyCVarSettings(const bool bOverrideValues);
 
