MoviePipeline: The UI now allows cvars to be enabled/disabled. Also shows a help tooltip for the cvar to describe what it does, as well as value resolution tooltip for where the cvar value originated (the startup value, the preset value, and the value set in MRQ).

If a cvar is disabled, it will not be included in cvar overrides when rendering, and the value shown in the UI will be the preset value of the cvar (if one exists), else the startup value.

Established a dependency on the ConsoleVariablesEditor. This is initially for gaining access to CommandInfo, but future updates will utilize more of ConsoleVariablesEditor.

Migrated ConsoleVariableSetting::ConsoleVariables away from a TMap to an array of ConsoleVariableEntry structs. This was done to support enable/disable state, as well as support
any additional cvar-related state that needs to be tracked in the future.
* Renamed ConsoleVariables to ConsoleVariables_DEPRECATED.
* ConsoleVariables are converted/merged into the new CVars property upon load.
* If scripting modifies ConsoleVariables after initial load, they are merged into CVars upon cvar apply. Note there is a slight potential for data loss here if a script modifies ConsoleVariables, but then MRQ is not run in PIE (meaning ApplyCVarSettings() is never run). ConsoleVariables will not save due to being deprecated, and it will never be merged into CVars.
* Scripts should use the new cvar mutator methods rather than accessing cvars directly.

diff --git a/Engine/Plugins/MovieScene/MovieRenderPipeline/MovieRenderPipeline.uplugin b/Engine/Plugins/MovieScene/MovieRenderPipeline/MovieRenderPipeline.uplugin
--- a/Engine/Plugins/MovieScene/MovieRenderPipeline/MovieRenderPipeline.uplugin
+++ b/Engine/Plugins/MovieScene/MovieRenderPipeline/MovieRenderPipeline.uplugin
@@ -20,6 +20,11 @@
 			"Enabled": true,
 			"Optional": false
 		},
+		{
+			"Name": "ConsoleVariables",
+			"Enabled": true,
+			"Optional": false
+		},
 		{
 			"Name": "EditorScriptingUtilities",
 			"Enabled": true
diff --git a/Engine/Plugins/MovieScene/MovieRenderPipeline/Source/MovieRenderPipelineEditor/MovieRenderPipelineEditor.Build.cs b/Engine/Plugins/MovieScene/MovieRenderPipeline/Source/MovieRenderPipelineEditor/MovieRenderPipelineEditor.Build.cs
--- a/Engine/Plugins/MovieScene/MovieRenderPipeline/Source/MovieRenderPipelineEditor/MovieRenderPipelineEditor.Build.cs
+++ b/Engine/Plugins/MovieScene/MovieRenderPipeline/Source/MovieRenderPipelineEditor/MovieRenderPipelineEditor.Build.cs
@@ -63,7 +63,8 @@ public class MovieRenderPipelineEditor : ModuleRules
 				"JsonUtilities",
 				"ScreenShotComparisonTools",
 				"AutomationMessages",
-				"ToolWidgets"
+				"ToolWidgets",
+				"ConsoleVariablesEditor",
 			}
         );
 
diff --git a/Engine/Plugins/MovieScene/MovieRenderPipeline/Source/MovieRenderPipelineEditor/Private/Customizations/ConsoleVariableCustomization.h b/Engine/Plugins/MovieScene/MovieRenderPipeline/Source/MovieRenderPipelineEditor/Private/Customizations/ConsoleVariableCustomization.h
new file mode 100644
--- /dev/null
+++ b/Engine/Plugins/MovieScene/MovieRenderPipeline/Source/MovieRenderPipelineEditor/Private/Customizations/ConsoleVariableCustomization.h
@@ -0,0 +1,279 @@
+﻿// Copyright Epic Games, Inc. All Rights Reserved.
+#pragma once
+
+#include "Internationalization/Internationalization.h"
+#include "Misc/DefaultValueHelper.h"
+#include "MoviePipelineConsoleVariableSetting.h"
+#include "PropertyCustomizationHelpers.h"
+#include "Widgets/Input/SEditableTextBox.h"
+#include "Widgets/Input/SNumericEntryBox.h"
+
+#define LOCTEXT_NAMESPACE "SMoviePipelineEditor"
+
+/* Customizes how console variables are displayed in the details pane. */
+class FConsoleVariablesDetailsCustomization : public IPropertyTypeCustomization
+{
+public:
+	static TSharedRef<IPropertyTypeCustomization> MakeInstance()
+	{
+		return MakeShared<FConsoleVariablesDetailsCustomization>();
+	}
+
+protected:
+	//~ Begin IPropertyTypeCustomization interface
+	virtual void CustomizeHeader(TSharedRef<IPropertyHandle> InStructPropertyHandle, FDetailWidgetRow& HeaderRow, IPropertyTypeCustomizationUtils& CustomizationUtils) override
+	{
+		const TSharedRef<IPropertyHandle> NameProp = InStructPropertyHandle->GetChildHandle(GET_MEMBER_NAME_CHECKED(FMoviePipelineConsoleVariableEntry, Name)).ToSharedRef();
+		const TSharedRef<IPropertyHandle> ValueProp = InStructPropertyHandle->GetChildHandle(GET_MEMBER_NAME_CHECKED(FMoviePipelineConsoleVariableEntry, Value)).ToSharedRef();
+		const TSharedRef<IPropertyHandle> IsEnabledProp = InStructPropertyHandle->GetChildHandle(GET_MEMBER_NAME_CHECKED(FMoviePipelineConsoleVariableEntry, bIsEnabled)).ToSharedRef();
+
+		// Action for deleting a cvar from the parent array
+		const FExecuteAction DeleteAction = FExecuteAction::CreateLambda([InStructPropertyHandle]
+		{
+			const TSharedPtr<IPropertyHandleArray> ParentPropertyHandleArray = InStructPropertyHandle->GetParentHandle()->AsArray();
+			const int32 ArrayIndex = InStructPropertyHandle->IsValidHandle() ? InStructPropertyHandle->GetIndexInArray() : INDEX_NONE;
+			if (ParentPropertyHandleArray.IsValid() && ArrayIndex >= 0)
+			{
+				ParentPropertyHandleArray->DeleteItem(ArrayIndex);
+			}
+		});
+
+		HeaderRow
+		.WholeRowContent()
+		[
+			SNew(SHorizontalBox)
+			+ SHorizontalBox::Slot()
+			.Padding(0, 2)
+			.AutoWidth()
+			.HAlign(HAlign_Left)
+			[
+				SNew(SCheckBox)
+				.IsChecked(this, &FConsoleVariablesDetailsCustomization::GetIsEnabledCheckBoxState, IsEnabledProp)
+				.OnCheckStateChanged(this, &FConsoleVariablesDetailsCustomization::OnIsEnabledCheckBoxStateChanged, IsEnabledProp)
+			]
+
+			+ SHorizontalBox::Slot()
+			.Padding(5, 2)
+			.FillWidth(0.75f)
+			[
+				SNew(SEditableTextBox)
+				.IsEnabled(this, &FConsoleVariablesDetailsCustomization::IsEnabled, IsEnabledProp)
+				.Font(CustomizationUtils.GetRegularFont())
+				.Text(this, &FConsoleVariablesDetailsCustomization::GetConsoleVariableText, NameProp)
+				.ToolTipText(this, &FConsoleVariablesDetailsCustomization::GetConsoleVariableHelpText, NameProp)
+				.OnTextCommitted(this, &FConsoleVariablesDetailsCustomization::OnConsoleVariableNameTextCommitted, NameProp)
+			]
+
+			+ SHorizontalBox::Slot()
+			.Padding(0, 2)
+			.FillWidth(0.25f)
+			[
+				SNew(SSpinBox<float>)
+				.IsEnabled(this, &FConsoleVariablesDetailsCustomization::IsEnabled, IsEnabledProp)
+				.MaxFractionalDigits(3)
+				.Font(CustomizationUtils.GetRegularFont())
+				.Value(this, &FConsoleVariablesDetailsCustomization::GetConsoleVariableValue, ValueProp)
+				.OnValueChanged(this, &FConsoleVariablesDetailsCustomization::OnConsoleVariableValueChanged, ValueProp)
+				.OnValueCommitted(this, &FConsoleVariablesDetailsCustomization::OnConsoleVariableValueCommitted, ValueProp)
+				.ToolTipText(this, &FConsoleVariablesDetailsCustomization::GetConsoleVariableTooltip, InStructPropertyHandle)
+			]
+
+			+ SHorizontalBox::Slot()
+			.Padding(5, 2)
+			.AutoWidth()
+			[
+				PropertyCustomizationHelpers::MakeInsertDeleteDuplicateButton({}, DeleteAction, {})
+			]
+		];
+	}
+	
+	virtual void CustomizeChildren(TSharedRef<IPropertyHandle> InStructPropertyHandle, IDetailChildrenBuilder& StructBuilder, IPropertyTypeCustomizationUtils& StructCustomizationUtils) override
+	{
+	}
+	//~ End IPropertyTypeCustomization interface
+
+	/* Get a CVar settings object from a property handle. Returned pointer may be invalid if the object does not exist. */
+	UMoviePipelineConsoleVariableSetting* GetCVarSettingFromHandle(const TSharedPtr<IPropertyHandle>& PropertyHandle) const
+	{
+		if (!PropertyHandle.IsValid() || !PropertyHandle->IsValidHandle())
+		{
+			return nullptr;
+		}
+		
+		TArray<UObject *> OuterObjects;
+		PropertyHandle->GetOuterObjects(OuterObjects);
+
+		if (OuterObjects.Num() != 1)
+		{
+			return nullptr;
+		}
+
+		// Get the outermost settings object
+		return Cast<UMoviePipelineConsoleVariableSetting>(OuterObjects[0]);
+	}
+
+	/* Get a CVar entry object from a property handle. Returned pointer may be invalid if the object does not exist. */
+	FMoviePipelineConsoleVariableEntry* GetEntryForHandle(const TSharedPtr<IPropertyHandle>& PropertyHandle) const
+	{
+		if (!PropertyHandle.IsValid() || !PropertyHandle->IsValidHandle())
+		{
+			return nullptr;
+		}
+		
+		UMoviePipelineConsoleVariableSetting* Setting = GetCVarSettingFromHandle(PropertyHandle);
+		if (!Setting)
+		{
+			return nullptr;
+		}
+		
+		// Get the cvar entry associated w/ this property
+		const int32 EntryIndex = PropertyHandle->GetIndexInArray();
+		return Setting->GetCVarAtIndex(EntryIndex);
+	}
+
+	ECheckBoxState GetIsEnabledCheckBoxState(TSharedRef<IPropertyHandle> PropertyHandle) const
+	{
+		bool bValue = false;
+		if (PropertyHandle->IsValidHandle() && (PropertyHandle->GetValue(bValue) != FPropertyAccess::Fail))
+		{
+			return bValue ? ECheckBoxState::Checked : ECheckBoxState::Unchecked;
+		}
+
+		return ECheckBoxState::Unchecked;
+	}
+
+	void OnIsEnabledCheckBoxStateChanged(ECheckBoxState NewState, TSharedRef<IPropertyHandle> PropertyHandle)
+	{
+		if (PropertyHandle->IsValidHandle())
+		{
+			const bool bValue = (NewState == ECheckBoxState::Checked) ? true : false;
+			PropertyHandle->SetValue(bValue);
+		}
+	}
+
+	bool IsEnabled(TSharedRef<IPropertyHandle> InIsEnabledPropertyHandle) const
+	{
+		bool bValue = false;
+		if (InIsEnabledPropertyHandle->IsValidHandle() && (InIsEnabledPropertyHandle->GetValue(bValue) != FPropertyAccess::Fail))
+		{
+			return bValue;
+		}
+
+		return false;
+	}
+
+	FText GetConsoleVariableText(TSharedRef<IPropertyHandle> PropertyHandle) const
+	{
+		FString Value;
+		if (PropertyHandle->IsValidHandle() && (PropertyHandle->GetValue(Value) != FPropertyAccess::Fail))
+		{
+			return FText::FromString(Value);
+		}
+
+		return FText();
+	}
+	
+	FText GetConsoleVariableHelpText(TSharedRef<IPropertyHandle> PropertyHandle) const
+	{
+		// Get the handle to the parent struct property
+		const TSharedPtr<IPropertyHandle> ParentPropertyHandle = PropertyHandle->GetParentHandle();
+		
+		if (const FMoviePipelineConsoleVariableEntry* CVarEntry = GetEntryForHandle(ParentPropertyHandle))
+		{
+			const TSharedPtr<FConsoleVariablesEditorCommandInfo> CommandInfo = CVarEntry->CommandInfo.Pin();
+			if (CommandInfo.IsValid())
+			{
+				return FText::FromString(CommandInfo->GetHelpText());
+			}
+		}
+
+		return FText();
+	}
+
+	void OnConsoleVariableNameTextCommitted(const FText& NewText, ETextCommit::Type CommitType, TSharedRef<IPropertyHandle> PropertyHandle)
+	{
+		if (PropertyHandle->IsValidHandle())
+		{
+			PropertyHandle->SetValue(NewText.ToString());
+		}
+	}
+
+	float GetConsoleVariableValue(TSharedRef<IPropertyHandle> PropertyHandle) const
+	{
+		// Get the handle to the parent struct property
+		const TSharedPtr<IPropertyHandle> ParentPropertyHandle = PropertyHandle->GetParentHandle();
+		
+		const FMoviePipelineConsoleVariableEntry* CVarEntry = GetEntryForHandle(ParentPropertyHandle);
+		if (!CVarEntry)
+		{
+			return 0.f;
+		}
+
+		// If this cvar override is disabled, find what its value would be as if this override did not exist
+		if (!CVarEntry->bIsEnabled)
+		{
+			// Get the non-override value of this cvar
+			const UMoviePipelineConsoleVariableSetting* CVarSetting = GetCVarSettingFromHandle(ParentPropertyHandle);
+			if (CVarSetting)
+			{
+				float CVarFloatValue = 0.f;
+				if (FDefaultValueHelper::ParseFloat(CVarSetting->ResolveDisabledValue(*CVarEntry), CVarFloatValue))
+				{
+					return CVarFloatValue;
+				}
+			}
+		}
+		
+		float Value;
+		if (PropertyHandle->IsValidHandle() && (PropertyHandle->GetValue(Value) != FPropertyAccess::Fail))
+		{
+			return Value;
+		}
+
+		return 0.f;
+	}
+
+	void OnConsoleVariableValueChanged(float Value, TSharedRef<IPropertyHandle> PropertyHandle)
+	{
+		if (PropertyHandle->IsValidHandle())
+		{
+			PropertyHandle->SetValue(Value);
+		}
+	}
+
+	void OnConsoleVariableValueCommitted(float Value, ETextCommit::Type CommitType, TSharedRef<IPropertyHandle> PropertyHandle)
+	{
+		OnConsoleVariableValueChanged(Value, PropertyHandle);
+	}
+
+	FText GetConsoleVariableTooltip(TSharedRef<IPropertyHandle> StructPropertyHandle) const
+	{
+		static const FText ToolTipFormatText = LOCTEXT("CVarToolTipFormatText", "Custom Value: {0}\nPreset Value: {1}\nStartup Value: {2} (Set By {3})");
+		static const FText NoPresetValueText = LOCTEXT("NoPresetValue", "None");
+		
+		const UMoviePipelineConsoleVariableSetting* Setting = GetCVarSettingFromHandle(StructPropertyHandle);
+		const FMoviePipelineConsoleVariableEntry* CVarEntry = GetEntryForHandle(StructPropertyHandle);
+		if (!Setting || !CVarEntry)
+		{
+			return FText();
+		}
+		
+		const TSharedPtr<FConsoleVariablesEditorCommandInfo> CommandInfo = CVarEntry->CommandInfo.Pin();
+		if (!CommandInfo.IsValid())
+		{
+			return FText();
+		}
+		
+		const FString CustomValue = FString::SanitizeFloat(CVarEntry->Value);
+		const FString PresetValue = Setting->ResolvePresetValue(CVarEntry->Name);
+		const FString StartupValue = CommandInfo->StartupValueAsString;
+		
+		return FText::Format(ToolTipFormatText,
+			FText::FromString(CustomValue),
+			PresetValue.IsEmpty() ? NoPresetValueText : FText::FromString(PresetValue),
+			FText::FromString(StartupValue),
+			FConsoleVariablesEditorCommandInfo::ConvertConsoleVariableSetByFlagToText(CommandInfo->StartupSource));
+	}
+};
+
+#undef LOCTEXT_NAMESPACE
\ No newline at end of file
diff --git a/Engine/Plugins/MovieScene/MovieRenderPipeline/Source/MovieRenderPipelineEditor/Private/Widgets/SMoviePipelineConfigEditor.cpp b/Engine/Plugins/MovieScene/MovieRenderPipeline/Source/MovieRenderPipelineEditor/Private/Widgets/SMoviePipelineConfigEditor.cpp
--- a/Engine/Plugins/MovieScene/MovieRenderPipeline/Source/MovieRenderPipelineEditor/Private/Widgets/SMoviePipelineConfigEditor.cpp
+++ b/Engine/Plugins/MovieScene/MovieRenderPipeline/Source/MovieRenderPipelineEditor/Private/Widgets/SMoviePipelineConfigEditor.cpp
@@ -55,6 +55,8 @@
 #include "DetailWidgetRow.h"
 #include "IPropertyUtilities.h"
 
+#include "Customizations/ConsoleVariableCustomization.h"
+#include "MoviePipelineConsoleVariableSetting.h"
 #include "MoviePipelineOutputSetting.h"
 
 #define LOCTEXT_NAMESPACE "SMoviePipelineEditor"
@@ -732,6 +734,7 @@ void SMoviePipelineConfigEditor::UpdateDetails()
 		else
 		{
 			TSharedRef<IDetailsView> Details = PropertyEditorModule.CreateDetailView(DetailsViewArgs);
+			Details->RegisterInstancedCustomPropertyTypeLayout(FMoviePipelineConsoleVariableEntry::StaticStruct()->GetFName(), FOnGetPropertyTypeCustomizationInstance::CreateStatic(FConsoleVariablesDetailsCustomization::MakeInstance));
 			Details->RegisterInstancedCustomPropertyTypeLayout("FrameNumber", FOnGetPropertyTypeCustomizationInstance::CreateStatic(&FFrameNumberDetailsCustomization::MakeInstance, NumericTypeInterface));
 			Details->RegisterInstancedCustomPropertyLayout(UMoviePipelineOutputSetting::StaticClass(), FOnGetDetailCustomizationInstance::CreateStatic(FOutputFormatDetailsCustomization::MakeInstance, SharedThis(this)));
 
diff --git a/Engine/Plugins/MovieScene/MovieRenderPipeline/Source/MovieRenderPipelineSettings/MovieRenderPipelineSettings.Build.cs b/Engine/Plugins/MovieScene/MovieRenderPipeline/Source/MovieRenderPipelineSettings/MovieRenderPipelineSettings.Build.cs
--- a/Engine/Plugins/MovieScene/MovieRenderPipeline/Source/MovieRenderPipelineSettings/MovieRenderPipelineSettings.Build.cs
+++ b/Engine/Plugins/MovieScene/MovieRenderPipeline/Source/MovieRenderPipelineSettings/MovieRenderPipelineSettings.Build.cs
@@ -30,5 +30,10 @@ public class MovieRenderPipelineSettings : ModuleRules
 				"OpenColorIO",
 			}
 		);
+		
+		if (Target.bBuildEditor)
+		{
+			PublicDependencyModuleNames.Add("ConsoleVariablesEditor");
+		}
 	}
 }
diff --git a/Engine/Plugins/MovieScene/MovieRenderPipeline/Source/MovieRenderPipelineSettings/Private/MoviePipelineConsoleVariableSetting.cpp b/Engine/Plugins/MovieScene/MovieRenderPipeline/Source/MovieRenderPipelineSettings/Private/MoviePipelineConsoleVariableSetting.cpp
--- a/Engine/Plugins/MovieScene/MovieRenderPipeline/Source/MovieRenderPipelineSettings/Private/MoviePipelineConsoleVariableSetting.cpp
+++ b/Engine/Plugins/MovieScene/MovieRenderPipeline/Source/MovieRenderPipelineSettings/Private/MoviePipelineConsoleVariableSetting.cpp
@@ -133,6 +133,77 @@ FText UMoviePipelineConsoleVariableSetting::GetFooterText(UMoviePipelineExecutor
 		"The current job contains a Level Sequence with a Console Variables Track, additional settings are configured in Sequencer."));
 }
 
+FString UMoviePipelineConsoleVariableSetting::ResolvePresetValue(const FString& InCVarName) const
+{
+	FString ResolvedValue;
+
+	// Iterate the presets in reverse; cvars in presets at the end of the array take precedence.
+	for (int32 Index = ConsoleVariablePresets.Num() - 1; Index >= 0; --Index)
+	{
+		const TScriptInterface<IMovieSceneConsoleVariableTrackInterface>& Preset = ConsoleVariablePresets[Index];
+		if (!Preset)
+		{
+			continue;
+		}
+		
+		const bool bOnlyIncludeChecked = true;
+		TArray<TPair<FString, FString>> PresetCVars;
+		Preset->GetConsoleVariablesForTrack(bOnlyIncludeChecked, PresetCVars);
+		
+		for (const TPair<FString, FString>& CVar : PresetCVars)
+		{
+			if (CVar.Key.Equals(InCVarName))
+			{
+				// Found a matching cvar in the preset
+				return CVar.Value;
+			}
+		}
+	}
+
+	return FString();
+}
+
+FString UMoviePipelineConsoleVariableSetting::ResolveDisabledValue(const FMoviePipelineConsoleVariableEntry& InEntry) const
+{
+	// There may be multiple cvar overrides w/ the same name, so even if the provided entry is disabled, we still need
+	// to look through the other entries for cvars with the same name. Iterate in reverse, since cvars at the end of the
+	// array take precedence.
+	for (int32 Index = CVars.Num() - 1; Index >= 0; --Index)
+	{
+		const FMoviePipelineConsoleVariableEntry& CVarEntry = CVars[Index];
+		if (CVarEntry.Name.Equals(InEntry.Name) && CVarEntry.bIsEnabled)
+		{
+			return FString::SanitizeFloat(CVarEntry.Value);
+		}
+	}
+
+	// If no override value was found, look for a value from the presets
+	const FString PresetValue = ResolvePresetValue(InEntry.Name);
+	if (!PresetValue.IsEmpty())
+	{
+		return PresetValue;
+	}
+
+	// Fall back to the startup value of the cvar
+	const TSharedPtr<FConsoleVariablesEditorCommandInfo> CommandInfo = InEntry.CommandInfo.Pin();
+	if (CommandInfo.IsValid())
+	{
+		return CommandInfo->StartupValueAsString;
+	}
+
+	return FString();
+}
+
+FMoviePipelineConsoleVariableEntry* UMoviePipelineConsoleVariableSetting::GetCVarAtIndex(const int32 InIndex)
+{
+	if (!CVars.IsValidIndex(InIndex))
+	{
+		return nullptr;
+	}
+
+	return &CVars[InIndex];
+}
+
 #endif // WITH_EDITOR
 
 void UMoviePipelineConsoleVariableSetting::SetupForPipelineImpl(UMoviePipeline* InPipeline)
@@ -149,34 +220,35 @@ void UMoviePipelineConsoleVariableSetting::ApplyCVarSettings(const bool bOverrid
 {
 	if (bOverrideValues)
 	{
-		MergeConsoleVariables();
+		MergeInOldConsoleVariables();
+		MergeInPresetConsoleVariables();
 		PreviousConsoleVariableValues.Reset();
 		PreviousConsoleVariableValues.SetNumZeroed(MergedConsoleVariables.Num());
 	}
 
 	int32 Index = 0;
-	for(const TPair<FString, float>& KVP : MergedConsoleVariables)
+	for(const FMoviePipelineConsoleVariableEntry& CVarEntry : MergedConsoleVariables)
 	{
 		// We don't use the shared macro here because we want to soft-warn the user instead of tripping an ensure over missing cvar values.
-		const FString TrimmedCvar = KVP.Key.TrimStartAndEnd();
+		const FString TrimmedCvar = CVarEntry.Name.TrimStartAndEnd();
 		IConsoleVariable* CVar = IConsoleManager::Get().FindConsoleVariable(*TrimmedCvar); 
 		if (CVar)
 		{
 			if (bOverrideValues)
 			{
 				PreviousConsoleVariableValues[Index] = CVar->GetFloat();
-				UE::MoviePipeline::SetValue(CVar, KVP.Value);
-				UE_LOG(LogMovieRenderPipeline, Log, TEXT("Applying CVar \"%s\" PreviousValue: %f NewValue: %f"), *KVP.Key, PreviousConsoleVariableValues[Index], KVP.Value);
+				UE::MoviePipeline::SetValue(CVar, CVarEntry.Value);
+				UE_LOG(LogMovieRenderPipeline, Log, TEXT("Applying CVar \"%s\" PreviousValue: %f NewValue: %f"), *CVarEntry.Name, PreviousConsoleVariableValues[Index], CVarEntry.Value);
 			}
 			else
 			{
 				UE::MoviePipeline::SetValue(CVar, PreviousConsoleVariableValues[Index]);
-				UE_LOG(LogMovieRenderPipeline, Log, TEXT("Restoring CVar \"%s\" PreviousValue: %f NewValue: %f"), *KVP.Key, KVP.Value, PreviousConsoleVariableValues[Index]);
+				UE_LOG(LogMovieRenderPipeline, Log, TEXT("Restoring CVar \"%s\" PreviousValue: %f NewValue: %f"), *CVarEntry.Name, CVarEntry.Value, PreviousConsoleVariableValues[Index]);
 			}
 		}
 		else
 		{
-			UE_LOG(LogMovieRenderPipeline, Warning, TEXT("Failed to apply CVar \"%s\" due to no cvar by that name. Ignoring."), *KVP.Key);
+			UE_LOG(LogMovieRenderPipeline, Warning, TEXT("Failed to apply CVar \"%s\" due to no cvar by that name. Ignoring."), *CVarEntry.Name);
 		}
 
 		Index++;
@@ -200,7 +272,170 @@ void UMoviePipelineConsoleVariableSetting::ApplyCVarSettings(const bool bOverrid
 	}
 }
 
-void UMoviePipelineConsoleVariableSetting::MergeConsoleVariables()
+TArray<FMoviePipelineConsoleVariableEntry> UMoviePipelineConsoleVariableSetting::GetConsoleVariables() const
+{
+	return CVars;
+}
+
+bool UMoviePipelineConsoleVariableSetting::RemoveConsoleVariable(const FString& Name, const bool bRemoveAllInstances)
+{
+	if (bRemoveAllInstances)
+	{
+		const int32 NumRemoved = CVars.RemoveAll([&Name](const FMoviePipelineConsoleVariableEntry& Entry)
+		{
+			return Entry.Name.Equals(Name);
+		});
+
+		return NumRemoved != 0;
+	}
+
+	const int32 LastMatch = CVars.FindLastByPredicate([&Name](const FMoviePipelineConsoleVariableEntry& Entry)
+	{
+		return Entry.Name.Equals(Name);
+	});
+
+	if (LastMatch != INDEX_NONE)
+	{
+		CVars.RemoveAt(LastMatch);
+		return true;
+	}
+
+	return false;
+}
+
+bool UMoviePipelineConsoleVariableSetting::AddOrUpdateConsoleVariable(const FString& Name, const float Value)
+{
+	const int32 LastMatch = CVars.FindLastByPredicate([&Name](const FMoviePipelineConsoleVariableEntry& Entry)
+	{
+		return Entry.Name.Equals(Name);
+	});
+
+	if (LastMatch != INDEX_NONE)
+	{
+		CVars[LastMatch].Value = Value;
+		return true;
+	}
+
+	CVars.Add(FMoviePipelineConsoleVariableEntry(Name, Value));
+
+	return true;
+}
+
+bool UMoviePipelineConsoleVariableSetting::UpdateConsoleVariableEnableState(const FString& Name, const bool bIsEnabled)
+{
+	const int32 LastMatch = CVars.FindLastByPredicate([&Name](const FMoviePipelineConsoleVariableEntry& Entry)
+	{
+		return Entry.Name.Equals(Name);
+	});
+
+	if (LastMatch != INDEX_NONE)
+	{
+		CVars[LastMatch].bIsEnabled = bIsEnabled;
+		return true;
+	}
+
+	return false;
+}
+
+void UMoviePipelineConsoleVariableSetting::PostLoad()
+{
+	Super::PostLoad();
+
+#if WITH_EDITOR
+	PRAGMA_DISABLE_DEPRECATION_WARNINGS
+
+	// Convert the old ConsoleVariables map into the new struct array
+	if (!ConsoleVariables_DEPRECATED.IsEmpty())
+	{
+		for (const TPair<FString, float>& Pair : ConsoleVariables_DEPRECATED)
+		{
+			// The converted cvar will be enabled by default
+			CVars.Add(FMoviePipelineConsoleVariableEntry(Pair.Key, Pair.Value));
+		}
+
+		ConsoleVariables_DEPRECATED.Empty();
+	}
+	
+	PRAGMA_ENABLE_DEPRECATION_WARNINGS
+
+	// Initialize the CommandInfo structs
+	for (FMoviePipelineConsoleVariableEntry& Entry : CVars)
+	{
+		Entry.UpdateCommandInfo();
+	}
+#endif	// WITH_EDITOR
+}
+
+#if WITH_EDITOR
+void UMoviePipelineConsoleVariableSetting::PostEditChangeChainProperty(FPropertyChangedChainEvent& PropertyChangedEvent)
+{
+	Super::PostEditChangeChainProperty(PropertyChangedEvent);
+
+	const FProperty* Property = PropertyChangedEvent.Property;
+	const FProperty* MemberProperty = nullptr;
+	if (PropertyChangedEvent.PropertyChain.GetActiveMemberNode())
+	{
+		MemberProperty = PropertyChangedEvent.PropertyChain.GetActiveMemberNode()->GetValue();
+	}
+
+	if (!MemberProperty || !Property)
+	{
+		return;
+	}
+
+	// If the name of one of the cvar overrides changes, generate a new CommandInfo for it
+	if (MemberProperty->GetFName() == GET_MEMBER_NAME_CHECKED(UMoviePipelineConsoleVariableSetting, CVars))
+	{
+		if (Property->GetFName() == GET_MEMBER_NAME_CHECKED(FMoviePipelineConsoleVariableEntry, Name))
+		{
+			const int32 ArrayIndex = PropertyChangedEvent.GetArrayIndex(MemberProperty->GetName());
+			if (CVars.IsValidIndex(ArrayIndex))
+			{
+				CVars[ArrayIndex].UpdateCommandInfo();
+			}
+		}
+	}
+}
+#endif	// WITH_EDITOR
+
+void UMoviePipelineConsoleVariableSetting::MergeInOldConsoleVariables()
+{
+	// Note: The old cvars in ConsoleVariables are merged in via PostLoad(), which covers most use cases.
+	// However, scripting may have modified ConsoleVariables after initial load, so they need to be
+	// merged in again in certain scenarios, like when cvars are applied.
+	//
+	// Data loss is possible if scripting modified ConsoleVariables, but this method never executes
+	// (ie, MRQ never gets run in PIE). ConsoleVariables is deprecated, thus changes to it will not be saved.
+	
+	PRAGMA_DISABLE_DEPRECATION_WARNINGS
+	
+	for (const TPair<FString, float>& OldCVar : ConsoleVariables_DEPRECATED)
+	{
+		// Copy the value in ConsoleVariables over to CVars if an entry in CVars already exists. Apply the value to the
+		// last matching entry in CVars, since the last entry will win.
+		bool bDidFindNewCVar = false;
+		for (int32 Index = CVars.Num() - 1; Index >= 0; --Index)
+		{
+			FMoviePipelineConsoleVariableEntry& NewCVar = CVars[Index];
+			if (NewCVar.Name == OldCVar.Key)
+			{
+				NewCVar.Value = OldCVar.Value;
+				bDidFindNewCVar = true;
+				break;
+			}
+		}
+
+		// Otherwise, add a new entry to CVars
+		if (!bDidFindNewCVar)
+		{
+			CVars.Add(FMoviePipelineConsoleVariableEntry(OldCVar.Key, OldCVar.Value));
+		}
+	}
+	
+	PRAGMA_ENABLE_DEPRECATION_WARNINGS
+}
+
+void UMoviePipelineConsoleVariableSetting::MergeInPresetConsoleVariables()
 {
 	MergedConsoleVariables.Reset();
 	
@@ -222,7 +457,7 @@ void UMoviePipelineConsoleVariableSetting::MergeConsoleVariables()
 			float CVarFloatValue = 0.0f;
 			if (FDefaultValueHelper::ParseFloat(CVarPair.Value, CVarFloatValue))
 			{
-				MergedConsoleVariables.Add(CVarPair.Key, CVarFloatValue);
+				MergedConsoleVariables.Add(FMoviePipelineConsoleVariableEntry(CVarPair.Key, CVarFloatValue));
 			}
 			else
 			{
@@ -233,9 +468,14 @@ void UMoviePipelineConsoleVariableSetting::MergeConsoleVariables()
 	}
 	
 	// Merge in the overrides
-	for (const TPair<FString, float>& KVP : ConsoleVariables)
+	for (const FMoviePipelineConsoleVariableEntry& Entry : CVars)
 	{
-		MergedConsoleVariables.Add(KVP);
+		if (!Entry.bIsEnabled)
+		{
+			continue;
+		}
+		
+		MergedConsoleVariables.Add(Entry);
 	}
 }
 
diff --git a/Engine/Plugins/MovieScene/MovieRenderPipeline/Source/MovieRenderPipelineSettings/Public/MoviePipelineConsoleVariableSetting.h b/Engine/Plugins/MovieScene/MovieRenderPipeline/Source/MovieRenderPipelineSettings/Public/MoviePipelineConsoleVariableSetting.h
--- a/Engine/Plugins/MovieScene/MovieRenderPipeline/Source/MovieRenderPipelineSettings/Public/MoviePipelineConsoleVariableSetting.h
+++ b/Engine/Plugins/MovieScene/MovieRenderPipeline/Source/MovieRenderPipelineSettings/Public/MoviePipelineConsoleVariableSetting.h
@@ -1,11 +1,68 @@
 // Copyright Epic Games, Inc. All Rights Reserved.
 #pragma once
 
+#if WITH_EDITOR
+#include "ConsoleVariablesEditorCommandInfo.h"
+#include "ConsoleVariablesEditorModule.h"
+#endif
+
 #include "MoviePipelineSetting.h"
+#include "HAL/IConsoleManager.h"
+
 #include "MoviePipelineConsoleVariableSetting.generated.h"
 
 class IMovieSceneConsoleVariableTrackInterface;
 
+/**
+ * Represents a console variable override within the Console Variable setting.
+ */
+USTRUCT(BlueprintType)
+struct FMoviePipelineConsoleVariableEntry
+{
+	GENERATED_BODY()
+	
+	FMoviePipelineConsoleVariableEntry(const FString& InName, const float InValue, const bool bInIsEnabled = true)
+		: Name(InName)
+		, Value(InValue)
+		, bIsEnabled(bInIsEnabled)
+	{
+#if WITH_EDITOR
+		UpdateCommandInfo();
+#endif
+	}
+
+	FMoviePipelineConsoleVariableEntry()
+		: Value(0), bIsEnabled(true)
+	{
+	}
+
+#if WITH_EDITOR
+	/**
+	 * Updates the CommandInfo pointer. Should be called in PostLoad to ensure CommandInfo is properly set. In other
+	 * cases, CommandInfo will be kept up-to-date.
+	 */
+	void UpdateCommandInfo()
+	{
+		FConsoleVariablesEditorModule& ConsoleVariablesEditorModule = FConsoleVariablesEditorModule::Get();
+		CommandInfo = ConsoleVariablesEditorModule.FindCommandInfoByName(Name);
+	}
+
+	TWeakPtr<FConsoleVariablesEditorCommandInfo> CommandInfo;
+#endif
+
+	/* The name of the console variable. */
+	UPROPERTY(EditAnywhere, BlueprintReadOnly, Category = "Settings")
+	FString Name;
+
+	/* The value of the console variable. */
+	UPROPERTY(EditAnywhere, BlueprintReadOnly, Category = "Settings")
+	float Value;
+
+	/* Enable state. If disabled, this cvar entry will be ignored when resolving the final value of the cvar. */
+	UPROPERTY(EditAnywhere, BlueprintReadOnly, Category = "Settings")
+	bool bIsEnabled;
+};
+
 UCLASS(BlueprintType)
 class MOVIERENDERPIPELINESETTINGS_API UMoviePipelineConsoleVariableSetting : public UMoviePipelineSetting
 {
@@ -16,6 +73,21 @@ public:
 #if WITH_EDITOR
 	virtual FText GetDisplayText() const override { return NSLOCTEXT("MovieRenderPipeline", "ConsoleVariableSettingDisplayName", "Console Variables"); }
 	virtual FText GetFooterText(UMoviePipelineExecutorJob* InJob) const override;
+
+	/**
+	 * Returns the value of the given console variable within the presets for this setting. If none of the presets contain
+	 * this cvar, an empty string is returned.
+	 */
+    FString ResolvePresetValue(const FString& InCVarName) const;
+
+	/**
+	 * Returns the value of the console variable if the given entry were to be disabled in this setting. Searches other
+	 * override cvars first, then presets, and lastly the startup value.
+	 */
+	FString ResolveDisabledValue(const FMoviePipelineConsoleVariableEntry& InEntry) const;
+
+	/* Returns the console variable entry at the specified index, else nullptr if the index is invalid. */
+	FMoviePipelineConsoleVariableEntry* GetCVarAtIndex(const int32 InIndex);
 #endif
 	virtual bool IsValidOnShots() const override { return true; }
 	virtual bool IsValidOnPrimary() const override { return true; }
@@ -43,8 +115,9 @@ public:
 	* An array of key/value pairs for console variable name and the value you wish to set for that cvar.
 	* The existing value will automatically be cached and restored afterwards.
 	*/
-	UPROPERTY(EditAnywhere, BlueprintReadWrite, Category = "Settings")
-	TMap<FString, float> ConsoleVariables;
+	UE_DEPRECATED(5.2, "ConsoleVariables has been deprecated. Please use the console variable getters/setters (GetConsoleVariables(), RemoveConsoleVariable(), AddOrUpdateConsoleVariable(), and UpdateConsoleVariableEnableState()) instead. If scripting modifies ConsoleVariables, MRQ must be run in PIE to ensure changes are migrated.")
+	UPROPERTY(meta = (DeprecatedProperty, DeprecationMessage = "ConsoleVariables has been deprecated. Please use the console variable getters/setters (GetConsoleVariables(), RemoveConsoleVariable(), AddOrUpdateConsoleVariable(), and UpdateConsoleVariableEnableState()) instead. If scripting modifies ConsoleVariables, MRQ must be run in PIE to ensure changes are migrated."))
+	TMap<FString, float> ConsoleVariables_DEPRECATED;
 
 	/**
 	* An array of console commands to execute when this shot is started. If you need to restore the value 
@@ -60,14 +133,61 @@ public:
 	*/
 	UPROPERTY(EditAnywhere, BlueprintReadWrite, Category = "Settings")
 	TArray<FString> EndConsoleCommands;
+	
+	/**
+	 * Gets a copy of all console variable overrides. These are not meant to be changed; use the mutator methods if
+	 * console variables need to be updated.
+	 */
+	UFUNCTION(BlueprintCallable, Category = "Settings")
+	TArray<FMoviePipelineConsoleVariableEntry> GetConsoleVariables() const;
+	
+	/**
+	 * Removes the console variable override with the specified name. If more than one with the same name exists, the
+	 * last one will be removed. Returns true if at least one console variable was removed, else false.
+	 * @param bRemoveAllInstances Remove all console variables overrides with the given name (not just the last one)
+	 */
+	UFUNCTION(BlueprintCallable, Category = "Settings")
+	bool RemoveConsoleVariable(const FString& Name, const bool bRemoveAllInstances = false);
+
+	/**
+	 * Adds a console variable override with the given name and value if one does not already exist. If the console
+	 * variable with the given name already exists, its value will be updated. Returns true if the operation was
+	 * successful, else false.
+	 */
+	UFUNCTION(BlueprintCallable, Category = "Settings")
+	bool AddOrUpdateConsoleVariable(const FString& Name, const float Value);
+
+	/**
+	 * Updates the enable state of the console variable override with the provided name. Returns true if the operation
+	 * was successful, else false.
+	 */
+	UFUNCTION(BlueprintCallable, Category = "Settings")
+	bool UpdateConsoleVariableEnableState(const FString& Name, const bool bIsEnabled);
+
+	// Begin UObject interface
+	virtual void PostLoad() override;
+
+#if WITH_EDITOR
+	virtual void PostEditChangeChainProperty(FPropertyChangedChainEvent& PropertyChangedEvent) override;
+#endif
+	// End UObject interface
 
 private:
+	/** If the (deprecated) ConsoleVariables property contains any data, merge it into the new CVars property. */
+	void MergeInOldConsoleVariables();
+	
 	/** Merge together preset and override cvars into MergedConsoleVariables. Discards result of a prior merge (if any). */
-	void MergeConsoleVariables();
+	void MergeInPresetConsoleVariables();
 
 private:
+	/**
+	 * An array of console variable overrides which are applied during render and reverted after the render completes.
+	 */
+	UPROPERTY(EditAnywhere, Category = "Settings", meta = (DisplayName = "Console Variables", DisplayAfter = "ConsoleVariablePresets", ScriptNoExport))
+	TArray<FMoviePipelineConsoleVariableEntry> CVars;
+	
 	TArray<float> PreviousConsoleVariableValues;
 
 	/** Merged result of preset cvars and override cvars. */
-	TMap<FString, float> MergedConsoleVariables;
+	TArray<FMoviePipelineConsoleVariableEntry> MergedConsoleVariables;
 };
