nDisplay MRQ: fixed video export for the multiple viewports

diff --git a/Engine/Plugins/MovieScene/MovieRenderPipeline/Source/MovieRenderPipelineCore/Private/MoviePipelineVideoOutputBase.cpp b/Engine/Plugins/MovieScene/MovieRenderPipeline/Source/MovieRenderPipelineCore/Private/MoviePipelineVideoOutputBase.cpp
--- a/Engine/Plugins/MovieScene/MovieRenderPipeline/Source/MovieRenderPipelineCore/Private/MoviePipelineVideoOutputBase.cpp
+++ b/Engine/Plugins/MovieScene/MovieRenderPipeline/Source/MovieRenderPipelineCore/Private/MoviePipelineVideoOutputBase.cpp
@@ -65,6 +65,15 @@ void UMoviePipelineVideoOutputBase::OnReceiveImageDataImpl(FMoviePipelineMergerO
 	TArray<MoviePipeline::FCompositePassInfo> CompositedPasses;
 	MoviePipeline::GetPassCompositeData(InMergedOutputFrame, CompositedPasses);
 
+	// Get Camera Names
+	TArray<FString> UsedCameraNames;
+	for (TPair<FMoviePipelinePassIdentifier, TUniquePtr<FImagePixelData>>& RenderPassData : InMergedOutputFrame->ImageOutputData)
+	{
+		UsedCameraNames.AddUnique(RenderPassData.Key.CameraName);
+	}
+	// Support per-camera video output
+	const bool bIncludeCameraName = UsedCameraNames.Num() > 1;
+
 	for (TPair<FMoviePipelinePassIdentifier, TUniquePtr<FImagePixelData>>& RenderPassData : InMergedOutputFrame->ImageOutputData)
 	{
 		// Don't write out a composited pass in this loop, as it will be merged with the Final Image and not written separately. 
@@ -99,13 +108,14 @@ void UMoviePipelineVideoOutputBase::OnReceiveImageDataImpl(FMoviePipelineMergerO
 			const bool bIncludeRenderPass = InMergedOutputFrame->ImageOutputData.Num() - CompositedPasses.Num() > 1;
 			const bool bTestFrameNumber = false;
 
-			UE::MoviePipeline::ValidateOutputFormatString(FileNameFormatString, bIncludeRenderPass, bTestFrameNumber);
+			UE::MoviePipeline::ValidateOutputFormatString(FileNameFormatString, bIncludeRenderPass, bTestFrameNumber, bIncludeCameraName);
 
 			// Strip any frame number tags so we don't get one video file per frame.
 			UE::MoviePipeline::RemoveFrameNumberFormatStrings(FileNameFormatString, true);
 
 			// Create specific data that needs to override 
 			TMap<FString, FString> FormatOverrides;
+			FormatOverrides.Add(TEXT("camera_name"), RenderPassData.Key.CameraName);
 			FormatOverrides.Add(TEXT("render_pass"), RenderPassData.Key.Name);
 			FormatOverrides.Add(TEXT("ext"), GetFilenameExtension());
 
