MovieRenderQueue: fixed high-res rendering crashed caused by wrong tile size calculation

diff --git a/Engine/Plugins/MovieScene/MovieRenderPipeline/Source/MovieRenderPipelineRenderPasses/Private/MoviePipelineDeferredPasses.cpp b/Engine/Plugins/MovieScene/MovieRenderPipeline/Source/MovieRenderPipelineRenderPasses/Private/MoviePipelineDeferredPasses.cpp
--- a/Engine/Plugins/MovieScene/MovieRenderPipeline/Source/MovieRenderPipelineRenderPasses/Private/MoviePipelineDeferredPasses.cpp
+++ b/Engine/Plugins/MovieScene/MovieRenderPipeline/Source/MovieRenderPipelineRenderPasses/Private/MoviePipelineDeferredPasses.cpp
@@ -157,9 +157,19 @@ void UMoviePipelineDeferredPassBase::SetupImpl(const MoviePipeline::FMoviePipeli
 	int32 TotalNumberOfAccumulators = 0;
 	for (int32 CamIndex = 0; CamIndex < NumCameras; CamIndex++)
 	{
+		const FIntPoint OutputCameraResolution = GetEffectiveOutputResolutionForCamera(CamIndex);
+
+		// Figure out how big each sub-region (tile) is.
+		FIntPoint BackbufferResolution = FIntPoint(
+			FMath::CeilToInt((float)OutputCameraResolution.X / (float)HighResSettings->TileCount),
+			FMath::CeilToInt((float)OutputCameraResolution.Y / (float)HighResSettings->TileCount));
+
+		// Then increase each sub-region by the overlap amount.
+		BackbufferResolution = HighResSettings->CalculatePaddedBackbufferSize(BackbufferResolution);
+
 		// Re-initialize the render target and surface queue for the current camera
-		GetOrCreateViewRenderTarget(GetEffectiveOutputResolutionForCamera(CamIndex));
-		CreateSurfaceQueueImpl(GetEffectiveOutputResolutionForCamera(CamIndex));
+		GetOrCreateViewRenderTarget(BackbufferResolution);
+		CreateSurfaceQueueImpl(BackbufferResolution);
 
 		FMultiCameraViewStateData& CameraData = CameraViewStateData.AddDefaulted_GetRef();
 
