MovieRenderQueue: Add autocomplete for console variables.

diff --git a/Engine/Plugins/MovieScene/MovieRenderPipeline/Source/MovieRenderPipelineEditor/Private/Customizations/ConsoleVariableCustomization.h b/Engine/Plugins/MovieScene/MovieRenderPipeline/Source/MovieRenderPipelineEditor/Private/Customizations/ConsoleVariableCustomization.h
--- a/Engine/Plugins/MovieScene/MovieRenderPipeline/Source/MovieRenderPipelineEditor/Private/Customizations/ConsoleVariableCustomization.h
+++ b/Engine/Plugins/MovieScene/MovieRenderPipeline/Source/MovieRenderPipelineEditor/Private/Customizations/ConsoleVariableCustomization.h
@@ -1,10 +1,11 @@
-﻿// Copyright Epic Games, Inc. All Rights Reserved.
+// Copyright Epic Games, Inc. All Rights Reserved.
 #pragma once
 
 #include "Internationalization/Internationalization.h"
 #include "Misc/DefaultValueHelper.h"
 #include "MoviePipelineConsoleVariableSetting.h"
 #include "PropertyCustomizationHelpers.h"
+#include "Views/Widgets/SConsoleVariablesEditorCustomConsoleInputBox.h"
 #include "Widgets/Input/SEditableTextBox.h"
 #include "Widgets/Input/SNumericEntryBox.h"
 
@@ -56,12 +57,14 @@ protected:
 			.Padding(5, 2)
 			.FillWidth(0.75f)
 			[
-				SNew(SEditableTextBox)
+				SNew(SConsoleVariablesEditorCustomConsoleInputBox)
 				.IsEnabled(this, &FConsoleVariablesDetailsCustomization::IsEnabled, IsEnabledProp)
 				.Font(CustomizationUtils.GetRegularFont())
 				.Text(this, &FConsoleVariablesDetailsCustomization::GetConsoleVariableText, NameProp)
-				.ToolTipText(this, &FConsoleVariablesDetailsCustomization::GetConsoleVariableHelpText, NameProp)
-				.OnTextCommitted(this, &FConsoleVariablesDetailsCustomization::OnConsoleVariableNameTextCommitted, NameProp)
+				.HideOnFocusLost(false)
+				.ClearOnCommit(false)
+				.HintText(LOCTEXT("ConsoleVariableInputHintText", "Enter Console Variable"))
+				.OnTextChanged(this, &FConsoleVariablesDetailsCustomization::OnConsoleVariableNameTextChanged, NameProp)
 			]
 
 			+ SHorizontalBox::Slot()
@@ -190,7 +193,7 @@ protected:
 		return FText();
 	}
 
-	void OnConsoleVariableNameTextCommitted(const FText& NewText, ETextCommit::Type CommitType, TSharedRef<IPropertyHandle> PropertyHandle)
+	void OnConsoleVariableNameTextChanged(const FText& NewText, TSharedRef<IPropertyHandle> PropertyHandle)
 	{
 		if (PropertyHandle->IsValidHandle())
 		{
