MoviePipeline: Reworked the queue saving UI, similar to how configuration presets work.

* The queue save/load button now indicates the name of the queue that's currently loaded, as well as its dirty state (indicated with an "*").
* The queue save/load menu now provides both "Save" and "Save As".
* When loading an existing queue, a warning will be displayed if the current queue has unsaved changes.

Note: The Modify() call was removed from UMoviePipelinePrimaryConfig because it is unnecessary and was causing the queue to become
dirtied after starting a render (even though the modifications were transient).

diff --git a/Engine/Plugins/MovieScene/MovieRenderPipeline/Source/MovieRenderPipelineCore/Private/MoviePipelinePrimaryConfig.cpp b/Engine/Plugins/MovieScene/MovieRenderPipeline/Source/MovieRenderPipelineCore/Private/MoviePipelinePrimaryConfig.cpp
--- a/Engine/Plugins/MovieScene/MovieRenderPipeline/Source/MovieRenderPipelineCore/Private/MoviePipelinePrimaryConfig.cpp
+++ b/Engine/Plugins/MovieScene/MovieRenderPipeline/Source/MovieRenderPipelineCore/Private/MoviePipelinePrimaryConfig.cpp
@@ -52,10 +52,6 @@ void UMoviePipelinePrimaryConfig::InitializeTransientSettings()
 		AllSettingClasses.RemoveSwap(ExistingSetting->GetClass());
 	}
 
-#if WITH_EDITOR
-	Modify();
-#endif
-
 	// Now we can initialize an instance of every remaining setting that we don't have.
 	for (const UClass* SettingClass : AllSettingClasses)
 	{
diff --git a/Engine/Plugins/MovieScene/MovieRenderPipeline/Source/MovieRenderPipelineCore/Private/MoviePipelineQueue.cpp b/Engine/Plugins/MovieScene/MovieRenderPipeline/Source/MovieRenderPipelineCore/Private/MoviePipelineQueue.cpp
--- a/Engine/Plugins/MovieScene/MovieRenderPipeline/Source/MovieRenderPipelineCore/Private/MoviePipelineQueue.cpp
+++ b/Engine/Plugins/MovieScene/MovieRenderPipeline/Source/MovieRenderPipelineCore/Private/MoviePipelineQueue.cpp
@@ -6,8 +6,25 @@
 #include "MoviePipelineBlueprintLibrary.h"
 #include "MovieRenderPipelineCoreModule.h"
 
+#include "UObject/UObjectGlobals.h"
+
 #include UE_INLINE_GENERATED_CPP_BY_NAME(MoviePipelineQueue)
 
+UMoviePipelineQueue::UMoviePipelineQueue()
+	: QueueSerialNumber(0)
+	, bIsDirty(false)
+{
+	// Ensure instances are always transactional
+	if (!HasAnyFlags(RF_ClassDefaultObject))
+	{
+		SetFlags(RF_Transactional);
+	}
+
+#if WITH_EDITOR
+	FCoreUObjectDelegates::OnObjectModified.AddUObject(this, &UMoviePipelineQueue::OnAnyObjectModified);
+#endif
+}
+
 UMoviePipelineExecutorJob* UMoviePipelineQueue::AllocateNewJob(TSubclassOf<UMoviePipelineExecutorJob> InJobType)
 {
 	if (!ensureAlwaysMsgf(InJobType, TEXT("Failed to specify a Job Type. Use the default in project setting or UMoviePipelineExecutorJob.")))
@@ -93,6 +110,10 @@ void UMoviePipelineQueue::CopyFrom(UMoviePipelineQueue* InQueue)
 		return;
 	}
 
+	// The copy should reflect the input queue's origin (ie, the queue asset it was originally based off of). Setting
+	// the origin to the input queue would change the meaning of what the origin is.
+	QueueOrigin = InQueue->QueueOrigin;
+
 #if WITH_EDITOR
 	Modify();
 #endif
@@ -136,7 +157,27 @@ void UMoviePipelineQueue::SetJobIndex(UMoviePipelineExecutorJob* InJob, int32 In
 	QueueSerialNumber++;
 }
 
+void UMoviePipelineQueue::PreSave(FObjectPreSaveContext ObjectSaveContext)
+{
+	Super::PreSave(ObjectSaveContext);
+
+	// Reset dirty state after saving
+	bIsDirty = false;
+}
+
 #if WITH_EDITOR
+void UMoviePipelineQueue::OnAnyObjectModified(UObject* InModifiedObject)
+{
+	// Mark as dirty if this queue or any of its owned objects have been modified
+	if (InModifiedObject)
+	{
+		if (InModifiedObject->IsIn(this) || (InModifiedObject == this))
+		{
+			bIsDirty = true;
+		}
+	}
+}
+
 void UMoviePipelineExecutorJob::PostEditChangeProperty(FPropertyChangedEvent& PropertyChangedEvent)
 {
 	Super::PostEditChangeProperty(PropertyChangedEvent);
diff --git a/Engine/Plugins/MovieScene/MovieRenderPipeline/Source/MovieRenderPipelineCore/Public/MoviePipelineQueue.h b/Engine/Plugins/MovieScene/MovieRenderPipeline/Source/MovieRenderPipelineCore/Public/MoviePipelineQueue.h
--- a/Engine/Plugins/MovieScene/MovieRenderPipeline/Source/MovieRenderPipelineCore/Public/MoviePipelineQueue.h
+++ b/Engine/Plugins/MovieScene/MovieRenderPipeline/Source/MovieRenderPipelineCore/Public/MoviePipelineQueue.h
@@ -383,7 +383,13 @@ protected:
 	virtual float GetStatusProgress_Implementation() const { return StatusProgress; }
 	virtual bool IsConsumed_Implementation() const { return bIsConsumed; }
 	virtual void OnDuplicated_Implementation();
-	virtual void SetIsEnabled_Implementation(const bool bInEnabled) { bEnabled = bInEnabled; }
+	virtual void SetIsEnabled_Implementation(const bool bInEnabled)
+	{
+		// Call Modify() so OnObjectModified picks up this change.
+		Modify();
+
+		bEnabled = bInEnabled;
+	}
 	virtual bool IsEnabled_Implementation() const { return bEnabled; }
 	// ~UMoviePipelineExecutorJob Interface
 
@@ -452,15 +458,7 @@ class MOVIERENDERPIPELINECORE_API UMoviePipelineQueue : public UObject
 	GENERATED_BODY()
 public:
 	
-	UMoviePipelineQueue()
-		: QueueSerialNumber(0)
-	{
-		// Ensure instances are always transactional
-		if (!HasAnyFlags(RF_ClassDefaultObject))
-		{
-			SetFlags(RF_Transactional);
-		}
-	}
+	UMoviePipelineQueue();
 	
 	/**
 	* Allocates a new Job in this Queue. The Queue owns the jobs for memory management purposes,
@@ -505,15 +503,44 @@ public:
 		return Jobs;
 	}
 
+	/**
+	 * Gets the queue that this queue was originally based on (if any). The origin will only be set on transient
+	 * queues; the origin will be nullptr for non-transient queues because the origin will be this object.
+	 */
+	UFUNCTION(BlueprintCallable, Category = "Movie Render Pipeline|Queue")
+	UMoviePipelineQueue* GetQueueOrigin() const { return QueueOrigin.LoadSynchronous(); }
+
+	/**
+	 * Sets the queue that this queue originated from (if any). The origin should only be set for transient queues.
+	 */
+	UFUNCTION(BlueprintCallable, Category = "Movie Render Pipeline|Queue")
+	void SetQueueOrigin(UMoviePipelineQueue* InConfig) { QueueOrigin = InConfig; }
+
 	/** 
 	* Replace the contents of this queue with a copy of the contents from another queue. 
 	*/
-	UFUNCTION(BlueprintCallable, Category = "Movie Render Pipeline")
+	UFUNCTION(BlueprintCallable, Category = "Movie Render Pipeline|Queue")
 	void CopyFrom(UMoviePipelineQueue* InQueue);
 	
 	/* Set the index of the given job */
-	UFUNCTION(BlueprintCallable, Category = "Movie Render Pipeline")
+	UFUNCTION(BlueprintCallable, Category = "Movie Render Pipeline|Queue")
 	void SetJobIndex(UMoviePipelineExecutorJob* InJob, int32 Index);
+
+#if WITH_EDITOR
+	/**
+	 * Gets the dirty state of this queue. Note that dirty state is only tracked when the editor is active.
+	 */
+	UFUNCTION(BlueprintCallable, Category = "Movie Render Pipeline|Queue")
+	bool IsDirty() const { return bIsDirty; }
+
+	/**
+	 * Sets the dirty state of this queue. Generally the queue will correctly track the dirty state; however, there are
+	 * situations where a queue may need its dirty state reset (eg, it may be appropriate to reset the dirty state after
+	 * a call to CopyFrom(), depending on the use case).
+	 */
+	UFUNCTION(BlueprintCallable, Category = "Movie Render Pipeline|Queue")
+	void SetIsDirty(const bool bNewDirtyState) { bIsDirty = bNewDirtyState; }
+#endif
 	
 	/**
 	 * Retrieve the serial number that is incremented when a job is added or removed from this list.
@@ -529,10 +556,27 @@ public:
 		QueueSerialNumber++;
 	}
 
+	// UObject Interface
+	virtual void PreSave(FObjectPreSaveContext ObjectSaveContext) override;
+	// ~UObject interface
+
+private:
+#if WITH_EDITOR
+	void OnAnyObjectModified(UObject* InModifiedObject);
+#endif
+
 private:
 	UPROPERTY(Instanced)
 	TArray<TObjectPtr<UMoviePipelineExecutorJob>> Jobs;
+
+	/* The queue that this queue originated from. Helpful for determining the source of the queue when this queue is transient. */
+	UPROPERTY()
+	TSoftObjectPtr<UMoviePipelineQueue> QueueOrigin;
 	
 private:
 	int32 QueueSerialNumber;
+
+	/* Dirty state is tracked explicitly so transient queues, which are based on a transient package, know when modifications
+	 * have been made. */
+	bool bIsDirty;
 };
\ No newline at end of file
diff --git a/Engine/Plugins/MovieScene/MovieRenderPipeline/Source/MovieRenderPipelineEditor/Private/Widgets/SMoviePipelineQueuePanel.cpp b/Engine/Plugins/MovieScene/MovieRenderPipeline/Source/MovieRenderPipelineEditor/Private/Widgets/SMoviePipelineQueuePanel.cpp
--- a/Engine/Plugins/MovieScene/MovieRenderPipeline/Source/MovieRenderPipelineEditor/Private/Widgets/SMoviePipelineQueuePanel.cpp
+++ b/Engine/Plugins/MovieScene/MovieRenderPipeline/Source/MovieRenderPipelineEditor/Private/Widgets/SMoviePipelineQueuePanel.cpp
@@ -39,6 +39,7 @@
 #include "Editor.h"
 #include "PropertyEditorModule.h"
 #include "IDetailsView.h"
+#include "Misc/MessageDialog.h"
 #include "Modules/ModuleManager.h"
 
 // UnrealEd Includes
@@ -157,7 +158,7 @@ void SMoviePipelineQueuePanel::Construct(const FArguments& InArgs)
 						.Padding(0, 1, 0, 0)
 						[
 							SNew(STextBlock)
-							.Text(LOCTEXT("SavedQueueToolbarButton", "Load/Save Queue"))
+							.Text(this, &SMoviePipelineQueuePanel::GetQueueMenuButtonText)
 						]
 					]
 				]
@@ -500,10 +501,39 @@ TSharedRef<SWidget> SMoviePipelineQueuePanel::OnGenerateSavedQueuesMenu()
 
 	IContentBrowserSingleton& ContentBrowser = FModuleManager::LoadModuleChecked<FContentBrowserModule>("ContentBrowser").Get();
 
+	static const FText NoLoadedQueueText = LOCTEXT("NoLoadedQueue", "None");
+	static const FText LoadedQueueFormatText = LOCTEXT("CurrentQueueFormat", "Current Queue: {0}");
+	const FString QueueName = GetQueueOriginName();
+	const FText LoadedQueueName = FText::Format(
+		LoadedQueueFormatText, !QueueName.IsEmpty() ? FText::FromString(QueueName) : NoLoadedQueueText);
+
 	MenuBuilder.AddMenuEntry(
-		LOCTEXT("SaveAsQueue_Text", "Save As Asset"),
-		LOCTEXT("SaveAsQueue_Tip", "Save the current configuration as a new preset that can be shared between multiple jobs, or imported later as the base of a new configuration."),
+		LoadedQueueName,
+		LoadedQueueName,
+		FSlateIcon(),
+		FUIAction(
+			FExecuteAction(),
+			FCanExecuteAction::CreateLambda([]() { return false; })),
+		NAME_None,
+		EUserInterfaceActionType::None
+	);
+
+	MenuBuilder.AddMenuSeparator();
+	
+	MenuBuilder.AddMenuEntry(
+		LOCTEXT("SaveQueue_Text", "Save Queue"),
+		LOCTEXT("SaveQueue_Tip", "Save the current configuration in its existing preset which can be shared between multiple jobs, or imported later as the base of a new configuration."),
 		FSlateIcon(FAppStyle::Get().GetStyleSetName(), "AssetEditor.SaveAsset"),
+		FUIAction(
+			FExecuteAction::CreateSP(this, &SMoviePipelineQueuePanel::OnSaveAsset),
+			FCanExecuteAction::CreateSP(this, &SMoviePipelineQueuePanel::IsQueueDirty)
+		)
+	);
+	
+	MenuBuilder.AddMenuEntry(
+		LOCTEXT("SaveAsQueue_Text", "Save Queue As"),
+		LOCTEXT("SaveAsQueue_Tip", "Save the current configuration as a new preset that can be shared between multiple jobs, or imported later as the base of a new configuration."),
+		FSlateIcon(FAppStyle::Get().GetStyleSetName(), "AssetEditor.SaveAssetAs"),
 		FUIAction(FExecuteAction::CreateSP(this, &SMoviePipelineQueuePanel::OnSaveAsAsset))
 	);
 
@@ -524,10 +554,10 @@ TSharedRef<SWidget> SMoviePipelineQueuePanel::OnGenerateSavedQueuesMenu()
 		AssetPickerConfig.AssetShowWarningText = LOCTEXT("NoQueueAssets_Warning", "No Queues Found");
 		AssetPickerConfig.Filter.ClassPaths.Add(UMoviePipelineQueue::StaticClass()->GetClassPathName());
 		AssetPickerConfig.Filter.bRecursiveClasses = true;
-		AssetPickerConfig.OnAssetSelected = FOnAssetSelected::CreateSP(this, &SMoviePipelineQueuePanel::OnImportSavedQueueAssest);
+		AssetPickerConfig.OnAssetSelected = FOnAssetSelected::CreateSP(this, &SMoviePipelineQueuePanel::OnImportSavedQueueAsset);
 	}
 
-	MenuBuilder.BeginSection(NAME_None, LOCTEXT("LoadQueue_MenuSection", "Load Queue"));
+	MenuBuilder.BeginSection(NAME_None, LOCTEXT("LoadQueue_MenuSection", "Import Queue"));
 	{
 		TSharedRef<SWidget> PresetPicker = SNew(SBox)
 			.MinDesiredWidth(400.f)
@@ -543,6 +573,19 @@ TSharedRef<SWidget> SMoviePipelineQueuePanel::OnGenerateSavedQueuesMenu()
 	return MenuBuilder.MakeWidget();
 }
 
+FText SMoviePipelineQueuePanel::GetQueueMenuButtonText() const
+{
+	const FString QueueName = GetQueueOriginName();
+	if (!QueueName.IsEmpty())
+	{
+		return FText::Format(FText::FromString(TEXT("{0}{1}")),
+			FText::FromString(QueueName),
+			FText::FromString(IsQueueDirty() ? TEXT(" *") : TEXT("")));
+	}
+	
+	return LOCTEXT("QueueSaveMenuUnsavedConfig_Text", "Unsaved Queue *");
+}
+
 bool SMoviePipelineQueuePanel::OpenSaveDialog(const FString& InDefaultPath, const FString& InNewNameSuggestion, FString& OutPackageName)
 {
 	FSaveAssetDialogConfig SaveAssetDialogConfig;
@@ -616,6 +659,20 @@ bool SMoviePipelineQueuePanel::GetSavePresetPackageName(const FString& InExistin
 	return true;
 }
 
+void SMoviePipelineQueuePanel::OnSaveAsset()
+{
+	UMoviePipelineQueue* QueueOrigin = GetQueueOrigin();
+	
+	// If the queue origin is not known, then treat the "Save" as a "Save As"
+	if (!QueueOrigin)
+	{
+		OnSaveAsAsset();
+		return;
+	}
+	
+	SaveTransientQueueToAsset(QueueOrigin);
+}
+
 void SMoviePipelineQueuePanel::OnSaveAsAsset()
 {
 	UMoviePipelineQueueSubsystem* Subsystem = GEditor->GetEditorSubsystem<UMoviePipelineQueueSubsystem>();
@@ -634,28 +691,66 @@ void SMoviePipelineQueuePanel::OnSaveAsAsset()
 	NewPackage->MarkAsFullyLoaded();
 	UMoviePipelineQueue* DuplicateQueue = DuplicateObject<UMoviePipelineQueue>(CurrentQueue, NewPackage, *NewAssetName);
 
-	if (DuplicateQueue)
+	SaveTransientQueueToAsset(DuplicateQueue);
+}
+
+void SMoviePipelineQueuePanel::SaveTransientQueueToAsset(UMoviePipelineQueue* DestinationQueue) const
+{
+	if (!DestinationQueue)
 	{
-		DuplicateQueue->SetFlags(RF_Public | RF_Standalone | RF_Transactional);
+		return;
+	}
+
+	UMoviePipelineQueueSubsystem* Subsystem = GEditor->GetEditorSubsystem<UMoviePipelineQueueSubsystem>();
+	check(Subsystem);
+	UMoviePipelineQueue* CurrentQueue = Subsystem->GetQueue();
 
-		FAssetRegistryModule::AssetCreated(DuplicateQueue);
+	DestinationQueue->CopyFrom(CurrentQueue);
+	DestinationQueue->SetQueueOrigin(nullptr);
+	DestinationQueue->MarkPackageDirty();
+	DestinationQueue->SetFlags(RF_Public | RF_Standalone | RF_Transactional);
 
-		FEditorFileUtils::EPromptReturnCode PromptReturnCode = FEditorFileUtils::PromptForCheckoutAndSave({ NewPackage }, false, false);
+	FAssetRegistryModule::AssetCreated(DestinationQueue);
+
+	const bool bCheckDirty = false;
+	const bool bPromptToSave = false;
+	FEditorFileUtils::EPromptReturnCode PromptReturnCode = FEditorFileUtils::PromptForCheckoutAndSave({ DestinationQueue->GetPackage() }, bCheckDirty, bPromptToSave);
+	if (PromptReturnCode == FEditorFileUtils::EPromptReturnCode::PR_Success)
+	{
+		LoadQueue(DestinationQueue);
 	}
 }
 
-void SMoviePipelineQueuePanel::OnImportSavedQueueAssest(const FAssetData& InPresetAsset)
+void SMoviePipelineQueuePanel::OnImportSavedQueueAsset(const FAssetData& InPresetAsset) const
 {
 	FSlateApplication::Get().DismissAllMenus();
 
+	if (IsQueueDirty())
+	{
+		const FText TitleText = LOCTEXT("UnsavedQueueWarningTitle", "Unsaved Changes to Queue");
+		const FText MessageText = LOCTEXT("UnsavedQueueWarningMessage", "The changes made to the current queue will be lost by importing another queue. Do you want to continue with this import?");
+
+		if (FMessageDialog::Open(EAppMsgType::YesNo, MessageText, &TitleText) == EAppReturnType::No)
+		{
+			return;
+		}
+	}
+
 	UMoviePipelineQueue* SavedQueue = CastChecked<UMoviePipelineQueue>(InPresetAsset.GetAsset());
+	LoadQueue(SavedQueue);
+}
+
+void SMoviePipelineQueuePanel::LoadQueue(UMoviePipelineQueue* SavedQueue) const
+{
 	if (SavedQueue)
 	{
-		// Duplicate the queue so we don't start modifying the one in the asset.
 		UMoviePipelineQueueSubsystem* Subsystem = GEditor->GetEditorSubsystem<UMoviePipelineQueueSubsystem>();
 		check(Subsystem);
 
-		Subsystem->GetQueue()->CopyFrom(SavedQueue);
+		UMoviePipelineQueue* CurrentQueue = Subsystem->GetQueue();
+		CurrentQueue->CopyFrom(SavedQueue);
+		CurrentQueue->SetQueueOrigin(SavedQueue);
+		CurrentQueue->SetIsDirty(false);
 
 		// Update the shot list in case the stored queue being copied is out of date with the sequence
 		for (UMoviePipelineExecutorJob* Job : Subsystem->GetQueue()->GetJobs())
@@ -688,4 +783,46 @@ void SMoviePipelineQueuePanel::OnImportSavedQueueAssest(const FAssetData& InPres
 	}
 }
 
+bool SMoviePipelineQueuePanel::IsQueueDirty() const
+{
+	UMoviePipelineQueueSubsystem* Subsystem = GEditor->GetEditorSubsystem<UMoviePipelineQueueSubsystem>();
+	check(Subsystem);
+
+	const UMoviePipelineQueue* Queue = Subsystem->GetQueue();
+
+	// The queue is considered dirty if the current queue has no origin (ie, it has never been saved) or it has been
+	// modified since it was loaded
+	return !GetQueueOrigin() || (Queue && Queue->IsDirty());
+}
+
+UMoviePipelineQueue* SMoviePipelineQueuePanel::GetQueueOrigin() const
+{
+	UMoviePipelineQueueSubsystem* Subsystem = GEditor->GetEditorSubsystem<UMoviePipelineQueueSubsystem>();
+	check(Subsystem);
+
+	const UMoviePipelineQueue* Queue = Subsystem->GetQueue();
+
+	return Queue ? Queue->GetQueueOrigin() : nullptr;
+}
+
+FString SMoviePipelineQueuePanel::GetQueueOriginName() const
+{
+	const UMoviePipelineQueue* QueueOrigin = GetQueueOrigin();
+	if (!QueueOrigin)
+	{
+		return FString();
+	}
+
+	if (const UPackage* Package = QueueOrigin->GetPackage())
+	{
+		FString OutPackageRoot, OutPackagePath, OutPackageName;
+		if (FPackageName::SplitLongPackageName(Package->GetName(), OutPackageRoot, OutPackagePath, OutPackageName))
+		{
+			return OutPackageName;
+		}
+	}
+
+	return FString();
+}
+
 #undef LOCTEXT_NAMESPACE // SMoviePipelineQueuePanel
\ No newline at end of file
diff --git a/Engine/Plugins/MovieScene/MovieRenderPipeline/Source/MovieRenderPipelineEditor/Private/Widgets/SMoviePipelineQueuePanel.h b/Engine/Plugins/MovieScene/MovieRenderPipeline/Source/MovieRenderPipelineEditor/Private/Widgets/SMoviePipelineQueuePanel.h
--- a/Engine/Plugins/MovieScene/MovieRenderPipeline/Source/MovieRenderPipelineEditor/Private/Widgets/SMoviePipelineQueuePanel.h
+++ b/Engine/Plugins/MovieScene/MovieRenderPipeline/Source/MovieRenderPipelineEditor/Private/Widgets/SMoviePipelineQueuePanel.h
@@ -6,6 +6,7 @@
 #include "Widgets/SCompoundWidget.h"
 
 class UMoviePipelineConfigBase;
+class UMoviePipelineQueue;
 class SMoviePipelineQueueEditor;
 class SWindow;
 class UMoviePipelineExecutorJob;
@@ -52,10 +53,17 @@ private:
 	bool IsDetailsViewEnabled() const;
 
 	TSharedRef<SWidget> OnGenerateSavedQueuesMenu();
+	FText GetQueueMenuButtonText() const;
 	bool OpenSaveDialog(const FString& InDefaultPath, const FString& InNewNameSuggestion, FString& OutPackageName);
 	bool GetSavePresetPackageName(const FString& InExistingName, FString& OutName);
+	void OnSaveAsset();
 	void OnSaveAsAsset();
-	void OnImportSavedQueueAssest(const FAssetData& InPresetAsset);
+	void SaveTransientQueueToAsset(UMoviePipelineQueue* DestinationQueue) const;
+	void OnImportSavedQueueAsset(const FAssetData& InPresetAsset) const;
+	void LoadQueue(UMoviePipelineQueue* SavedQueue) const;
+	bool IsQueueDirty() const;
+	UMoviePipelineQueue* GetQueueOrigin() const;
+	FString GetQueueOriginName() const;
 
 private:
 	/** Allocates a transient preset so that the user can use the pipeline without saving it to an asset first. */
