MoviePipeline: Fixed crash which occurs when the Deferred Rendering setting has "Add Default Layer" checked, but no stencil layers are supplied.
The render is aborted and an error generated indicating that stencil layers need to be added.

diff --git a/Engine/Plugins/MovieScene/MovieRenderPipeline/Source/MovieRenderPipelineRenderPasses/Private/MoviePipelineDeferredPasses.cpp b/Engine/Plugins/MovieScene/MovieRenderPipeline/Source/MovieRenderPipelineRenderPasses/Private/MoviePipelineDeferredPasses.cpp
--- a/Engine/Plugins/MovieScene/MovieRenderPipeline/Source/MovieRenderPipelineRenderPasses/Private/MoviePipelineDeferredPasses.cpp
+++ b/Engine/Plugins/MovieScene/MovieRenderPipeline/Source/MovieRenderPipelineRenderPasses/Private/MoviePipelineDeferredPasses.cpp
@@ -128,6 +128,13 @@ void UMoviePipelineDeferredPassBase::SetupImpl(const MoviePipeline::FMoviePipeli
 	Super::SetupImpl(InPassInitSettings);
 	LLM_SCOPE_BYNAME(TEXT("MoviePipeline/DeferredPassSetup"));
 
+	if (bAddDefaultLayer && (GetNumStencilLayers() == 0))
+	{
+		UE_LOG(LogMovieRenderPipeline, Error, TEXT("The 'Add Default Layer' deferred rendering option requires at least one Actor or Data Layer to be specified."));
+		GetPipeline()->Shutdown(true);
+		return;
+	}
+
 	{
 		TSoftObjectPtr<UMaterialInterface> StencilMatRef = TSoftObjectPtr<UMaterialInterface>(FSoftObjectPath(StencilLayerMaterialAsset));
 		StencilLayerMaterial = StencilMatRef.LoadSynchronous();
