MoviePipeline: Setting configs should return their origin as nullptr for "self". Only transient configs should have their origin set. This is consistent with the API for origins in UMoviePipelineQueue.

diff --git a/Engine/Plugins/MovieScene/MovieRenderPipeline/Source/MovieRenderPipelineCore/Public/MoviePipelineConfigBase.h b/Engine/Plugins/MovieScene/MovieRenderPipeline/Source/MovieRenderPipelineCore/Public/MoviePipelineConfigBase.h
--- a/Engine/Plugins/MovieScene/MovieRenderPipeline/Source/MovieRenderPipelineCore/Public/MoviePipelineConfigBase.h
+++ b/Engine/Plugins/MovieScene/MovieRenderPipeline/Source/MovieRenderPipelineCore/Public/MoviePipelineConfigBase.h
@@ -47,14 +47,14 @@ public:
 	virtual TArray<UMoviePipelineSetting*> GetUserSettings() const { return Settings; }
 
 	/**
-	 * Gets the config that this config was originally based on (if any). Helpful for determining where a transient
-	 * config came from.
+	 * Gets the config that this config was originally based on (if any). The origin will only be set on transient
+	 * configs; the origin will be nullptr for non-transient configs because the origin will be this object.
 	 */
 	UFUNCTION(BlueprintCallable, Category = "Movie Render Pipeline")
 	UMoviePipelineConfigBase* GetConfigOrigin() const { return ConfigOrigin.LoadSynchronous(); }
 
 	/**
-	 * Sets the config that this config originated from (if any).
+	 * Sets the config that this config originated from (if any). The origin should only be set for transient configs.
 	 */
 	UFUNCTION(BlueprintCallable, Category = "Movie Render Pipeline")
 	void SetConfigOrigin(UMoviePipelineConfigBase* InConfig) { ConfigOrigin = InConfig; }
diff --git a/Engine/Plugins/MovieScene/MovieRenderPipeline/Source/MovieRenderPipelineEditor/Private/Widgets/SMoviePipelineConfigPanel.cpp b/Engine/Plugins/MovieScene/MovieRenderPipeline/Source/MovieRenderPipelineEditor/Private/Widgets/SMoviePipelineConfigPanel.cpp
--- a/Engine/Plugins/MovieScene/MovieRenderPipeline/Source/MovieRenderPipelineEditor/Private/Widgets/SMoviePipelineConfigPanel.cpp
+++ b/Engine/Plugins/MovieScene/MovieRenderPipeline/Source/MovieRenderPipelineEditor/Private/Widgets/SMoviePipelineConfigPanel.cpp
@@ -556,6 +556,7 @@ void SMoviePipelineConfigPanel::SaveTransientPresetToAsset(UMoviePipelineConfigB
 	if (TransientPreset && DestinationPreset)
 	{
 		DestinationPreset->CopyFrom(TransientPreset);
+		DestinationPreset->SetConfigOrigin(nullptr);
 		DestinationPreset->MarkPackageDirty();
 
 		// We need to be able to save an empty string into the file. When we duplicate the Output Setting in CopyFrom, the PostLoad is fired which fixes it back up
