MoviePipeline: Fixed a regression where video output codecs didn't get burn-ins composited into the video despite composite being enabled.

diff --git a/Engine/Plugins/MovieScene/MovieRenderPipeline/Source/MovieRenderPipelineCore/Private/MoviePipelineVideoOutputBase.cpp b/Engine/Plugins/MovieScene/MovieRenderPipeline/Source/MovieRenderPipelineCore/Private/MoviePipelineVideoOutputBase.cpp
--- a/Engine/Plugins/MovieScene/MovieRenderPipeline/Source/MovieRenderPipelineCore/Private/MoviePipelineVideoOutputBase.cpp
+++ b/Engine/Plugins/MovieScene/MovieRenderPipeline/Source/MovieRenderPipelineCore/Private/MoviePipelineVideoOutputBase.cpp
@@ -231,10 +231,26 @@ void UMoviePipelineVideoOutputBase::OnReceiveImageDataImpl(FMoviePipelineMergerO
 
 		//FGraphEventRef Event = Task.Execute([this, OutputWriter, RawRenderPassData]
 		//	{
-			if (RenderPassData.Key == FMoviePipelinePassIdentifier(TEXT("FinalImage")))
+			if (RenderPassData.Key.Name == TEXT("FinalImage"))
 			{
+				TArray<MoviePipeline::FCompositePassInfo> CompositesForThisCamera;
+				for (MoviePipeline::FCompositePassInfo& CompositePass : CompositedPasses)
+				{
+					// Match them up by camera name so multiple passes intended for different camera names work.
+					if (RenderPassData.Key.CameraName == CompositePass.PassIdentifier.CameraName)
+					{
+						// Create a new composite pass (but move the actual pixel data), otherwise when the main
+						// loop tries to check if we should skip writing out this image pass in the future, it fails
+						// because we've MoveTemp'd CompositePass and thus the name checks no longer pass.
+						MoviePipeline::FCompositePassInfo NewPassInfo;
+						NewPassInfo.PassIdentifier = CompositePass.PassIdentifier;
+						NewPassInfo.PixelData = MoveTemp(CompositePass.PixelData);
+						CompositesForThisCamera.Add(MoveTemp(NewPassInfo));
+					}
+				}
+
 				// Enqueue a encode for this frame onto our worker thread.
-				this->WriteFrame_EncodeThread(OutputWriter->Get<0>().Get(), RawRenderPassData, MoveTemp(CompositedPasses));
+				this->WriteFrame_EncodeThread(OutputWriter->Get<0>().Get(), RawRenderPassData, MoveTemp(CompositesForThisCamera));
 			}
 			else
 			{
