MRQ: Improve exported color space metadata (for multi-layer exrs).

The metadata now respects active color settings, and we supply the color space source/destination names when relevant.

#rnx

diff --git a/Engine/Plugins/MovieScene/MovieRenderPipeline/Source/MovieRenderPipelineRenderPasses/Private/MoviePipelineEXROutput.cpp b/Engine/Plugins/MovieScene/MovieRenderPipeline/Source/MovieRenderPipelineRenderPasses/Private/MoviePipelineEXROutput.cpp
--- a/Engine/Plugins/MovieScene/MovieRenderPipeline/Source/MovieRenderPipelineRenderPasses/Private/MoviePipelineEXROutput.cpp
+++ b/Engine/Plugins/MovieScene/MovieRenderPipeline/Source/MovieRenderPipelineRenderPasses/Private/MoviePipelineEXROutput.cpp
@@ -19,6 +19,7 @@
 #include "Modules/ModuleManager.h"
 #include "MoviePipelineUtils.h"
 #include "ColorSpace.h"
+#include "HDRHelper.h"
 
 THIRD_PARTY_INCLUDES_START
 #include "OpenEXR/ImfChannelList.h"
@@ -166,15 +167,19 @@ bool FEXRImageWriteTask::WriteToDisk()
 		// Insert our key-value pair metadata (if any, can be an arbitrary set of key/value pairs)
 		AddFileMetadata(Header);
 
-		// Write the working color space into EXR chromaticities
-		const UE::Color::FColorSpace& WCS = UE::Color::FColorSpace::GetWorking();
-		Imf::Chromaticities Chromaticities = {
-			IMATH_NAMESPACE::V2f((float)WCS.GetRedChromaticity().X, (float)WCS.GetRedChromaticity().Y),
-			IMATH_NAMESPACE::V2f((float)WCS.GetGreenChromaticity().X, (float)WCS.GetGreenChromaticity().Y),
-			IMATH_NAMESPACE::V2f((float)WCS.GetBlueChromaticity().X, (float)WCS.GetBlueChromaticity().Y),
-			IMATH_NAMESPACE::V2f((float)WCS.GetWhiteChromaticity().X, (float)WCS.GetWhiteChromaticity().Y),
-		};
-		Imf::addChromaticities(Header, Chromaticities);
+		if (ColorSpaceChromaticities.Num() > 0)
+		{
+			if (ensureMsgf(ColorSpaceChromaticities.Num() == 4, TEXT("Four chromaticity coordinates expected.")))
+			{
+				Imf::Chromaticities Chromaticities = {
+					IMATH_NAMESPACE::V2f((float)ColorSpaceChromaticities[0].X, (float)ColorSpaceChromaticities[0].Y),
+					IMATH_NAMESPACE::V2f((float)ColorSpaceChromaticities[1].X, (float)ColorSpaceChromaticities[1].Y),
+					IMATH_NAMESPACE::V2f((float)ColorSpaceChromaticities[2].X, (float)ColorSpaceChromaticities[2].Y),
+					IMATH_NAMESPACE::V2f((float)ColorSpaceChromaticities[3].X, (float)ColorSpaceChromaticities[3].Y),
+				};
+				Imf::addChromaticities(Header, Chromaticities);
+			}
+		}
 
 		FExrMemStreamOut OutputFile; 
 		
@@ -501,6 +506,25 @@ void UMoviePipelineImageSequenceOutput_EXR::OnReceiveImageDataImpl(FMoviePipelin
 		}
 		MultiLayerImageTask->FileMetadata = NewFileMetdataMap;
 
+		// Add color space metadata to the output: xy chromaticity coordinates and/or the color space source/dest names.
+		// TODO: Support is also needed for regular exrs via the image wrapper module.
+		{
+			UMoviePipelineColorSetting* ColorSetting = GetPipeline()->GetPipelinePrimaryConfig()->FindSetting<UMoviePipelineColorSetting>();
+
+			FColorSpaceMetadata ColorSpaceMetadata = GetColorSpaceMetadata(ColorSetting);
+
+			if (!ColorSpaceMetadata.SourceName.IsEmpty())
+			{
+				MultiLayerImageTask->FileMetadata.Add("unreal/colorSpace/source", ColorSpaceMetadata.SourceName);
+			}
+			if (!ColorSpaceMetadata.DestinationName.IsEmpty())
+			{
+				MultiLayerImageTask->FileMetadata.Add("unreal/colorSpace/destination", ColorSpaceMetadata.DestinationName);
+			}
+
+			MultiLayerImageTask->ColorSpaceChromaticities = ColorSpaceMetadata.Chromaticities;
+		}
+
 		int32 LayerIndex = 0;
 		bool bRequiresTransparentOutput = false;
 		int32 ShotIndex = 0;
@@ -565,3 +589,74 @@ void UMoviePipelineImageSequenceOutput_EXR::OnReceiveImageDataImpl(FMoviePipelin
 	}
 }
 
+UMoviePipelineImageSequenceOutput_EXR::FColorSpaceMetadata UMoviePipelineImageSequenceOutput_EXR::GetColorSpaceMetadata(UMoviePipelineColorSetting* InColorSettings)
+{
+	auto GetDisplayGamutTypeFn = [](EDisplayColorGamut InDisplayGamut) -> TPair<UE::Color::EColorSpace, FString>
+	{
+		switch (InDisplayGamut)
+		{
+		case EDisplayColorGamut::sRGB_D65: return { UE::Color::EColorSpace::sRGB, TEXT("sRGB") };
+		case EDisplayColorGamut::DCIP3_D65: return { UE::Color::EColorSpace::P3DCI, TEXT("P3DCI") };
+		case EDisplayColorGamut::Rec2020_D65: return { UE::Color::EColorSpace::Rec2020, TEXT("Rec2020") };
+		case EDisplayColorGamut::ACES_D60: return { UE::Color::EColorSpace::ACESAP0, TEXT("ACESAP0") };
+		case EDisplayColorGamut::ACEScg_D60: return { UE::Color::EColorSpace::ACESAP1, TEXT("ACESAP1") };
+
+		default:
+			checkNoEntry();
+			return {};
+		}
+	};
+
+	const UE::Color::FColorSpace& WCS = UE::Color::FColorSpace::GetWorking();
+	FColorSpaceMetadata OutMetadata;
+
+	if (InColorSettings)
+	{
+		if (InColorSettings->OCIOConfiguration.bIsEnabled)
+		{
+			// Note: OpenColorIO does not expose chromaticity information so we only provide transform names.
+			const FOpenColorIOColorConversionSettings& ConversionSettings = InColorSettings->OCIOConfiguration.ColorConfiguration;
+
+			if (ConversionSettings.IsDisplayView())
+			{
+				switch (ConversionSettings.DisplayViewDirection)
+				{
+				case EOpenColorIOViewTransformDirection::Forward:
+					OutMetadata.SourceName = ConversionSettings.SourceColorSpace.ToString();
+					OutMetadata.DestinationName = ConversionSettings.DestinationDisplayView.ToString();
+					break;
+				case EOpenColorIOViewTransformDirection::Inverse:
+					OutMetadata.SourceName = ConversionSettings.DestinationDisplayView.ToString();
+					OutMetadata.DestinationName = ConversionSettings.SourceColorSpace.ToString();
+					break;
+
+				default:
+					checkNoEntry();
+				}
+			}
+			else
+			{
+				OutMetadata.SourceName = ConversionSettings.SourceColorSpace.ToString();
+				OutMetadata.DestinationName = ConversionSettings.DestinationColorSpace.ToString();
+			}
+		}
+		else if (!InColorSettings->bDisableToneCurve)
+		{
+			TPair<UE::Color::EColorSpace, FString> ColorSpaceType = GetDisplayGamutTypeFn(HDRGetDefaultDisplayColorGamut());
+			UE::Color::FColorSpace OuptutCS = UE::Color::FColorSpace(ColorSpaceType.Key);
+
+			OutMetadata.DestinationName = ColorSpaceType.Value;
+			OutMetadata.Chromaticities = { OuptutCS.GetRedChromaticity(), OuptutCS.GetGreenChromaticity(), OuptutCS.GetBlueChromaticity(), OuptutCS.GetWhiteChromaticity() };
+		}
+		else
+		{
+			OutMetadata.Chromaticities = { WCS.GetRedChromaticity(), WCS.GetGreenChromaticity(), WCS.GetBlueChromaticity(), WCS.GetWhiteChromaticity() };
+		}
+	}
+	else
+	{
+		OutMetadata.Chromaticities = { WCS.GetRedChromaticity(), WCS.GetGreenChromaticity(), WCS.GetBlueChromaticity(), WCS.GetWhiteChromaticity() };
+	}
+
+	return OutMetadata;
+}
diff --git a/Engine/Plugins/MovieScene/MovieRenderPipeline/Source/MovieRenderPipelineRenderPasses/Private/MoviePipelineEXROutput.h b/Engine/Plugins/MovieScene/MovieRenderPipeline/Source/MovieRenderPipelineRenderPasses/Private/MoviePipelineEXROutput.h
--- a/Engine/Plugins/MovieScene/MovieRenderPipeline/Source/MovieRenderPipelineRenderPasses/Private/MoviePipelineEXROutput.h
+++ b/Engine/Plugins/MovieScene/MovieRenderPipeline/Source/MovieRenderPipelineRenderPasses/Private/MoviePipelineEXROutput.h
@@ -22,6 +22,8 @@ THIRD_PARTY_INCLUDES_END
 
 #include "MoviePipelineEXROutput.generated.h"
 
+class UMoviePipelineColorSetting;
+
 UENUM(BlueprintType)
 enum class EEXRCompressionFormat : uint8
 {
@@ -74,6 +76,9 @@ public:
 	/** Overscan info used to create apropriate dataWindow for EXR output. Goes from 0.0 to 1.0. */
 	float OverscanPercentage;
 
+	/** Color space chromaticity metadata. */
+	TArray<FVector2d> ColorSpaceChromaticities;
+
 	FEXRImageWriteTask()
 		: bOverwriteFile(true)
 		, Compression(EEXRCompressionFormat::PIZ)
@@ -143,4 +148,17 @@ public:
 	*/
 	UPROPERTY(EditAnywhere, BlueprintReadWrite, Category = "EXR")
 	bool bMultilayer;
+
+protected:
+	struct FColorSpaceMetadata
+	{
+		FString SourceName;
+		FString DestinationName;
+		TArray<FVector2d> Chromaticities;
+	};
+
+	/**
+	* Get color space chromaticities, source and destination names from the color settings or working color space.
+	*/
+	static FColorSpaceMetadata GetColorSpaceMetadata(UMoviePipelineColorSetting* InColorSettings);
 };
\ No newline at end of file
