MoviePipeline: Updated the Console Variables setting to display the cvars contained within each cvar preset.

* Added a new customization to display the cvars in a preset.
* CVar enable state (within a preset) exposed via UMovieSceneConsoleVariableTrackInterface::IsConsoleVariableEnabled().
* Added UMoviePipelineConsoleVariableSetting::AddConsoleVariable() so cvar overrides can be created from a preset.

diff --git a/Engine/Plugins/Editor/ConsoleVariablesEditor/Source/ConsoleVariablesEditorRuntime/Private/ConsoleVariablesAsset.cpp b/Engine/Plugins/Editor/ConsoleVariablesEditor/Source/ConsoleVariablesEditorRuntime/Private/ConsoleVariablesAsset.cpp
--- a/Engine/Plugins/Editor/ConsoleVariablesEditor/Source/ConsoleVariablesEditorRuntime/Private/ConsoleVariablesAsset.cpp
+++ b/Engine/Plugins/Editor/ConsoleVariablesEditor/Source/ConsoleVariablesEditorRuntime/Private/ConsoleVariablesAsset.cpp
@@ -50,6 +50,19 @@ void UConsoleVariablesAsset::GetConsoleVariablesForTrack(bool bOnlyIncludeChecke
 	}
 }
 
+bool UConsoleVariablesAsset::IsConsoleVariableEnabled(const FString& ConsoleVariableName)
+{
+	for (const FConsoleVariablesEditorAssetSaveData& Command : SavedCommands)
+	{
+		if (Command.CommandName == ConsoleVariableName)
+		{
+			return Command.CheckedState == ECheckBoxState::Checked;
+		}
+	}
+
+	return false;
+}
+
 void UConsoleVariablesAsset::ExecuteSavedCommands(UObject* WorldContextObject, bool bOnlyIncludeChecked) const
 {
 	if (const TObjectPtr<UWorld> World = GEngine->GetWorldFromContextObject(WorldContextObject, EGetWorldErrorMode::LogAndReturnNull))
diff --git a/Engine/Plugins/Editor/ConsoleVariablesEditor/Source/ConsoleVariablesEditorRuntime/Public/ConsoleVariablesAsset.h b/Engine/Plugins/Editor/ConsoleVariablesEditor/Source/ConsoleVariablesEditorRuntime/Public/ConsoleVariablesAsset.h
--- a/Engine/Plugins/Editor/ConsoleVariablesEditor/Source/ConsoleVariablesEditorRuntime/Public/ConsoleVariablesAsset.h
+++ b/Engine/Plugins/Editor/ConsoleVariablesEditor/Source/ConsoleVariablesEditorRuntime/Public/ConsoleVariablesAsset.h
@@ -104,6 +104,7 @@ private:
 
 	//~ IMovieSceneConsoleVariableTrackInterface overrides
 	virtual void GetConsoleVariablesForTrack(bool bOnlyIncludeChecked, TArray<TPair<FString, FString>>& OutVariables) override;
+	virtual bool IsConsoleVariableEnabled(const FString& ConsoleVariableName) override;
 
 private:
 
diff --git a/Engine/Plugins/MovieScene/MovieRenderPipeline/Source/MovieRenderPipelineEditor/Private/Customizations/ConsoleVariableSettingCustomization.h b/Engine/Plugins/MovieScene/MovieRenderPipeline/Source/MovieRenderPipelineEditor/Private/Customizations/ConsoleVariableSettingCustomization.h
new file mode 100644
--- /dev/null
+++ b/Engine/Plugins/MovieScene/MovieRenderPipeline/Source/MovieRenderPipelineEditor/Private/Customizations/ConsoleVariableSettingCustomization.h
@@ -0,0 +1,161 @@
+// Copyright Epic Games, Inc. All Rights Reserved.
+
+#pragma once
+
+#include "DetailCategoryBuilder.h"
+#include "DetailLayoutBuilder.h"
+#include "IDetailChildrenBuilder.h"
+#include "IDetailCustomization.h"
+#include "IDetailGroup.h"
+#include "MoviePipelineConsoleVariableSetting.h"
+#include "PropertyCustomizationHelpers.h"
+#include "PropertyEditor/Private/PropertyNode.h"
+#include "PropertyHandle.h"
+#include "Sections/MovieSceneConsoleVariableTrackInterface.h"
+
+#define LOCTEXT_NAMESPACE "SMoviePipelineEditor"
+
+/** Customize how properties in UMoviePipelineConsoleVariableSetting appear in the details panel. */
+class FConsoleVariablesSettingDetailsCustomization : public IDetailCustomization
+{
+public:
+	static TSharedRef<IDetailCustomization> MakeInstance()
+	{
+		return MakeShared<FConsoleVariablesSettingDetailsCustomization>();
+	}
+
+protected:
+	//~ Begin IDetailCustomization interface
+	virtual void CustomizeDetails(IDetailLayoutBuilder& DetailBuilder) override
+	{
+		IDetailCategoryBuilder& SettingsCategory = DetailBuilder.EditCategory(TEXT("Settings"));
+		const TSharedRef<IPropertyHandle> ConsoleVariablePresetsHandle = DetailBuilder.GetProperty(GET_MEMBER_NAME_CHECKED(UMoviePipelineConsoleVariableSetting, ConsoleVariablePresets));
+
+		// Customize how the console variable presets array looks. Each preset is shown as a group, and that group's children
+		// are the console variables contained within the preset.
+		const TSharedRef<FDetailArrayBuilder> PropertyBuilder = MakeShareable(new FDetailArrayBuilder(ConsoleVariablePresetsHandle));
+		PropertyBuilder->OnGenerateArrayElementWidget(FOnGenerateArrayElementWidget::CreateSP(this, &FConsoleVariablesSettingDetailsCustomization::GeneratePresetGroup, &DetailBuilder));
+		SettingsCategory.AddCustomBuilder(PropertyBuilder, false);
+
+		// Regenerate the preset details when the preset array is updated
+		ConsoleVariablePresetsHandle->SetOnChildPropertyValueChanged(FSimpleDelegate::CreateLambda([&DetailBuilder]()
+		{
+			DetailBuilder.ForceRefreshDetails();
+		}));
+	}
+	//~ End IDetailCustomization interface
+
+	UMoviePipelineConsoleVariableSetting* GetCVarSettingFromHandle(const TSharedPtr<IPropertyHandle>& PropertyHandle) const
+	{
+		if (!PropertyHandle.IsValid() || !PropertyHandle->IsValidHandle())
+		{
+			return nullptr;
+		}
+		
+		TArray<UObject*> OuterObjects;
+		PropertyHandle->GetOuterObjects(OuterObjects);
+
+		if (OuterObjects.Num() != 1)
+		{
+			return nullptr;
+		}
+
+		// Get the outermost cvar setting object
+		return Cast<UMoviePipelineConsoleVariableSetting>(OuterObjects[0]);
+	}
+
+	void GeneratePresetGroup(TSharedRef<IPropertyHandle> ElementProperty, int32 ElementIndex, IDetailChildrenBuilder& ChildrenBuilder, IDetailLayoutBuilder* DetailLayout)
+	{
+		static const FText ConsoleVariableOverrideText = LOCTEXT("CreateConsoleVariableOverride", "Override Console Variable");
+		static const FText ConsoleVariableDisabledText = LOCTEXT("DisabledConsoleVariable", "This console variable is disabled.");
+
+		// Add the preset asset chooser as the group header
+		const FString& PropName = ElementProperty->GetPropertyNode()->GetPropertyPath();
+		IDetailGroup& Group = ChildrenBuilder.AddGroup(FName(PropName), FText::FromString(PropName));
+		Group.HeaderRow()
+		[
+			ElementProperty->CreatePropertyValueWidget()
+		];
+
+		UMoviePipelineConsoleVariableSetting* CVarSetting = GetCVarSettingFromHandle(ElementProperty);
+		if (!CVarSetting)
+		{
+			return;
+		}
+
+		const TArray<TScriptInterface<IMovieSceneConsoleVariableTrackInterface>>& CVarPresets = CVarSetting->ConsoleVariablePresets;
+		if (!CVarPresets.IsValidIndex(ElementIndex))
+		{
+			return;
+		}
+
+		const TScriptInterface<IMovieSceneConsoleVariableTrackInterface>& CVarPreset = CVarPresets[ElementIndex];
+		if (!CVarPreset)
+		{
+			return;
+		}
+
+		const bool bOnlyIncludeChecked = false;
+		TArray<TTuple<FString, FString>> CVars;
+		CVarPreset->GetConsoleVariablesForTrack(bOnlyIncludeChecked, CVars);
+
+		// Show every console variable that's included in this preset; each console variable gets its own row
+		for (const TTuple<FString, FString>& CVar : CVars)
+		{
+			// Each console variable gets a menu that allows the user to create an override outside of the preset
+			FMenuBuilder CreateOverrideMenu(true, nullptr, nullptr, true);
+			FUIAction AddOverrideAction(
+				FExecuteAction::CreateLambda([CVar, CVarSetting]()
+				{
+					CVarSetting->AddConsoleVariable(CVar.Key, 0.f);
+				})
+			);
+			CreateOverrideMenu.AddMenuEntry(ConsoleVariableOverrideText, FText::GetEmpty(), FSlateIcon(), AddOverrideAction);
+			
+			Group.AddWidgetRow()
+			.WholeRowContent()
+			[
+				SNew(SHorizontalBox)
+				
+				+ SHorizontalBox::Slot()
+				.Padding(5, 0)
+				.FillWidth(0.75f)
+				.VAlign(VAlign_Center)
+				[
+					SNew(STextBlock)
+					.IsEnabled(CVarPreset->IsConsoleVariableEnabled(CVar.Key))
+					.Font(DetailLayout->GetDetailFont())
+					.Text(FText::FromString(CVar.Key))
+					.ToolTipText(!CVarPreset->IsConsoleVariableEnabled(CVar.Key) ? ConsoleVariableDisabledText : FText::GetEmpty())
+				]
+				
+				+ SHorizontalBox::Slot()
+				.Padding(0, 0)
+				.FillWidth(0.25f)
+				.VAlign(VAlign_Center)
+				[
+					SNew(SEditableTextBox)
+					.IsEnabled(false)
+					.Font(DetailLayout->GetDetailFont())
+					.Text(FText::FromString(CVar.Value))
+				]
+				
+				+ SHorizontalBox::Slot()
+				.Padding(5, 0)
+				.AutoWidth()
+				[
+					SNew(SComboButton)
+					.ComboButtonStyle(FAppStyle::Get(), "SimpleComboButton")
+					.ForegroundColor(FSlateColor::UseForeground())
+					.HasDownArrow(true)
+					.MenuContent()
+					[
+						CreateOverrideMenu.MakeWidget()
+					]
+				]
+			];
+		}
+	}
+};
+
+#undef LOCTEXT_NAMESPACE
\ No newline at end of file
diff --git a/Engine/Plugins/MovieScene/MovieRenderPipeline/Source/MovieRenderPipelineEditor/Private/Widgets/SMoviePipelineConfigEditor.cpp b/Engine/Plugins/MovieScene/MovieRenderPipeline/Source/MovieRenderPipelineEditor/Private/Widgets/SMoviePipelineConfigEditor.cpp
--- a/Engine/Plugins/MovieScene/MovieRenderPipeline/Source/MovieRenderPipelineEditor/Private/Widgets/SMoviePipelineConfigEditor.cpp
+++ b/Engine/Plugins/MovieScene/MovieRenderPipeline/Source/MovieRenderPipelineEditor/Private/Widgets/SMoviePipelineConfigEditor.cpp
@@ -56,6 +56,7 @@
 #include "IPropertyUtilities.h"
 
 #include "Customizations/ConsoleVariableCustomization.h"
+#include "Customizations/ConsoleVariableSettingCustomization.h"
 #include "MoviePipelineConsoleVariableSetting.h"
 #include "MoviePipelineOutputSetting.h"
 
@@ -737,7 +738,7 @@ void SMoviePipelineConfigEditor::UpdateDetails()
 			Details->RegisterInstancedCustomPropertyTypeLayout(FMoviePipelineConsoleVariableEntry::StaticStruct()->GetFName(), FOnGetPropertyTypeCustomizationInstance::CreateStatic(FConsoleVariablesDetailsCustomization::MakeInstance));
 			Details->RegisterInstancedCustomPropertyTypeLayout("FrameNumber", FOnGetPropertyTypeCustomizationInstance::CreateStatic(&FFrameNumberDetailsCustomization::MakeInstance, NumericTypeInterface));
 			Details->RegisterInstancedCustomPropertyLayout(UMoviePipelineOutputSetting::StaticClass(), FOnGetDetailCustomizationInstance::CreateStatic(FOutputFormatDetailsCustomization::MakeInstance, SharedThis(this)));
-
+			Details->RegisterInstancedCustomPropertyLayout(UMoviePipelineConsoleVariableSetting::StaticClass(), FOnGetDetailCustomizationInstance::CreateStatic(FConsoleVariablesSettingDetailsCustomization::MakeInstance));
 
 			Details->SetObjects(Pair.Value);
 			DetailsBox->AddSlot()
diff --git a/Engine/Plugins/MovieScene/MovieRenderPipeline/Source/MovieRenderPipelineSettings/Private/MoviePipelineConsoleVariableSetting.cpp b/Engine/Plugins/MovieScene/MovieRenderPipeline/Source/MovieRenderPipelineSettings/Private/MoviePipelineConsoleVariableSetting.cpp
--- a/Engine/Plugins/MovieScene/MovieRenderPipeline/Source/MovieRenderPipelineSettings/Private/MoviePipelineConsoleVariableSetting.cpp
+++ b/Engine/Plugins/MovieScene/MovieRenderPipeline/Source/MovieRenderPipelineSettings/Private/MoviePipelineConsoleVariableSetting.cpp
@@ -321,6 +321,13 @@ bool UMoviePipelineConsoleVariableSetting::AddOrUpdateConsoleVariable(const FStr
 	return true;
 }
 
+bool UMoviePipelineConsoleVariableSetting::AddConsoleVariable(const FString& Name, const float Value)
+{
+	CVars.Add(FMoviePipelineConsoleVariableEntry(Name, Value));
+
+	return true;
+}
+
 bool UMoviePipelineConsoleVariableSetting::UpdateConsoleVariableEnableState(const FString& Name, const bool bIsEnabled)
 {
 	const int32 LastMatch = CVars.FindLastByPredicate([&Name](const FMoviePipelineConsoleVariableEntry& Entry)
diff --git a/Engine/Plugins/MovieScene/MovieRenderPipeline/Source/MovieRenderPipelineSettings/Public/MoviePipelineConsoleVariableSetting.h b/Engine/Plugins/MovieScene/MovieRenderPipeline/Source/MovieRenderPipelineSettings/Public/MoviePipelineConsoleVariableSetting.h
--- a/Engine/Plugins/MovieScene/MovieRenderPipeline/Source/MovieRenderPipelineSettings/Public/MoviePipelineConsoleVariableSetting.h
+++ b/Engine/Plugins/MovieScene/MovieRenderPipeline/Source/MovieRenderPipelineSettings/Public/MoviePipelineConsoleVariableSetting.h
@@ -151,15 +151,25 @@ public:
 
 	/**
 	 * Adds a console variable override with the given name and value if one does not already exist. If the console
-	 * variable with the given name already exists, its value will be updated. Returns true if the operation was
-	 * successful, else false.
+	 * variable with the given name already exists, its value will be updated (the last one will be updated if there are
+	 * duplicates with the same name). Returns true if the operation was successful, else false.
+	 * @see AddConsoleVariable()
 	 */
 	UFUNCTION(BlueprintCallable, Category = "Settings")
 	bool AddOrUpdateConsoleVariable(const FString& Name, const float Value);
 
 	/**
-	 * Updates the enable state of the console variable override with the provided name. Returns true if the operation
-	 * was successful, else false.
+	 * Adds a console variable override with the given name and value, and will add a duplicate if one with the provided
+	 * name already exists. Returns true if the operation was successful, else false.
+	 * @see AddOrUpdateConsoleVariable()
+	 */
+	UFUNCTION(BlueprintCallable, Category = "Settings")
+	bool AddConsoleVariable(const FString& Name, const float Value);
+
+	/**
+	 * Updates the enable state of the console variable override with the provided name. If there are duplicate cvars
+	 * with the same name, the last one with the provided name will be updated. Returns true if the operation was
+	 * successful, else false.
 	 */
 	UFUNCTION(BlueprintCallable, Category = "Settings")
 	bool UpdateConsoleVariableEnableState(const FString& Name, const bool bIsEnabled);
diff --git a/Engine/Source/Runtime/MovieSceneTracks/Public/Sections/MovieSceneConsoleVariableTrackInterface.h b/Engine/Source/Runtime/MovieSceneTracks/Public/Sections/MovieSceneConsoleVariableTrackInterface.h
--- a/Engine/Source/Runtime/MovieSceneTracks/Public/Sections/MovieSceneConsoleVariableTrackInterface.h
+++ b/Engine/Source/Runtime/MovieSceneTracks/Public/Sections/MovieSceneConsoleVariableTrackInterface.h
@@ -24,4 +24,5 @@ public:
 	GENERATED_BODY()
 
 	virtual void GetConsoleVariablesForTrack(bool bOnlyIncludeChecked, TArray<TTuple<FString, FString>>& OutVariables) = 0;
+	virtual bool IsConsoleVariableEnabled(const FString& ConsoleVariableName) = 0;
 };
