MoviePipeline: Fixed an issue where UWorldSettings::MinUndilatedFrameTime/UWorldSettings::MaxUndilatedFrameTime  were being modified and not restored after a render.

diff --git a/Engine/Plugins/MovieScene/MovieRenderPipeline/Source/MovieRenderPipelineCore/Private/MoviePipeline.cpp b/Engine/Plugins/MovieScene/MovieRenderPipeline/Source/MovieRenderPipelineCore/Private/MoviePipeline.cpp
--- a/Engine/Plugins/MovieScene/MovieRenderPipeline/Source/MovieRenderPipelineCore/Private/MoviePipeline.cpp
+++ b/Engine/Plugins/MovieScene/MovieRenderPipeline/Source/MovieRenderPipelineCore/Private/MoviePipeline.cpp
@@ -204,6 +204,7 @@ void UMoviePipeline::Initialize(UMoviePipelineExecutorJob* InJob)
 
 	CachedSequenceHierarchyRoot = MakeShared<MoviePipeline::FCameraCutSubSectionHierarchyNode>();
 	MoviePipeline::CacheCompleteSequenceHierarchy(TargetSequence, CachedSequenceHierarchyRoot);
+	CustomTimeStep->CacheWorldSettings();
 
 	// Override the frame range on the target sequence if needed first before anyone has a chance to modify it.
 	{
@@ -572,6 +573,7 @@ void UMoviePipeline::TransitionToState(const EMovieRenderPipelineState InNewStat
 			// and cause it to re-allocate at the currrent size on the next render request, which is likely to be the size
 			// of the PIE window (720p) or the Viewport itself.
 			UKismetSystemLibrary::ExecuteConsoleCommand(GetWorld(), TEXT("r.ResetRenderTargetsExtent"), nullptr);
+			CustomTimeStep->RestoreCachedWorldSettings();
 
 			GAreScreenMessagesEnabled = bPrevGScreenMessagesEnabled;
 
@@ -830,6 +832,24 @@ bool UMoviePipelineCustomTimeStep::UpdateTimeStep(UEngine* /*InEngine*/)
 	return false;
 }
 
+void UMoviePipelineCustomTimeStep::CacheWorldSettings()
+{
+	if (AWorldSettings* WorldSettings = GetWorld()->GetWorldSettings())
+	{
+		PrevMinUndilatedFrameTime = WorldSettings->MinUndilatedFrameTime;
+		PrevMaxUndilatedFrameTime = WorldSettings->MaxUndilatedFrameTime;
+	}
+}
+
+void UMoviePipelineCustomTimeStep::RestoreCachedWorldSettings()
+{
+	if (AWorldSettings* WorldSettings = GetWorld()->GetWorldSettings())
+	{
+		WorldSettings->MinUndilatedFrameTime = PrevMinUndilatedFrameTime;
+		WorldSettings->MaxUndilatedFrameTime = PrevMaxUndilatedFrameTime;
+	}
+}
+
 void UMoviePipelineCustomTimeStep::SetCachedFrameTiming(const MoviePipeline::FFrameTimeStepCache& InTimeCache)
 { 
 	if (ensureMsgf(!FMath::IsNearlyZero(InTimeCache.UndilatedDeltaTime), TEXT("An incorrect or uninitialized time step was used! Delta Time of 0 isn't allowed.")))
diff --git a/Engine/Plugins/MovieScene/MovieRenderPipeline/Source/MovieRenderPipelineCore/Public/MoviePipeline.h b/Engine/Plugins/MovieScene/MovieRenderPipeline/Source/MovieRenderPipelineCore/Public/MoviePipeline.h
--- a/Engine/Plugins/MovieScene/MovieRenderPipeline/Source/MovieRenderPipelineCore/Public/MoviePipeline.h
+++ b/Engine/Plugins/MovieScene/MovieRenderPipeline/Source/MovieRenderPipelineCore/Public/MoviePipeline.h
@@ -517,6 +517,7 @@ public:
 	static FString DefaultDebugWidgetAsset;
 };
 
+
 UCLASS()
 class UMoviePipelineCustomTimeStep : public UEngineCustomTimeStep
 {
@@ -532,9 +533,16 @@ public:
 
 	void SetCachedFrameTiming(const MoviePipeline::FFrameTimeStepCache& InTimeCache);
 
+	// Cache and Restore some of the World Settings settings, as rendering with MRQ in a runtime world needs to restore any changes made to World Settings.
+	void CacheWorldSettings();
+	void RestoreCachedWorldSettings();
 private:
 	/** We don't do any thinking on our own, instead we just spit out the numbers stored in our time cache. */
 	MoviePipeline::FFrameTimeStepCache TimeCache;
+
+	// Not cached in TimeCache as TimeCache is reset every frame.
+	float PrevMinUndilatedFrameTime;
+	float PrevMaxUndilatedFrameTime;
 };
 
 struct FMoviePipelineTimeController : public FMovieSceneTimeController
