MoviePipeline: Added a workaround for the UI Renderer setting having issues with some Layout Widgets not liking the 720p PIE preview window (vs say a 4k render).

To use, go to the Project Settings, then Movie Pipeline In Editor, and enable bResizePIEWindowToOutputResolution. Your PIE window is still limited by your monitor resolution so you will need to launch the editor with "-ForceRes" if you wish to render at a resolution higher than your monitor resolution.

diff --git a/Engine/Plugins/MovieScene/MovieRenderPipeline/Source/MovieRenderPipelineEditor/Private/MoviePipelinePIEExecutor.cpp b/Engine/Plugins/MovieScene/MovieRenderPipeline/Source/MovieRenderPipelineEditor/Private/MoviePipelinePIEExecutor.cpp
--- a/Engine/Plugins/MovieScene/MovieRenderPipeline/Source/MovieRenderPipelineEditor/Private/MoviePipelinePIEExecutor.cpp
+++ b/Engine/Plugins/MovieScene/MovieRenderPipeline/Source/MovieRenderPipelineEditor/Private/MoviePipelinePIEExecutor.cpp
@@ -4,6 +4,7 @@
 #include "MoviePipelineShotConfig.h"
 #include "MoviePipeline.h"
 #include "MoviePipelineGameOverrideSetting.h"
+#include "MoviePipelineOutputSetting.h"
 #include "MovieRenderPipelineCoreModule.h"
 #include "Framework/Application/SlateApplication.h"
 #include "MovieRenderPipelineDataTypes.h"
@@ -113,9 +114,20 @@ void UMoviePipelinePIEExecutor::Start(const UMoviePipelineExecutorJob* InJob)
 	// Start capturing logging messages
 	ValidationMessageGatherer.StartGathering();
 
+	FVector2D WindowSize = FVector2D(1280, 720);
+
+	const UMoviePipelinePIEExecutorSettings* ExecutorSettings = GetDefault<UMoviePipelinePIEExecutorSettings>();
+	if (ExecutorSettings->bResizePIEWindowToOutputResolution)
+	{
+		// Allow creating the PIE window at the size of the output to work around UMG scaling bugs when the main viewport is one resolution and then we re-render it at another.
+		UMoviePipelineOutputSetting* OutputSetting = Cast<UMoviePipelineOutputSetting>(InJob->GetConfiguration()->FindOrAddSettingByClass(UMoviePipelineOutputSetting::StaticClass()));
+		WindowSize = FVector2D(OutputSetting->OutputResolution.X, OutputSetting->OutputResolution.Y);
+	}
+
+
 	// Create a Slate window to hold our preview.
 	TSharedRef<SWindow> CustomWindow = SNew(SWindow)
-		.ClientSize(FVector2D(1280, 720))
+		.ClientSize(WindowSize)
 		.AutoCenter(EAutoCenter::PrimaryWorkArea)
 		.UseOSWindowBorder(true)
 		.FocusWhenFirstShown(false)
diff --git a/Engine/Plugins/MovieScene/MovieRenderPipeline/Source/MovieRenderPipelineEditor/Private/MoviePipelinePIEExecutorSettings.h b/Engine/Plugins/MovieScene/MovieRenderPipeline/Source/MovieRenderPipelineEditor/Private/MoviePipelinePIEExecutorSettings.h
--- a/Engine/Plugins/MovieScene/MovieRenderPipeline/Source/MovieRenderPipelineEditor/Private/MoviePipelinePIEExecutorSettings.h
+++ b/Engine/Plugins/MovieScene/MovieRenderPipeline/Source/MovieRenderPipelineEditor/Private/MoviePipelinePIEExecutorSettings.h
@@ -26,5 +26,17 @@ public:
 	* any work (such as evaluating frames for warm-up).
 	*/
 	UPROPERTY(EditAnywhere, BlueprintReadWrite, config, meta = (UIMin = "0", ClampMin = "0", UIMax = "150"), Category = "Startup")
-	int32 InitialDelayFrameCount;
+	int32 InitialDelayFrameCount = 0;
+
+	/**
+	* Should the PIE Window be created at the same resolution as the MRQ Output? By default we create the window at 720p for a nicer
+	* user experience, but this can be used to work around a widget scaling issue with UMG Widgets when using the UI Renderer
+	* setting. PIE is still limited by your monitor's resolution so you will need a monitor at least as big as your requested output
+	* for this to work (or can be combined with launching the editor with -ForceRes). 
+	*
+	* Warning: Don't use this setting in combination with HighResTiling, as the main backbuffer will have to get resized to your final
+	* output resolution which will be too large.
+	*/
+	UPROPERTY(EditAnywhere, BlueprintReadWrite, config, Category = "Startup")
+	bool bResizePIEWindowToOutputResolution = false;
 };
\ No newline at end of file
