MoviePipeline: Prototype for the render layer subsystem. Experimental!

Also provides prototypes for queries, modifiers, and collections.

diff --git a/Engine/Plugins/MovieScene/MovieRenderPipeline/Source/MovieRenderPipelineCore/Private/Graph/MoviePipelineRenderLayerSubsystem.cpp b/Engine/Plugins/MovieScene/MovieRenderPipeline/Source/MovieRenderPipelineCore/Private/Graph/MoviePipelineRenderLayerSubsystem.cpp
new file mode 100644
--- /dev/null
+++ b/Engine/Plugins/MovieScene/MovieRenderPipeline/Source/MovieRenderPipelineCore/Private/Graph/MoviePipelineRenderLayerSubsystem.cpp
@@ -0,0 +1,487 @@
+// Copyright Epic Games, Inc. All Rights Reserved.
+
+#include "Graph/MoviePipelineRenderLayerSubsystem.h"
+
+#include "EngineUtils.h"
+#include "MovieRenderPipelineCoreModule.h"
+#include "Subsystems/WorldSubsystem.h"
+
+void UMoviePipelineMaterialModifier::ApplyModifier(const UWorld* World)
+{
+	UMaterialInterface* NewMaterial = MaterialToApply.LoadSynchronous();
+	if (!NewMaterial)
+	{
+		return;
+	}
+
+	ModifiedComponents.Empty();
+	
+	for (const UMoviePipelineCollection* Collection : Collections)
+	{
+		if (!Collection)
+		{
+			continue;
+		}
+
+		for (const AActor* Actor : Collection->GetMatchingActors(World, bUseInvertedActors))
+		{
+			const bool bIncludeFromChildActors = true;
+			TArray<UPrimitiveComponent*> PrimitiveComponents;
+			Actor->GetComponents<UPrimitiveComponent>(PrimitiveComponents, bIncludeFromChildActors);
+
+			for (UPrimitiveComponent* PrimitiveComponent : PrimitiveComponents)
+			{
+				TArray<FMaterialSlotAssignment>& ModifiedMaterials = ModifiedComponents.FindOrAdd(PrimitiveComponent);
+				
+				for (int32 Index = 0; Index < PrimitiveComponent->GetNumMaterials(); ++Index)
+				{
+					ModifiedMaterials.Add(FMaterialSlotAssignment(Index, PrimitiveComponent->GetMaterial(Index)));
+				
+					PrimitiveComponent->SetMaterial(Index, NewMaterial);
+				}
+			}
+		}
+	}
+}
+
+void UMoviePipelineMaterialModifier::UndoModifier()
+{
+	for (const FComponentToMaterialMap::ElementType& ModifiedComponent : ModifiedComponents) 
+	{
+		UPrimitiveComponent* MeshComponent = ModifiedComponent.Key.LoadSynchronous();
+		const TArray<FMaterialSlotAssignment>& OldMaterials = ModifiedComponent.Value;
+
+		if (!MeshComponent)
+		{
+			continue;
+		}
+
+		for (const FMaterialSlotAssignment& MaterialPair : OldMaterials)
+		{
+			UMaterialInterface* MaterialInterface = MaterialPair.Value.LoadSynchronous();
+			if (!MaterialInterface)
+			{
+				continue;
+			}
+
+			const int32 ElementIndex = MaterialPair.Key;
+			MeshComponent->SetMaterial(ElementIndex, MaterialInterface);
+		}
+	}
+
+	ModifiedComponents.Empty();
+}
+
+void UMoviePipelineVisibilityModifier::ApplyModifier(const UWorld* World)
+{
+	ModifiedActors.Empty();
+	
+	for (const UMoviePipelineCollection* Collection : Collections)
+	{
+		if (!Collection)
+		{
+			continue;
+		}
+
+		for (AActor* Actor : Collection->GetMatchingActors(World, bUseInvertedActors))
+		{
+			ModifiedActors.Add(Actor, Actor->IsHidden());
+			SetActorHiddenState(Actor, bIsHidden);
+		}
+	}
+}
+
+void UMoviePipelineVisibilityModifier::UndoModifier()
+{
+	for (const TTuple<TSoftObjectPtr<AActor>, bool>& Pair : ModifiedActors)
+	{
+		SetActorHiddenState(Pair.Key.LoadSynchronous(), Pair.Value);
+	}
+
+	ModifiedActors.Empty();
+}
+
+void UMoviePipelineVisibilityModifier::SetActorHiddenState(AActor* Actor, const bool bInIsHidden) const
+{
+	Actor->SetActorHiddenInGame(bInIsHidden);
+
+#if WITH_EDITOR
+	Actor->SetIsTemporarilyHiddenInEditor(bInIsHidden);
+#endif
+}
+
+// TODO: This really should be "DoesComponentMatchQuery()"
+bool UMoviePipelineCollectionCommonQuery::DoesActorMatchQuery(const AActor* Actor) const
+{
+	if (!Actor)
+	{
+		return false;
+	}
+
+	// TODO: This method should short-circuit so all query types aren't executed for every actor when the query mode is OR
+
+	bool bMatchesActorNames = false;
+	if (!ActorNames.IsEmpty())
+	{
+		bMatchesActorNames = ActorNames.Contains(Actor->GetName());
+	}
+	
+	bool bMatchesComponentTypes = false;
+	if (!ComponentTypes.IsEmpty())
+	{
+		const bool bIncludeFromChildActors = true;
+		TArray<UActorComponent*> ActorComponents;
+		Actor->GetComponents(ActorComponents, bIncludeFromChildActors);
+		for (const UActorComponent* Component : ActorComponents)
+		{
+			for (const UClass* ComponentType : ComponentTypes)
+			{
+				if (Component->IsA(ComponentType))
+				{
+					bMatchesComponentTypes = true;
+					break;
+				}
+			}
+		}
+	}
+
+	bool bMatchesTags = false;
+	if (!Tags.IsEmpty())
+	{
+		for (const FName& Tag : Tags)
+		{
+			if (Actor->Tags.Contains(Tag))
+			{
+				bMatchesTags = true;
+				break;
+			}
+		}
+	}
+
+	const bool bUsingActorNames = !ActorNames.IsEmpty();
+	const bool bUsingComponentTypes = !ComponentTypes.IsEmpty();
+	const bool bUsingTags = !Tags.IsEmpty();
+
+	if (QueryMode == EMoviePipelineCollectionCommonQueryMode::And)
+	{
+		return (!bUsingActorNames || bMatchesActorNames) &&
+			   (!bUsingTags || bMatchesTags) &&
+			   (!bUsingComponentTypes || bMatchesComponentTypes);
+	}
+	
+	return bMatchesActorNames || bMatchesTags || bMatchesComponentTypes;
+}
+
+TArray<AActor*> UMoviePipelineCollection::GetMatchingActors(const UWorld* World, const bool bInvertResult) const
+{
+	TArray<AActor*> MatchingActors;
+
+	for (TActorIterator<AActor> ActorItr(World); ActorItr; ++ActorItr)
+	{
+		AActor* Actor = *ActorItr;
+		if (!Actor)
+		{
+			continue;
+		}
+
+		// If there aren't any queries, and the result should be inverted, just include the actor
+		if (bInvertResult && Queries.IsEmpty())
+		{
+			MatchingActors.Add(Actor);
+			continue;
+		}
+
+		for (const UMoviePipelineCollectionQuery* Query : Queries)
+		{
+			const bool bActorMatchesQuery = Query->DoesActorMatchQuery(Actor);
+			
+			if (bActorMatchesQuery && !bInvertResult)
+			{
+				MatchingActors.Add(Actor);
+				break;
+			}
+
+			if (!bActorMatchesQuery && bInvertResult)
+			{
+				MatchingActors.Add(Actor);
+				break;
+			}
+		}
+	}
+	
+	return MatchingActors;
+}
+
+void UMoviePipelineCollectionModifier::AddCollection(UMoviePipelineCollection* Collection)
+{
+	// Don't allow adding a duplicate collection
+	for (const UMoviePipelineCollection* ExistingCollection : Collections)
+	{
+		if (Collection && ExistingCollection && Collection->GetCollectionName().Equals(ExistingCollection->GetCollectionName()))
+		{
+			return;
+		}
+	}
+	
+	Collections.Add(Collection);
+}
+
+void UMoviePipelineCollection::AddQuery(UMoviePipelineCollectionQuery* Query)
+{
+	if (!Queries.Contains(Query))
+	{
+		Queries.Add(Query);
+	}
+}
+
+UMoviePipelineCollection* UMoviePipelineRenderLayer::GetCollectionByName(const FString& Name) const
+{
+	for (const UMoviePipelineCollectionModifier* Modifier : Modifiers)
+	{
+		if (!Modifier)
+		{
+			continue;
+		}
+		
+		for (UMoviePipelineCollection* Collection : Modifier->GetCollections())
+		{
+			if (Collection && Collection->GetCollectionName().Equals(Name))
+			{
+				return Collection;
+			}
+		}
+	}
+
+	return nullptr;
+}
+
+void UMoviePipelineRenderLayer::AddModifier(UMoviePipelineCollectionModifier* Modifier)
+{
+	if (!Modifiers.Contains(Modifier))
+	{
+		Modifiers.Add(Modifier);
+	}
+}
+
+void UMoviePipelineRenderLayer::RemoveModifier(UMoviePipelineCollectionModifier* Modifier)
+{
+	Modifiers.Remove(Modifier);
+}
+
+void UMoviePipelineRenderLayer::Preview(const UWorld* World)
+{
+	if (!World)
+	{
+		return;
+	}
+	
+	// Apply all modifiers
+	for (UMoviePipelineCollectionModifier* Modifier : Modifiers)
+	{
+		Modifier->ApplyModifier(World);
+	}
+}
+
+void UMoviePipelineRenderLayer::UndoPreview(const UWorld* World)
+{
+	if (!World)
+	{
+		return;
+	}
+
+	// Undo actions performed by all modifiers. Do this in the reverse order that they were applied, since the undo
+	// state of one modifier may depend on modifiers that were previously applied.
+	for (int32 Index = Modifiers.Num() - 1; Index >= 0; Index--)
+	{
+		if (UMoviePipelineCollectionModifier* Modifier = Modifiers[Index])
+		{
+			Modifier->UndoModifier();
+		}
+	}
+}
+
+
+UMoviePipelineRenderLayerSubsystem* UMoviePipelineRenderLayerSubsystem::GetFromWorld(const UWorld* World)
+{
+	if (World)
+	{
+		return UWorld::GetSubsystem<UMoviePipelineRenderLayerSubsystem>(World);
+	}
+
+	return nullptr;
+}
+
+void UMoviePipelineRenderLayerSubsystem::Initialize(FSubsystemCollectionBase& Collection)
+{
+	VisualizationEmptyCollection = NewObject<UMoviePipelineCollection>(GetTransientPackage(), NAME_None, RF_Transient);
+
+	// By default, the visualizer should hide everything in the world
+	VisualizationModifier_HideWorld = NewObject<UMoviePipelineVisibilityModifier>(GetTransientPackage(), NAME_None, RF_Transient);
+	VisualizationModifier_HideWorld->AddCollection(VisualizationEmptyCollection);
+	VisualizationModifier_HideWorld->SetHidden(true);
+	VisualizationModifier_HideWorld->SetIsInverted(true);
+
+	// Selectively show collections in the visualization
+	VisualizationModifier_VisibleCollections = NewObject<UMoviePipelineVisibilityModifier>(GetTransientPackage(), NAME_None, RF_Transient);
+
+	// The visualizer render layer will hide the world, selectively show specified collections, and then apply any other provided modifiers
+	VisualizationRenderLayer = NewObject<UMoviePipelineRenderLayer>(GetTransientPackage(), NAME_None, RF_Transient);
+	VisualizationRenderLayer->AddModifier(VisualizationModifier_HideWorld);
+	VisualizationRenderLayer->AddModifier(VisualizationModifier_VisibleCollections);
+}
+
+void UMoviePipelineRenderLayerSubsystem::Deinitialize()
+{
+}
+
+void UMoviePipelineRenderLayerSubsystem::Reset()
+{
+	ClearAllPreviews();
+	RenderLayers.Empty();
+}
+
+bool UMoviePipelineRenderLayerSubsystem::AddRenderLayer(UMoviePipelineRenderLayer* RenderLayer)
+{
+	if (!RenderLayer)
+	{
+		UE_LOG(LogMovieRenderPipeline, Warning, TEXT("Invalid render layer provided to AddRenderLayer()."));
+		return false;
+	}
+	
+	const bool bRenderLayerExists = RenderLayers.ContainsByPredicate([RenderLayer](const UMoviePipelineRenderLayer* RL)
+	{
+		return RL && (RenderLayer->GetRenderLayerName() == RL->GetName());
+	});
+
+	if (bRenderLayerExists)
+	{
+		UE_LOG(LogMovieRenderPipeline, Warning, TEXT("Render layer '%s' already exists in the render layer subsystem; it will not be added again."), *RenderLayer->GetRenderLayerName());
+		return false;
+	}
+
+	RenderLayers.Add(RenderLayer);
+	return true;
+}
+
+void UMoviePipelineRenderLayerSubsystem::RemoveRenderLayer(const FString& RenderLayerName)
+{
+	if (ActiveRenderLayer && (ActiveRenderLayer->GetName() == RenderLayerName))
+	{
+		ClearAllPreviews();
+	}
+	
+	const uint32 Index = RenderLayers.IndexOfByPredicate([&RenderLayerName](const UMoviePipelineRenderLayer* RenderLayer)
+	{
+		return RenderLayer->GetRenderLayerName() == RenderLayerName;
+	});
+
+	if (Index != INDEX_NONE)
+	{
+		RenderLayers.RemoveAt(Index);
+	}
+}
+
+void UMoviePipelineRenderLayerSubsystem::SetActiveRenderLayerByObj(UMoviePipelineRenderLayer* RenderLayer)
+{
+	if (!RenderLayer)
+	{
+		return;
+	}
+	
+	ClearAllPreviews();
+	SetAndPreviewRenderLayer(RenderLayer);
+}
+
+void UMoviePipelineRenderLayerSubsystem::SetActiveRenderLayerByName(const FString& RenderLayerName)
+{
+	const uint32 Index = RenderLayers.IndexOfByPredicate([&RenderLayerName](const UMoviePipelineRenderLayer* RenderLayer)
+	{
+		return RenderLayer->GetRenderLayerName() == RenderLayerName;
+	});
+
+	if (Index != INDEX_NONE)
+	{
+		SetActiveRenderLayerByObj(RenderLayers[Index]);
+	}
+}
+
+void UMoviePipelineRenderLayerSubsystem::ClearActiveRenderLayer()
+{
+	ClearAllPreviews();
+}
+
+void UMoviePipelineRenderLayerSubsystem::PreviewCollection(UMoviePipelineCollection* Collection)
+{
+	if (!Collection)
+	{
+		return;
+	}
+
+	ClearAllPreviews();
+	
+	ActiveCollection = Collection;
+	VisualizationModifier_VisibleCollections->AddCollection(ActiveCollection);
+
+	SetAndPreviewRenderLayer(VisualizationRenderLayer);
+}
+
+void UMoviePipelineRenderLayerSubsystem::ClearCollectionPreview()
+{
+	ClearAllPreviews();
+}
+
+void UMoviePipelineRenderLayerSubsystem::ClearAllPreviews()
+{
+	// Render layer previews and collection previews both use the active render layer, so undoing the preview this
+	// way will clear previews for both
+	if (ActiveRenderLayer)
+	{
+		ActiveRenderLayer->UndoPreview(GetWorld());
+
+		// Remove the modifier preview if present (requires an active render layer)
+		if (ActiveModifier)
+		{
+			ActiveRenderLayer->RemoveModifier(ActiveModifier);
+		}
+	}
+
+	// Reset the viz modifier for the collection preview
+	VisualizationModifier_VisibleCollections->SetCollections({});
+
+	ActiveRenderLayer = nullptr;
+	ActiveCollection = nullptr;
+	ActiveModifier = nullptr;
+}
+
+void UMoviePipelineRenderLayerSubsystem::SetAndPreviewRenderLayer(UMoviePipelineRenderLayer* RenderLayer)
+{
+	ActiveRenderLayer = RenderLayer;
+	ActiveRenderLayer->Preview(GetWorld());
+}
+
+void UMoviePipelineRenderLayerSubsystem::PreviewModifier(UMoviePipelineCollectionModifier* Modifier)
+{
+	if (!Modifier)
+	{
+		return;
+	}
+
+	ClearAllPreviews();
+
+	ActiveModifier = Modifier;
+	VisualizationRenderLayer->AddModifier(ActiveModifier);
+
+	// Add the modifier's collections to the viz as well, so the actors that the modifier affects are actually visible
+	// TODO: This may need special handing for visibility modifiers
+	for (UMoviePipelineCollection* Collection : ActiveModifier->GetCollections())
+	{
+		VisualizationModifier_VisibleCollections->AddCollection(Collection);
+	}
+	
+	SetAndPreviewRenderLayer(VisualizationRenderLayer);
+}
+
+void UMoviePipelineRenderLayerSubsystem::ClearModifierPreview()
+{
+	ClearAllPreviews();
+}
\ No newline at end of file
diff --git a/Engine/Plugins/MovieScene/MovieRenderPipeline/Source/MovieRenderPipelineCore/Public/Graph/MoviePipelineRenderLayerSubsystem.h b/Engine/Plugins/MovieScene/MovieRenderPipeline/Source/MovieRenderPipelineCore/Public/Graph/MoviePipelineRenderLayerSubsystem.h
new file mode 100644
--- /dev/null
+++ b/Engine/Plugins/MovieScene/MovieRenderPipeline/Source/MovieRenderPipelineCore/Public/Graph/MoviePipelineRenderLayerSubsystem.h
@@ -0,0 +1,352 @@
+﻿// Copyright Epic Games, Inc. All Rights Reserved.
+
+#pragma once
+
+#include "CoreMinimal.h"
+
+#include "MoviePipelineRenderLayerSubsystem.generated.h"
+
+/**
+ * EXPERIMENTAL
+ */
+
+/**
+ * Generates a collection of actors via provided queries.
+ */
+UCLASS(BlueprintType)
+class UMoviePipelineCollection : public UObject
+{
+	GENERATED_BODY()
+
+public:
+	UFUNCTION(BlueprintCallable, Category = "Render Layers")
+	TArray<AActor*> GetMatchingActors(const UWorld* World, const bool bInvertResult = false) const;
+
+	UFUNCTION(BlueprintCallable, Category = "Render Layers")
+	void AddQuery(UMoviePipelineCollectionQuery* Query);
+
+	UFUNCTION(BlueprintCallable, Category = "Render Layers")
+	void SetCollectionName(const FString& InName) { CollectionName = InName; }
+
+	UFUNCTION(BlueprintCallable, Category = "Render Layers")
+	FString GetCollectionName() const { return CollectionName; }
+
+private:
+	UPROPERTY()
+	FString CollectionName;
+
+	UPROPERTY()
+	TArray<TObjectPtr<UMoviePipelineCollectionQuery>> Queries;
+};
+
+/**
+ * Base class for providing actor modification functionality via collections.
+ */
+UCLASS(Abstract)
+class UMoviePipelineCollectionModifier : public UObject
+{
+	GENERATED_BODY()
+
+public:
+	/** Adds a collection to the existing set of collections in this modifier. */
+	UFUNCTION(BlueprintCallable, Category = "Render Layers")
+	void AddCollection(UMoviePipelineCollection* Collection);
+
+	/** Overwrites the existing collections with the provided array of collections. */
+	UFUNCTION(BlueprintCallable, Category = "Render Layers")
+	void SetCollections(const TArray<UMoviePipelineCollection*> InCollections) { Collections = InCollections; }
+	
+	UFUNCTION(BlueprintCallable, Category = "Render Layers")
+	TArray<UMoviePipelineCollection*> GetCollections() const { return Collections; }
+
+	UFUNCTION(BlueprintCallable, Category = "Render Layers")
+	void SetIsInverted(const bool bIsInverted) { bUseInvertedActors = bIsInverted; }
+
+	UFUNCTION(BlueprintCallable, Category = "Render Layers")
+	bool IsInverted() const { return bUseInvertedActors; }
+	
+	virtual void ApplyModifier(const UWorld* World) PURE_VIRTUAL(UMoviePipelineCollectionModifier::ApplyModifier, );
+	virtual void UndoModifier() PURE_VIRTUAL(UMoviePipelineCollectionModifier::UndoModifier, );
+
+protected:
+	/** The collections which this modifier will operate on. */
+	UPROPERTY()
+	TArray<TObjectPtr<UMoviePipelineCollection>> Collections;
+
+	/** Whether an inverted collection of actors should be used. */
+	UPROPERTY()
+	bool bUseInvertedActors = false;
+};
+
+/**
+ * Modifies actor materials.
+ */
+UCLASS(BlueprintType)
+class UMoviePipelineMaterialModifier : public UMoviePipelineCollectionModifier
+{
+	GENERATED_BODY()
+
+public:
+	UFUNCTION(BlueprintCallable, Category = "Render Layers")
+	void SetMaterial(TSoftObjectPtr<UMaterialInterface> InMaterial) { MaterialToApply = InMaterial; }
+
+	UFUNCTION(BlueprintCallable, Category = "Render Layers")
+	virtual void ApplyModifier(const UWorld* World) override;
+
+	UFUNCTION(BlueprintCallable, Category = "Render Layers")
+	virtual void UndoModifier() override;
+
+private:
+	typedef TTuple<int32, TSoftObjectPtr<UMaterialInterface>> FMaterialSlotAssignment;
+	typedef TMap<TSoftObjectPtr<UPrimitiveComponent>, TArray<FMaterialSlotAssignment>> FComponentToMaterialMap;
+	
+	/** Maps a component to its original material assignments (per index). */
+	FComponentToMaterialMap ModifiedComponents;
+
+	UPROPERTY()
+	TSoftObjectPtr<UMaterialInterface> MaterialToApply;
+};
+
+/**
+ * Modifies actor visibility.
+ */
+UCLASS(BlueprintType)
+class UMoviePipelineVisibilityModifier : public UMoviePipelineCollectionModifier
+{
+	GENERATED_BODY()
+
+public:
+	UFUNCTION(BlueprintCallable, Category = "Render Layers")
+	void SetHidden(const bool bInIsHidden) { bIsHidden = bInIsHidden; }
+
+	UFUNCTION(BlueprintCallable, Category = "Render Layers")
+	bool IsHidden() const { return bIsHidden; }
+
+	UFUNCTION(BlueprintCallable, Category = "Render Layers")
+	virtual void ApplyModifier(const UWorld* World) override;
+
+	UFUNCTION(BlueprintCallable, Category = "Render Layers")
+	virtual void UndoModifier() override;
+
+private:
+	void SetActorHiddenState(AActor* Actor, const bool bInIsHidden) const;
+
+private:
+	/** Maps an actor to its original hidden state. */
+	UPROPERTY(Transient)
+	TMap<TSoftObjectPtr<AActor>, bool> ModifiedActors;
+
+	UPROPERTY()
+	bool bIsHidden = false;
+};
+
+/**
+ * Base class for providing the ability to determine if an actor matches a query.
+ */
+UCLASS(Abstract)
+class UMoviePipelineCollectionQuery : public UObject
+{
+	GENERATED_BODY()
+
+public:
+	virtual bool DoesActorMatchQuery(const AActor* Actor) const PURE_VIRTUAL(UMoviePipelineCollectionQuery::DoesActorMatchQuery, return false; );  
+};
+
+UENUM(BlueprintType)
+enum class EMoviePipelineCollectionCommonQueryMode : uint8
+{
+	And UMETA(ToolTip = "All specifiers in the query must be true"),
+	Or UMETA(ToolTip = "At least one specifier in the query must be true")
+};
+
+/**
+ * Provides common actor querying functionality (names, tags, components, etc). These individual sub-queries can be
+ * AND'd or OR'd together (eg, matches provided names OR provided tags, vs matches provided names AND provided tags).
+ */
+UCLASS(BlueprintType)
+class UMoviePipelineCollectionCommonQuery : public UMoviePipelineCollectionQuery
+{
+	GENERATED_BODY()
+
+public:
+	// TODO: Add other common query operations (level, etc)
+	UFUNCTION(BlueprintCallable, Category = "Render Layers")
+	void SetActorNames(const TArray<FString>& InActorNames) { ActorNames = InActorNames; }
+
+	UFUNCTION(BlueprintCallable, Category = "Render Layers")
+	void SetTags(const TArray<FName>& InTags) { Tags = InTags; }
+
+	UFUNCTION(BlueprintCallable, Category = "Render Layers")
+	void SetComponentTypes(TArray<UClass*> InComponentTypes) { ComponentTypes = InComponentTypes; }
+
+	UFUNCTION(BlueprintCallable, Category = "Render Layers")
+	void SetQueryMode(const EMoviePipelineCollectionCommonQueryMode InQueryMode) { QueryMode = InQueryMode; }
+
+	UFUNCTION(BlueprintCallable, Category = "Render Layers")
+	virtual bool DoesActorMatchQuery(const AActor* Actor) const override;
+
+private:
+	UPROPERTY()
+	TArray<TObjectPtr<UClass>> ComponentTypes;
+
+	UPROPERTY()
+	EMoviePipelineCollectionCommonQueryMode QueryMode = EMoviePipelineCollectionCommonQueryMode::And;
+
+	UPROPERTY()
+	TArray<FString> ActorNames;
+
+	UPROPERTY()
+	TArray<FName> Tags;
+};
+
+/**
+ * Provides a means of assembling modifiers together to generate a desired view of a scene. 
+ */
+UCLASS(BlueprintType)
+class UMoviePipelineRenderLayer : public UObject
+{
+	GENERATED_BODY()
+
+public:
+	UMoviePipelineRenderLayer() = default;
+
+	UFUNCTION(BlueprintCallable, Category = "Render Layers")
+	FString GetRenderLayerName() const { return RenderLayerName; };
+
+	UFUNCTION(BlueprintCallable, Category = "Render Layers")
+	void SetRenderLayerName(const FString& NewName) { RenderLayerName = NewName; }
+
+	UFUNCTION(BlueprintCallable, Category = "Render Layers")
+	UMoviePipelineCollection* GetCollectionByName(const FString& Name) const;
+
+	UFUNCTION(BlueprintCallable, Category = "Render Layers")
+	void AddModifier(UMoviePipelineCollectionModifier* Modifier);
+
+	UFUNCTION(BlueprintCallable, Category = "Render Layers")
+	TArray<UMoviePipelineCollectionModifier*> GetModifiers() const { return Modifiers; };
+
+	UFUNCTION(BlueprintCallable, Category = "Render Layers")
+	void RemoveModifier(UMoviePipelineCollectionModifier* Modifier);
+
+	UFUNCTION(BlueprintCallable, Category = "Render Layers")
+	void Preview(const UWorld* World);
+
+	UFUNCTION(BlueprintCallable, Category = "Render Layers")
+	void UndoPreview(const UWorld* World);
+
+private:
+	/** The name of this render layer. */
+	UPROPERTY()
+	FString RenderLayerName;
+
+	/** The modifiers that are active when this render layer is active. */
+	UPROPERTY()
+	TArray<TObjectPtr<UMoviePipelineCollectionModifier>> Modifiers;
+};
+
+/**
+ * The primary means of controlling render layers in MRQ. Render layers can be added/registered with the subsystem, then
+ * made active in order to view them. Collections and modifiers can also be viewed, but they do not need to be added to
+ * the subsystem ahead of time.
+ */
+UCLASS()
+class MOVIERENDERPIPELINECORE_API UMoviePipelineRenderLayerSubsystem : public UWorldSubsystem
+{
+	GENERATED_BODY()
+
+public:
+	UMoviePipelineRenderLayerSubsystem() = default;
+
+	/* Get this subsystem for a specific world. Handy for use from Python. */
+	UFUNCTION(BlueprintCallable, Category = "Render Layers")
+	static UMoviePipelineRenderLayerSubsystem* GetFromWorld(const UWorld* World);
+
+	//~ Begin USubsystem interface
+	virtual bool ShouldCreateSubsystem(UObject* Outer) const override { return true; }
+	virtual void Initialize(FSubsystemCollectionBase& Collection) override;
+	virtual void Deinitialize() override;
+	//~ End USubsystem interface
+
+	/** Clear out all tracked render layers and collections. */
+	UFUNCTION(BlueprintCallable, Category = "Render Layers")
+	void Reset();
+
+	/**
+	 * Adds a render layer to the system, which can later be made active by SetActiveRenderLayer*(). Returns true
+	 * if the layer was added successfully, else false.
+	 */
+	UFUNCTION(BlueprintCallable, Category = "Render Layers")
+	bool AddRenderLayer(UMoviePipelineRenderLayer* RenderLayer);
+
+	UFUNCTION(BlueprintCallable, Category = "Render Layers")
+	TArray<UMoviePipelineRenderLayer*>& GetRenderLayers() { return RenderLayers; }
+
+	UFUNCTION(BlueprintCallable, Category = "Render Layers")
+	void RemoveRenderLayer(const FString& RenderLayerName);
+
+	UFUNCTION(BlueprintCallable, Category = "Render Layers")
+	UMoviePipelineRenderLayer* GetActiveRenderLayer() const { return ActiveRenderLayer; }
+
+	/** Previews the layer with the given name. The layer needs to have been registered with AddRenderLayer(). */
+	UFUNCTION(BlueprintCallable, Category = "Render Layers")
+	void SetActiveRenderLayerByName(const FString& RenderLayerName);
+
+	/** Previews the given layer. The layer does not need to have been registered with AddRenderLayer(). */
+	UFUNCTION(BlueprintCallable, Category = "Render Layers")
+	void SetActiveRenderLayerByObj(UMoviePipelineRenderLayer* RenderLayer);
+
+	UFUNCTION(BlueprintCallable, Category = "Render Layers")
+	void ClearActiveRenderLayer();
+
+	UFUNCTION(BlueprintCallable, Category = "Render Layers")
+	void PreviewCollection(UMoviePipelineCollection* Collection);
+
+	UFUNCTION(BlueprintCallable, Category = "Render Layers")
+	void ClearCollectionPreview();
+
+	UFUNCTION(BlueprintCallable, Category = "Render Layers")
+	void PreviewModifier(UMoviePipelineCollectionModifier* Modifier);
+
+	UFUNCTION(BlueprintCallable, Category = "Render Layers")
+	void ClearModifierPreview();
+
+private:
+	/** Clears the render layer, collection, and modifier previews if any of them are active. */
+	void ClearAllPreviews();
+
+	/** Sets the active render layer and previews it. */
+	void SetAndPreviewRenderLayer(UMoviePipelineRenderLayer* RenderLayer);
+
+private:
+	/** Render layers which have been added/registered with the subsystem. These can be found by name. */
+	UPROPERTY(Transient)
+	TArray<TObjectPtr<UMoviePipelineRenderLayer>> RenderLayers;
+
+	/** The render layer that is currently being viewed/previewed. */
+	UPROPERTY(Transient)
+	TObjectPtr<UMoviePipelineRenderLayer> ActiveRenderLayer = nullptr;
+
+	/** The collection that is currently being viewed/previewed. */
+	UPROPERTY(Transient)
+	TObjectPtr<UMoviePipelineCollection> ActiveCollection = nullptr;
+
+	/** The modifier that is currently being viewed/previewed. */
+	UPROPERTY(Transient)
+	TObjectPtr<UMoviePipelineCollectionModifier> ActiveModifier = nullptr;
+
+	/** A render layer dedicated to visualizing collections and modifiers. */
+	UPROPERTY(Transient)
+	TObjectPtr<UMoviePipelineRenderLayer> VisualizationRenderLayer = nullptr;
+
+	/** Empty collection used for visualization purposes (in conjunction w/ the viz render layer). */
+	UPROPERTY(Transient)
+	TObjectPtr<UMoviePipelineCollection> VisualizationEmptyCollection = nullptr;
+
+	/** A modifier used for visualization purposes (to hide the entire world). */
+	UPROPERTY(Transient)
+	TObjectPtr<UMoviePipelineVisibilityModifier> VisualizationModifier_HideWorld = nullptr;
+
+	/** A modifier used for visualization purposes (to show collections used with the modifier). */
+	UPROPERTY(Transient)
+	TObjectPtr<UMoviePipelineVisibilityModifier> VisualizationModifier_VisibleCollections = nullptr;
+};
\ No newline at end of file
