MoviePipeline: Added a lot of stub code for an upcoming experimental way of enhanced job configuration for MRQ jobs. Does not have any features at this time and will be a while before there is any sort of feature parity with the existing MRQ configuration.

diff --git a/Engine/Plugins/MovieScene/MovieRenderPipeline/MovieRenderPipeline.uplugin b/Engine/Plugins/MovieScene/MovieRenderPipeline/MovieRenderPipeline.uplugin
--- a/Engine/Plugins/MovieScene/MovieRenderPipeline/MovieRenderPipeline.uplugin
+++ b/Engine/Plugins/MovieScene/MovieRenderPipeline/MovieRenderPipeline.uplugin
@@ -42,6 +42,10 @@
 		{
 			"Name": "SequencerScripting",
 			"Enabled": true
+		},
+		{
+			"Name": "StructUtils",
+			"Enabled": true
 		}
 	],
 	"Modules": [
diff --git a/Engine/Plugins/MovieScene/MovieRenderPipeline/Source/MovieRenderPipelineCore/MovieRenderPipelineCore.Build.cs b/Engine/Plugins/MovieScene/MovieRenderPipeline/Source/MovieRenderPipelineCore/MovieRenderPipelineCore.Build.cs
--- a/Engine/Plugins/MovieScene/MovieRenderPipeline/Source/MovieRenderPipelineCore/MovieRenderPipelineCore.Build.cs
+++ b/Engine/Plugins/MovieScene/MovieRenderPipeline/Source/MovieRenderPipelineCore/MovieRenderPipelineCore.Build.cs
@@ -36,6 +36,7 @@ public class MovieRenderPipelineCore : ModuleRules
 				"ImageWriteQueue", // For debug tile writing
 				"OpenColorIO",
 				"CinematicCamera",
+				"StructUtils"
 			}
 		);
 
diff --git a/Engine/Plugins/MovieScene/MovieRenderPipeline/Source/MovieRenderPipelineCore/Private/Graph/MovieGraphConfig.cpp b/Engine/Plugins/MovieScene/MovieRenderPipeline/Source/MovieRenderPipelineCore/Private/Graph/MovieGraphConfig.cpp
new file mode 100644
--- /dev/null
+++ b/Engine/Plugins/MovieScene/MovieRenderPipeline/Source/MovieRenderPipelineCore/Private/Graph/MovieGraphConfig.cpp
@@ -0,0 +1,223 @@
+// Copyright Epic Games, Inc. All Rights Reserved.
+
+#include "Graph/MovieGraphConfig.h"
+#include "Graph/MovieGraphEdge.h"
+#include "MovieRenderPipelineCoreModule.h"
+
+
+UMovieGraphConfig::UMovieGraphConfig()
+{
+	InputNode = CreateDefaultSubobject<UMovieGraphInputNode>(TEXT("DefaultInputNode"));
+	OutputNode = CreateDefaultSubobject<UMovieGraphOutputNode>(TEXT("DefaultOutputNode"));
+	
+	if (!HasAnyFlags(RF_ClassDefaultObject))
+	{
+		InputNode->UpdatePins();
+		OutputNode->UpdatePins();
+	}
+}
+
+void UMovieGraphConfig::TraverseTest(UMovieGraphNode* InNode)
+{
+	UE_LOG(LogTemp, Warning, TEXT("Visiting Node: %s"), *InNode->GetFName().ToString());
+
+	for (UMovieGraphPin* Pin : InNode->InputPins)
+	{
+		for (UMovieGraphEdge* Edge : Pin->Edges)
+		{
+			if (UMovieGraphPin* OtherPin = Edge->GetOtherPin(Pin))
+			{
+				if (OtherPin->Node != InNode)
+				{
+					TraverseTest(OtherPin->Node);
+				}
+			}
+		}
+	}
+}
+
+UMovieGraphNode* UMovieGraphConfig::FindTraversalStartForContext(const FMovieGraphTraversalContext& InContext) const
+{
+	return OutputNode;
+}
+
+void UMovieGraphConfig::TraversalTest()
+{
+	FMovieGraphTraversalContext Context;
+	Context.ShotName = TEXT("Test");
+	Context.RenderLayerName = TEXT("Test2");
+
+	UMovieGraphNode* StartNode = FindTraversalStartForContext(Context);
+	if (StartNode)
+	{
+		UE_LOG(LogTemp, Warning, TEXT("Traversing!"));
+		TraverseTest(StartNode);
+	}
+	else
+	{
+		UE_LOG(LogTemp, Error, TEXT("Could not find traversal entry point for context"));
+	}
+}
+
+bool UMovieGraphConfig::AddLabeledEdge(UMovieGraphNode* FromNode, const FName& FromPinLabel, UMovieGraphNode* ToNode, const FName& ToPinLabel)
+{
+	if (!FromNode || !ToNode)
+	{
+		UE_LOG(LogMovieRenderPipeline, Error, TEXT("AddLabeledEdge: Invalid Edge Nodes"));
+		return false;
+	}
+
+	UMovieGraphPin* FromPin = FromNode->GetOutputPin(FromPinLabel);
+	if (!FromPin)
+	{
+		UE_LOG(LogMovieRenderPipeline, Error, TEXT("AddLabeledEdge: FromNode: %s does not have a pin with the label: %s"), *FromNode->GetName(), *FromPinLabel.ToString());
+		return false;
+	}
+
+	UMovieGraphPin* ToPin = ToNode->GetInputPin(ToPinLabel);
+	if (!ToPin)
+	{
+		UE_LOG(LogMovieRenderPipeline, Error, TEXT("AddLabeledEdge: ToNode: %s does not have a pin with the label: %s"), *ToNode->GetName(), *ToPinLabel.ToString());
+		return false;
+	}
+
+	// Add the edge
+	FromPin->AddEdgeTo(ToPin);
+	bool bConnectionBrokeOtherEdges = false;
+//	if (!ToPin->AllowMultipleConnections())
+//	{
+//		bConnectionBrokeOtherEdges = ToPin->BreakAllIncompatibleEdges();
+//	}
+//
+//#if WITH_EDITOR
+//	NotifyGraphChanged(EMoviePipelineGraphChangeType::Structural);
+//#endif
+
+	return bConnectionBrokeOtherEdges;
+}
+
+bool UMovieGraphConfig::RemoveEdge(UMovieGraphNode* FromNode, const FName& FromPinLabel, UMovieGraphNode* ToNode, const FName& ToPinLabel)
+{
+	if (!FromNode || !ToNode)
+	{
+		UE_LOG(LogMovieRenderPipeline, Error, TEXT("RemoveEdge: Invalid Edge Nodes"));
+		return false;
+	}
+
+	UMovieGraphPin* FromPin = FromNode->GetOutputPin(FromPinLabel);
+	if (!FromPin)
+	{
+		UE_LOG(LogMovieRenderPipeline, Error, TEXT("RemoveEdge: FromNode: %s does not have a pin with the label: %s"), *FromNode->GetName(), *FromPinLabel.ToString());
+		return false;
+	}
+
+	UMovieGraphPin* ToPin = ToNode->GetInputPin(ToPinLabel);
+	if (!ToPin)
+	{
+		UE_LOG(LogMovieRenderPipeline, Error, TEXT("RemoveEdge: ToNode: %s does not have a pin with the label: %s"), *ToNode->GetName(), *ToPinLabel.ToString());
+		return false;
+	}
+
+	bool bChanged = ToPin->BreakEdgeTo(FromPin);
+
+//#if WITH_EDITOR
+// 	   if(bChanged) {
+//	NotifyGraphChanged(EMoviePipelineGraphChangeType::Structural);
+// 	   }
+//#endif
+
+	return bChanged;
+}
+
+bool UMovieGraphConfig::RemoveAllInboundEdges(UMovieGraphNode* InNode)
+{
+	if (!InNode)
+	{
+		UE_LOG(LogMovieRenderPipeline, Error, TEXT("RemoveAllInboundEdges: Invalid Edge Nodes"));
+		return false;
+	}
+
+	bool bChanged = false;
+	for (UMovieGraphPin* InputPin : InNode->InputPins)
+	{
+		bChanged |= InputPin->BreakAllEdges();
+	}
+
+	//#if WITH_EDITOR
+	// 	   if(bChanged) {
+	//	NotifyGraphChanged(EMoviePipelineGraphChangeType::Structural);
+	// 	   }
+	//#endif
+
+	return bChanged;
+}
+
+bool UMovieGraphConfig::RemoveAllOutboundEdges(UMovieGraphNode* InNode)
+{
+	if (!InNode)
+	{
+		UE_LOG(LogMovieRenderPipeline, Error, TEXT("RemoveAllOutboundEdges: Invalid Edge Nodes"));
+		return false;
+	}
+
+	bool bChanged = false;
+	for (UMovieGraphPin* OutputPin : InNode->OutputPins)
+	{
+		bChanged |= OutputPin->BreakAllEdges();
+	}
+
+	//#if WITH_EDITOR
+	// 	   if(bChanged) {
+	//	NotifyGraphChanged(EMoviePipelineGraphChangeType::Structural);
+	// 	   }
+	//#endif
+
+	return bChanged;
+}
+
+bool UMovieGraphConfig::RemoveInboundEdges(UMovieGraphNode* InNode, const FName& InPinName)
+{
+	if (!InNode)
+	{
+		UE_LOG(LogMovieRenderPipeline, Error, TEXT("RemoveInboundEdges: Invalid Edge Nodes"));
+		return false;
+	}
+
+	bool bChanged = false;
+	if(UMovieGraphPin* InputPin = InNode->GetInputPin(InPinName))
+	{
+		bChanged |= InputPin->BreakAllEdges();
+	}
+
+	//#if WITH_EDITOR
+	// 	   if(bChanged) {
+	//	NotifyGraphChanged(EMoviePipelineGraphChangeType::Structural);
+	// 	   }
+	//#endif
+
+	return bChanged;
+}
+
+bool UMovieGraphConfig::RemoveOutboundEdges(UMovieGraphNode* InNode, const FName& InPinName)
+{
+	if (!InNode)
+	{
+		UE_LOG(LogMovieRenderPipeline, Error, TEXT("RemoveOutboundEdges: Invalid Edge Nodes"));
+		return false;
+	}
+
+	bool bChanged = false;
+	if (UMovieGraphPin* OutputPin = InNode->GetOutputPin(InPinName))
+	{
+		bChanged |= OutputPin->BreakAllEdges();
+	}
+
+	//#if WITH_EDITOR
+	// 	   if(bChanged) {
+	//	NotifyGraphChanged(EMoviePipelineGraphChangeType::Structural);
+	// 	   }
+	//#endif
+
+	return bChanged;
+}
+
diff --git a/Engine/Plugins/MovieScene/MovieRenderPipeline/Source/MovieRenderPipelineCore/Private/Graph/MovieGraphNode.cpp b/Engine/Plugins/MovieScene/MovieRenderPipeline/Source/MovieRenderPipelineCore/Private/Graph/MovieGraphNode.cpp
new file mode 100644
--- /dev/null
+++ b/Engine/Plugins/MovieScene/MovieRenderPipeline/Source/MovieRenderPipelineCore/Private/Graph/MovieGraphNode.cpp
@@ -0,0 +1,146 @@
+// Copyright Epic Games, Inc. All Rights Reserved.
+
+#include "Graph/MovieGraphNode.h"
+#include "Graph/MovieGraphConfig.h"
+
+
+UMovieGraphNode::UMovieGraphNode()
+{
+}
+
+void UMovieGraphNode::UpdatePins()
+{
+	
+	TArray<FMovieGraphPinProperties> InputPinProperties = GetInputPinProperties();
+	TArray<FMovieGraphPinProperties> OutputPinProperties = GetOutputPinProperties();
+
+	auto UpdatePins = [this](TArray<UMovieGraphPin*>& Pins, const TArray<FMovieGraphPinProperties>& PinProperties)
+	{
+		bool bAppliedEdgeChanges = false;
+		bool bChangedPins = false;
+
+		// Find unmatched pins vs. properties (via name matching)
+		TArray<UMovieGraphPin*> UnmatchedPins;
+		for (UMovieGraphPin* Pin : Pins)
+		{
+			if (const FMovieGraphPinProperties* MatchingProperties = PinProperties.FindByPredicate([Pin](const FMovieGraphPinProperties& Prop) { return Prop.Label == Pin->Properties.Label; }))
+			{
+				if (Pin->Properties != *MatchingProperties)
+				{
+					Pin->Modify();
+					Pin->Properties = *MatchingProperties;
+					//bAppliedEdgeCHanges |= Pin->BreakAllIncompatibleEdges();
+					bChangedPins = true;
+				}
+			}
+			else
+			{
+				UnmatchedPins.Add(Pin);
+			}
+		}
+
+		// Now do the opposite, find any properties that don't have pins
+		TArray<FMovieGraphPinProperties> UnmatchedProperties;
+		for (const FMovieGraphPinProperties& Properties : PinProperties)
+		{
+			if (!Pins.FindByPredicate([&Properties](const UMovieGraphPin* Pin) { return Pin->Properties.Label == Properties.Label; }))
+			{
+				UnmatchedProperties.Add(Properties);
+			}
+		}
+
+		if (!UnmatchedPins.IsEmpty() || !UnmatchedProperties.IsEmpty())
+		{
+			Modify();
+			bChangedPins = true;
+		}
+
+		// Remove old pins
+		for (int32 UnmatchedPinIndex = UnmatchedPins.Num() - 1; UnmatchedPinIndex >= 0; UnmatchedPinIndex--)
+		{
+			const int32 PinIndex = Pins.IndexOfByKey(UnmatchedPins[UnmatchedPinIndex]);
+			check(PinIndex >= 0);
+
+			//bAppliedEdgeCHanges |= Pins[PinIndex]->BreakAllEdges();
+			Pins.RemoveAt(PinIndex);
+		}
+
+		// Add new pins
+		for (const FMovieGraphPinProperties& UnmatchedProperty : UnmatchedProperties)
+		{
+			const int32 InsertIndex = PinProperties.IndexOfByKey(UnmatchedProperty);
+			UMovieGraphPin* NewPin = NewObject<UMovieGraphPin>(this);
+			NewPin->Node = this;
+			NewPin->Properties = UnmatchedProperty;
+			Pins.Insert(NewPin, InsertIndex);
+		}
+
+		// return bAppliedEdgeChanges ? EPCGChangeType::Edge : None) | (bChangedPins ? EPCGChangeType::Node : None
+	};
+
+	UpdatePins(InputPins, InputPinProperties);
+	UpdatePins(OutputPins, OutputPinProperties);
+}
+
+void UMovieGraphNode::UpdateDynamicProperties()
+{
+	TArray<FPropertyBagPropertyDesc> DesiredDynamicProperties = GetDynamicPropertyDescriptions();
+
+	// Check to see if we need to remake our property bag
+	bool bHasAllProperties = DynamicProperties.GetNumPropertiesInBag() == DesiredDynamicProperties.Num();
+	if (bHasAllProperties)
+	{
+		// If there's still the same number of properties before/after, then we need to do the more expensive
+		// check, which is to see if every property in the desired list is already inside the property bag.
+		for (const FPropertyBagPropertyDesc& Desc : DesiredDynamicProperties)
+		{
+			bool bBagContainsProperty = DynamicProperties.FindPropertyDescByName(Desc.Name) != nullptr;
+			if (!bBagContainsProperty)
+			{
+				bHasAllProperties = false;
+				break;
+			}
+		}
+	}
+
+	// If we don't have all the properties in our bag already, we need to generate a new bag with the correct
+	// layout, and then we have to migrate the existing bag over (so existing, matching, values stay).
+	if (!bHasAllProperties)
+	{
+		FInstancedPropertyBag NewPropertyBag;
+		NewPropertyBag.AddProperties(DesiredDynamicProperties);
+
+		DynamicProperties.MigrateToNewBagInstance(NewPropertyBag);
+	}
+}
+
+UMovieGraphConfig* UMovieGraphNode::GetGraph() const
+{
+	return Cast<UMovieGraphConfig>(GetOuter());
+}
+
+UMovieGraphPin* UMovieGraphNode::GetInputPin(const FName& Label) const
+{
+	for (UMovieGraphPin* InputPin : InputPins)
+	{
+		if (InputPin->Properties.Label == Label)
+		{
+			return InputPin;
+		}
+	}
+
+	return nullptr;
+}
+
+UMovieGraphPin* UMovieGraphNode::GetOutputPin(const FName& Label) const
+{
+	for (UMovieGraphPin* OutputPin : OutputPins)
+	{
+		if (OutputPin->Properties.Label == Label)
+		{
+			return OutputPin;
+		}
+	}
+
+	return nullptr;
+}
diff --git a/Engine/Plugins/MovieScene/MovieRenderPipeline/Source/MovieRenderPipelineCore/Private/Graph/MovieGraphPin.cpp b/Engine/Plugins/MovieScene/MovieRenderPipeline/Source/MovieRenderPipelineCore/Private/Graph/MovieGraphPin.cpp
new file mode 100644
--- /dev/null
+++ b/Engine/Plugins/MovieScene/MovieRenderPipeline/Source/MovieRenderPipelineCore/Private/Graph/MovieGraphPin.cpp
@@ -0,0 +1,122 @@
+// Copyright Epic Games, Inc. All Rights Reserved.
+
+#include "Graph/MovieGraphPin.h"
+#include "Graph/MovieGraphNode.h"
+#include "Graph/MovieGraphEdge.h"
+#include "MovieRenderPipelineCoreModule.h"
+
+bool UMovieGraphPin::AddEdgeTo(UMovieGraphPin* InOtherPin)
+{
+	if (!InOtherPin)
+	{
+		UE_LOG(LogMovieRenderPipeline, Error, TEXT("AddEdgeTo: Invalid InOtherPin"));
+		return false;
+	}
+
+	// Check to make sure the connection doesn't already exist
+	for (UMovieGraphEdge* Edge : Edges)
+	{
+		if (Edge->GetOtherPin(this) == InOtherPin)
+		{
+			return false;
+		}
+	}
+
+	// Don't allow connection between two output streams
+	const bool bThisPinIsUpstream = IsOutputPin();
+	const bool bOtherPinIsUpstream = InOtherPin->IsOutputPin();
+	if (!ensure(bThisPinIsUpstream != bOtherPinIsUpstream))
+	{
+		return false;
+	}
+
+	Modify();
+	InOtherPin->Modify();
+
+	UMovieGraphEdge* NewEdge = NewObject<UMovieGraphEdge>(this);
+	Edges.Add(NewEdge);
+	InOtherPin->Edges.Add(NewEdge);
+
+	NewEdge->InputPin = bThisPinIsUpstream ? this : InOtherPin;
+	NewEdge->OutputPin = bThisPinIsUpstream ? InOtherPin : this;
+	return true;
+}
+
+bool UMovieGraphPin::BreakEdgeTo(UMovieGraphPin* InOtherPin)
+{
+	if (!InOtherPin)
+	{
+		UE_LOG(LogMovieRenderPipeline, Error, TEXT("BreakEdgeTo: Invalid InOtherPin"));
+		return false;
+	}
+
+	for (UMovieGraphEdge* Edge : Edges)
+	{
+		if (Edge->GetOtherPin(this) == InOtherPin)
+		{
+			Modify();
+			InOtherPin->Modify();
+			ensure(InOtherPin->Edges.Remove(Edge));
+
+			Edges.Remove(Edge);
+			return true;
+		}
+	}
+
+	return false;
+}
+
+bool UMovieGraphPin::BreakAllEdges()
+{
+	bool bChanged = false;
+	if (!Edges.IsEmpty())
+	{
+		Modify();
+	}
+
+	for (UMovieGraphEdge* Edge : Edges)
+	{
+		if (UMovieGraphPin* OtherPin = Edge->GetOtherPin(this))
+		{
+			OtherPin->Modify();
+			ensure(OtherPin->Edges.Remove(Edge));
+			bChanged = true;
+		}
+	}
+
+	Edges.Reset();
+	return bChanged;
+}
+
+bool UMovieGraphPin::IsConnected() const
+{
+	for (UMovieGraphEdge* Edge : Edges)
+	{
+		if (Edge->IsValid())
+		{
+			return true;
+		}
+	}
+
+	return false;
+}
+bool UMovieGraphPin::IsOutputPin() const
+{
+	check(Node);
+	return Node->GetOutputPin(Properties.Label) == this;
+}
+
+int32 UMovieGraphPin::EdgeCount() const
+{
+	int32 EdgeCount = 0;
+	for (UMovieGraphEdge* Edge : Edges)
+	{
+		if (Edge->IsValid())
+		{
+			EdgeCount++;
+		}
+	}
+
+	return EdgeCount;
+}
+
diff --git a/Engine/Plugins/MovieScene/MovieRenderPipeline/Source/MovieRenderPipelineCore/Public/Graph/MovieGraphConfig.h b/Engine/Plugins/MovieScene/MovieRenderPipeline/Source/MovieRenderPipelineCore/Public/Graph/MovieGraphConfig.h
new file mode 100644
--- /dev/null
+++ b/Engine/Plugins/MovieScene/MovieRenderPipeline/Source/MovieRenderPipelineCore/Public/Graph/MovieGraphConfig.h
@@ -0,0 +1,85 @@
+// Copyright Epic Games, Inc. All Rights Reserved.
+
+#pragma once
+#include "CoreMinimal.h"
+#include "EdGraph/EdGraph.h"
+#include "Graph/MovieGraphNode.h"
+#include "MovieGraphConfig.generated.h"
+
+// Forward Declare
+class UMovieGraphNode;
+
+USTRUCT(BlueprintType)
+struct FMovieGraphTraversalContext
+{
+	GENERATED_BODY();
+
+	UPROPERTY(BlueprintReadWrite, EditAnywhere, Category = "Blah")
+	FString ShotName;
+
+	UPROPERTY(BlueprintReadWrite, EditAnywhere, Category = "Blah")
+	FString RenderLayerName;
+};
+
+/**
+* This is the runtime representation of the UMoviePipelineEdGraph which contains the actual strongly
+* typed graph network that is read by the MoviePipeline. There is an editor-only representation of
+* this graph (UMoviePipelineEdGraph).
+*/
+UCLASS()
+class MOVIERENDERPIPELINECORE_API UMovieGraphConfig : public UObject
+{
+	GENERATED_BODY()
+
+public:
+	UMovieGraphConfig();
+public:
+	bool AddLabeledEdge(UMovieGraphNode* FromNode, const FName& FromPinLabel, UMovieGraphNode* ToNode, const FName& ToPinLabel);
+	bool RemoveEdge(UMovieGraphNode* FromNode, const FName& FromPinName, UMovieGraphNode* ToNode, const FName& ToPinName);
+	bool RemoveAllInboundEdges(UMovieGraphNode* InNode);
+	bool RemoveAllOutboundEdges(UMovieGraphNode* InNode);
+	bool RemoveInboundEdges(UMovieGraphNode* InNode, const FName& InPinName);
+	bool RemoveOutboundEdges(UMovieGraphNode* InNode, const FName& InPinName);
+
+	UMovieGraphNode* GetInputNode() const { return InputNode; }
+	UMovieGraphNode* GetOutputNode() const { return OutputNode; }
+	const TArray<TObjectPtr<UMovieGraphNode>>& GetNodes() const { return AllNodes; }
+
+protected:
+	UMovieGraphNode* FindTraversalStartForContext(const FMovieGraphTraversalContext& InContext) const;
+public:
+	UPROPERTY()
+	TArray<TObjectPtr<UMovieGraphNode>> AllNodes;
+
+	UPROPERTY()
+	TObjectPtr<UMovieGraphNode> InputNode;
+
+	UPROPERTY()
+	TObjectPtr<UMovieGraphNode> OutputNode;
+
+#if WITH_EDITORONLY_DATA
+	// Not strongly typed to avoid a circular dependency between the editor only module
+	// and the runtime module, but it should be a UMoviePipelineEdGraph.
+	UPROPERTY()
+	TObjectPtr<UEdGraph> PipelineEdGraph;
+#endif
+
+	template<class T>
+	T* ConstructRuntimeNode(TSubclassOf<UMovieGraphNode> PipelineGraphNodeClass = T::StaticClass())
+	{
+		// Construct a new object with ourselves as the outer, then keep track of it.
+		// ToDo: This is a runtime node, kept track in AllNodes, which is ultimately editor only. Probably
+		// because the system it's based on (SoundCues) have a root node and links nodes together later?
+		T* RuntimeNode = NewObject<T>(this, PipelineGraphNodeClass, NAME_None, RF_Transactional);
+		RuntimeNode->UpdateDynamicProperties();
+		RuntimeNode->UpdatePins();
+#if WITH_EDITOR
+		AllNodes.Add(RuntimeNode);
+#endif // WITH_EDITORONLY_DATA
+		return RuntimeNode;
+	}
+
+	void TraversalTest();
+	void TraverseTest(UMovieGraphNode* InNode);
+
+};
\ No newline at end of file
diff --git a/Engine/Plugins/MovieScene/MovieRenderPipeline/Source/MovieRenderPipelineCore/Public/Graph/MovieGraphEdge.h b/Engine/Plugins/MovieScene/MovieRenderPipeline/Source/MovieRenderPipelineCore/Public/Graph/MovieGraphEdge.h
new file mode 100644
--- /dev/null
+++ b/Engine/Plugins/MovieScene/MovieRenderPipeline/Source/MovieRenderPipelineCore/Public/Graph/MovieGraphEdge.h
@@ -0,0 +1,32 @@
+// Copyright Epic Games, Inc. All Rights Reserved.
+
+#pragma once
+#include "CoreMinimal.h"
+#include "MovieGraphEdge.generated.h"
+
+class UMovieGraphPin;
+
+UCLASS()
+class MOVIERENDERPIPELINECORE_API UMovieGraphEdge : public UObject
+{
+	GENERATED_BODY()
+
+public:
+	UPROPERTY()
+	TObjectPtr<UMovieGraphPin> InputPin;
+	
+	UPROPERTY()
+	TObjectPtr<UMovieGraphPin> OutputPin;
+
+	bool IsValid() const
+	{
+		return InputPin.Get() && OutputPin.Get();
+	}
+
+	UMovieGraphPin* GetOtherPin(const UMovieGraphPin* InFirstPin)
+	{
+		check(InFirstPin == InputPin || InFirstPin == OutputPin);
+		return InFirstPin == InputPin ? OutputPin : InputPin;
+	}
+
+};
\ No newline at end of file
diff --git a/Engine/Plugins/MovieScene/MovieRenderPipeline/Source/MovieRenderPipelineCore/Public/Graph/MovieGraphNode.h b/Engine/Plugins/MovieScene/MovieRenderPipeline/Source/MovieRenderPipelineCore/Public/Graph/MovieGraphNode.h
new file mode 100644
--- /dev/null
+++ b/Engine/Plugins/MovieScene/MovieRenderPipeline/Source/MovieRenderPipelineCore/Public/Graph/MovieGraphNode.h
@@ -0,0 +1,206 @@
+// Copyright Epic Games, Inc. All Rights Reserved.
+
+#pragma once
+#include "CoreMinimal.h"
+#include "MovieGraphPin.h"
+#include "InstancedStruct.h"
+#include "PropertyBag.h"
+#include "MovieGraphNode.generated.h"
+
+
+// Forward Declares
+class UMovieGraphPin;
+
+#if WITH_EDITOR
+class UEdGraphNode;
+#endif
+
+/**
+* This is a base class for all nodes that can exist in the UMovieGraphConfig network.
+* In the editor, each node in the network will have an editor-only representation too 
+* which contains data about it's visual position in the graph, comments, etc.
+*/
+UCLASS()
+class MOVIERENDERPIPELINECORE_API UMovieGraphNode: public UObject
+{
+	GENERATED_BODY()
+
+	friend class UMovieGraphConfig;
+	friend class UMovieGraphEdge;
+	
+public:
+	UMovieGraphNode();
+
+	const TArray<TObjectPtr<UMovieGraphPin>>& GetInputPins() const { return InputPins; }
+	const TArray<TObjectPtr<UMovieGraphPin>>& GetOutputPins() const { return OutputPins; }
+	
+	virtual TArray<FMovieGraphPinProperties> GetInputPinProperties() const
+	{
+		return TArray<FMovieGraphPinProperties>();
+	}
+	virtual TArray<FMovieGraphPinProperties> GetOutputPinProperties() const
+	{
+		return TArray<FMovieGraphPinProperties>();
+	}
+
+	virtual TArray<FPropertyBagPropertyDesc> GetDynamicPropertyDescriptions() const
+	{
+		return TArray<FPropertyBagPropertyDesc>();
+	}
+
+	virtual TArray<FName> GetExposedDynamicProperties() const
+	{ 
+		return TArray<FName>();
+	}
+
+	void UpdatePins();
+	void UpdateDynamicProperties();
+	class UMovieGraphConfig* GetGraph() const;
+	UMovieGraphPin* GetInputPin(const FName& InPinLabel) const;
+	UMovieGraphPin* GetOutputPin(const FName& InPinLabel) const;
+protected:
+public:
+#if WITH_EDITORONLY_DATA
+	/** Editor Node Graph representation. Not strongly typed to avoid circular dependency between editor/runtime modules. */
+	UPROPERTY()
+	TObjectPtr<UEdGraphNode>	GraphNode;
+
+	class UEdGraphNode* GetGraphNode() const;
+#endif
+
+#if WITH_EDITOR
+	virtual FText GetMenuDescription() const { return FText(); }
+	virtual FText GetMenuCategory() const { return FText(); }
+#endif
+
+protected:
+	UPROPERTY()
+	TArray<TObjectPtr<UMovieGraphPin>> InputPins;
+	
+	UPROPERTY()
+	TArray<TObjectPtr<UMovieGraphPin>> OutputPins;
+
+	UPROPERTY(EditAnywhere, meta=(FixedLayout), Category = "Node")
+	FInstancedPropertyBag DynamicProperties;
+
+	UPROPERTY()
+	TArray<FName> ExposedDynamicPropertyNames;
+};
+
+// Dummy test nodes
+UCLASS()
+class UMoviePipelineCollectionNode : public UMovieGraphNode
+{
+	GENERATED_BODY()
+public:
+	UMoviePipelineCollectionNode()
+	{
+	}
+
+	virtual TArray<FMovieGraphPinProperties> GetInputPinProperties() const override
+	{
+		TArray<FMovieGraphPinProperties> Properties;
+		Properties.Add(FMovieGraphPinProperties(TEXT("Input"), false));
+		return Properties;
+	}
+
+	virtual TArray<FMovieGraphPinProperties> GetOutputPinProperties() const override
+	{
+		TArray<FMovieGraphPinProperties> Properties;
+		Properties.Add(FMovieGraphPinProperties(TEXT("Output"), false));
+		return Properties;
+	}
+
+	virtual TArray<FPropertyBagPropertyDesc> GetDynamicPropertyDescriptions() const
+	{
+		TArray<FPropertyBagPropertyDesc> Properties;
+		FPropertyBagPropertyDesc FloatEditConProperty = FPropertyBagPropertyDesc("bOverride_TestPropName", EPropertyBagPropertyType::Bool);
+		FPropertyBagPropertyDesc FloatProperty = FPropertyBagPropertyDesc("TestPropName", EPropertyBagPropertyType::Float);
+#if WITH_EDITOR
+		FloatEditConProperty.MetaData.Add(FPropertyBagPropertyDescMetaData("InlineEditConditionToggle", "true"));
+		FloatProperty.MetaData.Add(FPropertyBagPropertyDescMetaData("EditCondition", "bOverride_TestPropName"));
+#endif
+		
+		Properties.Add(FloatEditConProperty);
+		Properties.Add(FloatProperty);
+		return Properties;
+	}
+#if WITH_EDITOR
+	virtual FText GetMenuDescription() const override
+	{
+		return NSLOCTEXT("debug", "collection nodename", "Component Collection");
+	}
+	
+	virtual FText GetMenuCategory() const override
+	{
+		return NSLOCTEXT("debug", "collection cat", "Rendering");
+	}
+#endif
+};
+
+UCLASS()
+class UMoviePipelineRenderLayerNode : public UMovieGraphNode
+{
+	GENERATED_BODY()
+	
+public:
+	virtual TArray<FMovieGraphPinProperties> GetInputPinProperties() const override
+	{
+		TArray<FMovieGraphPinProperties> Properties;
+		Properties.Add(FMovieGraphPinProperties(TEXT("Test Input 3"), false));
+		Properties.Add(FMovieGraphPinProperties(TEXT("Test Input 4"), false));
+		Properties.Add(FMovieGraphPinProperties(TEXT("Test Input 5"), false));
+		Properties.Add(FMovieGraphPinProperties(TEXT("Test Input 6"), false));
+		return Properties;
+	}
+
+	virtual TArray<FMovieGraphPinProperties> GetOutputPinProperties() const override
+	{
+		TArray<FMovieGraphPinProperties> Properties;
+		Properties.Add(FMovieGraphPinProperties(TEXT("Test Output 2"), false));
+		return Properties;
+	}
+
+#if WITH_EDITOR
+	virtual FText GetMenuDescription() const override
+	{
+		return NSLOCTEXT("debug", "collection nodename2", "Deferred Render Layer");
+	}
+	
+	virtual FText GetMenuCategory() const override
+	{
+		return NSLOCTEXT("debug", "collection cat", "Rendering");
+	}
+#endif
+	UPROPERTY(BlueprintReadWrite, EditAnywhere, Category = "Settings2")
+	FString LayerName;
+};
+
+
+UCLASS()
+class MOVIERENDERPIPELINECORE_API UMovieGraphOutputNode : public UMovieGraphNode
+{
+	GENERATED_BODY()
+
+public:
+	virtual TArray<FMovieGraphPinProperties> GetInputPinProperties() const override
+	{
+		TArray<FMovieGraphPinProperties> Properties;
+		Properties.Add(FMovieGraphPinProperties(TEXT("Output"), false));
+		return Properties;
+	}
+};
+
+UCLASS()
+class MOVIERENDERPIPELINECORE_API UMovieGraphInputNode : public UMovieGraphNode
+{
+	GENERATED_BODY()
+
+public:
+	virtual TArray<FMovieGraphPinProperties> GetOutputPinProperties() const override
+	{
+		TArray<FMovieGraphPinProperties> Properties;
+		Properties.Add(FMovieGraphPinProperties(TEXT("Input"), false));
+		return Properties;
+	}
+};
\ No newline at end of file
diff --git a/Engine/Plugins/MovieScene/MovieRenderPipeline/Source/MovieRenderPipelineCore/Public/Graph/MovieGraphPin.h b/Engine/Plugins/MovieScene/MovieRenderPipeline/Source/MovieRenderPipelineCore/Public/Graph/MovieGraphPin.h
new file mode 100644
--- /dev/null
+++ b/Engine/Plugins/MovieScene/MovieRenderPipeline/Source/MovieRenderPipelineCore/Public/Graph/MovieGraphPin.h
@@ -0,0 +1,63 @@
+// Copyright Epic Games, Inc. All Rights Reserved.
+
+#pragma once
+#include "CoreMinimal.h"
+#include "MovieGraphPin.generated.h"
+
+// Forward Declares
+class UMovieGraphNode;
+class UMovieGraphEdge;
+
+USTRUCT(BlueprintType)
+struct FMovieGraphPinProperties
+{
+	GENERATED_BODY()
+
+	FMovieGraphPinProperties() = default;
+	explicit FMovieGraphPinProperties(const FName& InLabel, bool bInAllowMultipleConnections)
+		: Label(InLabel)
+		, bAllowMultipleConnections(bInAllowMultipleConnections)
+	{}
+
+	UPROPERTY(BlueprintReadWrite, EditAnywhere, Category = "Settings")
+	FName Label = NAME_None;
+
+	UPROPERTY(BlueprintReadWrite, EditAnywhere, Category = "Settings")
+	bool bAllowMultipleConnections = true;
+
+	bool operator==(const FMovieGraphPinProperties& Other) const
+	{
+		return Label == Other.Label
+			&& bAllowMultipleConnections == Other.bAllowMultipleConnections;
+	}
+
+	bool operator !=(const FMovieGraphPinProperties& Other) const
+	{
+		return !(*this == Other);
+	}
+};
+
+
+UCLASS()
+class MOVIERENDERPIPELINECORE_API UMovieGraphPin : public UObject
+{
+	GENERATED_BODY()
+
+public:
+	bool AddEdgeTo(UMovieGraphPin* InOtherPin);
+	bool BreakEdgeTo(UMovieGraphPin* InOtherPin);
+	bool BreakAllEdges();
+	bool IsConnected() const;
+	bool IsOutputPin() const;
+	int32 EdgeCount() const;
+public:
+	// The node that this pin belongs to.
+	UPROPERTY(BlueprintReadOnly, Category = "Properties")
+	TObjectPtr<UMovieGraphNode> Node = nullptr;
+
+	UPROPERTY(BlueprintReadOnly, EditAnywhere, Category = "Properties", meta = (ShowOnlyInnerProperties))
+	FMovieGraphPinProperties Properties;
+
+	UPROPERTY(BlueprintReadOnly, Category = "Properties")
+	TArray<TObjectPtr<UMovieGraphEdge>> Edges;
+};
\ No newline at end of file
diff --git a/Engine/Plugins/MovieScene/MovieRenderPipeline/Source/MovieRenderPipelineCore/Public/MoviePipelineQueue.h b/Engine/Plugins/MovieScene/MovieRenderPipeline/Source/MovieRenderPipelineCore/Public/MoviePipelineQueue.h
--- a/Engine/Plugins/MovieScene/MovieRenderPipeline/Source/MovieRenderPipelineCore/Public/MoviePipelineQueue.h
+++ b/Engine/Plugins/MovieScene/MovieRenderPipeline/Source/MovieRenderPipelineCore/Public/MoviePipelineQueue.h
@@ -8,6 +8,7 @@
 #include "MoviePipelineShotConfig.h"
 #include "MoviePipelineConfigBase.h"
 #include "MovieSceneSequenceID.h"
+#include "Graph/MovieGraphConfig.h"
 
 #include "MoviePipelineQueue.generated.h"
 
@@ -131,6 +132,57 @@ public:
 		return ShotOverridePresetOrigin.Get();
 	}
 
+	/**
+	 * Returns true if this job is using graph-style configuration, else false.
+	 */
+	UFUNCTION(BlueprintPure, Category = "Movie Render Pipeline")
+	bool IsUsingGraphConfiguration() const
+	{
+		return GraphPreset.IsValid() || GraphConfig != nullptr;
+	}
+
+	/**
+	 * Gets the graph-style preset that this job is using. If the job is not using a graph-style preset, returns nullptr.
+	 * @see GetGraphConfig()
+	 */
+	UFUNCTION(BlueprintPure, Category = "Movie Render Pipeline")
+	UMovieGraphConfig* GetGraphPreset() const
+	{
+		return GraphPreset.LoadSynchronous();
+	}
+
+	/**
+	 * Gets the graph-style config that this job is using. If the job is not using a graph-style config, returns nullptr.
+	 * @see GetGraphPreset()
+	 */
+	UFUNCTION(BlueprintPure, Category = "Movie Render Pipeline")
+	UMovieGraphConfig* GetGraphConfig() const
+	{
+		return GraphConfig;
+	}
+
+	/**
+	 * Sets the graph-style preset that this job will use. Note that this will cause the graph to switch over to using
+	 * graph-style configuration if it is not already using it.
+	 * @see SetGraphConfig()
+	 */
+	UFUNCTION(BlueprintCallable, Category = "Movie Render Pipeline")
+	void SetGraphPreset(const UMovieGraphConfig* InGraphPreset)
+	{
+		GraphPreset = InGraphPreset;
+	}
+
+	/**
+	 * Sets the graph-style config that this job will use. Note that this will cause the graph to switch over to using
+	 * graph-style configuration if it is not already using it.
+	 * @see SetGraphPreset()
+	 */
+	UFUNCTION(BlueprintCallable, Category = "Movie Render Pipeline")
+	void SetGraphConfig(UMovieGraphConfig* InGraphConfig)
+	{
+		GraphConfig = InGraphConfig;
+	}
+
 	/** Returns whether this should should be rendered */
 	UFUNCTION(BlueprintPure, Category = "Movie Render Pipeline")
 	bool ShouldRender() const
@@ -188,6 +240,14 @@ private:
 
 	UPROPERTY()
 	TSoftObjectPtr<UMoviePipelineShotConfig> ShotOverridePresetOrigin;
+
+	/** The graph-based configuration that this shot is using; Can be nullptr. */
+	UPROPERTY()
+	TObjectPtr<UMovieGraphConfig> GraphConfig;
+
+	/** The graph-based configuration preset that this shot is using. Can be nullptr. */
+	UPROPERTY()
+	TSoftObjectPtr<UMovieGraphConfig> GraphPreset;
 };
 
 /**
@@ -364,6 +424,57 @@ public:
 	UFUNCTION(BlueprintCallable, Category = "Movie Render Pipeline")
 	void SetConfiguration(UMoviePipelinePrimaryConfig* InPreset);
 
+	/**
+	 * Returns true if this job is using graph-style configuration, else false.
+	 */
+	UFUNCTION(BlueprintPure, Category = "Movie Render Pipeline")
+	bool IsUsingGraphConfiguration() const
+	{
+		return GraphPreset.IsValid() || GraphConfig != nullptr;
+	}
+
+	/**
+	 * Gets the graph-style preset that this job is using. If the job is not using a graph-style preset, returns nullptr.
+	 * @see GetPresetOrigin()
+	 */
+	UFUNCTION(BlueprintPure, Category = "Movie Render Pipeline")
+	UMovieGraphConfig* GetGraphPreset() const
+	{
+		return GraphPreset.LoadSynchronous();
+	}
+
+	/**
+	 * Gets the graph-style config that this job is using. If the job is not using a graph-style config, returns nullptr.
+	 * @see GetGraphPreset()
+	 */
+	UFUNCTION(BlueprintPure, Category = "Movie Render Pipeline")
+	UMovieGraphConfig* GetGraphConfig() const
+	{
+		return GraphConfig;
+	}
+
+	/**
+	 * Sets the graph-style preset that this job will use. Note that this will cause the graph to switch over to using
+	 * graph-style configuration if it is not already using it.
+	 * @see SetGraphConfig()
+	 */
+	UFUNCTION(BlueprintCallable, Category = "Movie Render Pipeline")
+	void SetGraphPreset(const UMovieGraphConfig* InGraphPreset)
+	{
+		GraphPreset = InGraphPreset;
+	}
+
+	/**
+	 * Sets the graph-style config that this job will use. Note that this will cause the graph to switch over to using
+	 * graph-style configuration if it is not already using it.
+	 * @see SetGraphPreset()
+	 */
+	UFUNCTION(BlueprintCallable, Category = "Movie Render Pipeline")
+	void SetGraphConfig(UMovieGraphConfig* InGraphConfig)
+	{
+		GraphConfig = InGraphConfig;
+	}
+
 	UFUNCTION(BlueprintSetter, Category = "Movie Render Pipeline")
 	void SetSequence(FSoftObjectPath InSequence);
 
@@ -446,6 +557,14 @@ private:
 	/** Whether this job is enabled and should be rendered. */
 	UPROPERTY()
 	bool bEnabled;
+
+	/** The graph-based configuration that this job is using; this graph has not been saved as an asset. Can be nullptr.  */
+	UPROPERTY()
+	TObjectPtr<UMovieGraphConfig> GraphConfig;
+
+	/** The graph-based configuration preset that this job is using. Can be nullptr. */
+	UPROPERTY()
+	TSoftObjectPtr<UMovieGraphConfig> GraphPreset;
 };
 
 /**
diff --git a/Engine/Plugins/MovieScene/MovieRenderPipeline/Source/MovieRenderPipelineEditor/Private/Graph/MovieEdGraph.cpp b/Engine/Plugins/MovieScene/MovieRenderPipeline/Source/MovieRenderPipelineEditor/Private/Graph/MovieEdGraph.cpp
new file mode 100644
--- /dev/null
+++ b/Engine/Plugins/MovieScene/MovieRenderPipeline/Source/MovieRenderPipelineEditor/Private/Graph/MovieEdGraph.cpp
@@ -0,0 +1,121 @@
+// Copyright Epic Games, Inc. All Rights Reserved.
+
+#include "Graph/MovieEdGraph.h"
+#include "Graph/MovieGraphConfig.h"
+#include "Graph/MovieGraphEdge.h"
+#include "Graph/MovieGraphPin.h"
+#include "Graph/MovieGraphNode.h"
+#include "MovieEdGraphOutputNode.h"
+#include "MovieEdGraphInputNode.h"
+#include "MovieRenderPipelineCoreModule.h"
+
+UMovieGraphConfig* UMoviePipelineEdGraph::GetPipelineGraph() const
+{
+	return CastChecked<UMovieGraphConfig>(GetOuter());
+}
+
+
+void UMoviePipelineEdGraph::InitFromRuntimeGraph(UMovieGraphConfig* InGraph)
+{
+	// Don't allow reinitialization of an existing graph
+	check(InGraph && !bInitialized);
+
+	const bool bSelectNewNode = false;
+	TMap<UMovieGraphNode*, UMoviePipelineEdGraphNodeBase*> NodeLookup;
+
+	// Input
+	{
+		UMovieGraphNode* InputNode = InGraph->GetInputNode();
+		FGraphNodeCreator<UMoviePipelineEdGraphNodeInput> InputNodeCreator(*this);
+		UMoviePipelineEdGraphNodeInput* InputGraphNode = InputNodeCreator.CreateNode(bSelectNewNode);
+		InputGraphNode->Construct(InputNode);
+		InputNodeCreator.Finalize();
+	}
+
+	// Output
+	{
+		UMovieGraphNode* OutputNode = InGraph->GetOutputNode();
+		FGraphNodeCreator<UMoviePipelineEdGraphNodeOutput> OutputNodeCreator(*this);
+		UMoviePipelineEdGraphNodeOutput* OutputGraphNode = OutputNodeCreator.CreateNode(bSelectNewNode);
+		OutputGraphNode->Construct(OutputNode);
+		OutputNodeCreator.Finalize();
+	}
+
+	// Create the rest of the nodes in the graph
+	for (const TObjectPtr<UMovieGraphNode> RuntimeNode : InGraph->GetNodes())
+	{
+		FGraphNodeCreator<UMoviePipelineEdGraphNode> NodeCreator(*this);
+		UMoviePipelineEdGraphNode* GraphNode = NodeCreator.CreateNode(bSelectNewNode);
+		GraphNode->Construct(RuntimeNode);
+		NodeCreator.Finalize();
+		NodeLookup.Add(RuntimeNode, GraphNode);
+	}
+
+	// Now that we've added an Editor Graph representation for every node in the graph, link
+	// the editor nodes together to match the Runtime Layout.
+	for (const TPair< UMovieGraphNode*, UMoviePipelineEdGraphNodeBase*> Pair : NodeLookup)
+	{
+		const bool bCreateInboundLinks = false;
+		const bool bCreateOutboundLinks = true;
+		CreateLinks(Pair.Value, bCreateInboundLinks, bCreateOutboundLinks, NodeLookup);
+	}
+
+	bInitialized = true;
+}
+
+void UMoviePipelineEdGraph::CreateLinks(UMoviePipelineEdGraphNodeBase* InGraphNode, bool bCreateInboundLinks, bool bCreateOutboundLinks,
+	const TMap<UMovieGraphNode*, UMoviePipelineEdGraphNodeBase*>& RuntimeNodeToEdNodeMap)
+{
+	check(InGraphNode);
+	const UMovieGraphNode* RuntimeNode = InGraphNode->GetRuntimeNode();
+	check(RuntimeNode);
+
+	auto CreateLinks = [&](const TArray<TObjectPtr<UMovieGraphPin>>& Pins, EEdGraphPinDirection PrimaryDirection)
+	{
+		for (const UMovieGraphPin* Pin : Pins)
+		{
+			UEdGraphPin* EdPin = InGraphNode->FindPin(Pin->Properties.Label, PrimaryDirection);
+			if (!EdPin)
+			{
+				UE_LOG(LogMovieRenderPipeline, Error, TEXT("Invalid Pin for %s"), *Pin->Properties.Label.ToString());
+				continue;
+			}
+
+			for (const UMovieGraphEdge* Edge : Pin->Edges)
+			{
+				if (!Edge->IsValid())
+				{
+					UE_LOG(LogMovieRenderPipeline, Error, TEXT("Invalid edge for Pin: %s"), *Pin->Properties.Label.ToString());
+					continue;
+				}
+
+				UMovieGraphPin* OtherPin = PrimaryDirection == EEdGraphPinDirection::EGPD_Input ?
+					Edge->InputPin : Edge->OutputPin;
+				if (UMoviePipelineEdGraphNodeBase* const* ConnectedGraphNode = RuntimeNodeToEdNodeMap.Find(OtherPin->Node))
+				{
+					EEdGraphPinDirection SecondaryDirection = PrimaryDirection == EEdGraphPinDirection::EGPD_Input ?
+						EEdGraphPinDirection::EGPD_Output : EEdGraphPinDirection::EGPD_Input;
+					if (UEdGraphPin* OutEdPin = (*ConnectedGraphNode)->FindPin(OtherPin->Properties.Label, SecondaryDirection))
+					{
+						OutEdPin->MakeLinkTo(EdPin);
+					}
+					else
+					{
+						UE_LOG(LogMovieRenderPipeline, Error, TEXT("Could not create link to Pin %s from Node %s"), *Pin->Properties.Label.ToString(), *OtherPin->Node->GetFName().ToString());
+						continue;
+					}
+				}
+			}
+		}
+	};
+
+
+	if (bCreateInboundLinks)
+	{
+		CreateLinks(RuntimeNode->GetInputPins(), EEdGraphPinDirection::EGPD_Input);
+	}
+	if (bCreateOutboundLinks)
+	{
+		CreateLinks(RuntimeNode->GetInputPins(), EEdGraphPinDirection::EGPD_Output);
+	}
+}
\ No newline at end of file
diff --git a/Engine/Plugins/MovieScene/MovieRenderPipeline/Source/MovieRenderPipelineEditor/Private/Graph/MovieEdGraph.h b/Engine/Plugins/MovieScene/MovieRenderPipeline/Source/MovieRenderPipelineEditor/Private/Graph/MovieEdGraph.h
new file mode 100644
--- /dev/null
+++ b/Engine/Plugins/MovieScene/MovieRenderPipeline/Source/MovieRenderPipelineEditor/Private/Graph/MovieEdGraph.h
@@ -0,0 +1,35 @@
+// Copyright Epic Games, Inc. All Rights Reserved.
+
+#pragma once
+#include "CoreMinimal.h"
+#include "EdGraph/EdGraph.h"
+#include "MovieEdGraph.generated.h"
+
+class UMoviePipelineEdGraphNodeBase;
+class UMovieGraphNode;
+class UMovieGraphConfig;
+
+/**
+* This is the editor-only graph representation of the UMovieGraphConfig. This contains
+* editor only nodes (which have information about their X/Y position graphs, their widgets, etc.)
+* where each node in this ed graph is tied to a node in the runtime UMovieGraphConfig.
+*/
+UCLASS()
+class UMoviePipelineEdGraph : public UEdGraph
+{
+	GENERATED_BODY()
+
+public:
+
+	/** Initialize this Editor Graph from a Runtime Graph */
+	void InitFromRuntimeGraph(UMovieGraphConfig* InGraph);
+
+	/** Returns the runtime UMovieGraphConfig that contains this editor graph */
+	class UMovieGraphConfig* GetPipelineGraph() const;
+
+protected:
+	void CreateLinks(UMoviePipelineEdGraphNodeBase* InGraphNode, bool bCreateInboundLinks, bool bCreateOutboundLinks,
+		const TMap<UMovieGraphNode*, UMoviePipelineEdGraphNodeBase*>& RuntimeNodeToEdNodeMap);
+
+	bool bInitialized = false;
+};
\ No newline at end of file
diff --git a/Engine/Plugins/MovieScene/MovieRenderPipeline/Source/MovieRenderPipelineEditor/Private/Graph/MovieEdGraphInputNode.cpp b/Engine/Plugins/MovieScene/MovieRenderPipeline/Source/MovieRenderPipelineEditor/Private/Graph/MovieEdGraphInputNode.cpp
new file mode 100644
--- /dev/null
+++ b/Engine/Plugins/MovieScene/MovieRenderPipeline/Source/MovieRenderPipelineEditor/Private/Graph/MovieEdGraphInputNode.cpp
@@ -0,0 +1,19 @@
+// Copyright Epic Games, Inc. All Rights Reserved.
+
+#include "Graph/MovieEdGraphInputNode.h"
+
+#define LOCTEXT_NAMESPACE "MoviePipelineGraph"
+FText UMoviePipelineEdGraphNodeInput::GetNodeTitle(ENodeTitleType::Type TitleType) const
+{
+	return FText::FromName(TEXT("Input"));
+}
+
+void UMoviePipelineEdGraphNodeInput::AllocateDefaultPins()
+{
+	if (RuntimeNode)
+	{
+		CreatePins(/*InInputPins=*/{}, RuntimeNode->GetOutputPins());
+	}
+}
+
+#undef LOCTEXT_NAMESPACE
\ No newline at end of file
diff --git a/Engine/Plugins/MovieScene/MovieRenderPipeline/Source/MovieRenderPipelineEditor/Private/Graph/MovieEdGraphInputNode.h b/Engine/Plugins/MovieScene/MovieRenderPipeline/Source/MovieRenderPipelineEditor/Private/Graph/MovieEdGraphInputNode.h
new file mode 100644
--- /dev/null
+++ b/Engine/Plugins/MovieScene/MovieRenderPipeline/Source/MovieRenderPipelineEditor/Private/Graph/MovieEdGraphInputNode.h
@@ -0,0 +1,21 @@
+// Copyright Epic Games, Inc. All Rights Reserved.
+
+#pragma once
+#include "Graph/MovieEdGraphNode.h"
+
+#include "MovieEdGraphInputNode.generated.h"
+
+
+UCLASS()
+class UMoviePipelineEdGraphNodeInput : public UMoviePipelineEdGraphNodeBase
+{
+	GENERATED_BODY()
+
+public:
+	// ~Begin UEdGraphNode interface
+	virtual FText GetNodeTitle(ENodeTitleType::Type TitleType) const override;
+	virtual void AllocateDefaultPins() override;
+	virtual bool CanUserDeleteNode() const override { return false; }
+	virtual bool CanDuplicateNode() const override { return false; }
+	// ~End UEdGraphNode interface
+};
diff --git a/Engine/Plugins/MovieScene/MovieRenderPipeline/Source/MovieRenderPipelineEditor/Private/Graph/MovieEdGraphNode.cpp b/Engine/Plugins/MovieScene/MovieRenderPipeline/Source/MovieRenderPipelineEditor/Private/Graph/MovieEdGraphNode.cpp
new file mode 100644
--- /dev/null
+++ b/Engine/Plugins/MovieScene/MovieRenderPipeline/Source/MovieRenderPipelineEditor/Private/Graph/MovieEdGraphNode.cpp
@@ -0,0 +1,101 @@
+// Copyright Epic Games, Inc. All Rights Reserved.
+
+#include "Graph/MovieEdGraphNode.h"
+#include "Graph/MovieGraphPin.h"
+#include "Graph/MovieGraphNode.h"
+
+void UMoviePipelineEdGraphNodeBase::Construct(UMovieGraphNode* InRuntimeNode)
+{
+	check(InRuntimeNode);
+	RuntimeNode = InRuntimeNode;
+
+	// NodePosX = InRuntimeNode->PositionX;
+	// NodePosY = InRuntimeNode->PositionY;
+	// NodeComment = InRuntimeNode->NodeComment;
+	// bCommentBubblePinned = InRuntimeNode->bCommentBubblePinned;
+	// bCommentBubbleVisible = InRuntimeNode->bCommentBubbleVisible;
+}
+
+FEdGraphPinType UMoviePipelineEdGraphNodeBase::GetPinType(const UMovieGraphPin* InPin)
+{
+	FEdGraphPinType EdPinType;
+	EdPinType.ResetToDefaults();
+	
+	EdPinType.PinCategory = FName("TestCategory");
+	EdPinType.PinSubCategory = FName("TestSubCategory");
+	return EdPinType;
+}
+
+void UMoviePipelineEdGraphNode::AllocateDefaultPins()
+{
+	if(RuntimeNode)
+	{
+		for(const UMovieGraphPin* InputPin : RuntimeNode->GetInputPins())
+		{
+			CreatePin(EEdGraphPinDirection::EGPD_Input, GetPinType(InputPin), InputPin->Properties.Label);
+		}
+		
+		for(const UMovieGraphPin* OutputPin : RuntimeNode->GetOutputPins())
+		{
+			CreatePin(EEdGraphPinDirection::EGPD_Output, GetPinType(OutputPin), OutputPin->Properties.Label);
+		}
+	}
+}
+
+FText UMoviePipelineEdGraphNode::GetNodeTitle(ENodeTitleType::Type TitleType) const
+{
+	if (RuntimeNode)
+	{
+		return RuntimeNode->GetMenuDescription();
+	}
+
+	return NSLOCTEXT("MoviePipelineGraph", "GraphTestTitle", "TestTitle");
+}
+
+FText UMoviePipelineEdGraphNode::GetTooltipText() const
+{
+	return NSLOCTEXT("MoviePipelineGraph", "GraphTestTooltip", "Test Tooltip");
+}
+
+bool UMoviePipelineEdGraphNodeBase::ShouldCreatePin(const UMovieGraphPin* InPin) const
+{
+	return true;
+}
+
+void UMoviePipelineEdGraphNodeBase::CreatePins(const TArray<UMovieGraphPin*>& InInputPins, const TArray<UMovieGraphPin*>& InOutputPins)
+{
+	bool bHasAdvancedPin = false;
+
+	for (const UMovieGraphPin* InputPin : InInputPins)
+	{
+		if (!ShouldCreatePin(InputPin))
+		{
+			continue;
+		}
+
+		UEdGraphPin* Pin = CreatePin(EEdGraphPinDirection::EGPD_Input, GetPinType(InputPin), InputPin->Properties.Label);
+		// Pin->bAdvancedView = InputPin->Properties.bAdvancedPin;
+		bHasAdvancedPin |= Pin->bAdvancedView;
+	}
+
+	for (const UMovieGraphPin* OutputPin : InOutputPins)
+	{
+		if (!ShouldCreatePin(OutputPin))
+		{
+			continue;
+		}
+
+		UEdGraphPin* Pin = CreatePin(EEdGraphPinDirection::EGPD_Output, GetPinType(OutputPin), OutputPin->Properties.Label);
+		// Pin->bAdvancedView = OutputPin->Properties.bAdvancedPin;
+		bHasAdvancedPin |= Pin->bAdvancedView;
+	}
+
+	if (bHasAdvancedPin && AdvancedPinDisplay == ENodeAdvancedPins::NoPins)
+	{
+		AdvancedPinDisplay = ENodeAdvancedPins::Hidden;
+	}
+	else if (!bHasAdvancedPin)
+	{
+		AdvancedPinDisplay = ENodeAdvancedPins::NoPins;
+	}
+}
\ No newline at end of file
diff --git a/Engine/Plugins/MovieScene/MovieRenderPipeline/Source/MovieRenderPipelineEditor/Private/Graph/MovieEdGraphNode.h b/Engine/Plugins/MovieScene/MovieRenderPipeline/Source/MovieRenderPipelineEditor/Private/Graph/MovieEdGraphNode.h
new file mode 100644
--- /dev/null
+++ b/Engine/Plugins/MovieScene/MovieRenderPipeline/Source/MovieRenderPipelineEditor/Private/Graph/MovieEdGraphNode.h
@@ -0,0 +1,50 @@
+// Copyright Epic Games, Inc. All Rights Reserved.
+
+#pragma once
+#include "CoreTypes.h"
+#include "EdGraph/EdGraphNode.h"
+#include "Graph/MovieGraphNode.h"
+
+#include "MovieEdGraphNode.generated.h"
+
+class UMovieGraphNode;
+
+UCLASS(Abstract)
+class UMoviePipelineEdGraphNodeBase : public UEdGraphNode
+{
+	GENERATED_BODY()
+public:
+	void SetRuntimeNode(UMovieGraphNode* InNode)
+	{
+		RuntimeNode = InNode;
+		InNode->GraphNode = this;
+	}
+
+	UMovieGraphNode* GetRuntimeNode() const { return RuntimeNode; }
+
+	void Construct(UMovieGraphNode* InRuntimeNode);
+protected:
+	virtual bool ShouldCreatePin(const UMovieGraphPin* InPin) const;
+
+	void CreatePins(const TArray<UMovieGraphPin*>& InInputPins, const TArray<UMovieGraphPin*>& InOutputPins);
+
+	static FEdGraphPinType GetPinType(const UMovieGraphPin* InPin);
+	
+	UPROPERTY()
+	TObjectPtr<UMovieGraphNode> RuntimeNode;
+};
+
+UCLASS()
+class UMoviePipelineEdGraphNode : public UMoviePipelineEdGraphNodeBase
+{
+	GENERATED_BODY()
+
+public:
+	//~ Begin UEdGraphNode Interface
+	virtual void AllocateDefaultPins() override;
+	virtual FText GetNodeTitle(ENodeTitleType::Type TitleType) const override;
+	virtual FText GetTooltipText() const override;
+	//~ End UEdGraphNode Interface
+
+
+};
diff --git a/Engine/Plugins/MovieScene/MovieRenderPipeline/Source/MovieRenderPipelineEditor/Private/Graph/MovieEdGraphOutputNode.cpp b/Engine/Plugins/MovieScene/MovieRenderPipeline/Source/MovieRenderPipelineEditor/Private/Graph/MovieEdGraphOutputNode.cpp
new file mode 100644
--- /dev/null
+++ b/Engine/Plugins/MovieScene/MovieRenderPipeline/Source/MovieRenderPipelineEditor/Private/Graph/MovieEdGraphOutputNode.cpp
@@ -0,0 +1,19 @@
+// Copyright Epic Games, Inc. All Rights Reserved.
+
+#include "Graph/MovieEdGraphOutputNode.h"
+
+#define LOCTEXT_NAMESPACE "MoviePipelineGraph"
+
+FText UMoviePipelineEdGraphNodeOutput::GetNodeTitle(ENodeTitleType::Type TitleType) const
+{
+	return FText::FromName(TEXT("Output"));
+}
+
+void UMoviePipelineEdGraphNodeOutput::AllocateDefaultPins()
+{
+	if (RuntimeNode)
+	{
+		CreatePins(RuntimeNode->GetInputPins(), /*InOutputPins=*/{});
+	}
+}
+#undef LOCTEXT_NAMESPACE
\ No newline at end of file
diff --git a/Engine/Plugins/MovieScene/MovieRenderPipeline/Source/MovieRenderPipelineEditor/Private/Graph/MovieEdGraphOutputNode.h b/Engine/Plugins/MovieScene/MovieRenderPipeline/Source/MovieRenderPipelineEditor/Private/Graph/MovieEdGraphOutputNode.h
new file mode 100644
--- /dev/null
+++ b/Engine/Plugins/MovieScene/MovieRenderPipeline/Source/MovieRenderPipelineEditor/Private/Graph/MovieEdGraphOutputNode.h
@@ -0,0 +1,20 @@
+// Copyright Epic Games, Inc. All Rights Reserved.
+
+#pragma once
+#include "Graph/MovieEdGraphNode.h"
+
+#include "MovieEdGraphOutputNode.generated.h"
+
+UCLASS()
+class UMoviePipelineEdGraphNodeOutput : public UMoviePipelineEdGraphNodeBase
+{
+	GENERATED_BODY()
+
+public:
+	// ~Begin UEdGraphNode interface
+	virtual FText GetNodeTitle(ENodeTitleType::Type TitleType) const override;
+	virtual void AllocateDefaultPins() override;
+	virtual bool CanUserDeleteNode() const override { return false; }
+	virtual bool CanDuplicateNode() const override { return false; }
+	// ~End UEdGraphNode interface
+};
\ No newline at end of file
diff --git a/Engine/Plugins/MovieScene/MovieRenderPipeline/Source/MovieRenderPipelineEditor/Private/Graph/MovieGraphAssetEditor.cpp b/Engine/Plugins/MovieScene/MovieRenderPipeline/Source/MovieRenderPipelineEditor/Private/Graph/MovieGraphAssetEditor.cpp
new file mode 100644
--- /dev/null
+++ b/Engine/Plugins/MovieScene/MovieRenderPipeline/Source/MovieRenderPipelineEditor/Private/Graph/MovieGraphAssetEditor.cpp
@@ -0,0 +1,28 @@
+// Copyright Epic Games, Inc. All Rights Reserved.
+
+#include "MovieGraphAssetEditor.h"
+#include "MovieGraphAssetToolkit.h"
+
+#include UE_INLINE_GENERATED_CPP_BY_NAME(MovieGraphAssetEditor)
+
+// Temporary cvar to enable/disable upgrading to a graph-based configuration
+static TAutoConsoleVariable<bool> CVarMoviePipelineEnableRenderGraph(
+	TEXT("MoviePipeline.EnableRenderGraph"),
+	false,
+	TEXT("Determines if the Render Graph feature is enabled in the UI. This is a highly experimental feature and is not ready for use.")
+);
+
+TSharedPtr<FBaseAssetToolkit> UMovieGraphAssetEditor::CreateToolkit()
+{
+	return MakeShared<FMovieGraphAssetToolkit>(this);
+}
+
+void UMovieGraphAssetEditor::GetObjectsToEdit(TArray<UObject*>& InObjectsToEdit)
+{
+	InObjectsToEdit.Add(ObjectToEdit);
+}
+
+void UMovieGraphAssetEditor::SetObjectToEdit(UObject* InObject)
+{
+	ObjectToEdit = InObject;
+}
diff --git a/Engine/Plugins/MovieScene/MovieRenderPipeline/Source/MovieRenderPipelineEditor/Private/Graph/MovieGraphAssetEditor.h b/Engine/Plugins/MovieScene/MovieRenderPipeline/Source/MovieRenderPipelineEditor/Private/Graph/MovieGraphAssetEditor.h
new file mode 100644
--- /dev/null
+++ b/Engine/Plugins/MovieScene/MovieRenderPipeline/Source/MovieRenderPipelineEditor/Private/Graph/MovieGraphAssetEditor.h
@@ -0,0 +1,25 @@
+﻿// Copyright Epic Games, Inc. All Rights Reserved.
+
+#pragma once
+
+#include "Tools/UAssetEditor.h"
+
+#include "MovieGraphAssetEditor.generated.h"
+
+UCLASS()
+class UMovieGraphAssetEditor : public UAssetEditor
+{
+	GENERATED_BODY()
+
+public:
+	//~ Begin UAssetEditor overrides
+	virtual TSharedPtr<FBaseAssetToolkit> CreateToolkit() override;
+	virtual void GetObjectsToEdit(TArray<UObject*>& InObjectsToEdit) override;
+	//~ End UAssetEditor overrides
+	
+	void SetObjectToEdit(UObject* InObject);
+
+private:
+	UPROPERTY(Transient)
+	TObjectPtr<UObject> ObjectToEdit;
+};
\ No newline at end of file
diff --git a/Engine/Plugins/MovieScene/MovieRenderPipeline/Source/MovieRenderPipelineEditor/Private/Graph/MovieGraphAssetToolkit.cpp b/Engine/Plugins/MovieScene/MovieRenderPipeline/Source/MovieRenderPipelineEditor/Private/Graph/MovieGraphAssetToolkit.cpp
new file mode 100644
--- /dev/null
+++ b/Engine/Plugins/MovieScene/MovieRenderPipeline/Source/MovieRenderPipelineEditor/Private/Graph/MovieGraphAssetToolkit.cpp
@@ -0,0 +1,89 @@
+// Copyright Epic Games, Inc. All Rights Reserved.
+
+#include "MovieGraphAssetToolkit.h"
+
+#include "Framework/MultiBox/MultiBoxBuilder.h"
+#include "Styling/CoreStyle.h"
+#include "ToolMenus.h"
+#include "Widgets/Docking/SDockTab.h"
+#include "Widgets/Images/SImage.h"
+#include "Widgets/Graph/SMoviePipelineGraphPanel.h"
+#include "Widgets/SOverlay.h"
+
+#define LOCTEXT_NAMESPACE "MovieGraphAssetToolkit"
+
+const FName FMovieGraphAssetToolkit::ContentTabId(TEXT("MovieGraphAssetToolkit"));
+
+FMovieGraphAssetToolkit::FMovieGraphAssetToolkit(UAssetEditor* InOwningAssetEditor)
+	: FBaseAssetToolkit(InOwningAssetEditor)
+{
+	// Note: Changes to the layout should include a increment to the layout's ID, i.e.
+	// MoviePipelineRenderGraphEditor[X] -> MoviePipelineRenderGraphEditor[X+1]. Otherwise, layouts may be messed up
+	// without a full reset to layout defaults inside the editor.
+	const FName LayoutString = FName("MoviePipelineRenderGraphEditor1");
+
+	// Override the Default Layout provided by FBaseAssetToolkit to hide the viewport and details panel tabs.
+	StandaloneDefaultLayout = FTabManager::NewLayout(FName(LayoutString))
+		->AddArea
+		(
+			FTabManager::NewPrimaryArea()->SetOrientation(Orient_Vertical)
+			->Split
+			(
+				FTabManager::NewStack()
+				->SetSizeCoefficient(1.f)
+				->SetHideTabWell(true)
+				->AddTab(ContentTabId, ETabState::OpenedTab)
+			)
+		);
+}
+
+void FMovieGraphAssetToolkit::RegisterTabSpawners(const TSharedRef<FTabManager>& InTabManager)
+{
+	WorkspaceMenuCategory = InTabManager->AddLocalWorkspaceMenuCategory(
+		LOCTEXT("WorkspaceMenu_MoviePipelineRenderGraphAssetToolkit", "Movie Pipeline Render Graph Asset Editor"));
+
+	FBaseAssetToolkit::RegisterTabSpawners(InTabManager);
+
+	InTabManager->RegisterTabSpawner(ContentTabId, FOnSpawnTab::CreateSP(this, &FMovieGraphAssetToolkit::SpawnTab_RenderGraphEditor))
+		.SetDisplayName(LOCTEXT("RenderGraphTab", "Render Graph"))
+		.SetGroup(WorkspaceMenuCategory.ToSharedRef());
+}
+
+void FMovieGraphAssetToolkit::UnregisterTabSpawners(const TSharedRef<FTabManager>& InTabManager)
+{
+	InTabManager->UnregisterTabSpawner(ContentTabId);
+}
+
+TSharedRef<SDockTab> FMovieGraphAssetToolkit::SpawnTab_RenderGraphEditor(const FSpawnTabArgs& Args)
+{
+	// Tabs and Windows have different backgrounds, so we copy the window style to make the 'void' area look similar.
+	const FSlateBrush* BackgroundBrush = &FCoreStyle::Get().GetWidgetStyle<FWindowStyle>("Window").ChildBackgroundBrush;
+
+	TSharedRef<SDockTab> NewDockTab = SNew(SDockTab)
+		.TabColorScale(GetTabColorScale())
+		[
+			SNew(SOverlay)
+			+ SOverlay::Slot()
+			[
+				SNew(SImage)
+				.Image(BackgroundBrush)
+			]
+			+ SOverlay::Slot()
+			[
+				SNew(SVerticalBox)
+				+ SVerticalBox::Slot()
+				[
+					SNew(SMoviePipelineGraphPanel)
+				]
+			]
+		];
+
+	return NewDockTab;
+}
+
+FName FMovieGraphAssetToolkit::GetToolkitFName() const
+{
+	return FName("MovieGraphAssetToolkit");
+}
+
+#undef LOCTEXT_NAMESPACE
diff --git a/Engine/Plugins/MovieScene/MovieRenderPipeline/Source/MovieRenderPipelineEditor/Private/Graph/MovieGraphAssetToolkit.h b/Engine/Plugins/MovieScene/MovieRenderPipeline/Source/MovieRenderPipelineEditor/Private/Graph/MovieGraphAssetToolkit.h
new file mode 100644
--- /dev/null
+++ b/Engine/Plugins/MovieScene/MovieRenderPipeline/Source/MovieRenderPipelineEditor/Private/Graph/MovieGraphAssetToolkit.h
@@ -0,0 +1,25 @@
+﻿// Copyright Epic Games, Inc. All Rights Reserved.
+
+#pragma once
+
+#include "CoreMinimal.h"
+#include "Tools/BaseAssetToolkit.h"
+
+class FMovieGraphAssetToolkit :  public FBaseAssetToolkit
+{
+public:
+	FMovieGraphAssetToolkit(UAssetEditor* InOwningAssetEditor);
+	virtual ~FMovieGraphAssetToolkit() override {}
+
+	//~ Begin FBaseAssetToolkit interface
+	virtual void RegisterTabSpawners(const TSharedRef<FTabManager>& InTabManager) override;
+	virtual void UnregisterTabSpawners(const TSharedRef<FTabManager>& InTabManager) override;
+	virtual FName GetToolkitFName() const override;
+	//~ End FBaseAssetToolkit interface
+
+private:
+	TSharedRef<SDockTab> SpawnTab_RenderGraphEditor(const FSpawnTabArgs& Args);
+
+private:
+	static const FName ContentTabId;
+};
diff --git a/Engine/Plugins/MovieScene/MovieRenderPipeline/Source/MovieRenderPipelineEditor/Private/Graph/MovieGraphSchema.cpp b/Engine/Plugins/MovieScene/MovieRenderPipeline/Source/MovieRenderPipelineEditor/Private/Graph/MovieGraphSchema.cpp
new file mode 100644
--- /dev/null
+++ b/Engine/Plugins/MovieScene/MovieRenderPipeline/Source/MovieRenderPipelineEditor/Private/Graph/MovieGraphSchema.cpp
@@ -0,0 +1,255 @@
+// Copyright Epic Games, Inc. All Rights Reserved.
+
+#include "MovieGraphSchema.h"
+#include "Graph/MovieGraphConfig.h"
+#include "Graph/MovieGraphNode.h"
+#include "Graph/MovieEdGraph.h"
+#include "Graph/MovieEdGraphNode.h"
+
+TArray<UClass*> UMovieGraphSchema::MoviePipelineNodeClasses;
+
+#define LOCTEXT_NAMESPACE "MoviePipelineGraphSchema"
+
+void UMovieGraphSchema::CreateDefaultNodesForGraph(UEdGraph& Graph) const
+{
+	/*UMovieGraphConfig* RuntimeGraph = CastChecked<UMoviePipelineEdGraph>(&Graph)->GetPipelineGraph();
+	const FScopedTransaction Transaction(LOCTEXT("GraphEditor_NewNode", "Create Pipeline Graph Node."));
+	RuntimeGraph->Modify();
+	const bool bSelectNewNode = false;
+
+	// Input Node
+	{
+		UMovieGraphNode* RuntimeNode = RuntimeGraph->ConstructRuntimeNode<UMovieGraphInputNode>();
+
+		// Now create the editor graph node
+		FGraphNodeCreator<UMoviePipelineEdGraphNodeInput> NodeCreator(Graph);
+		UMoviePipelineEdGraphNodeBase* GraphNode = NodeCreator.CreateUserInvokedNode(bSelectNewNode);
+		GraphNode->SetRuntimeNode(RuntimeNode);
+		NodeCreator.Finalize();
+	}
+
+	// Output Node
+	{
+		UMovieGraphNode* RuntimeNode = RuntimeGraph->ConstructRuntimeNode<UMovieGraphOutputNode>();
+
+		// Now create the editor graph node
+		FGraphNodeCreator<UMoviePipelineEdGraphNodeOutput> NodeCreator(Graph);
+		UMoviePipelineEdGraphNodeBase* GraphNode = NodeCreator.CreateUserInvokedNode(bSelectNewNode);
+		GraphNode->SetRuntimeNode(RuntimeNode);
+		NodeCreator.Finalize();
+	}*/
+}
+
+void UMovieGraphSchema::InitMoviePipelineNodeClasses()
+{
+	if (MoviePipelineNodeClasses.Num() > 0)
+	{
+		return;
+	}
+
+	for (TObjectIterator<UClass> It; It; ++It)
+	{
+		if (It->IsChildOf(UMovieGraphNode::StaticClass())
+			&& !It->HasAnyClassFlags(CLASS_Abstract))
+		{
+			MoviePipelineNodeClasses.Add(*It);
+		}
+	}
+
+	MoviePipelineNodeClasses.Sort();
+}
+
+void UMovieGraphSchema::GetGraphContextActions(FGraphContextMenuBuilder& ContextMenuBuilder) const
+{
+	InitMoviePipelineNodeClasses();
+
+	/* {
+		FText Category = NSLOCTEXT("MoveiPipeline", "Category", "Cat1");
+		FText MenuDesc = NSLOCTEXT("MoveiPipeline", "MenuDesc", "New Render Layer");
+		FText Description = NSLOCTEXT("MoveiPipeline", "Description", "Description");
+
+		TSharedPtr<FMovieGraphSchemaAction_NewNativeElement> NewAction(new FMovieGraphSchemaAction_NewNativeElement(Category, MenuDesc, Description, 0));
+		NewAction->NodeClass = UMovieGraphRenderLayerNode::StaticClass();
+		ContextMenuBuilder.AddAction(NewAction);
+	}
+
+	{
+		FText Category = NSLOCTEXT("MoveiPipeline", "Category", "Cat1");
+		FText MenuDesc = NSLOCTEXT("MoveiPipeline", "MenuDesc", "New Collection");
+		FText Description = NSLOCTEXT("MoveiPipeline", "Description", "Description");
+
+		NewAction->NodeClass = UMovieGraphCollectionNode::StaticClass();
+		ContextMenuBuilder.AddAction(NewAction);
+	}*/
+
+	for (UClass* PipelineNodeClass : MoviePipelineNodeClasses)
+	{
+		const UMovieGraphNode* PipelineNode = PipelineNodeClass->GetDefaultObject<UMovieGraphNode>();
+		if (!ContextMenuBuilder.FromPin || ContextMenuBuilder.FromPin->Direction == EGPD_Input)
+		{
+			const FText Name = PipelineNode->GetMenuDescription();
+			const FText Category = PipelineNode->GetMenuCategory();
+			const FText Tooltip = NSLOCTEXT("MoveiPipeline", "Description", "Placeholder Tooltip");
+
+			TSharedPtr<FMovieGraphSchemaAction_NewNativeElement> NewAction(new FMovieGraphSchemaAction_NewNativeElement(Category, Name, Tooltip, 0));
+			NewAction->NodeClass = PipelineNodeClass;
+
+			ContextMenuBuilder.AddAction(NewAction);
+		}
+	}
+}
+
+const FPinConnectionResponse UMovieGraphSchema::CanCreateConnection(const UEdGraphPin* PinA, const UEdGraphPin* PinB) const
+{
+	// No Circular Connections
+	if (PinA->GetOwningNode() == PinB->GetOwningNode())
+	{
+		return FPinConnectionResponse(CONNECT_RESPONSE_DISALLOW, NSLOCTEXT("MoviePipeline", "CircularPinError", "No Circular Connections!"));
+	}
+	// Make sure the pins are not on the same node
+	return FPinConnectionResponse(CONNECT_RESPONSE_MAKE, NSLOCTEXT("MoviePipeline", "PinConnect", "Connect nodes"));
+}
+
+bool UMovieGraphSchema::TryCreateConnection(UEdGraphPin* InA, UEdGraphPin* InB) const
+{
+	// See if the native UEdGraph connection goes through
+	const bool bModified = Super::TryCreateConnection(InA, InB);
+
+	// If it does, try to propagate the change to our runtime graph
+	if (bModified)
+	{
+		check(InA && InB);
+		const UEdGraphPin* A = (InA->Direction == EGPD_Output) ? InA : InB;
+		const UEdGraphPin* B = (InA->Direction == EGPD_Input) ? InA : InB;
+		check(A->Direction == EGPD_Output && B->Direction == EGPD_Input);
+
+		UMoviePipelineEdGraphNodeBase* EdGraphNodeA = CastChecked<UMoviePipelineEdGraphNodeBase>(A->GetOwningNode());
+		UMoviePipelineEdGraphNodeBase* EdGraphNodeB = CastChecked<UMoviePipelineEdGraphNodeBase>(B->GetOwningNode());
+
+		UMovieGraphNode* RuntimeNodeA = EdGraphNodeA->GetRuntimeNode();
+		UMovieGraphNode* RuntimeNodeB = EdGraphNodeB->GetRuntimeNode();
+		check(RuntimeNodeA && RuntimeNodeB);
+
+		UMovieGraphConfig* RuntimeGraph = RuntimeNodeA->GetGraph();
+		check(RuntimeGraph);
+
+		const bool bReconstructNodeB = RuntimeGraph->AddLabeledEdge(RuntimeNodeA, A->PinName, RuntimeNodeB, B->PinName);
+		//if (bReconstructNodeB)
+		//{
+		//	RuntimeNodeB->ReconstructNode();
+		//}
+	}
+
+	return bModified;
+}
+
+void UMovieGraphSchema::BreakPinLinks(UEdGraphPin& TargetPin, bool bSendsNodeNotification) const
+{
+	const FScopedTransaction Transaction(LOCTEXT("MoviePipelineGraphEditor_BreakPinLinks", "Break Pin Links"));
+	Super::BreakPinLinks(TargetPin, bSendsNodeNotification);
+
+	UEdGraphNode* GraphNode = TargetPin.GetOwningNode();
+	UMoviePipelineEdGraphNodeBase* MoviePipelineEdGraphNode = CastChecked<UMoviePipelineEdGraphNodeBase>(GraphNode);
+
+	UMovieGraphNode* RuntimeNode = MoviePipelineEdGraphNode->GetRuntimeNode();
+	check(RuntimeNode);
+
+	UMovieGraphConfig* RuntimeGraph = RuntimeNode->GetGraph();
+	check(RuntimeGraph);
+
+	if (TargetPin.Direction == EEdGraphPinDirection::EGPD_Input)
+	{
+		RuntimeGraph->RemoveInboundEdges(RuntimeNode, TargetPin.PinName);
+	}
+	else if (TargetPin.Direction == EEdGraphPinDirection::EGPD_Output)
+	{
+		RuntimeGraph->RemoveOutboundEdges(RuntimeNode, TargetPin.PinName);
+	}
+}
+
+void UMovieGraphSchema::BreakSinglePinLink(UEdGraphPin* SourcePin, UEdGraphPin* TargetPin) const
+{
+	const FScopedTransaction Transaction(LOCTEXT("MoviePipelineGraphEditor_BreakSinglePinLinks", "Break Single Pin Link"));
+	Super::BreakSinglePinLink(SourcePin, TargetPin);
+
+	UEdGraphNode* SourceGraphNode = SourcePin->GetOwningNode();
+	UEdGraphNode* TargetGraphNode = TargetPin->GetOwningNode();
+
+	UMoviePipelineEdGraphNodeBase* SourcePipelineGraphNode = CastChecked<UMoviePipelineEdGraphNodeBase>(SourceGraphNode);
+	UMoviePipelineEdGraphNodeBase* TargetPipelineGraphNode = CastChecked<UMoviePipelineEdGraphNodeBase>(TargetGraphNode);
+
+	UMovieGraphNode* SourceRuntimeNode = SourcePipelineGraphNode->GetRuntimeNode();
+	UMovieGraphNode* TargetRuntimeNode = TargetPipelineGraphNode->GetRuntimeNode();
+	check(SourceRuntimeNode && TargetRuntimeNode);
+
+	UMovieGraphConfig* RuntimeGraph = SourceRuntimeNode->GetGraph();
+	check(RuntimeGraph);
+
+	RuntimeGraph->RemoveEdge(SourceRuntimeNode, SourcePin->PinName, TargetRuntimeNode, TargetPin->PinName);
+}
+
+FLinearColor UMovieGraphSchema::GetPinTypeColor(const FEdGraphPinType& PinType) const
+{
+	return FLinearColor::Red;
+	/*const FName& TypeName = PinType.PinCategory;
+	if (TypeName == UEdGraphSchema_K2::PC_Struct)
+	{
+		if (UStruct* Struct = Cast<UStruct>(PinType.PinSubCategoryObject))
+		{
+			if (Struct->IsChildOf(FRigVMExecuteContext::StaticStruct()))
+			{
+				return FLinearColor::White;
+			}
+
+			if (Struct->IsChildOf(RigVMTypeUtils::GetWildCardCPPTypeObject()))
+			{
+				return FLinearColor(FVector3f::OneVector * 0.25f);
+			}
+
+			if (Struct == FRigElementKey::StaticStruct() || Struct == FRigElementKeyCollection::StaticStruct())
+			{
+				return FLinearColor(0.0, 0.6588, 0.9490);
+			}
+
+			if (Struct == FRigElementKey::StaticStruct() || Struct == FRigPose::StaticStruct())
+			{
+				return FLinearColor(0.0, 0.3588, 0.5490);
+			}
+
+			// external types can register their own colors, check if there are any
+			if (IControlRigDeveloperModule* Module = FModuleManager::GetModulePtr<IControlRigDeveloperModule>("ControlRigDeveloper"))
+			{
+				if (const FLinearColor* Color = Module->FindPinTypeColor(Struct))
+				{
+					return *Color;
+				}
+			}
+		}
+	}*/
+
+	//return GetDefault<UEdGraphSchema_K2>()->GetPinTypeColor(PinType);
+}
+
+
+
+UEdGraphNode* FMovieGraphSchemaAction_NewNativeElement::PerformAction(UEdGraph* ParentGraph, UEdGraphPin* FromPin, const FVector2D Location, bool bSelectNewNode)
+{
+	UMovieGraphConfig* RuntimeGraph = CastChecked<UMoviePipelineEdGraph>(ParentGraph)->GetPipelineGraph();
+	const FScopedTransaction Transaction(LOCTEXT("GraphEditor_NewNode", "Create Pipeline Graph Node."));
+	RuntimeGraph->Modify();
+
+	UMovieGraphNode* RuntimeNode = RuntimeGraph->ConstructRuntimeNode<UMovieGraphNode>(NodeClass);
+
+	// Now create the editor graph node
+	FGraphNodeCreator<UMoviePipelineEdGraphNode> NodeCreator(*ParentGraph);
+	UMoviePipelineEdGraphNode* GraphNode = NodeCreator.CreateUserInvokedNode(bSelectNewNode);
+	GraphNode->SetRuntimeNode(RuntimeNode);
+	GraphNode->NodePosX = Location.X;
+	GraphNode->NodePosY = Location.Y;
+
+	// Finalize generates a guid, calls a post-place callback, and allocates default pins if needed
+	NodeCreator.Finalize();
+	return GraphNode;
+}
+
+#undef LOCTEXT_NAMESPACE // "MoviePipelineGraphSchema"
diff --git a/Engine/Plugins/MovieScene/MovieRenderPipeline/Source/MovieRenderPipelineEditor/Private/Graph/MovieGraphSchema.h b/Engine/Plugins/MovieScene/MovieRenderPipeline/Source/MovieRenderPipelineEditor/Private/Graph/MovieGraphSchema.h
new file mode 100644
--- /dev/null
+++ b/Engine/Plugins/MovieScene/MovieRenderPipeline/Source/MovieRenderPipelineEditor/Private/Graph/MovieGraphSchema.h
@@ -0,0 +1,63 @@
+// Copyright Epic Games, Inc. All Rights Reserved.
+
+#pragma once
+#include "CoreTypes.h"
+#include "EdGraph/EdGraphSchema.h"
+#include "MovieGraphSchema.generated.h"
+
+class UMovieGraphNode;
+
+
+UCLASS()
+class UMovieGraphSchema : public UEdGraphSchema
+{
+	GENERATED_BODY()
+	//~ Begin EdGraphSchema Interface
+	virtual void CreateDefaultNodesForGraph(UEdGraph& Graph) const override;
+	virtual void GetGraphContextActions(FGraphContextMenuBuilder& ContextMenuBuilder) const override;
+	//virtual void GetContextMenuActions(class UToolMenu* Menu, class UGraphNodeContextMenuContext* Context) const override;
+	virtual const FPinConnectionResponse CanCreateConnection(const UEdGraphPin* A, const UEdGraphPin* B) const override;
+	virtual bool TryCreateConnection(UEdGraphPin* InA, UEdGraphPin* InB) const override;
+	//virtual const FPinConnectionResponse CanMergeNodes(const UEdGraphNode* A, const UEdGraphNode* B) const override;
+	virtual FLinearColor GetPinTypeColor(const FEdGraphPinType& PinType) const override;
+
+	//virtual int32 GetNodeSelectionCount(const UEdGraph* Graph) const override;
+	//virtual bool IsCacheVisualizationOutOfDate(int32 InVisualizationCacheID) const override;
+	//virtual int32 GetCurrentVisualizationCacheID() const override;
+	//virtual void ForceVisualizationCacheClear() const override;
+	//virtual TSharedPtr<FEdGraphSchemaAction> GetCreateCommentAction() const override;
+	virtual void BreakPinLinks(UEdGraphPin& TargetPin, bool bSendsNodeNotification) const override;
+	virtual void BreakSinglePinLink(UEdGraphPin* SourcePin, UEdGraphPin* TargetPin) const override;
+	//~ End EdGraphSchema Interface
+	static void InitMoviePipelineNodeClasses();
+	
+	
+private:
+	static TArray<UClass*> MoviePipelineNodeClasses;
+	
+};
+
+
+USTRUCT()
+struct FMovieGraphSchemaAction_NewNativeElement : public FEdGraphSchemaAction
+{
+	GENERATED_BODY()
+	
+	// Inherit the base class's constructors
+	using FEdGraphSchemaAction::FEdGraphSchemaAction;
+	
+	// Simple type info
+	static FName StaticGetTypeId()
+	{
+		static FName Type("FMovieGraphSchemaAction_NewNativeElement");
+		return Type;
+	}
+	
+	UPROPERTY()
+	TSubclassOf<UMovieGraphNode> NodeClass;
+
+	//~ Begin FEdGraphSchemaAction Interface
+	virtual FName GetTypeId() const override { return StaticGetTypeId(); }
+	virtual UEdGraphNode* PerformAction(UEdGraph* ParentGraph, UEdGraphPin* FromPin, const FVector2D Location, bool bSelectNewNode = true) override;
+	//~ End FEdGraphSchemaAction Interface
+};
\ No newline at end of file
diff --git a/Engine/Plugins/MovieScene/MovieRenderPipeline/Source/MovieRenderPipelineEditor/Private/Widgets/Graph/SMoviePipelineGraphPanel.cpp b/Engine/Plugins/MovieScene/MovieRenderPipeline/Source/MovieRenderPipelineEditor/Private/Widgets/Graph/SMoviePipelineGraphPanel.cpp
new file mode 100644
--- /dev/null
+++ b/Engine/Plugins/MovieScene/MovieRenderPipeline/Source/MovieRenderPipelineEditor/Private/Widgets/Graph/SMoviePipelineGraphPanel.cpp
@@ -0,0 +1,611 @@
+// Copyright Epic Games, Inc. All Rights Reserved.
+
+// Movie Pipeline Includes
+#include "Widgets/Graph/SMoviePipelineGraphPanel.h"
+#include "Widgets/MoviePipelineWidgetConstants.h"
+#include "Widgets/SMoviePipelineQueueEditor.h"
+#include "Widgets/SMoviePipelineConfigPanel.h"
+#include "MovieRenderPipelineSettings.h"
+#include "MovieRenderPipelineStyle.h"
+#include "MoviePipelineBlueprintLibrary.h"
+#include "Sections/MovieSceneCinematicShotSection.h"
+#include "MoviePipelineQueue.h"
+#include "MoviePipelinePrimaryConfig.h"
+#include "MoviePipelineQueueSubsystem.h"
+
+// Slate Includes
+#include "Widgets/SBoxPanel.h"
+#include "Widgets/Input/SComboBox.h"
+#include "Widgets/Layout/SBorder.h"
+#include "Widgets/Text/STextBlock.h"
+#include "Widgets/Input/SButton.h"
+#include "Widgets/SWindow.h"
+#include "Styling/AppStyle.h"
+#include "Framework/Application/SlateApplication.h"
+#include "Framework/Notifications/NotificationManager.h"
+#include "Widgets/Layout/SSplitter.h"
+#include "Widgets/Layout/SWidgetSwitcher.h"
+#include "Widgets/Notifications/SNotificationList.h"
+#include "Styling/SlateIconFinder.h"
+#include "Widgets/Images/SImage.h"
+#include "Framework/MultiBox/MultiBoxBuilder.h"
+#include "SPrimaryButton.h"
+
+// ContentBrowser Includes
+#include "IContentBrowserSingleton.h"
+#include "ContentBrowserModule.h"
+
+// Misc
+#include "Editor.h"
+#include "PropertyEditorModule.h"
+#include "IDetailsView.h"
+#include "Modules/ModuleManager.h"
+
+// UnrealEd Includes
+#include "ScopedTransaction.h"
+#include "AssetRegistry/AssetData.h"
+#include "AssetRegistry/AssetRegistryModule.h"
+#include "FileHelpers.h"
+#include "AssetToolsModule.h"
+#include "Misc/FileHelper.h"
+
+#include "Graph/MovieGraphConfig.h"
+#include "Graph/MovieEdGraph.h"
+#include "Graph/MovieGraphSchema.h"
+#include "Kismet2/BlueprintEditorUtils.h"
+
+#define LOCTEXT_NAMESPACE "SMoviePipelineGraphPanel"
+
+UE_DISABLE_OPTIMIZATION_SHIP
+void SMoviePipelineGraphPanel::Construct(const FArguments& InArgs)
+{
+	FPropertyEditorModule& PropertyEditorModule = FModuleManager::Get().LoadModuleChecked<FPropertyEditorModule>("PropertyEditor");
+	FDetailsViewArgs DetailsViewArgs;
+	DetailsViewArgs.bAllowSearch = false;
+	DetailsViewArgs.NameAreaSettings = FDetailsViewArgs::HideNameArea;
+	DetailsViewArgs.bHideSelectionTip = true;
+	DetailsViewArgs.ColumnWidth = 0.7f;
+
+	JobDetailsPanelWidget = PropertyEditorModule.CreateDetailView(DetailsViewArgs);
+
+	// Create the child widgets that need to know about our pipeline
+	PipelineQueueEditorWidget = SNew(SMoviePipelineQueueEditor)
+		.OnEditConfigRequested(this, &SMoviePipelineGraphPanel::OnEditJobConfigRequested)
+		.OnPresetChosen(this, &SMoviePipelineGraphPanel::OnJobPresetChosen)
+		.OnJobSelectionChanged(this, &SMoviePipelineGraphPanel::OnSelectionChanged);
+
+
+	{
+		// Automatically select the first job in the queue
+		UMoviePipelineQueueSubsystem* Subsystem = GEditor->GetEditorSubsystem<UMoviePipelineQueueSubsystem>();
+		check(Subsystem);
+
+		TArray<UMoviePipelineExecutorJob*> Jobs;
+		if (Subsystem->GetQueue()->GetJobs().Num() > 0)
+		{
+			Jobs.Add(Subsystem->GetQueue()->GetJobs()[0]);
+		}
+
+		// Go through the UI so it updates the UI selection too and then this will loop back
+		// around to OnSelectionChanged to update ourself.
+		PipelineQueueEditorWidget->SetSelectedJobs(Jobs);
+	}
+
+	SGraphEditor::FGraphEditorEvents InEvents;
+	InEvents.OnSelectionChanged = SGraphEditor::FOnSelectionChanged::CreateSP(this, &SMoviePipelineGraphPanel::OnSelectedNodesChanged);
+	InEvents.OnNodeDoubleClicked = FSingleNodeEvent::CreateSP(this, &SMoviePipelineGraphPanel::OnNodeDoubleClicked);
+	InEvents.OnTextCommitted = FOnNodeTextCommitted::CreateSP(this, &SMoviePipelineGraphPanel::OnNodeTitleCommitted);
+
+	CurrentGraph = NewObject<UMovieGraphConfig>(GetTransientPackage());
+	CurrentGraph->AddToRoot();
+	
+	UMoviePipelineEdGraph* GraphToEdit = Cast<UMoviePipelineEdGraph>(FBlueprintEditorUtils::CreateNewGraph(CurrentGraph, TEXT("MoviePipelineEdGraph"), UMoviePipelineEdGraph::StaticClass(), UMovieGraphSchema::StaticClass()));
+
+	// Probably not ideal.. USoundCue has a CreateGraph() node that does this (#if WITH_EDITOR) but then requires an interface
+	// for the editor half to avoid the circular dependency.
+	CurrentGraph->PipelineEdGraph = GraphToEdit;
+	GraphToEdit->InitFromRuntimeGraph(CurrentGraph);
+
+
+	const UEdGraphSchema* Schema = GraphToEdit->GetSchema();
+	Schema->CreateDefaultNodesForGraph(*GraphToEdit);
+
+	TSharedRef<SWidget> GraphEditor =
+		SNew(SGraphEditor)
+		.IsEditable(true)
+		// .Appearance(this, &FTestEditor::GetGraphAppearance)
+		// .TitleBar(TitleBarWidget)
+		.GraphToEdit(GraphToEdit)
+		.GraphEvents(InEvents);
+
+	ChildSlot
+	[
+		SNew(SSplitter)
+		.Orientation(EOrientation::Orient_Horizontal)
+		+ SSplitter::Slot()
+		.Value(3)
+		[
+			GraphEditor
+		]
+		+ SSplitter::Slot()
+		.Value(1)
+		[
+			SNew(SBorder)
+			.BorderImage(FAppStyle::GetBrush("ToolPanel.GroupBorder"))
+			.Padding(FMargin(1.f, 1.0f))
+			.Content()
+			[
+				JobDetailsPanelWidget.ToSharedRef()
+			]
+		]
+		+SSplitter::Slot()
+		[
+			SNew(SButton)
+			.Text(LOCTEXT("OK", "OK"))
+			.OnClicked(this, &SMoviePipelineGraphPanel::OnDebugButtonClicked)
+		]
+	];
+}
+
+UE_ENABLE_OPTIMIZATION_SHIP
+FReply SMoviePipelineGraphPanel::OnDebugButtonClicked()
+{
+	CurrentGraph->TraversalTest();
+	return FReply::Handled();
+}
+
+void SMoviePipelineGraphPanel::OnSelectedNodesChanged(const TSet<class UObject*>& NewSelection)
+{
+	TArray<UObject*> SelectedObjects;
+	for (UObject* Obj : NewSelection)
+	{
+		SelectedObjects.Add(Obj);
+	}
+
+	JobDetailsPanelWidget->SetObjects(SelectedObjects);
+}
+
+void SMoviePipelineGraphPanel::OnNodeDoubleClicked(class UEdGraphNode* Node)
+{
+}
+
+void SMoviePipelineGraphPanel::OnNodeTitleCommitted(const FText& NewText, ETextCommit::Type CommitInfo, UEdGraphNode* NodeBeingChanged)
+{
+	if (NodeBeingChanged)
+	{
+		const FScopedTransaction Transaction(NSLOCTEXT("K2_RenameNode", "RenameNode", "Rename Node"));
+		NodeBeingChanged->Modify();
+		NodeBeingChanged->OnRenameNode(NewText.ToString());
+	}
+}
+
+FReply SMoviePipelineGraphPanel::OnRenderLocalRequested()
+{
+	UMoviePipelineQueueSubsystem* Subsystem = GEditor->GetEditorSubsystem<UMoviePipelineQueueSubsystem>();
+	check(Subsystem);
+
+	const UMovieRenderPipelineProjectSettings* ProjectSettings = GetDefault<UMovieRenderPipelineProjectSettings>();
+	TSubclassOf<UMoviePipelineExecutorBase> ExecutorClass = ProjectSettings->DefaultLocalExecutor.TryLoadClass<UMoviePipelineExecutorBase>();
+
+	// OnRenderLocalRequested should only get called if IsRenderLocalEnabled() returns true, meaning there's a valid class.
+	check(ExecutorClass != nullptr);
+	Subsystem->RenderQueueWithExecutor(ExecutorClass);
+	return FReply::Handled();
+}
+
+bool SMoviePipelineGraphPanel::IsRenderLocalEnabled() const
+{
+	UMoviePipelineQueueSubsystem* Subsystem = GEditor->GetEditorSubsystem<UMoviePipelineQueueSubsystem>();
+	check(Subsystem);
+
+	const UMovieRenderPipelineProjectSettings* ProjectSettings = GetDefault<UMovieRenderPipelineProjectSettings>();
+	const bool bHasExecutor = ProjectSettings->DefaultLocalExecutor.TryLoadClass<UMoviePipelineExecutorBase>() != nullptr;
+	const bool bNotRendering = !Subsystem->IsRendering();
+	const bool bConfigWindowIsOpen = WeakEditorWindow.IsValid();
+
+	bool bAtLeastOneJobAvailable = false;
+	for (UMoviePipelineExecutorJob* Job : Subsystem->GetQueue()->GetJobs())
+	{
+		if (!Job->IsConsumed() && Job->IsEnabled())
+		{
+			bAtLeastOneJobAvailable = true;
+			break;
+		}
+	}
+
+	const bool bWorldIsActive = GEditor->IsPlaySessionInProgress();
+	return bHasExecutor && bNotRendering && bAtLeastOneJobAvailable && !bWorldIsActive && !bConfigWindowIsOpen;
+}
+
+FReply SMoviePipelineGraphPanel::OnRenderRemoteRequested()
+{
+	UMoviePipelineQueueSubsystem* Subsystem = GEditor->GetEditorSubsystem<UMoviePipelineQueueSubsystem>();
+	check(Subsystem);
+
+	const UMovieRenderPipelineProjectSettings* ProjectSettings = GetDefault<UMovieRenderPipelineProjectSettings>();
+	TSubclassOf<UMoviePipelineExecutorBase> ExecutorClass = ProjectSettings->DefaultRemoteExecutor.TryLoadClass<UMoviePipelineExecutorBase>();
+
+	// OnRenderRemoteRequested should only get called if IsRenderRemoteEnabled() returns true, meaning there's a valid class.
+	check(ExecutorClass != nullptr);
+
+	Subsystem->RenderQueueWithExecutor(ExecutorClass);
+	return FReply::Handled();
+}
+
+bool SMoviePipelineGraphPanel::IsRenderRemoteEnabled() const
+{
+	UMoviePipelineQueueSubsystem* Subsystem = GEditor->GetEditorSubsystem<UMoviePipelineQueueSubsystem>();
+	check(Subsystem);
+
+	const UMovieRenderPipelineProjectSettings* ProjectSettings = GetDefault<UMovieRenderPipelineProjectSettings>();
+	const bool bHasExecutor = ProjectSettings->DefaultRemoteExecutor.TryLoadClass<UMoviePipelineExecutorBase>() != nullptr;
+	const bool bNotRendering = !Subsystem->IsRendering();
+	const bool bConfigWindowIsOpen = WeakEditorWindow.IsValid();
+
+	bool bAtLeastOneJobAvailable = false;
+	for (UMoviePipelineExecutorJob* Job : Subsystem->GetQueue()->GetJobs())
+	{
+		if (!Job->IsConsumed() && Job->IsEnabled())
+		{
+			bAtLeastOneJobAvailable = true;
+			break;
+		}
+	}
+
+	return bHasExecutor && bNotRendering && bAtLeastOneJobAvailable && !bConfigWindowIsOpen;
+}
+
+void SMoviePipelineGraphPanel::OnJobPresetChosen(TWeakObjectPtr<UMoviePipelineExecutorJob> InJob, TWeakObjectPtr<UMoviePipelineExecutorShot> InShot)
+{
+	UMovieRenderPipelineProjectSettings* ProjectSettings = GetMutableDefault<UMovieRenderPipelineProjectSettings>();
+	if (!InShot.IsValid())
+	{
+		// Store the preset so the next job they make will use it.
+		ProjectSettings->LastPresetOrigin = InJob->GetPresetOrigin();
+	}
+	ProjectSettings->SaveConfig();
+}
+
+void SMoviePipelineGraphPanel::OnEditJobConfigRequested(TWeakObjectPtr<UMoviePipelineExecutorJob> InJob, TWeakObjectPtr<UMoviePipelineExecutorShot> InShot)
+{
+	// Only allow one editor open at once for now.
+	if (WeakEditorWindow.IsValid())
+	{
+		FWidgetPath ExistingWindowPath;
+		if (FSlateApplication::Get().FindPathToWidget(WeakEditorWindow.Pin().ToSharedRef(), ExistingWindowPath, EVisibility::All))
+		{
+			WeakEditorWindow.Pin()->BringToFront();
+			FSlateApplication::Get().SetAllUserFocus(ExistingWindowPath, EFocusCause::SetDirectly);
+		}
+
+		return;
+	}
+
+	TSubclassOf<UMoviePipelineConfigBase> ConfigType;
+	UMoviePipelineConfigBase* BasePreset = nullptr;
+	UMoviePipelineConfigBase* BaseConfig = nullptr;
+	if (InShot.IsValid())
+	{
+		ConfigType = UMoviePipelineShotConfig::StaticClass();
+		BasePreset = InShot->GetShotOverridePresetOrigin();
+		BaseConfig = InShot->GetShotOverrideConfiguration();
+	}
+	else
+	{
+		ConfigType = UMoviePipelinePrimaryConfig::StaticClass();
+		BasePreset = InJob->GetPresetOrigin();
+		BaseConfig = InJob->GetConfiguration();
+	}
+
+	TSharedRef<SWindow> EditorWindow =
+		SNew(SWindow)
+		.ClientSize(FVector2D(700, 600));
+
+	TSharedRef<SMoviePipelineConfigPanel> ConfigEditorPanel =
+		SNew(SMoviePipelineConfigPanel, ConfigType)
+		.Job(InJob)
+		.Shot(InShot)
+		.OnConfigurationModified(this, &SMoviePipelineGraphPanel::OnConfigUpdatedForJob)
+		.OnConfigurationSetToPreset(this, &SMoviePipelineGraphPanel::OnConfigUpdatedForJobToPreset)
+		.BasePreset(BasePreset)
+		.BaseConfig(BaseConfig);
+
+	EditorWindow->SetContent(ConfigEditorPanel);
+
+
+	TSharedPtr<SWindow> ParentWindow = FSlateApplication::Get().FindWidgetWindow(AsShared());
+	if (ParentWindow.IsValid())
+	{
+		FSlateApplication::Get().AddWindowAsNativeChild(EditorWindow, ParentWindow.ToSharedRef());
+	}
+
+	WeakEditorWindow = EditorWindow;
+}
+
+void SMoviePipelineGraphPanel::OnConfigWindowClosed()
+{
+	if (WeakEditorWindow.IsValid())
+	{
+		WeakEditorWindow.Pin()->RequestDestroyWindow();
+	}
+}
+
+void SMoviePipelineGraphPanel::OnConfigUpdatedForJob(TWeakObjectPtr<UMoviePipelineExecutorJob> InJob, TWeakObjectPtr<UMoviePipelineExecutorShot> InShot, UMoviePipelineConfigBase* InConfig)
+{
+	if (InJob.IsValid())
+	{
+		if (InShot.IsValid())
+		{
+			if (UMoviePipelineShotConfig* ShotConfig = Cast<UMoviePipelineShotConfig>(InConfig))
+			{
+				InShot->SetShotOverrideConfiguration(ShotConfig);
+			}
+		}
+		else
+		{
+			if (UMoviePipelinePrimaryConfig* PrimaryConfig = Cast<UMoviePipelinePrimaryConfig>(InConfig))
+			{
+				InJob->SetConfiguration(PrimaryConfig);
+			}
+		}
+	}
+
+	OnConfigWindowClosed();
+}
+
+void SMoviePipelineGraphPanel::OnConfigUpdatedForJobToPreset(TWeakObjectPtr<UMoviePipelineExecutorJob> InJob, TWeakObjectPtr<UMoviePipelineExecutorShot> InShot, UMoviePipelineConfigBase* InConfig)
+{
+	if (InJob.IsValid())
+	{
+		if (InShot.IsValid())
+		{
+			if (UMoviePipelineShotConfig* ShotConfig = Cast<UMoviePipelineShotConfig>(InConfig))
+			{
+				InShot->SetShotOverridePresetOrigin(ShotConfig);
+			}
+		}
+		else
+		{
+			if (UMoviePipelinePrimaryConfig* PrimaryConfig = Cast<UMoviePipelinePrimaryConfig>(InConfig))
+			{
+				InJob->SetPresetOrigin(PrimaryConfig);
+			}
+		}
+	}
+
+	// Store the preset they used as the last set one
+	OnJobPresetChosen(InJob, InShot);
+
+	OnConfigWindowClosed();
+}
+
+void SMoviePipelineGraphPanel::OnSelectionChanged(const TArray<UMoviePipelineExecutorJob*>& InSelectedJobs)
+{
+	TArray<UObject*> Jobs;
+	for (UMoviePipelineExecutorJob* Job : InSelectedJobs)
+	{
+		Jobs.Add(Job);
+	}
+	
+	JobDetailsPanelWidget->SetObjects(Jobs);
+	NumSelectedJobs = InSelectedJobs.Num();
+}
+
+int32 SMoviePipelineGraphPanel::GetDetailsViewWidgetIndex() const
+{
+	return NumSelectedJobs == 0;
+}
+
+bool SMoviePipelineGraphPanel::IsDetailsViewEnabled() const
+{
+	TArray<TWeakObjectPtr<UObject>> OutObjects= JobDetailsPanelWidget->GetSelectedObjects();
+
+	bool bAllEnabled = true;
+	for (TWeakObjectPtr<UObject> Object : OutObjects)
+	{
+		const UMoviePipelineExecutorJob* Job = Cast<UMoviePipelineExecutorJob>(Object);
+		if (Job && Job->IsConsumed())
+		{
+			bAllEnabled = false;
+			break;
+		}
+	}
+
+	return bAllEnabled;
+}
+
+TSharedRef<SWidget> SMoviePipelineGraphPanel::OnGenerateSavedQueuesMenu()
+{
+	FMenuBuilder MenuBuilder(true, nullptr);
+
+	IContentBrowserSingleton& ContentBrowser = FModuleManager::LoadModuleChecked<FContentBrowserModule>("ContentBrowser").Get();
+
+	MenuBuilder.AddMenuEntry(
+		LOCTEXT("SaveAsQueue_Text", "Save As Asset"),
+		LOCTEXT("SaveAsQueue_Tip", "Save the current configuration as a new preset that can be shared between multiple jobs, or imported later as the base of a new configuration."),
+		FSlateIcon(FAppStyle::Get().GetStyleSetName(), "AssetEditor.SaveAsset"),
+		FUIAction(FExecuteAction::CreateSP(this, &SMoviePipelineGraphPanel::OnSaveAsAsset))
+	);
+
+	FAssetPickerConfig AssetPickerConfig;
+	{
+		AssetPickerConfig.SelectionMode = ESelectionMode::Single;
+		AssetPickerConfig.InitialAssetViewType = EAssetViewType::Column;
+		AssetPickerConfig.bFocusSearchBoxWhenOpened = true;
+		AssetPickerConfig.bAllowNullSelection = false;
+		AssetPickerConfig.bShowBottomToolbar = true;
+		AssetPickerConfig.bAutohideSearchBar = false;
+		AssetPickerConfig.bAllowDragging = false;
+		AssetPickerConfig.bCanShowClasses = false;
+		AssetPickerConfig.bShowPathInColumnView = true;
+		AssetPickerConfig.bShowTypeInColumnView = false;
+		AssetPickerConfig.bSortByPathInColumnView = false;
+
+		AssetPickerConfig.AssetShowWarningText = LOCTEXT("NoQueueAssets_Warning", "No Queues Found");
+		AssetPickerConfig.Filter.ClassPaths.Add(UMoviePipelineQueue::StaticClass()->GetClassPathName());
+		AssetPickerConfig.Filter.bRecursiveClasses = true;
+		AssetPickerConfig.OnAssetSelected = FOnAssetSelected::CreateSP(this, &SMoviePipelineGraphPanel::OnImportSavedQueueAssest);
+	}
+
+	MenuBuilder.BeginSection(NAME_None, LOCTEXT("LoadQueue_MenuSection", "Load Queue"));
+	{
+		TSharedRef<SWidget> PresetPicker = SNew(SBox)
+			.MinDesiredWidth(400.f)
+			.MinDesiredHeight(400.f)
+			[
+				ContentBrowser.CreateAssetPicker(AssetPickerConfig)
+			];
+
+		MenuBuilder.AddWidget(PresetPicker, FText(), true, false);
+	}
+	MenuBuilder.EndSection();
+
+	return MenuBuilder.MakeWidget();
+}
+
+bool SMoviePipelineGraphPanel::OpenSaveDialog(const FString& InDefaultPath, const FString& InNewNameSuggestion, FString& OutPackageName)
+{
+	FSaveAssetDialogConfig SaveAssetDialogConfig;
+	{
+		SaveAssetDialogConfig.DefaultPath = InDefaultPath;
+		SaveAssetDialogConfig.DefaultAssetName = InNewNameSuggestion;
+		SaveAssetDialogConfig.AssetClassNames.Add(UMoviePipelineQueue::StaticClass()->GetClassPathName());
+		SaveAssetDialogConfig.ExistingAssetPolicy = ESaveAssetDialogExistingAssetPolicy::AllowButWarn;
+		SaveAssetDialogConfig.DialogTitleOverride = LOCTEXT("SaveQueueAssetDialogTitle", "Save Queue Asset");
+	}
+
+	FContentBrowserModule& ContentBrowserModule = FModuleManager::LoadModuleChecked<FContentBrowserModule>("ContentBrowser");
+	FString SaveObjectPath = ContentBrowserModule.Get().CreateModalSaveAssetDialog(SaveAssetDialogConfig);
+
+	if (!SaveObjectPath.IsEmpty())
+	{
+		OutPackageName = FPackageName::ObjectPathToPackageName(SaveObjectPath);
+		return true;
+	}
+
+	return false;
+}
+
+bool SMoviePipelineGraphPanel::GetSavePresetPackageName(const FString& InExistingName, FString& OutName)
+{
+	UMovieRenderPipelineProjectSettings* ConfigSettings = GetMutableDefault<UMovieRenderPipelineProjectSettings>();
+
+	// determine default package path
+	const FString DefaultSaveDirectory = ConfigSettings->PresetSaveDir.Path;
+
+	FString DialogStartPath;
+	FPackageName::TryConvertFilenameToLongPackageName(DefaultSaveDirectory, DialogStartPath);
+	if (DialogStartPath.IsEmpty())
+	{
+		DialogStartPath = TEXT("/Game");
+	}
+
+	// determine default asset name
+	FString DefaultName = InExistingName;
+
+	FString UniquePackageName;
+	FString UniqueAssetName;
+
+	FAssetToolsModule& AssetToolsModule = FModuleManager::LoadModuleChecked<FAssetToolsModule>("AssetTools");
+	AssetToolsModule.Get().CreateUniqueAssetName(DialogStartPath / DefaultName, TEXT(""), UniquePackageName, UniqueAssetName);
+
+	FString DialogStartName = FPaths::GetCleanFilename(UniqueAssetName);
+
+	FString UserPackageName;
+	FString NewPackageName;
+
+	// get destination for asset
+	bool bFilenameValid = false;
+	while (!bFilenameValid)
+	{
+		if (!OpenSaveDialog(DialogStartPath, DialogStartName, UserPackageName))
+		{
+			return false;
+		}
+
+		NewPackageName = UserPackageName;
+
+		FText OutError;
+		bFilenameValid = FFileHelper::IsFilenameValidForSaving(NewPackageName, OutError);
+	}
+
+	// Update to the last location they saved to so it remembers their settings next time.
+	ConfigSettings->PresetSaveDir.Path = FPackageName::GetLongPackagePath(UserPackageName);
+	ConfigSettings->SaveConfig();
+	OutName = MoveTemp(NewPackageName);
+	return true;
+}
+
+void SMoviePipelineGraphPanel::OnSaveAsAsset()
+{
+	UMoviePipelineQueueSubsystem* Subsystem = GEditor->GetEditorSubsystem<UMoviePipelineQueueSubsystem>();
+	check(Subsystem);
+	UMoviePipelineQueue* CurrentQueue = Subsystem->GetQueue();
+
+	FString PackageName;
+	if (!GetSavePresetPackageName(CurrentQueue->GetName(), PackageName))
+	{
+		return;
+	}
+	
+	// Saving into a new package
+	const FString NewAssetName = FPackageName::GetLongPackageAssetName(PackageName);
+	UPackage* NewPackage = CreatePackage(*PackageName);
+	NewPackage->MarkAsFullyLoaded();
+	UMoviePipelineQueue* DuplicateQueue = DuplicateObject<UMoviePipelineQueue>(CurrentQueue, NewPackage, *NewAssetName);
+
+	if (DuplicateQueue)
+	{
+		DuplicateQueue->SetFlags(RF_Public | RF_Standalone | RF_Transactional);
+
+		FAssetRegistryModule::AssetCreated(DuplicateQueue);
+
+		FEditorFileUtils::EPromptReturnCode PromptReturnCode = FEditorFileUtils::PromptForCheckoutAndSave({ NewPackage }, false, false);
+	}
+}
+
+void SMoviePipelineGraphPanel::OnImportSavedQueueAssest(const FAssetData& InPresetAsset)
+{
+	FSlateApplication::Get().DismissAllMenus();
+
+	UMoviePipelineQueue* SavedQueue = CastChecked<UMoviePipelineQueue>(InPresetAsset.GetAsset());
+	if (SavedQueue)
+	{
+		// Duplicate the queue so we don't start modifying the one in the asset.
+		UMoviePipelineQueueSubsystem* Subsystem = GEditor->GetEditorSubsystem<UMoviePipelineQueueSubsystem>();
+		check(Subsystem);
+
+		Subsystem->GetQueue()->CopyFrom(SavedQueue);
+
+		// Update the shot list in case the stored queue being copied is out of date with the sequence
+		for (UMoviePipelineExecutorJob* Job : Subsystem->GetQueue()->GetJobs())
+		{
+			ULevelSequence* LoadedSequence = Cast<ULevelSequence>(Job->Sequence.TryLoad());
+			if (LoadedSequence)
+			{
+				bool bShotsChanged = false;
+				UMoviePipelineBlueprintLibrary::UpdateJobShotListFromSequence(LoadedSequence, Job, bShotsChanged);
+
+				if (bShotsChanged)
+				{
+					FNotificationInfo Info(LOCTEXT("QueueShotsUpdated", "Shots have changed since the queue was saved, please resave the queue"));
+					Info.ExpireDuration = 5.0f;
+					FSlateNotificationManager::Get().AddNotification(Info)->SetCompletionState(SNotificationItem::CS_Fail);
+				}
+			}
+		}
+
+		// Automatically select the first job in the queue
+		TArray<UMoviePipelineExecutorJob*> Jobs;
+		if (Subsystem->GetQueue()->GetJobs().Num() > 0)
+		{
+			Jobs.Add(Subsystem->GetQueue()->GetJobs()[0]);
+		}
+
+		// Go through the UI so it updates the UI selection too and then this will loop back
+		// around to OnSelectionChanged to update ourself.
+		PipelineQueueEditorWidget->SetSelectedJobs(Jobs);
+	}
+}
+
+#undef LOCTEXT_NAMESPACE // SMoviePipelineGraphPanel
\ No newline at end of file
diff --git a/Engine/Plugins/MovieScene/MovieRenderPipeline/Source/MovieRenderPipelineEditor/Private/Widgets/Graph/SMoviePipelineGraphPanel.h b/Engine/Plugins/MovieScene/MovieRenderPipeline/Source/MovieRenderPipelineEditor/Private/Widgets/Graph/SMoviePipelineGraphPanel.h
new file mode 100644
--- /dev/null
+++ b/Engine/Plugins/MovieScene/MovieRenderPipeline/Source/MovieRenderPipelineEditor/Private/Widgets/Graph/SMoviePipelineGraphPanel.h
@@ -0,0 +1,81 @@
+// Copyright Epic Games, Inc. All Rights Reserved.
+
+#pragma once
+
+#include "Widgets/DeclarativeSyntaxSupport.h"
+#include "Widgets/SCompoundWidget.h"
+
+class UMoviePipelineConfigBase;
+class SMoviePipelineQueueEditor;
+class SWindow;
+class UMoviePipelineExecutorJob;
+class UMoviePipelineExecutorShot;
+class IDetailsView;
+struct FAssetData;
+
+/**
+ * Outermost widget that is used for adding and removing jobs from the Movie Pipeline Queue Subsystem.
+ */
+class SMoviePipelineGraphPanel : public SCompoundWidget
+{
+public:
+	SLATE_BEGIN_ARGS(SMoviePipelineGraphPanel)
+		: _BasePreset(nullptr)
+
+		{}
+
+		/*~ All following arguments are mutually-exclusive */
+		/*-------------------------------------------------*/
+		/** A preset asset to base the pipeline off */
+		SLATE_ARGUMENT(UMoviePipelineConfigBase*, BasePreset)
+
+	SLATE_END_ARGS()
+
+	void Construct(const FArguments& InArgs);
+
+private:
+	FReply OnRenderLocalRequested();
+	bool IsRenderLocalEnabled() const;
+	FReply OnRenderRemoteRequested();
+	bool IsRenderRemoteEnabled() const;
+
+	/** When they want to edit the current configuration for the job */
+	void OnEditJobConfigRequested(TWeakObjectPtr<UMoviePipelineExecutorJob> InJob, TWeakObjectPtr<UMoviePipelineExecutorShot> InShot);
+	/** When an existing preset is chosen for the specified job. */
+	void OnJobPresetChosen(TWeakObjectPtr<UMoviePipelineExecutorJob> InJob, TWeakObjectPtr<UMoviePipelineExecutorShot> InShot);
+	void OnConfigUpdatedForJob(TWeakObjectPtr<UMoviePipelineExecutorJob> InJob, TWeakObjectPtr<UMoviePipelineExecutorShot> InShot, UMoviePipelineConfigBase* InConfig);
+	void OnConfigUpdatedForJobToPreset(TWeakObjectPtr<UMoviePipelineExecutorJob> InJob, TWeakObjectPtr<UMoviePipelineExecutorShot> InShot, UMoviePipelineConfigBase* InConfig);
+	void OnConfigWindowClosed();
+
+	void OnSelectionChanged(const TArray<UMoviePipelineExecutorJob*>& InSelectedJobs);
+	int32 GetDetailsViewWidgetIndex() const;
+	bool IsDetailsViewEnabled() const;
+
+	TSharedRef<SWidget> OnGenerateSavedQueuesMenu();
+	bool OpenSaveDialog(const FString& InDefaultPath, const FString& InNewNameSuggestion, FString& OutPackageName);
+	bool GetSavePresetPackageName(const FString& InExistingName, FString& OutName);
+	void OnSaveAsAsset();
+	void OnImportSavedQueueAssest(const FAssetData& InPresetAsset);
+	
+	void OnSelectedNodesChanged(const TSet<class UObject*>& NewSelection);
+	void OnNodeDoubleClicked(class UEdGraphNode* Node);
+	void OnNodeTitleCommitted(const FText& NewText, ETextCommit::Type CommitInfo, UEdGraphNode* NodeBeingChanged);
+	
+	FReply OnDebugButtonClicked();
+	TObjectPtr<class UMovieGraphConfig> CurrentGraph;
+private:
+	/** Allocates a transient preset so that the user can use the pipeline without saving it to an asset first. */
+	//UMoviePipelineConfigBase* AllocateTransientPreset();
+
+
+private:
+	/** The main movie pipeline queue editor widget */
+	TSharedPtr<SMoviePipelineQueueEditor> PipelineQueueEditorWidget;
+
+	/** The details panel for the selected job(s) */
+	TSharedPtr<IDetailsView> JobDetailsPanelWidget;
+
+	TWeakPtr<SWindow> WeakEditorWindow;
+	
+	int32 NumSelectedJobs;
+};
\ No newline at end of file
diff --git a/Engine/Plugins/MovieScene/MovieRenderPipeline/Source/MovieRenderPipelineEditor/Private/Widgets/SMoviePipelineQueueEditor.cpp b/Engine/Plugins/MovieScene/MovieRenderPipeline/Source/MovieRenderPipelineEditor/Private/Widgets/SMoviePipelineQueueEditor.cpp
--- a/Engine/Plugins/MovieScene/MovieRenderPipeline/Source/MovieRenderPipelineEditor/Private/Widgets/SMoviePipelineQueueEditor.cpp
+++ b/Engine/Plugins/MovieScene/MovieRenderPipeline/Source/MovieRenderPipelineEditor/Private/Widgets/SMoviePipelineQueueEditor.cpp
@@ -13,6 +13,7 @@
 #include "Sections/MovieSceneCameraCutSection.h"
 #include "MoviePipelineCommands.h"
 #include "MoviePipelineEditorBlueprintLibrary.h"
+#include "Graph/MovieGraphConfig.h"
 
 // Slate Includes
 #include "Widgets/Input/SComboButton.h"
@@ -203,6 +204,11 @@ public:
 		UMoviePipelineExecutorJob* Job = WeakJob.Get();
 		if (Job)
 		{
+			if (Job->IsUsingGraphConfiguration())
+			{
+				return LOCTEXT("QueueEditorDefaultJobGraph_Text", "Default Graph");
+			}
+			
 			// If the job has a preset origin (ie, its config is based off a preset w/o any modifications), use its
 			// display name. If the config has a preset origin (ie, it's based off a preset, but has modifications), use
 			// that display name. Otherwise, fall back to the config's display name.
@@ -234,7 +240,14 @@ public:
 		UMoviePipelineExecutorJob* Job = WeakJob.Get();
 		if (Job)
 		{
-			Job->SetPresetOrigin(CastChecked<UMoviePipelinePrimaryConfig>(AssetData.GetAsset()));
+			if (Job->IsUsingGraphConfiguration())
+			{
+				Job->SetGraphPreset(CastChecked<UMovieGraphConfig>(AssetData.GetAsset()));
+			}
+			else
+			{
+				Job->SetPresetOrigin(CastChecked<UMoviePipelinePrimaryConfig>(AssetData.GetAsset()));
+			}
 		}
 
 		OnChosePresetCallback.ExecuteIfBound(WeakJob, nullptr);
@@ -248,14 +261,34 @@ public:
 		UMoviePipelineExecutorJob* Job = WeakJob.Get();
 		if (Job)
 		{
-			// Copy from the CDO's version of the job to pick up the right name.
-			Job->SetConfiguration(GetMutableDefault<UMoviePipelineExecutorJob>()->GetConfiguration());
-			UMoviePipelineEditorBlueprintLibrary::EnsureJobHasDefaultSettings(Job);
+			if (Job->IsUsingGraphConfiguration())
+			{
+				Job->SetGraphPreset(GetMutableDefault<UMoviePipelineExecutorJob>()->GetGraphPreset());
+			}
+			else
+			{
+				// Copy from the CDO's version of the job to pick up the right name.
+				Job->SetConfiguration(GetMutableDefault<UMoviePipelineExecutorJob>()->GetConfiguration());
+				UMoviePipelineEditorBlueprintLibrary::EnsureJobHasDefaultSettings(Job);
+			}
 		}
 
 		OnChosePresetCallback.ExecuteIfBound(WeakJob, nullptr);
 	}
 
+	void OnReplaceWithRenderGraph()
+	{
+		UMoviePipelineExecutorJob* Job = WeakJob.Get();
+		if (!Job)
+		{
+			return;
+		}
+
+		// Note: Setting the graph preset will transition the job to use a graph-based configuration
+		UMovieGraphConfig* NewGraph = NewObject<UMovieGraphConfig>(Job);
+		Job->SetGraphPreset(NewGraph);
+	}
+
 	EVisibility GetPrimaryConfigModifiedVisibility() const
 	{
 		UMoviePipelineExecutorJob* Job = WeakJob.Get();
@@ -347,19 +380,31 @@ public:
 
 	TSharedRef<SWidget> OnGenerateConfigPresetPickerMenu()
 	{
-		return OnGenerateConfigPresetPickerMenuFromClass(UMoviePipelinePrimaryConfig::StaticClass(),
+		UMoviePipelineExecutorJob* Job = WeakJob.Get();
+
+		UClass* ConfigType = Job->IsUsingGraphConfiguration() ? UMovieGraphConfig::StaticClass() : UMoviePipelinePrimaryConfig::StaticClass();
+		
+		return OnGenerateConfigPresetPickerMenuFromClass(
+			ConfigType,
+			WeakJob,
 			FOnAssetSelected::CreateRaw(this, &FMoviePipelineQueueJobTreeItem::OnPickPresetFromAsset),
-			FExecuteAction::CreateRaw(this, &FMoviePipelineQueueJobTreeItem::OnPickNewPreset)
+			FExecuteAction::CreateRaw(this, &FMoviePipelineQueueJobTreeItem::OnPickNewPreset),
+			FExecuteAction::CreateRaw(this, &FMoviePipelineQueueJobTreeItem::OnReplaceWithRenderGraph)
 			);
 	}
 
-	static TSharedRef<SWidget> OnGenerateConfigPresetPickerMenuFromClass(TSubclassOf<UMoviePipelineConfigBase> InClass, FOnAssetSelected InOnAssetSelected, FExecuteAction InNewConfig)
+	static TSharedRef<SWidget> OnGenerateConfigPresetPickerMenuFromClass(UClass* InClass,
+		TWeakObjectPtr<UMoviePipelineExecutorJob> TargetJob, FOnAssetSelected InOnAssetSelected, FExecuteAction InNewConfig, FExecuteAction InNewRenderGraph)
 	{
 		FMenuBuilder MenuBuilder(true, nullptr);
 		IContentBrowserSingleton& ContentBrowser = FModuleManager::LoadModuleChecked<FContentBrowserModule>("ContentBrowser").Get();
 
 		FAssetPickerConfig AssetPickerConfig;
 		{
+			const FText NoAssetsFoundWarning = InClass->IsChildOf(UMoviePipelinePrimaryConfig::StaticClass())
+				? LOCTEXT("NoConfigs_Warning", "No Primary Configurations Found")
+				: LOCTEXT("NoGraphs_Warning", "No Render Graphs Found");
+			
 			AssetPickerConfig.SelectionMode = ESelectionMode::Single;
 			AssetPickerConfig.InitialAssetViewType = EAssetViewType::List;
 			AssetPickerConfig.bFocusSearchBoxWhenOpened = true;
@@ -374,7 +419,7 @@ public:
 			AssetPickerConfig.ThumbnailScale = 0.1f;
 			AssetPickerConfig.SaveSettingsName = TEXT("MoviePipelineConfigAsset");
 
-			AssetPickerConfig.AssetShowWarningText = LOCTEXT("NoConfigs_Warning", "No Primary Configurations Found");
+			AssetPickerConfig.AssetShowWarningText = NoAssetsFoundWarning;
 			AssetPickerConfig.Filter.ClassPaths.Add(InClass->GetClassPathName());
 			AssetPickerConfig.OnAssetSelected = InOnAssetSelected;
 		}
@@ -389,6 +434,19 @@ public:
 				NAME_None,
 				EUserInterfaceActionType::Button
 			);
+
+			const bool bIsGraphFeatureEnabled = IConsoleManager::Get().FindConsoleVariable(TEXT("MoviePipeline.EnableRenderGraph"))->GetBool();
+			if (bIsGraphFeatureEnabled && !TargetJob->IsUsingGraphConfiguration())
+			{
+				MenuBuilder.AddMenuEntry(
+					LOCTEXT("ReplaceWithGraph_Label", "Replace with Graph"),
+					LOCTEXT("ReplaceWithGraph_Tooltip", "Replaces the current configuration with a new graph representation."),
+					FSlateIcon(),
+					FUIAction(InNewRenderGraph),
+					NAME_None,
+					EUserInterfaceActionType::Button
+				);
+			}
 		}
 
 		MenuBuilder.BeginSection(NAME_None, LOCTEXT("ImportConfig_MenuSection", "Import Configuration"));
@@ -802,6 +860,11 @@ struct FMoviePipelineShotItem : IMoviePipelineQueueTreeItem
 		UMoviePipelineExecutorShot* Shot = WeakShot.Get();
 		if (Shot)
 		{
+			if (Shot->IsUsingGraphConfiguration())
+			{
+				return LOCTEXT("QueueEditorDefaultShotGraph_Text", "Default Graph");
+			}
+			
 			// If the shot has a preset origin (ie, its config is based off a preset w/o any modifications), use its
 			// display name. If the config has a preset origin (ie, it's based off a preset, but has modifications), use
 			// that display name. Otherwise, fall back to the config's display name.
@@ -853,6 +916,19 @@ struct FMoviePipelineShotItem : IMoviePipelineQueueTreeItem
 		OnChosePresetCallback.ExecuteIfBound(WeakJob, WeakShot);
 	}
 
+	void OnReplaceWithRenderGraph()
+	{
+		UMoviePipelineExecutorShot* Shot = WeakShot.Get();
+		if (!Shot)
+		{
+			return;
+		}
+
+		// Note: Setting the graph preset will transition the job to use a graph-based configuration
+		UMovieGraphConfig* NewGraph = NewObject<UMovieGraphConfig>(Shot);
+		Shot->SetGraphPreset(NewGraph);
+	}
+
 	EVisibility GetShotConfigModifiedVisibility() const
 	{
 		UMoviePipelineExecutorShot* Shot = WeakShot.Get();
@@ -871,9 +947,12 @@ struct FMoviePipelineShotItem : IMoviePipelineQueueTreeItem
 
 	TSharedRef<SWidget> OnGenerateShotConfigPresetPickerMenu()
 	{
-		return FMoviePipelineQueueJobTreeItem::OnGenerateConfigPresetPickerMenuFromClass(UMoviePipelineShotConfig::StaticClass(),
+		return FMoviePipelineQueueJobTreeItem::OnGenerateConfigPresetPickerMenuFromClass(
+			UMoviePipelineShotConfig::StaticClass(),
+			WeakJob,
 			FOnAssetSelected::CreateRaw(this, &FMoviePipelineShotItem::OnPickShotPresetFromAsset),
-			FExecuteAction::CreateRaw(this, &FMoviePipelineShotItem::OnPickNewShotPreset)
+			FExecuteAction::CreateRaw(this, &FMoviePipelineShotItem::OnPickNewShotPreset),
+			FExecuteAction::CreateRaw(this, &FMoviePipelineShotItem::OnReplaceWithRenderGraph)
 			);
 	}
 };
diff --git a/Engine/Plugins/MovieScene/MovieRenderPipeline/Source/MovieRenderPipelineEditor/Private/Widgets/SMoviePipelineQueuePanel.cpp b/Engine/Plugins/MovieScene/MovieRenderPipeline/Source/MovieRenderPipelineEditor/Private/Widgets/SMoviePipelineQueuePanel.cpp
--- a/Engine/Plugins/MovieScene/MovieRenderPipeline/Source/MovieRenderPipelineEditor/Private/Widgets/SMoviePipelineQueuePanel.cpp
+++ b/Engine/Plugins/MovieScene/MovieRenderPipeline/Source/MovieRenderPipelineEditor/Private/Widgets/SMoviePipelineQueuePanel.cpp
@@ -12,6 +12,8 @@
 #include "MoviePipelineQueue.h"
 #include "MoviePipelinePrimaryConfig.h"
 #include "MoviePipelineQueueSubsystem.h"
+#include "Graph/MovieGraphConfig.h"
+#include "Graph/MovieGraphAssetEditor.h"
 
 // Slate Includes
 #include "Widgets/SBoxPanel.h"
@@ -362,6 +364,31 @@ void SMoviePipelineQueuePanel::OnEditJobConfigRequested(TWeakObjectPtr<UMoviePip
 		return;
 	}
 
+	// Handle editing a graph vs. a standard configuration
+	UMovieGraphConfig* GraphToEdit = nullptr;
+	if (InShot.IsValid() && InShot->IsUsingGraphConfiguration())
+	{
+		GraphToEdit = InShot->GetGraphPreset();
+	}
+	else if (InJob.IsValid() && InJob->IsUsingGraphConfiguration())
+	{
+		GraphToEdit = InJob->GetGraphPreset();
+	}
+
+	if (GraphToEdit)
+	{
+		if (UAssetEditorSubsystem* AssetEditorSubsystem = GEditor->GetEditorSubsystem<UAssetEditorSubsystem>())
+		{
+			if (UMovieGraphAssetEditor* AssetEditor = NewObject<UMovieGraphAssetEditor>(AssetEditorSubsystem, NAME_None, RF_Transient))
+			{
+				AssetEditor->SetObjectToEdit(GraphToEdit);
+				AssetEditor->Initialize();
+			}
+		}
+
+		return;
+	}
+
 	TSubclassOf<UMoviePipelineConfigBase> ConfigType;
 	UMoviePipelineConfigBase* BasePreset = nullptr;
 	UMoviePipelineConfigBase* BaseConfig = nullptr;
diff --git a/Engine/Plugins/MovieScene/MovieRenderPipeline/Source/MovieRenderPipelineSettings/MovieRenderPipelineSettings.Build.cs b/Engine/Plugins/MovieScene/MovieRenderPipeline/Source/MovieRenderPipelineSettings/MovieRenderPipelineSettings.Build.cs
--- a/Engine/Plugins/MovieScene/MovieRenderPipeline/Source/MovieRenderPipelineSettings/MovieRenderPipelineSettings.Build.cs
+++ b/Engine/Plugins/MovieScene/MovieRenderPipeline/Source/MovieRenderPipelineSettings/MovieRenderPipelineSettings.Build.cs
@@ -16,6 +16,7 @@ public class MovieRenderPipelineSettings : ModuleRules
 				"Slate",
 				"RenderCore",
 				"RHI",
+				"ApplicationCore"
 			}
 		);
 
