NonUnity Compile Error fix for MovieGraphSchema

diff --git a/Engine/Plugins/MovieScene/MovieRenderPipeline/Source/MovieRenderPipelineEditor/Private/Graph/MovieGraphSchema.cpp b/Engine/Plugins/MovieScene/MovieRenderPipeline/Source/MovieRenderPipelineEditor/Private/Graph/MovieGraphSchema.cpp
--- a/Engine/Plugins/MovieScene/MovieRenderPipeline/Source/MovieRenderPipelineEditor/Private/Graph/MovieGraphSchema.cpp
+++ b/Engine/Plugins/MovieScene/MovieRenderPipeline/Source/MovieRenderPipelineEditor/Private/Graph/MovieGraphSchema.cpp
@@ -5,6 +5,10 @@
 #include "Graph/MovieGraphNode.h"
 #include "Graph/MovieEdGraph.h"
 #include "Graph/MovieEdGraphNode.h"
+#include "UObject/Class.h"
+#include "UObject/UObjectIterator.h"
+#include "ScopedTransaction.h"
+#include "GraphEditor.h"
 
 TArray<UClass*> UMovieGraphSchema::MoviePipelineNodeClasses;
 
diff --git a/Engine/Plugins/MovieScene/MovieRenderPipeline/Source/MovieRenderPipelineEditor/Private/Graph/MovieGraphSchema.h b/Engine/Plugins/MovieScene/MovieRenderPipeline/Source/MovieRenderPipelineEditor/Private/Graph/MovieGraphSchema.h
--- a/Engine/Plugins/MovieScene/MovieRenderPipeline/Source/MovieRenderPipelineEditor/Private/Graph/MovieGraphSchema.h
+++ b/Engine/Plugins/MovieScene/MovieRenderPipeline/Source/MovieRenderPipelineEditor/Private/Graph/MovieGraphSchema.h
@@ -3,6 +3,7 @@
 #pragma once
 #include "CoreTypes.h"
 #include "EdGraph/EdGraphSchema.h"
+#include "Templates/SubclassOf.h"
 #include "MovieGraphSchema.generated.h"
 
 class UMovieGraphNode;
diff --git a/Engine/Plugins/MovieScene/MovieRenderPipeline/Source/MovieRenderPipelineEditor/Private/Widgets/Graph/SMoviePipelineGraphPanel.cpp b/Engine/Plugins/MovieScene/MovieRenderPipeline/Source/MovieRenderPipelineEditor/Private/Widgets/Graph/SMoviePipelineGraphPanel.cpp
--- a/Engine/Plugins/MovieScene/MovieRenderPipeline/Source/MovieRenderPipelineEditor/Private/Widgets/Graph/SMoviePipelineGraphPanel.cpp
+++ b/Engine/Plugins/MovieScene/MovieRenderPipeline/Source/MovieRenderPipelineEditor/Private/Widgets/Graph/SMoviePipelineGraphPanel.cpp
@@ -53,6 +53,7 @@
 #include "Graph/MovieEdGraph.h"
 #include "Graph/MovieGraphSchema.h"
 #include "Kismet2/BlueprintEditorUtils.h"
+#include "GraphEditor.h"
 
 #define LOCTEXT_NAMESPACE "SMoviePipelineGraphPanel"
 
