NonUnity Fixes for MoviePipeline

diff --git a/Engine/Plugins/MovieScene/MovieRenderPipeline/Source/MovieRenderPipelineEditor/Private/Graph/MovieEdGraph.cpp b/Engine/Plugins/MovieScene/MovieRenderPipeline/Source/MovieRenderPipelineEditor/Private/Graph/MovieEdGraph.cpp
--- a/Engine/Plugins/MovieScene/MovieRenderPipeline/Source/MovieRenderPipelineEditor/Private/Graph/MovieEdGraph.cpp
+++ b/Engine/Plugins/MovieScene/MovieRenderPipeline/Source/MovieRenderPipelineEditor/Private/Graph/MovieEdGraph.cpp
@@ -8,6 +8,7 @@
 #include "MovieEdGraphOutputNode.h"
 #include "MovieEdGraphInputNode.h"
 #include "MovieRenderPipelineCoreModule.h"
+#include "EdGraph/EdGraphPin.h"
 
 UMovieGraphConfig* UMoviePipelineEdGraph::GetPipelineGraph() const
 {
diff --git a/Engine/Plugins/MovieScene/MovieRenderPipeline/Source/MovieRenderPipelineEditor/Private/Graph/MovieEdGraphNode.cpp b/Engine/Plugins/MovieScene/MovieRenderPipeline/Source/MovieRenderPipelineEditor/Private/Graph/MovieEdGraphNode.cpp
--- a/Engine/Plugins/MovieScene/MovieRenderPipeline/Source/MovieRenderPipelineEditor/Private/Graph/MovieEdGraphNode.cpp
+++ b/Engine/Plugins/MovieScene/MovieRenderPipeline/Source/MovieRenderPipelineEditor/Private/Graph/MovieEdGraphNode.cpp
@@ -3,6 +3,7 @@
 #include "Graph/MovieEdGraphNode.h"
 #include "Graph/MovieGraphPin.h"
 #include "Graph/MovieGraphNode.h"
+#include "EdGraph/EdGraphPin.h"
 
 void UMoviePipelineEdGraphNodeBase::Construct(UMovieGraphNode* InRuntimeNode)
 {
