MoviePipeline: Added ability to promote a node property to a pin via the dynamic property mechanism.

Currently this action is available in the context menu for a node, but will probably move elsewhere in the UI in later revisions.

diff --git a/Engine/Plugins/MovieScene/MovieRenderPipeline/Source/MovieRenderPipelineCore/Private/Graph/MovieGraphNode.cpp b/Engine/Plugins/MovieScene/MovieRenderPipeline/Source/MovieRenderPipelineCore/Private/Graph/MovieGraphNode.cpp
--- a/Engine/Plugins/MovieScene/MovieRenderPipeline/Source/MovieRenderPipelineCore/Private/Graph/MovieGraphNode.cpp
+++ b/Engine/Plugins/MovieScene/MovieRenderPipeline/Source/MovieRenderPipelineCore/Private/Graph/MovieGraphNode.cpp
@@ -8,12 +8,51 @@ UMovieGraphNode::UMovieGraphNode()
 {
 }
 
+TArray<FMovieGraphPinProperties> UMovieGraphNode::GetExposedDynamicPinProperties() const
+{
+	TArray<FMovieGraphPinProperties> Properties;
+
+	// Exposed dynamic properties are tracked by name; convert to pin properties
+	for (const FName& PropertyName : GetExposedDynamicProperties())
+	{
+		// Note: Currently hardcoded to not allow multiple connections, but this may need to change in the future
+		const bool bInAllowMultipleConnections = false;
+		Properties.Add(FMovieGraphPinProperties(PropertyName, bInAllowMultipleConnections));
+	}
+
+	return Properties;
+}
+
+void UMovieGraphNode::PromoteDynamicPropertyToPin(const FName& PropertyName)
+{
+	// Ensure this is a real dynamic property
+	bool bFoundMatchingProperty = false;
+	for (const FPropertyBagPropertyDesc& PropertyDescription : GetDynamicPropertyDescriptions())
+	{
+		if (PropertyDescription.Name == PropertyName)
+		{
+			bFoundMatchingProperty = true;
+			break;
+		}
+	}
+
+	// Add this property to promoted properties if it hasn't already been promoted
+	if (bFoundMatchingProperty && !ExposedDynamicPropertyNames.Contains(PropertyName))
+	{
+		ExposedDynamicPropertyNames.Add(PropertyName);
+	}
+
+	UpdatePins();
+}
+
 void UMovieGraphNode::UpdatePins()
 {
-	
 	TArray<FMovieGraphPinProperties> InputPinProperties = GetInputPinProperties();
 	TArray<FMovieGraphPinProperties> OutputPinProperties = GetOutputPinProperties();
 
+	// Include the exposed dynamic properties in the input pins
+	InputPinProperties.Append(GetExposedDynamicPinProperties());
+
 	auto UpdatePins = [this](TArray<UMovieGraphPin*>& Pins, const TArray<FMovieGraphPinProperties>& PinProperties)
 	{
 		bool bAppliedEdgeChanges = false;
@@ -80,6 +119,8 @@ void UMovieGraphNode::UpdatePins()
 
 	UpdatePins(InputPins, InputPinProperties);
 	UpdatePins(OutputPins, OutputPinProperties);
+
+	OnNodeChangedDelegate.Broadcast(this);
 }
 
 void UMovieGraphNode::UpdateDynamicProperties()
diff --git a/Engine/Plugins/MovieScene/MovieRenderPipeline/Source/MovieRenderPipelineCore/Public/Graph/MovieGraphNode.h b/Engine/Plugins/MovieScene/MovieRenderPipeline/Source/MovieRenderPipelineCore/Public/Graph/MovieGraphNode.h
--- a/Engine/Plugins/MovieScene/MovieRenderPipeline/Source/MovieRenderPipelineCore/Public/Graph/MovieGraphNode.h
+++ b/Engine/Plugins/MovieScene/MovieRenderPipeline/Source/MovieRenderPipelineCore/Public/Graph/MovieGraphNode.h
@@ -15,6 +15,8 @@ class UMovieGraphPin;
 class UEdGraphNode;
 #endif
 
+DECLARE_MULTICAST_DELEGATE_OneParam(FOnMovieGraphNodeChanged, UMovieGraphNode*);
+
 /**
 * This is a base class for all nodes that can exist in the UMovieGraphConfig network.
 * In the editor, each node in the network will have an editor-only representation too 
@@ -50,16 +52,21 @@ public:
 
 	virtual TArray<FName> GetExposedDynamicProperties() const
 	{ 
-		return TArray<FName>();
+		return ExposedDynamicPropertyNames;
 	}
 
+	/** Promotes the property with the given name to a pin on the node via a dynamic property. */
+	virtual void PromoteDynamicPropertyToPin(const FName& PropertyName);
+
 	void UpdatePins();
 	void UpdateDynamicProperties();
 	class UMovieGraphConfig* GetGraph() const;
 	UMovieGraphPin* GetInputPin(const FName& InPinLabel) const;
 	UMovieGraphPin* GetOutputPin(const FName& InPinLabel) const;
-protected:
+
 public:
+	FOnMovieGraphNodeChanged OnNodeChangedDelegate;
+
 #if WITH_EDITORONLY_DATA
 	/** Editor Node Graph representation. Not strongly typed to avoid circular dependency between editor/runtime modules. */
 	UPROPERTY()
@@ -73,6 +80,9 @@ public:
 	virtual FText GetMenuCategory() const { return FText(); }
 #endif
 
+protected:
+	virtual TArray<FMovieGraphPinProperties> GetExposedDynamicPinProperties() const;
+
 protected:
 	UPROPERTY()
 	TArray<TObjectPtr<UMovieGraphPin>> InputPins;
diff --git a/Engine/Plugins/MovieScene/MovieRenderPipeline/Source/MovieRenderPipelineEditor/MovieRenderPipelineEditor.Build.cs b/Engine/Plugins/MovieScene/MovieRenderPipeline/Source/MovieRenderPipelineEditor/MovieRenderPipelineEditor.Build.cs
--- a/Engine/Plugins/MovieScene/MovieRenderPipeline/Source/MovieRenderPipelineEditor/MovieRenderPipelineEditor.Build.cs
+++ b/Engine/Plugins/MovieScene/MovieRenderPipeline/Source/MovieRenderPipelineEditor/MovieRenderPipelineEditor.Build.cs
@@ -66,6 +66,7 @@ public class MovieRenderPipelineEditor : ModuleRules
 				"AutomationMessages",
 				"ToolWidgets",
 				"ConsoleVariablesEditor",
+				"StructUtils",
 			}
         );
 
diff --git a/Engine/Plugins/MovieScene/MovieRenderPipeline/Source/MovieRenderPipelineEditor/Private/Graph/MovieEdGraph.cpp b/Engine/Plugins/MovieScene/MovieRenderPipeline/Source/MovieRenderPipelineEditor/Private/Graph/MovieEdGraph.cpp
--- a/Engine/Plugins/MovieScene/MovieRenderPipeline/Source/MovieRenderPipelineEditor/Private/Graph/MovieEdGraph.cpp
+++ b/Engine/Plugins/MovieScene/MovieRenderPipeline/Source/MovieRenderPipelineEditor/Private/Graph/MovieEdGraph.cpp
@@ -64,6 +64,24 @@ void UMoviePipelineEdGraph::InitFromRuntimeGraph(UMovieGraphConfig* InGraph)
 	bInitialized = true;
 }
 
+void UMoviePipelineEdGraph::CreateLinks(UMoviePipelineEdGraphNodeBase* InGraphNode, bool bCreateInboundLinks, bool bCreateOutboundLinks)
+{
+	check(InGraphNode);
+	
+	// Build runtime node to editor node map
+	TMap<UMovieGraphNode*, UMoviePipelineEdGraphNodeBase*> RuntimeNodeToEdNodeMap;
+
+	for (const TObjectPtr<UEdGraphNode>& EdGraphNode : Nodes)
+	{
+		if (UMoviePipelineEdGraphNodeBase* MovieEdNode = Cast<UMoviePipelineEdGraphNodeBase>(EdGraphNode))
+		{
+			RuntimeNodeToEdNodeMap.Add(MovieEdNode->GetRuntimeNode(), MovieEdNode);
+		}
+	}
+
+	CreateLinks(InGraphNode, bCreateInboundLinks, bCreateOutboundLinks, RuntimeNodeToEdNodeMap);
+}
+
 void UMoviePipelineEdGraph::CreateLinks(UMoviePipelineEdGraphNodeBase* InGraphNode, bool bCreateInboundLinks, bool bCreateOutboundLinks,
 	const TMap<UMovieGraphNode*, UMoviePipelineEdGraphNodeBase*>& RuntimeNodeToEdNodeMap)
 {
@@ -117,6 +135,6 @@ void UMoviePipelineEdGraph::CreateLinks(UMoviePipelineEdGraphNodeBase* InGraphNo
 	}
 	if (bCreateOutboundLinks)
 	{
-		CreateLinks(RuntimeNode->GetInputPins(), EEdGraphPinDirection::EGPD_Output);
+		CreateLinks(RuntimeNode->GetOutputPins(), EEdGraphPinDirection::EGPD_Output);
 	}
 }
\ No newline at end of file
diff --git a/Engine/Plugins/MovieScene/MovieRenderPipeline/Source/MovieRenderPipelineEditor/Private/Graph/MovieEdGraph.h b/Engine/Plugins/MovieScene/MovieRenderPipeline/Source/MovieRenderPipelineEditor/Private/Graph/MovieEdGraph.h
--- a/Engine/Plugins/MovieScene/MovieRenderPipeline/Source/MovieRenderPipelineEditor/Private/Graph/MovieEdGraph.h
+++ b/Engine/Plugins/MovieScene/MovieRenderPipeline/Source/MovieRenderPipelineEditor/Private/Graph/MovieEdGraph.h
@@ -27,6 +27,9 @@ public:
 	/** Returns the runtime UMovieGraphConfig that contains this editor graph */
 	class UMovieGraphConfig* GetPipelineGraph() const;
 
+	/** Creates the links/edges between nodes in the graph */
+	void CreateLinks(UMoviePipelineEdGraphNodeBase* InGraphNode, bool bCreateInboundLinks, bool bCreateOutboundLinks);
+
 protected:
 	void CreateLinks(UMoviePipelineEdGraphNodeBase* InGraphNode, bool bCreateInboundLinks, bool bCreateOutboundLinks,
 		const TMap<UMovieGraphNode*, UMoviePipelineEdGraphNodeBase*>& RuntimeNodeToEdNodeMap);
diff --git a/Engine/Plugins/MovieScene/MovieRenderPipeline/Source/MovieRenderPipelineEditor/Private/Graph/MovieEdGraphNode.cpp b/Engine/Plugins/MovieScene/MovieRenderPipeline/Source/MovieRenderPipelineEditor/Private/Graph/MovieEdGraphNode.cpp
--- a/Engine/Plugins/MovieScene/MovieRenderPipeline/Source/MovieRenderPipelineEditor/Private/Graph/MovieEdGraphNode.cpp
+++ b/Engine/Plugins/MovieScene/MovieRenderPipeline/Source/MovieRenderPipelineEditor/Private/Graph/MovieEdGraphNode.cpp
@@ -1,15 +1,23 @@
 // Copyright Epic Games, Inc. All Rights Reserved.
 
 #include "Graph/MovieEdGraphNode.h"
+
 #include "Graph/MovieGraphPin.h"
 #include "Graph/MovieGraphNode.h"
 #include "EdGraph/EdGraphPin.h"
+#include "MovieEdGraph.h"
+#include "PropertyBag.h"
+#include "ToolMenu.h"
+
+#define LOCTEXT_NAMESPACE "MoviePipelineEdGraphNodeBase"
 
 void UMoviePipelineEdGraphNodeBase::Construct(UMovieGraphNode* InRuntimeNode)
 {
 	check(InRuntimeNode);
 	RuntimeNode = InRuntimeNode;
-
+	RuntimeNode->GraphNode = this;
+	RuntimeNode->OnNodeChangedDelegate.AddUObject(this, &UMoviePipelineEdGraphNodeBase::OnRuntimeNodeChanged);
+	
 	// NodePosX = InRuntimeNode->PositionX;
 	// NodePosY = InRuntimeNode->PositionY;
 	// NodeComment = InRuntimeNode->NodeComment;
@@ -50,12 +58,83 @@ FText UMoviePipelineEdGraphNode::GetNodeTitle(ENodeTitleType::Type TitleType) co
 		return RuntimeNode->GetMenuDescription();
 	}
 
-	return NSLOCTEXT("MoviePipelineGraph", "GraphTestTitle", "TestTitle");
+	return LOCTEXT("GraphTestTitle", "TestTitle");
 }
 
 FText UMoviePipelineEdGraphNode::GetTooltipText() const
 {
-	return NSLOCTEXT("MoviePipelineGraph", "GraphTestTooltip", "Test Tooltip");
+	return LOCTEXT("GraphTestTooltip", "Test Tooltip");
+}
+
+void UMoviePipelineEdGraphNode::GetNodeContextMenuActions(UToolMenu* Menu, UGraphNodeContextMenuContext* Context) const
+{
+	Super::GetNodeContextMenuActions(Menu, Context);
+	
+	if (!Context->Node || !RuntimeNode)
+	{
+		return;
+	}
+
+	GetPropertyPromotionContextMenuActions(Menu, Context);
+}
+
+void UMoviePipelineEdGraphNode::GetPropertyPromotionContextMenuActions(UToolMenu* Menu, UGraphNodeContextMenuContext* Context) const
+{
+	// Exclude the dynamic properties which are used for EditCondition metadata
+	static const FName EditConditionKey = FName("EditCondition");
+	TArray<FName> ExcludedPropertyNames;
+	for (const FPropertyBagPropertyDesc& PropertyDescription : RuntimeNode->GetDynamicPropertyDescriptions())
+	{
+		for (const FPropertyBagPropertyDescMetaData& Metadata : PropertyDescription.MetaData)
+		{
+			if (Metadata.Key == EditConditionKey)
+			{
+				ExcludedPropertyNames.Add(FName(Metadata.Value));
+			}
+		}
+	}
+
+	// Exclude dynamic properties which have already been promoted
+	ExcludedPropertyNames.Append(RuntimeNode->GetExposedDynamicProperties());
+
+	FToolMenuSection& Section = Menu->AddSection("MoviePipelineGraphExposeAsPin", LOCTEXT("ExposeAsPin", "Expose Property as Pin"));
+
+	const TArray<FPropertyBagPropertyDesc>& PropertyDescriptions = RuntimeNode->GetDynamicPropertyDescriptions();
+	for (const FPropertyBagPropertyDesc& PropertyDescription : PropertyDescriptions)
+	{
+		if (ExcludedPropertyNames.Contains(PropertyDescription.Name))
+		{
+			continue;
+		}
+		
+		Section.AddMenuEntry(
+			PropertyDescription.Name,
+			FText::FromName(PropertyDescription.Name),
+			LOCTEXT("PromotePropertyToPin", "Promote this property to a pin on this node."),
+			FSlateIcon(),
+			FUIAction(
+				FExecuteAction::CreateUObject(this, &UMoviePipelineEdGraphNode::PromotePropertyToPin, PropertyDescription.Name),
+				FCanExecuteAction())
+		);
+	}
+
+	if (PropertyDescriptions.Num() == ExcludedPropertyNames.Num())
+	{
+		Section.AddMenuEntry(
+			"NoPropertiesAvailable",
+			FText::FromString("No properties available"),
+			LOCTEXT("PromotePropertyToPin_NoneAvailable", "No properties are available to promote."),
+			FSlateIcon(),
+			FUIAction(
+				FExecuteAction(),
+				FCanExecuteAction::CreateLambda([]() { return false; }))
+		);
+	}
+}
+
+void UMoviePipelineEdGraphNode::PromotePropertyToPin(const FName PropertyName) const
+{
+	RuntimeNode->PromoteDynamicPropertyToPin(PropertyName);
 }
 
 bool UMoviePipelineEdGraphNodeBase::ShouldCreatePin(const UMovieGraphPin* InPin) const
@@ -99,4 +178,62 @@ void UMoviePipelineEdGraphNodeBase::CreatePins(const TArray<UMovieGraphPin*>& In
 	{
 		AdvancedPinDisplay = ENodeAdvancedPins::NoPins;
 	}
-}
\ No newline at end of file
+}
+
+void UMoviePipelineEdGraphNodeBase::OnRuntimeNodeChanged(UMovieGraphNode* InChangedNode)
+{
+	if (InChangedNode == GetRuntimeNode())
+	{
+		ReconstructNode();
+	}
+}
+
+void UMoviePipelineEdGraphNodeBase::ReconstructNode()
+{
+	ReconstructPins();
+
+	UMoviePipelineEdGraph* Graph = CastChecked<UMoviePipelineEdGraph>(GetGraph());
+	
+	// Reconstruct connections
+	const bool bCreateInbound = true;
+	const bool bCreateOutbound = true;
+	Graph->CreateLinks(this, bCreateInbound, bCreateOutbound);
+
+	Graph->NotifyGraphChanged();
+}
+
+void UMoviePipelineEdGraphNodeBase::ReconstructPins()
+{
+	// Store copy of old pins
+	TArray<UEdGraphPin*> OldPins = MoveTemp(Pins);
+	Pins.Reset();
+	
+	// Generate new pins
+	CreatePins(RuntimeNode->GetInputPins(), RuntimeNode->GetOutputPins());
+	
+	// Transfer persistent data from old to new pins
+	for (UEdGraphPin* OldPin : OldPins)
+	{
+		for (UEdGraphPin* NewPin : Pins)
+		{
+			if ((OldPin->PinName == NewPin->PinName) && (OldPin->Direction == NewPin->Direction))
+			{
+				// Remove invalid entries
+				OldPin->LinkedTo.Remove(nullptr);
+
+				NewPin->MovePersistentDataFromOldPin(*OldPin);
+				break;
+			}
+		}
+	}
+	
+	// Remove old pins
+	for (UEdGraphPin* OldPin : OldPins)
+	{
+		OldPin->BreakAllPinLinks();
+		OldPin->SubPins.Remove(nullptr);
+		DestroyPin(OldPin);
+	}
+}
+
+#undef LOCTEXT_NAMESPACE
\ No newline at end of file
diff --git a/Engine/Plugins/MovieScene/MovieRenderPipeline/Source/MovieRenderPipelineEditor/Private/Graph/MovieEdGraphNode.h b/Engine/Plugins/MovieScene/MovieRenderPipeline/Source/MovieRenderPipelineEditor/Private/Graph/MovieEdGraphNode.h
--- a/Engine/Plugins/MovieScene/MovieRenderPipeline/Source/MovieRenderPipelineEditor/Private/Graph/MovieEdGraphNode.h
+++ b/Engine/Plugins/MovieScene/MovieRenderPipeline/Source/MovieRenderPipelineEditor/Private/Graph/MovieEdGraphNode.h
@@ -14,22 +14,27 @@ class UMoviePipelineEdGraphNodeBase : public UEdGraphNode
 {
 	GENERATED_BODY()
 public:
-	void SetRuntimeNode(UMovieGraphNode* InNode)
-	{
-		RuntimeNode = InNode;
-		InNode->GraphNode = this;
-	}
-
 	UMovieGraphNode* GetRuntimeNode() const { return RuntimeNode; }
 
 	void Construct(UMovieGraphNode* InRuntimeNode);
+
+	void OnRuntimeNodeChanged(UMovieGraphNode* InChangedNode);
+
+	//~ Begin UEdGraphNode Interface
+	virtual void ReconstructNode() override;
+	//~ End UEdGraphNode Interface
+
 protected:
 	virtual bool ShouldCreatePin(const UMovieGraphPin* InPin) const;
 
 	void CreatePins(const TArray<UMovieGraphPin*>& InInputPins, const TArray<UMovieGraphPin*>& InOutputPins);
 
 	static FEdGraphPinType GetPinType(const UMovieGraphPin* InPin);
-	
+
+	/** Recreate the pins on this node, discarding all existing pins. */
+	void ReconstructPins();
+
+protected:
 	UPROPERTY()
 	TObjectPtr<UMovieGraphNode> RuntimeNode;
 };
@@ -44,7 +49,13 @@ public:
 	virtual void AllocateDefaultPins() override;
 	virtual FText GetNodeTitle(ENodeTitleType::Type TitleType) const override;
 	virtual FText GetTooltipText() const override;
+	virtual void GetNodeContextMenuActions(UToolMenu* Menu, UGraphNodeContextMenuContext* Context) const override;
 	//~ End UEdGraphNode Interface
 
+protected:
+	/** Promote the property with the given name to a pin exposed on the node. */
+	void PromotePropertyToPin(const FName PropertyName) const;
 
+private:
+	void GetPropertyPromotionContextMenuActions(UToolMenu* Menu, UGraphNodeContextMenuContext* Context) const;
 };
diff --git a/Engine/Plugins/MovieScene/MovieRenderPipeline/Source/MovieRenderPipelineEditor/Private/Graph/MovieGraphSchema.cpp b/Engine/Plugins/MovieScene/MovieRenderPipeline/Source/MovieRenderPipelineEditor/Private/Graph/MovieGraphSchema.cpp
--- a/Engine/Plugins/MovieScene/MovieRenderPipeline/Source/MovieRenderPipelineEditor/Private/Graph/MovieGraphSchema.cpp
+++ b/Engine/Plugins/MovieScene/MovieRenderPipeline/Source/MovieRenderPipelineEditor/Private/Graph/MovieGraphSchema.cpp
@@ -247,7 +247,7 @@ UEdGraphNode* FMovieGraphSchemaAction_NewNativeElement::PerformAction(UEdGraph*
 	// Now create the editor graph node
 	FGraphNodeCreator<UMoviePipelineEdGraphNode> NodeCreator(*ParentGraph);
 	UMoviePipelineEdGraphNode* GraphNode = NodeCreator.CreateUserInvokedNode(bSelectNewNode);
-	GraphNode->SetRuntimeNode(RuntimeNode);
+	GraphNode->Construct(RuntimeNode);
 	GraphNode->NodePosX = Location.X;
 	GraphNode->NodePosY = Location.Y;
 
