MoviePipeline: Added "Details" tab and "Members" (Inputs, Outputs, Variables) tab to the node graph.

The Details tab will react to selection changes in the graph. The Members tab is just boilerplate at the moment and needs some wiring to dynamically generate the correct contents based on what's being displayed in the graph.

Graphs are currently saved directly on the owning queue, rather than separate assets. To save a graph, save the associated queue. Nodes, node connections, and node positions currently round-trip.

diff --git a/Engine/Plugins/MovieScene/MovieRenderPipeline/Source/MovieRenderPipelineCore/Public/Graph/MovieGraphConfig.h b/Engine/Plugins/MovieScene/MovieRenderPipeline/Source/MovieRenderPipelineCore/Public/Graph/MovieGraphConfig.h
--- a/Engine/Plugins/MovieScene/MovieRenderPipeline/Source/MovieRenderPipelineCore/Public/Graph/MovieGraphConfig.h
+++ b/Engine/Plugins/MovieScene/MovieRenderPipeline/Source/MovieRenderPipelineCore/Public/Graph/MovieGraphConfig.h
@@ -60,7 +60,7 @@ public:
 #if WITH_EDITORONLY_DATA
 	// Not strongly typed to avoid a circular dependency between the editor only module
 	// and the runtime module, but it should be a UMoviePipelineEdGraph.
-	UPROPERTY()
+	UPROPERTY(Transient)
 	TObjectPtr<UEdGraph> PipelineEdGraph;
 #endif
 
diff --git a/Engine/Plugins/MovieScene/MovieRenderPipeline/Source/MovieRenderPipelineCore/Public/Graph/MovieGraphNode.h b/Engine/Plugins/MovieScene/MovieRenderPipeline/Source/MovieRenderPipelineCore/Public/Graph/MovieGraphNode.h
--- a/Engine/Plugins/MovieScene/MovieRenderPipeline/Source/MovieRenderPipelineCore/Public/Graph/MovieGraphNode.h
+++ b/Engine/Plugins/MovieScene/MovieRenderPipeline/Source/MovieRenderPipelineCore/Public/Graph/MovieGraphNode.h
@@ -23,7 +23,7 @@ DECLARE_MULTICAST_DELEGATE_OneParam(FOnMovieGraphNodeChanged, UMovieGraphNode*);
 * which contains data about it's visual position in the graph, comments, etc.
 */
 UCLASS()
-class MOVIERENDERPIPELINECORE_API UMovieGraphNode: public UObject
+class MOVIERENDERPIPELINECORE_API UMovieGraphNode : public UObject
 {
 	GENERATED_BODY()
 
@@ -64,6 +64,14 @@ public:
 	UMovieGraphPin* GetInputPin(const FName& InPinLabel) const;
 	UMovieGraphPin* GetOutputPin(const FName& InPinLabel) const;
 
+#if WITH_EDITOR
+	int32 GetNodePosX() const { return NodePosX; }
+	int32 GetNodePosY() const { return NodePosY; }
+
+	void SetNodePosX(const int32 InNodePosX) { NodePosX = InNodePosX; }
+	void SetNodePosY(const int32 InNodePosY) { NodePosY = InNodePosY; }
+#endif
+
 public:
 	FOnMovieGraphNodeChanged OnNodeChangedDelegate;
 
@@ -95,6 +103,14 @@ protected:
 
 	UPROPERTY()
 	TArray<FName> ExposedDynamicPropertyNames;
+
+#if WITH_EDITORONLY_DATA
+	UPROPERTY()
+	int32 NodePosX = 0;
+
+	UPROPERTY()
+	int32 NodePosY = 0;
+#endif
 };
 
 // Dummy test nodes
diff --git a/Engine/Plugins/MovieScene/MovieRenderPipeline/Source/MovieRenderPipelineEditor/MovieRenderPipelineEditor.Build.cs b/Engine/Plugins/MovieScene/MovieRenderPipeline/Source/MovieRenderPipelineEditor/MovieRenderPipelineEditor.Build.cs
--- a/Engine/Plugins/MovieScene/MovieRenderPipeline/Source/MovieRenderPipelineEditor/MovieRenderPipelineEditor.Build.cs
+++ b/Engine/Plugins/MovieScene/MovieRenderPipeline/Source/MovieRenderPipelineEditor/MovieRenderPipelineEditor.Build.cs
@@ -67,6 +67,7 @@ public class MovieRenderPipelineEditor : ModuleRules
 				"ToolWidgets",
 				"ConsoleVariablesEditor",
 				"StructUtils",
+				"GraphEditor",
 			}
         );
 
diff --git a/Engine/Plugins/MovieScene/MovieRenderPipeline/Source/MovieRenderPipelineEditor/Private/Graph/MovieEdGraph.cpp b/Engine/Plugins/MovieScene/MovieRenderPipeline/Source/MovieRenderPipelineEditor/Private/Graph/MovieEdGraph.cpp
--- a/Engine/Plugins/MovieScene/MovieRenderPipeline/Source/MovieRenderPipelineEditor/Private/Graph/MovieEdGraph.cpp
+++ b/Engine/Plugins/MovieScene/MovieRenderPipeline/Source/MovieRenderPipelineEditor/Private/Graph/MovieEdGraph.cpp
@@ -15,7 +15,6 @@ UMovieGraphConfig* UMoviePipelineEdGraph::GetPipelineGraph() const
 	return CastChecked<UMovieGraphConfig>(GetOuter());
 }
 
-
 void UMoviePipelineEdGraph::InitFromRuntimeGraph(UMovieGraphConfig* InGraph)
 {
 	// Don't allow reinitialization of an existing graph
@@ -31,6 +30,7 @@ void UMoviePipelineEdGraph::InitFromRuntimeGraph(UMovieGraphConfig* InGraph)
 		UMoviePipelineEdGraphNodeInput* InputGraphNode = InputNodeCreator.CreateNode(bSelectNewNode);
 		InputGraphNode->Construct(InputNode);
 		InputNodeCreator.Finalize();
+		NodeLookup.Add(InputNode, InputGraphNode);
 	}
 
 	// Output
@@ -40,6 +40,7 @@ void UMoviePipelineEdGraph::InitFromRuntimeGraph(UMovieGraphConfig* InGraph)
 		UMoviePipelineEdGraphNodeOutput* OutputGraphNode = OutputNodeCreator.CreateNode(bSelectNewNode);
 		OutputGraphNode->Construct(OutputNode);
 		OutputNodeCreator.Finalize();
+		NodeLookup.Add(OutputNode, OutputGraphNode);
 	}
 
 	// Create the rest of the nodes in the graph
@@ -128,7 +129,6 @@ void UMoviePipelineEdGraph::CreateLinks(UMoviePipelineEdGraphNodeBase* InGraphNo
 		}
 	};
 
-
 	if (bCreateInboundLinks)
 	{
 		CreateLinks(RuntimeNode->GetInputPins(), EEdGraphPinDirection::EGPD_Input);
diff --git a/Engine/Plugins/MovieScene/MovieRenderPipeline/Source/MovieRenderPipelineEditor/Private/Graph/MovieEdGraphNode.cpp b/Engine/Plugins/MovieScene/MovieRenderPipeline/Source/MovieRenderPipelineEditor/Private/Graph/MovieEdGraphNode.cpp
--- a/Engine/Plugins/MovieScene/MovieRenderPipeline/Source/MovieRenderPipelineEditor/Private/Graph/MovieEdGraphNode.cpp
+++ b/Engine/Plugins/MovieScene/MovieRenderPipeline/Source/MovieRenderPipelineEditor/Private/Graph/MovieEdGraphNode.cpp
@@ -5,6 +5,7 @@
 #include "Graph/MovieGraphPin.h"
 #include "Graph/MovieGraphNode.h"
 #include "EdGraph/EdGraphPin.h"
+#include "Misc/TransactionObjectEvent.h"
 #include "MovieEdGraph.h"
 #include "PropertyBag.h"
 #include "ToolMenu.h"
@@ -18,13 +19,24 @@ void UMoviePipelineEdGraphNodeBase::Construct(UMovieGraphNode* InRuntimeNode)
 	RuntimeNode->GraphNode = this;
 	RuntimeNode->OnNodeChangedDelegate.AddUObject(this, &UMoviePipelineEdGraphNodeBase::OnRuntimeNodeChanged);
 	
-	// NodePosX = InRuntimeNode->PositionX;
-	// NodePosY = InRuntimeNode->PositionY;
+	NodePosX = InRuntimeNode->GetNodePosX();
+	NodePosY = InRuntimeNode->GetNodePosY();
 	// NodeComment = InRuntimeNode->NodeComment;
 	// bCommentBubblePinned = InRuntimeNode->bCommentBubblePinned;
 	// bCommentBubbleVisible = InRuntimeNode->bCommentBubbleVisible;
 }
 
+void UMoviePipelineEdGraphNodeBase::PostTransacted(const FTransactionObjectEvent& TransactionEvent)
+{
+	const TArray<FName> ChangedProperties = TransactionEvent.GetChangedProperties();
+
+	if (ChangedProperties.Contains(GET_MEMBER_NAME_CHECKED(UEdGraphNode, NodePosX)) ||
+		ChangedProperties.Contains(GET_MEMBER_NAME_CHECKED(UEdGraphNode, NodePosY)))
+	{
+		UpdatePosition();
+	}
+}
+
 FEdGraphPinType UMoviePipelineEdGraphNodeBase::GetPinType(const UMovieGraphPin* InPin)
 {
 	FEdGraphPinType EdPinType;
@@ -35,6 +47,16 @@ FEdGraphPinType UMoviePipelineEdGraphNodeBase::GetPinType(const UMovieGraphPin*
 	return EdPinType;
 }
 
+void UMoviePipelineEdGraphNodeBase::UpdatePosition()
+{
+	if (RuntimeNode)
+	{
+		RuntimeNode->Modify();
+		RuntimeNode->SetNodePosX(NodePosX);
+		RuntimeNode->SetNodePosY(NodePosY);
+	}
+}
+
 void UMoviePipelineEdGraphNode::AllocateDefaultPins()
 {
 	if(RuntimeNode)
diff --git a/Engine/Plugins/MovieScene/MovieRenderPipeline/Source/MovieRenderPipelineEditor/Private/Graph/MovieEdGraphNode.h b/Engine/Plugins/MovieScene/MovieRenderPipeline/Source/MovieRenderPipelineEditor/Private/Graph/MovieEdGraphNode.h
--- a/Engine/Plugins/MovieScene/MovieRenderPipeline/Source/MovieRenderPipelineEditor/Private/Graph/MovieEdGraphNode.h
+++ b/Engine/Plugins/MovieScene/MovieRenderPipeline/Source/MovieRenderPipelineEditor/Private/Graph/MovieEdGraphNode.h
@@ -23,6 +23,10 @@ public:
 	//~ Begin UEdGraphNode Interface
 	virtual void ReconstructNode() override;
 	//~ End UEdGraphNode Interface
+	
+	//~ Begin UObject interface
+	virtual void PostTransacted(const FTransactionObjectEvent& TransactionEvent) override;
+	//~ End UObject interface
 
 protected:
 	virtual bool ShouldCreatePin(const UMovieGraphPin* InPin) const;
@@ -34,6 +38,8 @@ protected:
 	/** Recreate the pins on this node, discarding all existing pins. */
 	void ReconstructPins();
 
+	void UpdatePosition();
+
 protected:
 	UPROPERTY()
 	TObjectPtr<UMovieGraphNode> RuntimeNode;
diff --git a/Engine/Plugins/MovieScene/MovieRenderPipeline/Source/MovieRenderPipelineEditor/Private/Graph/MovieGraphAssetEditor.cpp b/Engine/Plugins/MovieScene/MovieRenderPipeline/Source/MovieRenderPipelineEditor/Private/Graph/MovieGraphAssetEditor.cpp
deleted file mode 100644
diff --git a/Engine/Plugins/MovieScene/MovieRenderPipeline/Source/MovieRenderPipelineEditor/Private/Graph/MovieGraphAssetEditor.h b/Engine/Plugins/MovieScene/MovieRenderPipeline/Source/MovieRenderPipelineEditor/Private/Graph/MovieGraphAssetEditor.h
deleted file mode 100644
diff --git a/Engine/Plugins/MovieScene/MovieRenderPipeline/Source/MovieRenderPipelineEditor/Private/Graph/MovieGraphAssetToolkit.cpp b/Engine/Plugins/MovieScene/MovieRenderPipeline/Source/MovieRenderPipelineEditor/Private/Graph/MovieGraphAssetToolkit.cpp
--- a/Engine/Plugins/MovieScene/MovieRenderPipeline/Source/MovieRenderPipelineEditor/Private/Graph/MovieGraphAssetToolkit.cpp
+++ b/Engine/Plugins/MovieScene/MovieRenderPipeline/Source/MovieRenderPipelineEditor/Private/Graph/MovieGraphAssetToolkit.cpp
@@ -2,88 +2,184 @@
 
 #include "MovieGraphAssetToolkit.h"
 
-#include "Framework/MultiBox/MultiBoxBuilder.h"
-#include "Styling/CoreStyle.h"
+#include "Graph/MovieGraphConfig.h"
+#include "PropertyEditorModule.h"
+#include "SMovieGraphMembersTabContent.h"
 #include "ToolMenus.h"
 #include "Widgets/Docking/SDockTab.h"
-#include "Widgets/Images/SImage.h"
 #include "Widgets/Graph/SMoviePipelineGraphPanel.h"
-#include "Widgets/SOverlay.h"
 
 #define LOCTEXT_NAMESPACE "MovieGraphAssetToolkit"
 
-const FName FMovieGraphAssetToolkit::ContentTabId(TEXT("MovieGraphAssetToolkit"));
+const FName FMovieGraphAssetToolkit::AppIdentifier(TEXT("MovieGraphAssetEditorApp"));
+const FName FMovieGraphAssetToolkit::GraphTabId(TEXT("MovieGraphAssetToolkit"));
+const FName FMovieGraphAssetToolkit::DetailsTabId(TEXT("MovieGraphAssetToolkitDetails"));
+const FName FMovieGraphAssetToolkit::MembersTabId(TEXT("MovieGraphAssetToolkitMembers"));
 
-FMovieGraphAssetToolkit::FMovieGraphAssetToolkit(UAssetEditor* InOwningAssetEditor)
-	: FBaseAssetToolkit(InOwningAssetEditor)
+// Temporary cvar to enable/disable upgrading to a graph-based configuration
+static TAutoConsoleVariable<bool> CVarMoviePipelineEnableRenderGraph(
+	TEXT("MoviePipeline.EnableRenderGraph"),
+	false,
+	TEXT("Determines if the Render Graph feature is enabled in the UI. This is a highly experimental feature and is not ready for use.")
+);
+
+void FMovieGraphAssetToolkit::RegisterTabSpawners(const TSharedRef<FTabManager>& InTabManager)
+{
+	WorkspaceMenuCategory = InTabManager->AddLocalWorkspaceMenuCategory(
+		LOCTEXT("WorkspaceMenu_MovieGraphAssetToolkit", "Render Graph Editor"));
+
+	FAssetEditorToolkit::RegisterTabSpawners(InTabManager);
+
+	InTabManager->RegisterTabSpawner(GraphTabId, FOnSpawnTab::CreateSP(this, &FMovieGraphAssetToolkit::SpawnTab_RenderGraphEditor))
+		.SetDisplayName(LOCTEXT("RenderGraphTab", "Render Graph"))
+		.SetGroup(WorkspaceMenuCategory.ToSharedRef())
+		.SetIcon(FSlateIcon(FAppStyle::GetAppStyleSetName(), "LevelEditor.Tabs.Viewports"));
+
+	InTabManager->RegisterTabSpawner(DetailsTabId, FOnSpawnTab::CreateSP(this, &FMovieGraphAssetToolkit::SpawnTab_RenderGraphDetails))
+		.SetDisplayName(LOCTEXT("DetailsTab", "Details"))
+		.SetGroup(WorkspaceMenuCategory.ToSharedRef())
+		.SetIcon(FSlateIcon(FAppStyle::GetAppStyleSetName(), "LevelEditor.Tabs.Details"));
+
+	InTabManager->RegisterTabSpawner(MembersTabId, FOnSpawnTab::CreateSP(this, &FMovieGraphAssetToolkit::SpawnTab_RenderGraphMembers))
+		.SetDisplayName(LOCTEXT("MembersTab", "Members"))
+		.SetGroup(WorkspaceMenuCategory.ToSharedRef())
+		.SetIcon(FSlateIcon(FAppStyle::GetAppStyleSetName(), "LevelEditor.Tabs.Outliner"));
+}
+
+void FMovieGraphAssetToolkit::UnregisterTabSpawners(const TSharedRef<FTabManager>& InTabManager)
+{
+	FAssetEditorToolkit::UnregisterTabSpawners(InTabManager);
+	
+	InTabManager->UnregisterTabSpawner(GraphTabId);
+	InTabManager->UnregisterTabSpawner(DetailsTabId);
+	InTabManager->UnregisterTabSpawner(MembersTabId);
+}
+
+void FMovieGraphAssetToolkit::InitMovieGraphAssetToolkit(const EToolkitMode::Type Mode, const TSharedPtr<class IToolkitHost>& InitToolkitHost, UMovieGraphConfig* InitGraph)
 {
+	InitialGraph = InitGraph;
+	
 	// Note: Changes to the layout should include a increment to the layout's ID, i.e.
 	// MoviePipelineRenderGraphEditor[X] -> MoviePipelineRenderGraphEditor[X+1]. Otherwise, layouts may be messed up
 	// without a full reset to layout defaults inside the editor.
 	const FName LayoutString = FName("MoviePipelineRenderGraphEditor1");
 
 	// Override the Default Layout provided by FBaseAssetToolkit to hide the viewport and details panel tabs.
-	StandaloneDefaultLayout = FTabManager::NewLayout(FName(LayoutString))
+	const TSharedRef<FTabManager::FLayout> Layout = FTabManager::NewLayout(FName(LayoutString))
 		->AddArea
 		(
 			FTabManager::NewPrimaryArea()->SetOrientation(Orient_Vertical)
 			->Split
 			(
-				FTabManager::NewStack()
+				FTabManager::NewSplitter()
+				->SetOrientation(Orient_Horizontal)
 				->SetSizeCoefficient(1.f)
-				->SetHideTabWell(true)
-				->AddTab(ContentTabId, ETabState::OpenedTab)
+				->Split
+				(
+					FTabManager::NewStack()
+					->SetSizeCoefficient(0.2f)
+					->AddTab(MembersTabId, ETabState::OpenedTab)
+				)
+				->Split
+				(
+					FTabManager::NewStack()
+					->SetSizeCoefficient(0.6f)
+					->SetHideTabWell(true)
+					->AddTab(GraphTabId, ETabState::OpenedTab)
+				)
+				->Split
+				(
+					FTabManager::NewStack()
+					->SetSizeCoefficient(0.2f)
+					->AddTab(DetailsTabId, ETabState::OpenedTab)
+				)
 			)
 		);
+
+	const bool bCreateDefaultStandaloneMenu = true;
+	const bool bCreateDefaultToolbar = true;
+	InitAssetEditor(Mode, InitToolkitHost, AppIdentifier, Layout, bCreateDefaultStandaloneMenu, bCreateDefaultToolbar, InitGraph);
 }
 
-void FMovieGraphAssetToolkit::RegisterTabSpawners(const TSharedRef<FTabManager>& InTabManager)
+TSharedRef<SDockTab> FMovieGraphAssetToolkit::SpawnTab_RenderGraphEditor(const FSpawnTabArgs& Args)
 {
-	WorkspaceMenuCategory = InTabManager->AddLocalWorkspaceMenuCategory(
-		LOCTEXT("WorkspaceMenu_MoviePipelineRenderGraphAssetToolkit", "Movie Pipeline Render Graph Asset Editor"));
-
-	FBaseAssetToolkit::RegisterTabSpawners(InTabManager);
+	TSharedRef<SDockTab> NewDockTab = SNew(SDockTab)
+		.TabColorScale(GetTabColorScale())
+		.Label(LOCTEXT("GraphTab_Title", "Graph"))
+		[
+			SNew(SVerticalBox)
+			+ SVerticalBox::Slot()
+			[
+				SAssignNew(MovieGraphWidget, SMoviePipelineGraphPanel)
+				.Graph(InitialGraph)
+				.OnGraphSelectionChanged_Lambda([this](const TArray<UObject*>& NewSelection) -> void
+				{
+					SelectedGraphObjectsDetailsWidget->SetObjects(NewSelection);
+				})
+			]
+		];
 
-	InTabManager->RegisterTabSpawner(ContentTabId, FOnSpawnTab::CreateSP(this, &FMovieGraphAssetToolkit::SpawnTab_RenderGraphEditor))
-		.SetDisplayName(LOCTEXT("RenderGraphTab", "Render Graph"))
-		.SetGroup(WorkspaceMenuCategory.ToSharedRef());
+	return NewDockTab;
 }
 
-void FMovieGraphAssetToolkit::UnregisterTabSpawners(const TSharedRef<FTabManager>& InTabManager)
+TSharedRef<SDockTab> FMovieGraphAssetToolkit::SpawnTab_RenderGraphMembers(const FSpawnTabArgs& Args)
 {
-	InTabManager->UnregisterTabSpawner(ContentTabId);
+	TSharedRef<SDockTab> NewDockTab = SNew(SDockTab)
+		.TabColorScale(GetTabColorScale())
+		.Label(LOCTEXT("MembersTab_Title", "Members"))
+		[
+			SNew(SMovieGraphMembersTabContent)
+		];
+
+	return NewDockTab;
 }
 
-TSharedRef<SDockTab> FMovieGraphAssetToolkit::SpawnTab_RenderGraphEditor(const FSpawnTabArgs& Args)
+TSharedRef<SDockTab> FMovieGraphAssetToolkit::SpawnTab_RenderGraphDetails(const FSpawnTabArgs& Args)
 {
-	// Tabs and Windows have different backgrounds, so we copy the window style to make the 'void' area look similar.
-	const FSlateBrush* BackgroundBrush = &FCoreStyle::Get().GetWidgetStyle<FWindowStyle>("Window").ChildBackgroundBrush;
+	FPropertyEditorModule& PropertyEditorModule = FModuleManager::Get().LoadModuleChecked<FPropertyEditorModule>("PropertyEditor");
+	FDetailsViewArgs DetailsViewArgs;
+	DetailsViewArgs.bAllowSearch = false;
+	DetailsViewArgs.NameAreaSettings = FDetailsViewArgs::HideNameArea;
+	DetailsViewArgs.bHideSelectionTip = true;
+	DetailsViewArgs.ViewIdentifier = "MovieGraphSettings";
 
-	TSharedRef<SDockTab> NewDockTab = SNew(SDockTab)
+	SelectedGraphObjectsDetailsWidget = PropertyEditorModule.CreateDetailView(DetailsViewArgs);
+	
+	return SNew(SDockTab)
 		.TabColorScale(GetTabColorScale())
+		.Label(LOCTEXT("DetailsTab_Title", "Details"))
 		[
-			SNew(SOverlay)
-			+ SOverlay::Slot()
-			[
-				SNew(SImage)
-				.Image(BackgroundBrush)
-			]
-			+ SOverlay::Slot()
-			[
-				SNew(SVerticalBox)
-				+ SVerticalBox::Slot()
-				[
-					SNew(SMoviePipelineGraphPanel)
-				]
-			]
+			SelectedGraphObjectsDetailsWidget.ToSharedRef()
 		];
-
-	return NewDockTab;
 }
 
 FName FMovieGraphAssetToolkit::GetToolkitFName() const
 {
-	return FName("MovieGraphAssetToolkit");
+	return FName("MovieGraphEditor");
+}
+
+FText FMovieGraphAssetToolkit::GetBaseToolkitName() const
+{
+	return LOCTEXT("MovieGraphAppLabel", "Movie Graph Editor");
+}
+
+FString FMovieGraphAssetToolkit::GetWorldCentricTabPrefix() const
+{
+	return TEXT("MovieGraphEditor");
+}
+
+FLinearColor FMovieGraphAssetToolkit::GetWorldCentricTabColorScale() const
+{
+	return FLinearColor::White;
+}
+
+void FMovieGraphAssetToolkit::SaveAsset_Execute()
+{
+	// Perform the default save process
+	// TODO: This will fail silently on a transient graph and won't trigger a Save As
+	FAssetEditorToolkit::SaveAsset_Execute();
+
+	// TODO: Any custom save logic here
 }
 
 #undef LOCTEXT_NAMESPACE
diff --git a/Engine/Plugins/MovieScene/MovieRenderPipeline/Source/MovieRenderPipelineEditor/Private/Graph/MovieGraphAssetToolkit.h b/Engine/Plugins/MovieScene/MovieRenderPipeline/Source/MovieRenderPipelineEditor/Private/Graph/MovieGraphAssetToolkit.h
--- a/Engine/Plugins/MovieScene/MovieRenderPipeline/Source/MovieRenderPipelineEditor/Private/Graph/MovieGraphAssetToolkit.h
+++ b/Engine/Plugins/MovieScene/MovieRenderPipeline/Source/MovieRenderPipelineEditor/Private/Graph/MovieGraphAssetToolkit.h
@@ -3,23 +3,46 @@
 #pragma once
 
 #include "CoreMinimal.h"
-#include "Tools/BaseAssetToolkit.h"
+#include "Toolkits/AssetEditorToolkit.h"
 
-class FMovieGraphAssetToolkit :  public FBaseAssetToolkit
+class FMovieGraphAssetToolkit :  public FAssetEditorToolkit
 {
 public:
-	FMovieGraphAssetToolkit(UAssetEditor* InOwningAssetEditor);
+	FMovieGraphAssetToolkit() {}
 	virtual ~FMovieGraphAssetToolkit() override {}
 
-	//~ Begin FBaseAssetToolkit interface
+	//~ Begin FAssetEditorToolkit interface
 	virtual void RegisterTabSpawners(const TSharedRef<FTabManager>& InTabManager) override;
 	virtual void UnregisterTabSpawners(const TSharedRef<FTabManager>& InTabManager) override;
 	virtual FName GetToolkitFName() const override;
-	//~ End FBaseAssetToolkit interface
+	virtual FText GetBaseToolkitName() const override;
+	virtual FString GetWorldCentricTabPrefix() const override;
+	virtual void SaveAsset_Execute() override;
+	//~ End FAssetEditorToolkit interface
+
+	//~ Begin IToolkit interface
+	virtual FLinearColor GetWorldCentricTabColorScale() const override;
+	//~ End IToolkit interface
+
+	void InitMovieGraphAssetToolkit(const EToolkitMode::Type Mode, const TSharedPtr<class IToolkitHost>& InitToolkitHost, class UMovieGraphConfig* InitGraph);
 
 private:
 	TSharedRef<SDockTab> SpawnTab_RenderGraphEditor(const FSpawnTabArgs& Args);
+	TSharedRef<SDockTab> SpawnTab_RenderGraphMembers(const FSpawnTabArgs& Args);
+	TSharedRef<SDockTab> SpawnTab_RenderGraphDetails(const FSpawnTabArgs& Args);
 
 private:
-	static const FName ContentTabId;
+	/** The details panel for the selected object(s) in the graph */
+	TSharedPtr<IDetailsView> SelectedGraphObjectsDetailsWidget;
+
+	/** The widget that contains the main node graph */
+	TSharedPtr<class SMoviePipelineGraphPanel> MovieGraphWidget;
+
+	/** The graph that the editor was initialized with */
+	TObjectPtr<UMovieGraphConfig> InitialGraph;
+
+	static const FName AppIdentifier;
+	static const FName GraphTabId;
+	static const FName DetailsTabId;
+	static const FName MembersTabId;
 };
diff --git a/Engine/Plugins/MovieScene/MovieRenderPipeline/Source/MovieRenderPipelineEditor/Private/Graph/SMovieGraphMembersTabContent.cpp b/Engine/Plugins/MovieScene/MovieRenderPipeline/Source/MovieRenderPipelineEditor/Private/Graph/SMovieGraphMembersTabContent.cpp
new file mode 100644
--- /dev/null
+++ b/Engine/Plugins/MovieScene/MovieRenderPipeline/Source/MovieRenderPipelineEditor/Private/Graph/SMovieGraphMembersTabContent.cpp
@@ -0,0 +1,118 @@
+// Copyright Epic Games, Inc. All Rights Reserved.
+
+#include "SMovieGraphMembersTabContent.h"
+
+#include "EdGraph/EdGraphSchema.h"
+#include "SGraphActionMenu.h"
+
+#define LOCTEXT_NAMESPACE "SMoviePipelineMembersTabContent"
+
+const TArray<FText> SMovieGraphMembersTabContent::ActionMenuSectionNames
+{
+	LOCTEXT("ActionMenuSectionName_Invalid", "INVALID"),
+	LOCTEXT("ActionMenuSectionName_Inputs", "Inputs"),
+	LOCTEXT("ActionMenuSectionName_Outputs", "Outputs"),
+	LOCTEXT("ActionMenuSectionName_Variables", "Variables")
+};
+
+void SMovieGraphMembersTabContent::Construct(const FArguments& InArgs)
+{
+	ChildSlot
+	[
+		SNew(SGraphActionMenu)
+		.AutoExpandActionMenu(true)
+		.OnCollectStaticSections(this, &SMovieGraphMembersTabContent::CollectStaticSections)
+		.OnGetSectionTitle(this, &SMovieGraphMembersTabContent::GetSectionTitle)
+		.OnGetSectionWidget(this, &SMovieGraphMembersTabContent::GetSectionWidget)
+		.UseSectionStyling(true)
+		.OnCollectAllActions(this, &SMovieGraphMembersTabContent::CollectAllActions)
+	];
+}
+
+void SMovieGraphMembersTabContent::CollectAllActions(FGraphActionListBuilderBase& OutAllActions)
+{
+	// TODO: These are placeholder actions and need to be dynamically populated by what is in the graph
+	
+	const FText InputActionDesc = LOCTEXT("InputAction", "This is an input");
+	const FText InputActionCategory;
+	const FText InputActionDescription = LOCTEXT("InputActionTooltip", "This is an input tooltip.");
+	const FText InputActionKeywords;
+	const int32 InputActionSectionID = 1;
+	const TSharedPtr<FEdGraphSchemaAction> InputAction(new FEdGraphSchemaAction(InputActionCategory, InputActionDesc, InputActionDescription, 0, InputActionKeywords, InputActionSectionID));
+
+	const FText OutputActionDesc = LOCTEXT("OutputAction", "This is an output");
+	const FText OutputActionCategory;
+	const FText OutputActionDescription = LOCTEXT("OutputActionTooltip", "This is an output tooltip.");
+	const FText OutputActionKeywords;
+	const int32 OutputActionSectionID = 2;
+	const TSharedPtr<FEdGraphSchemaAction> OutputAction(new FEdGraphSchemaAction(OutputActionCategory, OutputActionDesc, OutputActionDescription, 0, OutputActionKeywords, OutputActionSectionID));
+
+	const FText VariableActionDesc = LOCTEXT("VariableAction", "This is a variable");
+	const FText VariableActionCategory;
+	const FText VariableActionDescription = LOCTEXT("VariableActionTooltip", "This is a variable tooltip.");
+	const FText VariableActionKeywords;
+	const int32 VariableActionSectionID = 3;
+	const TSharedPtr<FEdGraphSchemaAction> VariableAction(new FEdGraphSchemaAction(VariableActionCategory, VariableActionDesc, VariableActionDescription, 0, VariableActionKeywords, VariableActionSectionID));
+	
+	FGraphActionMenuBuilder ActionMenuBuilder;
+	ActionMenuBuilder.AddAction(InputAction);
+	ActionMenuBuilder.AddAction(OutputAction);
+	ActionMenuBuilder.AddAction(VariableAction);
+	
+	OutAllActions.Append(ActionMenuBuilder);
+}
+
+void SMovieGraphMembersTabContent::CollectStaticSections(TArray<int32>& StaticSectionIDs)
+{
+	// Start at index 1 to skip the invalid section
+	for (int32 Index = 1; Index < static_cast<int32>(EActionSection::COUNT); ++Index)
+	{
+		StaticSectionIDs.Add(Index);
+	}
+}
+
+FText SMovieGraphMembersTabContent::GetSectionTitle(int32 InSectionID)
+{
+	if (ensure(ActionMenuSectionNames.IsValidIndex(InSectionID)))
+	{
+		return ActionMenuSectionNames[InSectionID];
+	}
+
+	return FText::GetEmpty();
+}
+
+TSharedRef<SWidget> SMovieGraphMembersTabContent::GetSectionWidget(TSharedRef<SWidget> RowWidget, int32 InSectionID)
+{
+	FText ButtonTooltip;
+
+	switch (static_cast<EActionSection>(InSectionID))
+	{
+	case EActionSection::Inputs:
+		ButtonTooltip = LOCTEXT("ButtonTooltip_Inputs", "Add Input");
+		break;
+
+	case EActionSection::Outputs:
+		ButtonTooltip = LOCTEXT("ButtonTooltip_Outputs", "Add Output");
+		break;
+
+	case EActionSection::Variables:
+		ButtonTooltip = LOCTEXT("ButtonTooltip_Variables", "Add Variable");
+		break;
+
+	default:
+		return SNullWidget::NullWidget;
+	}
+	
+	return
+		SNew(SButton)
+		.ButtonStyle(FAppStyle::Get(), "SimpleButton")
+		.ContentPadding(FMargin(1, 0))
+		.ToolTipText(ButtonTooltip)
+		[
+			SNew(SImage)
+			.Image(FAppStyle::Get().GetBrush("Icons.PlusCircle"))
+			.ColorAndOpacity(FSlateColor::UseForeground())
+		];
+}
+
+#undef LOCTEXT_NAMESPACE
\ No newline at end of file
diff --git a/Engine/Plugins/MovieScene/MovieRenderPipeline/Source/MovieRenderPipelineEditor/Private/Graph/SMovieGraphMembersTabContent.h b/Engine/Plugins/MovieScene/MovieRenderPipeline/Source/MovieRenderPipelineEditor/Private/Graph/SMovieGraphMembersTabContent.h
new file mode 100644
--- /dev/null
+++ b/Engine/Plugins/MovieScene/MovieRenderPipeline/Source/MovieRenderPipelineEditor/Private/Graph/SMovieGraphMembersTabContent.h
@@ -0,0 +1,41 @@
+﻿// Copyright Epic Games, Inc. All Rights Reserved.
+
+#pragma once
+
+#include "Widgets/Input/SButton.h"
+#include "Widgets/SCompoundWidget.h"
+
+struct FGraphActionListBuilderBase;
+
+/**
+ * Contents of the "Members" tab in the graph asset editor.
+ */
+class SMovieGraphMembersTabContent : public SCompoundWidget
+{
+public:
+	SLATE_BEGIN_ARGS(SMovieGraphMembersTabContent) {}
+	SLATE_END_ARGS();
+	
+	void Construct(const FArguments& InArgs);
+
+private:
+	/** The section identifier in the action widget. */
+	enum class EActionSection : uint8
+	{
+		Invalid,
+		
+		Inputs,
+		Outputs,
+		Variables,
+
+		COUNT
+	};
+
+	/** The names of the sections in the action widget. */
+	static const TArray<FText> ActionMenuSectionNames;
+	
+	void CollectAllActions(FGraphActionListBuilderBase& OutAllActions);
+	void CollectStaticSections(TArray<int32>& StaticSectionIDs);
+	FText GetSectionTitle(int32 InSectionID);
+	TSharedRef<SWidget> GetSectionWidget(TSharedRef<SWidget> RowWidget, int32 InSectionID);
+};
diff --git a/Engine/Plugins/MovieScene/MovieRenderPipeline/Source/MovieRenderPipelineEditor/Private/Widgets/Graph/SMoviePipelineGraphPanel.cpp b/Engine/Plugins/MovieScene/MovieRenderPipeline/Source/MovieRenderPipelineEditor/Private/Widgets/Graph/SMoviePipelineGraphPanel.cpp
--- a/Engine/Plugins/MovieScene/MovieRenderPipeline/Source/MovieRenderPipelineEditor/Private/Widgets/Graph/SMoviePipelineGraphPanel.cpp
+++ b/Engine/Plugins/MovieScene/MovieRenderPipeline/Source/MovieRenderPipelineEditor/Private/Widgets/Graph/SMoviePipelineGraphPanel.cpp
@@ -51,6 +51,7 @@
 
 #include "Graph/MovieGraphConfig.h"
 #include "Graph/MovieEdGraph.h"
+#include "Graph/MovieEdGraphNode.h"
 #include "Graph/MovieGraphSchema.h"
 #include "Kismet2/BlueprintEditorUtils.h"
 #include "GraphEditor.h"
@@ -60,14 +61,7 @@
 UE_DISABLE_OPTIMIZATION_SHIP
 void SMoviePipelineGraphPanel::Construct(const FArguments& InArgs)
 {
-	FPropertyEditorModule& PropertyEditorModule = FModuleManager::Get().LoadModuleChecked<FPropertyEditorModule>("PropertyEditor");
-	FDetailsViewArgs DetailsViewArgs;
-	DetailsViewArgs.bAllowSearch = false;
-	DetailsViewArgs.NameAreaSettings = FDetailsViewArgs::HideNameArea;
-	DetailsViewArgs.bHideSelectionTip = true;
-	DetailsViewArgs.ColumnWidth = 0.7f;
-
-	JobDetailsPanelWidget = PropertyEditorModule.CreateDetailView(DetailsViewArgs);
+	OnGraphSelectionChangedEvent = InArgs._OnGraphSelectionChanged;
 
 	// Create the child widgets that need to know about our pipeline
 	PipelineQueueEditorWidget = SNew(SMoviePipelineQueueEditor)
@@ -97,8 +91,7 @@ void SMoviePipelineGraphPanel::Construct(const FArguments& InArgs)
 	InEvents.OnNodeDoubleClicked = FSingleNodeEvent::CreateSP(this, &SMoviePipelineGraphPanel::OnNodeDoubleClicked);
 	InEvents.OnTextCommitted = FOnNodeTextCommitted::CreateSP(this, &SMoviePipelineGraphPanel::OnNodeTitleCommitted);
 
-	CurrentGraph = NewObject<UMovieGraphConfig>(GetTransientPackage());
-	CurrentGraph->AddToRoot();
+	CurrentGraph = InArgs._Graph;
 	
 	UMoviePipelineEdGraph* GraphToEdit = Cast<UMoviePipelineEdGraph>(FBlueprintEditorUtils::CreateNewGraph(CurrentGraph, TEXT("MoviePipelineEdGraph"), UMoviePipelineEdGraph::StaticClass(), UMovieGraphSchema::StaticClass()));
 
@@ -121,30 +114,7 @@ void SMoviePipelineGraphPanel::Construct(const FArguments& InArgs)
 
 	ChildSlot
 	[
-		SNew(SSplitter)
-		.Orientation(EOrientation::Orient_Horizontal)
-		+ SSplitter::Slot()
-		.Value(3)
-		[
-			GraphEditor
-		]
-		+ SSplitter::Slot()
-		.Value(1)
-		[
-			SNew(SBorder)
-			.BorderImage(FAppStyle::GetBrush("ToolPanel.GroupBorder"))
-			.Padding(FMargin(1.f, 1.0f))
-			.Content()
-			[
-				JobDetailsPanelWidget.ToSharedRef()
-			]
-		]
-		+SSplitter::Slot()
-		[
-			SNew(SButton)
-			.Text(LOCTEXT("OK", "OK"))
-			.OnClicked(this, &SMoviePipelineGraphPanel::OnDebugButtonClicked)
-		]
+		GraphEditor
 	];
 }
 
@@ -155,15 +125,22 @@ FReply SMoviePipelineGraphPanel::OnDebugButtonClicked()
 	return FReply::Handled();
 }
 
-void SMoviePipelineGraphPanel::OnSelectedNodesChanged(const TSet<class UObject*>& NewSelection)
+void SMoviePipelineGraphPanel::OnSelectedNodesChanged(const TSet<UObject*>& NewSelection)
 {
 	TArray<UObject*> SelectedObjects;
 	for (UObject* Obj : NewSelection)
 	{
-		SelectedObjects.Add(Obj);
+		// The selected objects will be editor objects. Get the underlying runtime object for each editor object.
+		if (const UMoviePipelineEdGraphNodeBase* NodeBase = Cast<UMoviePipelineEdGraphNodeBase>(Obj))
+		{
+			if (UMovieGraphNode* GraphNode = NodeBase->GetRuntimeNode())
+			{
+				SelectedObjects.Add(GraphNode);
+			}
+		}
 	}
 
-	JobDetailsPanelWidget->SetObjects(SelectedObjects);
+	OnGraphSelectionChangedEvent.ExecuteIfBound(SelectedObjects);
 }
 
 void SMoviePipelineGraphPanel::OnNodeDoubleClicked(class UEdGraphNode* Node)
@@ -388,33 +365,9 @@ void SMoviePipelineGraphPanel::OnSelectionChanged(const TArray<UMoviePipelineExe
 		Jobs.Add(Job);
 	}
 	
-	JobDetailsPanelWidget->SetObjects(Jobs);
 	NumSelectedJobs = InSelectedJobs.Num();
 }
 
-int32 SMoviePipelineGraphPanel::GetDetailsViewWidgetIndex() const
-{
-	return NumSelectedJobs == 0;
-}
-
-bool SMoviePipelineGraphPanel::IsDetailsViewEnabled() const
-{
-	TArray<TWeakObjectPtr<UObject>> OutObjects= JobDetailsPanelWidget->GetSelectedObjects();
-
-	bool bAllEnabled = true;
-	for (TWeakObjectPtr<UObject> Object : OutObjects)
-	{
-		const UMoviePipelineExecutorJob* Job = Cast<UMoviePipelineExecutorJob>(Object);
-		if (Job && Job->IsConsumed())
-		{
-			bAllEnabled = false;
-			break;
-		}
-	}
-
-	return bAllEnabled;
-}
-
 TSharedRef<SWidget> SMoviePipelineGraphPanel::OnGenerateSavedQueuesMenu()
 {
 	FMenuBuilder MenuBuilder(true, nullptr);
diff --git a/Engine/Plugins/MovieScene/MovieRenderPipeline/Source/MovieRenderPipelineEditor/Private/Widgets/Graph/SMoviePipelineGraphPanel.h b/Engine/Plugins/MovieScene/MovieRenderPipeline/Source/MovieRenderPipelineEditor/Private/Widgets/Graph/SMoviePipelineGraphPanel.h
--- a/Engine/Plugins/MovieScene/MovieRenderPipeline/Source/MovieRenderPipelineEditor/Private/Widgets/Graph/SMoviePipelineGraphPanel.h
+++ b/Engine/Plugins/MovieScene/MovieRenderPipeline/Source/MovieRenderPipelineEditor/Private/Widgets/Graph/SMoviePipelineGraphPanel.h
@@ -19,15 +19,17 @@ struct FAssetData;
 class SMoviePipelineGraphPanel : public SCompoundWidget
 {
 public:
+	DECLARE_DELEGATE_OneParam(FOnGraphSelectionChanged, TArray<UObject*>);
+	
 	SLATE_BEGIN_ARGS(SMoviePipelineGraphPanel)
-		: _BasePreset(nullptr)
+		: _Graph(nullptr)
 
 		{}
 
-		/*~ All following arguments are mutually-exclusive */
-		/*-------------------------------------------------*/
-		/** A preset asset to base the pipeline off */
-		SLATE_ARGUMENT(UMoviePipelineConfigBase*, BasePreset)
+		SLATE_EVENT(FOnGraphSelectionChanged, OnGraphSelectionChanged)
+
+		/** The graph that is initially displayed */
+		SLATE_ARGUMENT(class UMovieGraphConfig*, Graph)
 
 	SLATE_END_ARGS()
 
@@ -48,8 +50,6 @@ private:
 	void OnConfigWindowClosed();
 
 	void OnSelectionChanged(const TArray<UMoviePipelineExecutorJob*>& InSelectedJobs);
-	int32 GetDetailsViewWidgetIndex() const;
-	bool IsDetailsViewEnabled() const;
 
 	TSharedRef<SWidget> OnGenerateSavedQueuesMenu();
 	bool OpenSaveDialog(const FString& InDefaultPath, const FString& InNewNameSuggestion, FString& OutPackageName);
@@ -72,10 +72,9 @@ private:
 	/** The main movie pipeline queue editor widget */
 	TSharedPtr<SMoviePipelineQueueEditor> PipelineQueueEditorWidget;
 
-	/** The details panel for the selected job(s) */
-	TSharedPtr<IDetailsView> JobDetailsPanelWidget;
-
 	TWeakPtr<SWindow> WeakEditorWindow;
 	
 	int32 NumSelectedJobs;
+
+	FOnGraphSelectionChanged OnGraphSelectionChangedEvent;
 };
\ No newline at end of file
diff --git a/Engine/Plugins/MovieScene/MovieRenderPipeline/Source/MovieRenderPipelineEditor/Private/Widgets/SMoviePipelineQueuePanel.cpp b/Engine/Plugins/MovieScene/MovieRenderPipeline/Source/MovieRenderPipelineEditor/Private/Widgets/SMoviePipelineQueuePanel.cpp
--- a/Engine/Plugins/MovieScene/MovieRenderPipeline/Source/MovieRenderPipelineEditor/Private/Widgets/SMoviePipelineQueuePanel.cpp
+++ b/Engine/Plugins/MovieScene/MovieRenderPipeline/Source/MovieRenderPipelineEditor/Private/Widgets/SMoviePipelineQueuePanel.cpp
@@ -13,7 +13,7 @@
 #include "MoviePipelinePrimaryConfig.h"
 #include "MoviePipelineQueueSubsystem.h"
 #include "Graph/MovieGraphConfig.h"
-#include "Graph/MovieGraphAssetEditor.h"
+#include "Graph/MovieGraphAssetToolkit.h"
 
 // Slate Includes
 #include "Widgets/SBoxPanel.h"
@@ -368,23 +368,17 @@ void SMoviePipelineQueuePanel::OnEditJobConfigRequested(TWeakObjectPtr<UMoviePip
 	UMovieGraphConfig* GraphToEdit = nullptr;
 	if (InShot.IsValid() && InShot->IsUsingGraphConfiguration())
 	{
-		GraphToEdit = InShot->GetGraphPreset();
+		GraphToEdit = (InShot->GetGraphPreset() != nullptr) ? InShot->GetGraphPreset() : InShot->GetGraphConfig();
 	}
 	else if (InJob.IsValid() && InJob->IsUsingGraphConfiguration())
 	{
-		GraphToEdit = InJob->GetGraphPreset();
+		GraphToEdit = (InJob->GetGraphPreset() != nullptr) ? InJob->GetGraphPreset() : InJob->GetGraphConfig();
 	}
 
 	if (GraphToEdit)
 	{
-		if (UAssetEditorSubsystem* AssetEditorSubsystem = GEditor->GetEditorSubsystem<UAssetEditorSubsystem>())
-		{
-			if (UMovieGraphAssetEditor* AssetEditor = NewObject<UMovieGraphAssetEditor>(AssetEditorSubsystem, NAME_None, RF_Transient))
-			{
-				AssetEditor->SetObjectToEdit(GraphToEdit);
-				AssetEditor->Initialize();
-			}
-		}
+		const TSharedRef<FMovieGraphAssetToolkit> MovieGraphEditor(new FMovieGraphAssetToolkit());
+		MovieGraphEditor->InitMovieGraphAssetToolkit(EToolkitMode::Standalone, TSharedPtr<IToolkitHost>(), GraphToEdit);
 
 		return;
 	}
