MRQ: Fixed minor typo in output color space.

#rnx

diff --git a/Engine/Plugins/MovieScene/MovieRenderPipeline/Source/MovieRenderPipelineRenderPasses/Private/MoviePipelineEXROutput.cpp b/Engine/Plugins/MovieScene/MovieRenderPipeline/Source/MovieRenderPipelineRenderPasses/Private/MoviePipelineEXROutput.cpp
--- a/Engine/Plugins/MovieScene/MovieRenderPipeline/Source/MovieRenderPipelineRenderPasses/Private/MoviePipelineEXROutput.cpp
+++ b/Engine/Plugins/MovieScene/MovieRenderPipeline/Source/MovieRenderPipelineRenderPasses/Private/MoviePipelineEXROutput.cpp
@@ -643,10 +643,10 @@ UMoviePipelineImageSequenceOutput_EXR::FColorSpaceMetadata UMoviePipelineImageSe
 		else if (!InColorSettings->bDisableToneCurve)
 		{
 			TPair<UE::Color::EColorSpace, FString> ColorSpaceType = GetDisplayGamutTypeFn(HDRGetDefaultDisplayColorGamut());
-			UE::Color::FColorSpace OuptutCS = UE::Color::FColorSpace(ColorSpaceType.Key);
+			UE::Color::FColorSpace OutputCS = UE::Color::FColorSpace(ColorSpaceType.Key);
 
 			OutMetadata.DestinationName = ColorSpaceType.Value;
-			OutMetadata.Chromaticities = { OuptutCS.GetRedChromaticity(), OuptutCS.GetGreenChromaticity(), OuptutCS.GetBlueChromaticity(), OuptutCS.GetWhiteChromaticity() };
+			OutMetadata.Chromaticities = { OutputCS.GetRedChromaticity(), OutputCS.GetGreenChromaticity(), OutputCS.GetBlueChromaticity(), OutputCS.GetWhiteChromaticity() };
 		}
 		else
 		{
