MoviePipeline: Created a AssetDefinition for MovieGraphConfigs to allow creating them in the Content Browser. This is also gated behind the MoviePipeline.EnableRenderGraph cvar.

diff --git a/Engine/Plugins/MovieScene/MovieRenderPipeline/Source/MovieRenderPipelineEditor/Private/Graph/AssetDefinition_MovieGraphConfig.cpp b/Engine/Plugins/MovieScene/MovieRenderPipeline/Source/MovieRenderPipelineEditor/Private/Graph/AssetDefinition_MovieGraphConfig.cpp
new file mode 100644
--- /dev/null
+++ b/Engine/Plugins/MovieScene/MovieRenderPipeline/Source/MovieRenderPipelineEditor/Private/Graph/AssetDefinition_MovieGraphConfig.cpp
@@ -0,0 +1,24 @@
+// Copyright Epic Games, Inc. All Rights Reserved.
+
+#include "Graph/AssetDefinition_MovieGraphConfig.h"
+#include "Graph/MovieGraphAssetToolkit.h"
+
+#define LOCTEXT_NAMESPACE "AssetTypeActions"
+
+TSoftClassPtr<UObject> UAssetDefinition_MovieGraphConfig::GetAssetClass() const
+{
+	return UMovieGraphConfig::StaticClass();
+}
+
+EAssetCommandResult UAssetDefinition_MovieGraphConfig::OpenAssets(const FAssetOpenArgs& OpenArgs) const
+{
+	for (UMovieGraphConfig* GraphToEdit : OpenArgs.LoadObjects<UMovieGraphConfig>())
+	{
+		TSharedRef<FMovieGraphAssetToolkit> MovieGraphEditor = MakeShared<FMovieGraphAssetToolkit>();
+		MovieGraphEditor->InitMovieGraphAssetToolkit(EToolkitMode::Standalone, TSharedPtr<IToolkitHost>(), GraphToEdit);
+	}	
+
+	return EAssetCommandResult::Handled;
+}
+
+#undef LOCTEXT_NAMESPACE
diff --git a/Engine/Plugins/MovieScene/MovieRenderPipeline/Source/MovieRenderPipelineEditor/Private/Graph/AssetDefinition_MovieGraphConfig.h b/Engine/Plugins/MovieScene/MovieRenderPipeline/Source/MovieRenderPipelineEditor/Private/Graph/AssetDefinition_MovieGraphConfig.h
new file mode 100644
--- /dev/null
+++ b/Engine/Plugins/MovieScene/MovieRenderPipeline/Source/MovieRenderPipelineEditor/Private/Graph/AssetDefinition_MovieGraphConfig.h
@@ -0,0 +1,27 @@
+// Copyright Epic Games, Inc. All Rights Reserved.
+
+#pragma once
+
+#include "Graph/MovieGraphConfig.h"
+#include "AssetDefinitionDefault.h"
+
+#include "AssetDefinition_MovieGraphConfig.generated.h"
+
+UCLASS()
+class UAssetDefinition_MovieGraphConfig : public UAssetDefinitionDefault
+{
+	GENERATED_BODY()
+
+public:
+	// UAssetDefinition Begin
+	virtual FText GetAssetDisplayName() const override { return NSLOCTEXT("AssetTypeActions", "AssetTypeActions_MovieGraphConfig", "Movie Render Graph Config"); }
+	virtual FLinearColor GetAssetColor() const override { return FLinearColor(FColor(80, 200, 80)); }
+	virtual TSoftClassPtr<UObject> GetAssetClass() const override;
+	virtual TConstArrayView<FAssetCategoryPath> GetAssetCategories() const override
+	{
+		static const auto Categories = { EAssetCategoryPaths::Cinematics };
+		return Categories;
+	}
+	virtual EAssetCommandResult OpenAssets(const FAssetOpenArgs& OpenArgs) const override;
+	// UAssetDefinition End
+};
diff --git a/Engine/Plugins/MovieScene/MovieRenderPipeline/Source/MovieRenderPipelineEditor/Private/Graph/MovieGraphConfigFactory.cpp b/Engine/Plugins/MovieScene/MovieRenderPipeline/Source/MovieRenderPipelineEditor/Private/Graph/MovieGraphConfigFactory.cpp
new file mode 100644
--- /dev/null
+++ b/Engine/Plugins/MovieScene/MovieRenderPipeline/Source/MovieRenderPipelineEditor/Private/Graph/MovieGraphConfigFactory.cpp
@@ -0,0 +1,26 @@
+// Copyright Epic Games, Inc. All Rights Reserved.
+
+#include "Graph/MovieGraphConfigFactory.h"
+#include "Graph/MovieGraphConfig.h"
+#include "AssetTypeCategories.h"
+#include "HAL/IConsoleManager.h"
+
+#include UE_INLINE_GENERATED_CPP_BY_NAME(MovieGraphConfigFactory)
+
+UMovieGraphConfigFactory::UMovieGraphConfigFactory()
+{
+	bCreateNew = true;
+	bEditAfterNew = false;
+	SupportedClass = UMovieGraphConfig::StaticClass();
+}
+
+UObject* UMovieGraphConfigFactory::FactoryCreateNew(UClass* Class, UObject* InParent, FName Name, EObjectFlags Flags, UObject* Context, FFeedbackContext* Warn)
+{
+	return NewObject<UMovieGraphConfig>(InParent, Class, Name, Flags);
+}
+
+bool UMovieGraphConfigFactory::ShouldShowInNewMenu() const
+{
+	IConsoleVariable* RenderGraphCVar = IConsoleManager::Get().FindConsoleVariable(TEXT("MoviePipeline.EnableRenderGraph"));
+	return RenderGraphCVar && RenderGraphCVar->GetBool();
+}
diff --git a/Engine/Plugins/MovieScene/MovieRenderPipeline/Source/MovieRenderPipelineEditor/Public/Graph/MovieGraphConfigFactory.h b/Engine/Plugins/MovieScene/MovieRenderPipeline/Source/MovieRenderPipelineEditor/Public/Graph/MovieGraphConfigFactory.h
new file mode 100644
--- /dev/null
+++ b/Engine/Plugins/MovieScene/MovieRenderPipeline/Source/MovieRenderPipelineEditor/Public/Graph/MovieGraphConfigFactory.h
@@ -0,0 +1,18 @@
+// Copyright Epic Games, Inc. All Rights Reserved.
+
+#pragma once
+
+#include "Factories/Factory.h"
+#include "MovieGraphConfigFactory.generated.h"
+
+UCLASS(BlueprintType)
+class MOVIERENDERPIPELINEEDITOR_API UMovieGraphConfigFactory : public UFactory
+{
+    GENERATED_BODY()
+public:
+	UMovieGraphConfigFactory();
+	//~ Begin UFactory Interface
+	virtual UObject* FactoryCreateNew(UClass* Class, UObject* InParent, FName Name, EObjectFlags Flags, UObject* Context, FFeedbackContext* Warn) override;
+	virtual bool ShouldShowInNewMenu() const override;
+	//~ End UFactory Interface
+};
\ No newline at end of file
