MoviePipeline: Created a UMoviePipelineBase class to try and house the minimal/common API between UMoviePipeline and upcoming UMovieGraph.

diff --git a/Engine/Plugins/MovieScene/MovieRenderPipeline/Source/MovieRenderPipelineCore/Private/MoviePipeline.cpp b/Engine/Plugins/MovieScene/MovieRenderPipeline/Source/MovieRenderPipelineCore/Private/MoviePipeline.cpp
--- a/Engine/Plugins/MovieScene/MovieRenderPipeline/Source/MovieRenderPipelineCore/Private/MoviePipeline.cpp
+++ b/Engine/Plugins/MovieScene/MovieRenderPipeline/Source/MovieRenderPipelineCore/Private/MoviePipeline.cpp
@@ -369,7 +369,7 @@ void UMoviePipeline::RestoreTargetSequenceToOriginalState()
 }
 
 
-void UMoviePipeline::RequestShutdown(bool bIsError)
+void UMoviePipeline::RequestShutdownImpl(bool bIsError)
 {
 	// It's possible for a previous call to RequestionShutdown to have set an error before this call that may not
 	// We don't want to unset a previously set error state
@@ -394,7 +394,7 @@ void UMoviePipeline::RequestShutdown(bool bIsError)
 	}
 }
 
-void UMoviePipeline::Shutdown(bool bIsError)
+void UMoviePipeline::ShutdownImpl(bool bIsError)
 {
 	check(IsInGameThread());
 
diff --git a/Engine/Plugins/MovieScene/MovieRenderPipeline/Source/MovieRenderPipelineCore/Private/MoviePipelineInProcessExecutor.cpp b/Engine/Plugins/MovieScene/MovieRenderPipeline/Source/MovieRenderPipelineCore/Private/MoviePipelineInProcessExecutor.cpp
--- a/Engine/Plugins/MovieScene/MovieRenderPipeline/Source/MovieRenderPipelineCore/Private/MoviePipelineInProcessExecutor.cpp
+++ b/Engine/Plugins/MovieScene/MovieRenderPipeline/Source/MovieRenderPipelineCore/Private/MoviePipelineInProcessExecutor.cpp
@@ -12,6 +12,7 @@
 #include "Misc/PackageName.h"
 #include "Engine/GameEngine.h"
 #include "MoviePipelineUtils.h"
+#include "Graph/MovieGraphPipeline.h"
 
 #include UE_INLINE_GENERATED_CPP_BY_NAME(MoviePipelineInProcessExecutor)
 
@@ -124,8 +125,18 @@ void UMoviePipelineInProcessExecutor::OnMapLoadFinished(UWorld* NewWorld)
 		MoviePipelineClass = UMoviePipeline::StaticClass();
 	}
 
-	ActiveMoviePipeline = NewObject<UMoviePipeline>(NewWorld, MoviePipelineClass);
-	ActiveMoviePipeline->SetViewportInitArgs(ViewportInitArgs);
+	// Temporary
+	if (CurrentJob->GetGraphConfig())
+	{
+		MoviePipelineClass = UMovieGraphPipeline::StaticClass();
+	}
+
+	ActiveMoviePipeline = NewObject<UMoviePipelineBase>(NewWorld, MoviePipelineClass);
+	UMoviePipeline* PipelineAsLegacy = Cast<UMoviePipeline>(ActiveMoviePipeline);
+	if (PipelineAsLegacy)
+	{
+		PipelineAsLegacy->SetViewportInitArgs(ViewportInitArgs);
+	}
 
 	// We allow users to set a multi-frame delay before we actually run the Initialization function and start thinking.
 	// This solves cases where there are engine systems that need to finish loading before we do anything.
@@ -144,7 +155,14 @@ void UMoviePipelineInProcessExecutor::OnMapLoadFinished(UWorld* NewWorld)
 	if (ExecutorSettings->InitialDelayFrameCount == 0)
 	{
 		UE_LOG(LogMovieRenderPipeline, Log, TEXT("Zero Initial Delay, initializing..."));
-		ActiveMoviePipeline->Initialize(Queue->GetJobs()[CurrentPipelineIndex]);
+		if (PipelineAsLegacy)
+		{
+			PipelineAsLegacy->Initialize(Queue->GetJobs()[CurrentPipelineIndex]);
+		}
+		else if (UMovieGraphPipeline* PipelineAsGraph = Cast<UMovieGraphPipeline>(ActiveMoviePipeline))
+		{
+			// PipelineAsGraph->Initialize(Queue->GetJobs()[CurrentPipelineIndex], FMovieGraphInitConfig());
+		}
 		RemainingInitializationFrames = -1;
 	}
 	else
@@ -160,7 +178,15 @@ void UMoviePipelineInProcessExecutor::OnTick()
 		if (RemainingInitializationFrames == 0)
 		{
 			UE_LOG(LogMovieRenderPipeline, Log, TEXT("Delay finished, initializing..."));
-			ActiveMoviePipeline->Initialize(Queue->GetJobs()[CurrentPipelineIndex]);
+			UMoviePipeline* PipelineAsLegacy = Cast<UMoviePipeline>(ActiveMoviePipeline);
+			if (PipelineAsLegacy)
+			{
+				PipelineAsLegacy->Initialize(Queue->GetJobs()[CurrentPipelineIndex]);
+			}
+			else if (UMovieGraphPipeline* PipelineAsGraph = Cast<UMovieGraphPipeline>(ActiveMoviePipeline))
+			{
+				// PipelineAsGraph->Initialize(Queue->GetJobs()[CurrentPipelineIndex], FMovieGraphInitConfig());
+			}
 		}
 
 		RemainingInitializationFrames--;
@@ -188,8 +214,6 @@ void UMoviePipelineInProcessExecutor::OnApplicationQuit()
 void UMoviePipelineInProcessExecutor::OnMoviePipelineFinished(FMoviePipelineOutputData InOutputData)
 {
 	FCoreDelegates::OnBeginFrame.RemoveAll(this);
-	UMoviePipeline* MoviePipeline = ActiveMoviePipeline;
-	
 	OnIndividualJobFinishedDelegateNative.Broadcast(InOutputData);
 
 	if (ActiveMoviePipeline)
diff --git a/Engine/Plugins/MovieScene/MovieRenderPipeline/Source/MovieRenderPipelineCore/Private/MoviePipelineLinearExecutor.cpp b/Engine/Plugins/MovieScene/MovieRenderPipeline/Source/MovieRenderPipelineCore/Private/MoviePipelineLinearExecutor.cpp
--- a/Engine/Plugins/MovieScene/MovieRenderPipeline/Source/MovieRenderPipelineCore/Private/MoviePipelineLinearExecutor.cpp
+++ b/Engine/Plugins/MovieScene/MovieRenderPipeline/Source/MovieRenderPipelineCore/Private/MoviePipelineLinearExecutor.cpp
@@ -111,7 +111,7 @@ FText UMoviePipelineLinearExecutorBase::GetWindowTitle()
 	float CompletionPercentage = 0.f;
 	if (ActiveMoviePipeline)
 	{
-		CompletionPercentage = UMoviePipelineBlueprintLibrary::GetCompletionPercentage(ActiveMoviePipeline) * 100.f;
+		CompletionPercentage = UMoviePipelineBlueprintLibrary::GetCompletionPercentage(Cast<UMoviePipeline>(ActiveMoviePipeline)) * 100.f;
 	}
 
 	FText TitleFormatString = LOCTEXT("MoviePreviewWindowTitleFormat", "Movie Pipeline Render (Preview) [Job {CurrentCount}/{TotalCount} Total] Current Job: {PercentComplete}% Completed.");
diff --git a/Engine/Plugins/MovieScene/MovieRenderPipeline/Source/MovieRenderPipelineCore/Public/Graph/MovieGraphPipeline.h b/Engine/Plugins/MovieScene/MovieRenderPipeline/Source/MovieRenderPipelineCore/Public/Graph/MovieGraphPipeline.h
new file mode 100644
--- /dev/null
+++ b/Engine/Plugins/MovieScene/MovieRenderPipeline/Source/MovieRenderPipelineCore/Public/Graph/MovieGraphPipeline.h
@@ -0,0 +1,16 @@
+// Copyright Epic Games, Inc. All Rights Reserved.
+#pragma once
+
+#include "UObject/Object.h"
+#include "MoviePipelineBase.h"
+#include "MovieGraphPipeline.generated.h"
+
+// Forward Declares
+class UMoviePipelineExecutorJob;
+class UMoviePipelineExecutorShot;
+
+UCLASS(BlueprintType)
+class MOVIERENDERPIPELINECORE_API UMovieGraphPipeline : public UMoviePipelineBase
+{
+	GENERATED_BODY()
+};
\ No newline at end of file
diff --git a/Engine/Plugins/MovieScene/MovieRenderPipeline/Source/MovieRenderPipelineCore/Public/MoviePipeline.h b/Engine/Plugins/MovieScene/MovieRenderPipeline/Source/MovieRenderPipelineCore/Public/MoviePipeline.h
--- a/Engine/Plugins/MovieScene/MovieRenderPipeline/Source/MovieRenderPipelineCore/Public/MoviePipeline.h
+++ b/Engine/Plugins/MovieScene/MovieRenderPipeline/Source/MovieRenderPipelineCore/Public/MoviePipeline.h
@@ -10,6 +10,7 @@
 #endif
 #include "MovieSceneTimeController.h"
 #include "Async/Future.h"
+#include "MoviePipelineBase.h"
 #include "MoviePipeline.generated.h"
 
 // Forward Declares
@@ -34,12 +35,11 @@ class UTexture;
 typedef TTuple<TFuture<bool>, MoviePipeline::FMoviePipelineOutputFutureData> FMoviePipelineOutputFuture;
 
 DECLARE_MULTICAST_DELEGATE_TwoParams(FMoviePipelineFinishedNative, UMoviePipeline*, bool);
-DECLARE_MULTICAST_DELEGATE_OneParam(FMoviePipelineWorkFinishedNative, FMoviePipelineOutputData);
+
 DECLARE_DYNAMIC_MULTICAST_DELEGATE_TwoParams(FMoviePipelineFinished, UMoviePipeline*, MoviePipeline, bool, bFatalError);
-DECLARE_DYNAMIC_MULTICAST_DELEGATE_OneParam(FMoviePipelineWorkFinished, FMoviePipelineOutputData, Results);
 
 UCLASS(Blueprintable)
-class MOVIERENDERPIPELINECORE_API UMoviePipeline : public UObject
+class MOVIERENDERPIPELINECORE_API UMoviePipeline : public UMoviePipelineBase
 {
 	GENERATED_BODY()
 	
@@ -53,38 +53,11 @@ public:
 	UFUNCTION(BlueprintCallable, Category = "Movie Render Pipeline")
 	void Initialize(UMoviePipelineExecutorJob* InJob);
 
-
-	/**
-	* Request the movie pipeline to shut down at the next available time. The pipeline will attempt to abandon
-	* the current frame (such as if there are more temporal samples pending) but may be forced into finishing if
-	* there are spatial samples already submitted to the GPU. The shutdown flow will be run to ensure already
-	* completed work is written to disk. This is a non-blocking operation, use Shutdown() instead if you need to
-	* block until it is fully shut down.
-	*
-	* @param bError - Whether this is a request for early shut down due to an error
-	* 
-	* This function is thread safe.
-	*/
-	UFUNCTION(BlueprintCallable, Category = "Movie Render Pipeline")
-	void RequestShutdown(bool bIsError=false);
-	
-	/** 
-	* Abandons any future work on this Movie Pipeline and runs through the shutdown flow to ensure already
-	* completed work is written to disk. This is a blocking-operation and will not return until all outstanding
-	* work has been completed.
-	*
-	* @param bError - Whether this is an early shut down due to an error
-	*
-	* This function should only be called from the game thread.
-	*/
-	UFUNCTION(BlueprintCallable, Category = "Movie Render Pipeline")
-	void Shutdown(bool bError=false);
-
-	/**
-	* Has RequestShutdown() been called?
-	*/
-	UFUNCTION(BlueprintPure, Category = "Movie Render Pipeline")
-	bool IsShutdownRequested() const { return bShutdownRequested; }
+	// UMoviePipelineBase Interface
+	virtual void RequestShutdownImpl(bool bIsError) override;
+	virtual void ShutdownImpl(bool bError ) override;
+	virtual bool IsShutdownRequestedImpl() const override { return bShutdownRequested; }
+	// ~UMoviePipelineBase Interface
 
 	/**
 	* Returns the time this movie pipeline was initialized at.
@@ -115,18 +88,7 @@ public:
 	UPROPERTY(BlueprintAssignable, Category = "Movie Render Pipeline")
 	FMoviePipelineFinished OnMoviePipelineFinishedDelegate;
 
-	/**
-	* Called when we have completely finished this pipeline. This means that all frames have been rendered,
-	* all files written to disk, and any post-finalize exports have finished. This Pipeline will call
-	* Shutdown() on itself before calling this delegate to ensure we've unregistered from all delegates
-	* and are no longer trying to do anything (even if we still exist).
-	*
-	* The params struct in the return will have metadata about files written to disk for each shot.
-	*/
-	FMoviePipelineWorkFinishedNative& OnMoviePipelineWorkFinished()
-	{
-		return OnMoviePipelineWorkFinishedDelegateNative;
-	}
+
 	
 	/**
 	* Called when we have completely finished this pipeline. This means that all frames have been rendered,
@@ -139,16 +101,7 @@ public:
 	UPROPERTY(BlueprintAssignable, Category = "Movie Render Pipeline")
 	FMoviePipelineWorkFinished OnMoviePipelineWorkFinishedDelegate;
 
-	/**
-	* Only called if `IsFlushDiskWritesPerShot()` is set!
-	* Called after each shot is finished and files have been flushed to disk. The returned data in
-	* the params struct will have only the per-shot metadata for the just finished shot. Use
-	* OnMoviePipelineFinished() if you need all ot the metadata.
-	*/
-	FMoviePipelineWorkFinishedNative& OnMoviePipelineShotWorkFinished()
-	{
-		return OnMoviePipelineShotWorkFinishedDelegateNative;
-	}
+
 
 	/**
 	* Only called if `IsFlushDiskWritesPerShot()` is set!
@@ -454,11 +407,6 @@ private:
 
 	/** Deprecated. */
 	FMoviePipelineFinishedNative OnMoviePipelineFinishedDelegateNative;
-	/** Called when we have completely finished. This object will call Shutdown before this and stop ticking. */
-	FMoviePipelineWorkFinishedNative OnMoviePipelineWorkFinishedDelegateNative;
-
-	/** Called when each shot has finished work if IsFlushDiskWritesPerShot() is set. */
-	FMoviePipelineWorkFinishedNative OnMoviePipelineShotWorkFinishedDelegateNative;
 
 	/**
 	 * We have to apply camera motion vectors manually. So we keep the current and previous frame's camera view and rotation.
diff --git a/Engine/Plugins/MovieScene/MovieRenderPipeline/Source/MovieRenderPipelineCore/Public/MoviePipelineBase.h b/Engine/Plugins/MovieScene/MovieRenderPipeline/Source/MovieRenderPipelineCore/Public/MoviePipelineBase.h
new file mode 100644
--- /dev/null
+++ b/Engine/Plugins/MovieScene/MovieRenderPipeline/Source/MovieRenderPipelineCore/Public/MoviePipelineBase.h
@@ -0,0 +1,94 @@
+// Copyright Epic Games, Inc. All Rights Reserved.
+#pragma once
+
+#include "UObject/Object.h"
+#include "MovieRenderPipelineDataTypes.h"
+#include "MoviePipelineBase.generated.h"
+
+DECLARE_MULTICAST_DELEGATE_OneParam(FMoviePipelineWorkFinishedNative, FMoviePipelineOutputData);
+DECLARE_DYNAMIC_MULTICAST_DELEGATE_OneParam(FMoviePipelineWorkFinished, FMoviePipelineOutputData, Results);
+
+UCLASS(BlueprintType, Abstract)
+class MOVIERENDERPIPELINECORE_API UMoviePipelineBase : public UObject
+{
+	GENERATED_BODY()
+public:
+	/**
+	* Request the movie pipeline to shut down at the next available time. The pipeline will attempt to abandon
+	* the current frame (such as if there are more temporal samples pending) but may be forced into finishing if
+	* there are spatial samples already submitted to the GPU. The shutdown flow will be run to ensure already
+	* completed work is written to disk. This is a non-blocking operation, use Shutdown() instead if you need to
+	* block until it is fully shut down.
+	*
+	* @param bIsError - Whether this is a request for early shut down due to an error
+	*
+	* This function is thread safe.
+	*/
+	UFUNCTION(BlueprintCallable, Category = "Movie Render Pipeline")
+	void RequestShutdown(bool bIsError = false)
+	{
+		RequestShutdownImpl(bIsError);
+	}
+
+	/**
+	* Abandons any future work on this Movie Pipeline and runs through the shutdown flow to ensure already
+	* completed work is written to disk. This is a blocking-operation and will not return until all outstanding
+	* work has been completed.
+	*
+	* @param bIsError - Whether this is an early shut down due to an error
+	*
+	* This function should only be called from the game thread.
+	*/
+	UFUNCTION(BlueprintCallable, Category = "Movie Render Pipeline")
+	void Shutdown(bool bIsError = false)
+	{
+		ShutdownImpl(bIsError);
+	}
+
+	/**
+	* Has RequestShutdown() been called?
+	*/
+	UFUNCTION(BlueprintPure, Category = "Movie Render Pipeline")
+	bool IsShutdownRequested() const
+	{
+		return IsShutdownRequestedImpl();
+	}
+
+	UFUNCTION(BlueprintPure, Category = "Movie Render Pipeline")
+	EMovieRenderPipelineState GetPipelineState() const
+	{
+		return GetPipelineStateImpl();
+	}
+
+	/**
+	* Called when we have completely finished this pipeline. This means that all frames have been rendered,
+	* all files written to disk, and any post-finalize exports have finished. This Pipeline will call
+	* Shutdown() on itself before calling this delegate to ensure we've unregistered from all delegates
+	* and are no longer trying to do anything (even if we still exist).
+	*
+	* The params struct in the return will have metadata about files written to disk for each shot.
+	*/
+	FMoviePipelineWorkFinishedNative& OnMoviePipelineWorkFinished() { return OnMoviePipelineWorkFinishedImpl(); }
+	
+	/**
+	* Only called if `IsFlushDiskWritesPerShot()` is set!
+	* Called after each shot is finished and files have been flushed to disk. The returned data in
+	* the params struct will have only the per-shot metadata for the just finished shot. Use
+	* OnMoviePipelineFinished() if you need all of the metadata.
+	*/
+	FMoviePipelineWorkFinishedNative& OnMoviePipelineShotWorkFinished() { return OnMoviePipelineShotWorkFinishedImpl(); }
+
+protected:
+	virtual void RequestShutdownImpl(bool bIsError) {}
+	virtual void ShutdownImpl(bool bIsError) {}
+	virtual bool IsShutdownRequestedImpl() const { return false; }
+	virtual EMovieRenderPipelineState GetPipelineStateImpl() const { return EMovieRenderPipelineState::Uninitialized; }
+	virtual FMoviePipelineWorkFinishedNative& OnMoviePipelineWorkFinishedImpl() { return OnMoviePipelineWorkFinishedDelegateNative; }
+	virtual FMoviePipelineWorkFinishedNative& OnMoviePipelineShotWorkFinishedImpl() { return OnMoviePipelineShotWorkFinishedDelegateNative; }
+
+	/** Called when we have completely finished. This object will call Shutdown before this and stop ticking. */
+	FMoviePipelineWorkFinishedNative OnMoviePipelineWorkFinishedDelegateNative;
+
+	/** Called when each shot has finished work if IsFlushDiskWritesPerShot() is set. */
+	FMoviePipelineWorkFinishedNative OnMoviePipelineShotWorkFinishedDelegateNative;
+};
\ No newline at end of file
diff --git a/Engine/Plugins/MovieScene/MovieRenderPipeline/Source/MovieRenderPipelineCore/Public/MoviePipelineLinearExecutor.h b/Engine/Plugins/MovieScene/MovieRenderPipeline/Source/MovieRenderPipelineCore/Public/MoviePipelineLinearExecutor.h
--- a/Engine/Plugins/MovieScene/MovieRenderPipeline/Source/MovieRenderPipelineCore/Public/MoviePipelineLinearExecutor.h
+++ b/Engine/Plugins/MovieScene/MovieRenderPipeline/Source/MovieRenderPipelineCore/Public/MoviePipelineLinearExecutor.h
@@ -5,7 +5,7 @@
 #include "MovieRenderPipelineDataTypes.h"
 #include "MoviePipelineLinearExecutor.generated.h"
 
-class UMoviePipeline;
+class UMoviePipelineBase;
 class UMoviePipelineQueue;
 
 /**
@@ -53,7 +53,7 @@ protected:
 
 	/** A Movie Pipeline that has been spawned and is running (if any) */
 	UPROPERTY(Transient)
-	TObjectPtr<UMoviePipeline> ActiveMoviePipeline;
+	TObjectPtr<UMoviePipelineBase> ActiveMoviePipeline;
 
 	/** Which Pipeline Config are we currently working on. */
 	int32 CurrentPipelineIndex;
diff --git a/Engine/Plugins/MovieScene/MovieRenderPipeline/Source/MovieRenderPipelineEditor/Private/MoviePipelinePIEExecutor.cpp b/Engine/Plugins/MovieScene/MovieRenderPipeline/Source/MovieRenderPipelineEditor/Private/MoviePipelinePIEExecutor.cpp
--- a/Engine/Plugins/MovieScene/MovieRenderPipeline/Source/MovieRenderPipelineEditor/Private/MoviePipelinePIEExecutor.cpp
+++ b/Engine/Plugins/MovieScene/MovieRenderPipeline/Source/MovieRenderPipelineEditor/Private/MoviePipelinePIEExecutor.cpp
@@ -19,6 +19,7 @@
 #include "Modules/ModuleManager.h"
 #include "MessageLogModule.h"
 #include "Logging/MessageLog.h"
+#include "Graph/MovieGraphPipeline.h"
 
 #include UE_INLINE_GENERATED_CPP_BY_NAME(MoviePipelinePIEExecutor)
 
@@ -217,7 +218,7 @@ void UMoviePipelinePIEExecutor::OnPIEStartupFinished(bool)
 	// Allow the user to have overridden which Pipeline is actually run. This is an unlikely scenario but allows
 	// the user to create their own implementation while still re-using the rest of our UI and infrastructure.
 	const UMovieRenderPipelineProjectSettings* ProjectSettings = GetDefault<UMovieRenderPipelineProjectSettings>();
-	TSubclassOf<UMoviePipeline> PipelineClass = ProjectSettings->DefaultPipeline.TryLoadClass<UMoviePipeline>();
+	TSubclassOf<UMoviePipelineBase> PipelineClass = ProjectSettings->DefaultPipeline.TryLoadClass<UMoviePipelineBase>();
 	if (!PipelineClass)
 	{
 		UE_LOG(LogMovieRenderPipeline, Error, TEXT("Failed to load project specified MoviePipeline class type!"));
@@ -225,8 +226,17 @@ void UMoviePipelinePIEExecutor::OnPIEStartupFinished(bool)
 		return;
 	}
 
+	UMoviePipelineExecutorJob* CurrentJob = Queue->GetJobs()[CurrentPipelineIndex];
+	if (CurrentJob->GetGraphConfig())
+	{
+		PipelineClass = UMovieGraphPipeline::StaticClass();
+	}
+
 	// This Pipeline belongs to the world being created so that they have context for things they execute.
-	ActiveMoviePipeline = NewObject<UMoviePipeline>(ExecutingWorld, PipelineClass);
+	ActiveMoviePipeline = NewObject<UMoviePipelineBase>(ExecutingWorld, PipelineClass);
+
+	UMoviePipeline* PipelineAsLegacy = Cast<UMoviePipeline>(ActiveMoviePipeline);
+
 	
 	// We allow users to set a multi-frame delay before we actually run the Initialization function and start thinking.
 	// This solves cases where there are engine systems that need to finish loading before we do anything.
@@ -238,12 +248,23 @@ void UMoviePipelinePIEExecutor::OnPIEStartupFinished(bool)
 	// Listen for when the pipeline thinks it has finished.
 	ActiveMoviePipeline->OnMoviePipelineWorkFinished().AddUObject(this, &UMoviePipelinePIEExecutor::OnPIEMoviePipelineFinished);
 	ActiveMoviePipeline->OnMoviePipelineShotWorkFinished().AddUObject(this, &UMoviePipelinePIEExecutor::OnJobShotFinished);
-	ActiveMoviePipeline->SetViewportInitArgs(ViewportInitArgs);
+	if (PipelineAsLegacy)
+	{
+
+		PipelineAsLegacy->SetViewportInitArgs(ViewportInitArgs);
+	}
 	
 	if (ExecutorSettings->InitialDelayFrameCount == 0)
 	{
 		OnIndividualJobStartedImpl(Queue->GetJobs()[CurrentPipelineIndex]);
-		ActiveMoviePipeline->Initialize(Queue->GetJobs()[CurrentPipelineIndex]);
+		if (PipelineAsLegacy)
+		{
+			PipelineAsLegacy->Initialize(Queue->GetJobs()[CurrentPipelineIndex]);
+		}
+		else if (UMovieGraphPipeline* PipelineAsGraph = Cast<UMovieGraphPipeline>(ActiveMoviePipeline))
+		{
+			// PipelineAsGraph->Initialize(Queue->GetJobs()[CurrentPipelineIndex], FMovieGraphInitConfig());
+		}
 		RemainingInitializationFrames = -1;
 	}
 	else
@@ -261,13 +282,21 @@ void UMoviePipelinePIEExecutor::OnTick()
 	{
 		if (RemainingInitializationFrames == 0)
 		{
-			if (CustomInitializationTime.IsSet())
+			OnIndividualJobStartedImpl(Queue->GetJobs()[CurrentPipelineIndex]);
+			
+			UMoviePipeline* PipelineAsLegacy = Cast<UMoviePipeline>(ActiveMoviePipeline);
+			if (PipelineAsLegacy)
 			{
-				ActiveMoviePipeline->SetInitializationTime(CustomInitializationTime.GetValue());
+				PipelineAsLegacy->Initialize(Queue->GetJobs()[CurrentPipelineIndex]);
+				if (CustomInitializationTime.IsSet())
+				{
+					PipelineAsLegacy->SetInitializationTime(CustomInitializationTime.GetValue());
+				}
+			}
+			else if (UMovieGraphPipeline* PipelineAsGraph = Cast<UMovieGraphPipeline>(ActiveMoviePipeline))
+			{
+				// PipelineAsGraph->Initialize(Queue->GetJobs()[CurrentPipelineIndex], FMovieGraphInitConfig());
 			}
-
-			OnIndividualJobStartedImpl(Queue->GetJobs()[CurrentPipelineIndex]);
-			ActiveMoviePipeline->Initialize(Queue->GetJobs()[CurrentPipelineIndex]);
 		}
 
 		RemainingInitializationFrames--;
@@ -321,6 +350,7 @@ void UMoviePipelinePIEExecutor::OnPIEEnded(bool)
 	FEditorDelegates::EndPIE.RemoveAll(this);
 
 	CachedOutputDataParams = FMoviePipelineOutputData();
+	UMoviePipeline* PipelineAsLegacy = Cast<UMoviePipeline>(ActiveMoviePipeline);
 
 	// Only call Shutdown if the pipeline hasn't been finished.
 	if (ActiveMoviePipeline && ActiveMoviePipeline->GetPipelineState() != EMovieRenderPipelineState::Finished)
@@ -331,10 +361,10 @@ void UMoviePipelinePIEExecutor::OnPIEEnded(bool)
 		ActiveMoviePipeline->Shutdown(true);
 		UE_LOG(LogMovieRenderPipeline, Log, TEXT("MoviePipelinePIEExecutor: Stalling finished, pipeline has shut down."));
 	}
-	if (ActiveMoviePipeline)
+	if (PipelineAsLegacy)
 	{
 		// Cache this off so we can use it in DelayedFinishNotification, at which point ActiveMoviePipeline is null.
-		CachedOutputDataParams = ActiveMoviePipeline->GetOutputDataParams();
+		CachedOutputDataParams = PipelineAsLegacy->GetOutputDataParams();
 	}
 
 	// We need to null out this reference this frame, otherwise we try to hold onto a PIE
diff --git a/Engine/Plugins/MovieScene/MovieRenderPipeline/Source/MovieRenderPipelineEditor/Public/MovieRenderPipelineSettings.h b/Engine/Plugins/MovieScene/MovieRenderPipeline/Source/MovieRenderPipelineEditor/Public/MovieRenderPipelineSettings.h
--- a/Engine/Plugins/MovieScene/MovieRenderPipeline/Source/MovieRenderPipelineEditor/Public/MovieRenderPipelineSettings.h
+++ b/Engine/Plugins/MovieScene/MovieRenderPipeline/Source/MovieRenderPipelineEditor/Public/MovieRenderPipelineSettings.h
@@ -67,7 +67,7 @@ public:
 	* this will allow you to re-use the existing UI/Executors while providing your own logic for producing
 	* a single render.
 	*/
-	UPROPERTY(config, EditAnywhere, BlueprintReadWrite, meta = (MetaClass = "/Script/MovieRenderPipelineCore.MoviePipeline"), Category="Movie Render Pipeline")
+	UPROPERTY(config, EditAnywhere, BlueprintReadWrite, meta = (MetaClass = "/Script/MovieRenderPipelineCore.MoviePipelineBase"), Category="Movie Render Pipeline")
 	FSoftClassPath DefaultPipeline;
 
 	/**
