MoviePipeline: Added UMovieGraphPipeline as an equivilent to UMoviePipeline that works based on Graph Assets.

diff --git a/Engine/Plugins/MovieScene/MovieRenderPipeline/Source/MovieRenderPipelineCore/Private/Graph/MovieGraphDataTypes.cpp b/Engine/Plugins/MovieScene/MovieRenderPipeline/Source/MovieRenderPipelineCore/Private/Graph/MovieGraphDataTypes.cpp
new file mode 100644
--- /dev/null
+++ b/Engine/Plugins/MovieScene/MovieRenderPipeline/Source/MovieRenderPipelineCore/Private/Graph/MovieGraphDataTypes.cpp
@@ -0,0 +1,36 @@
+// Copyright Epic Games, Inc. All Rights Reserved.
+
+#include "Graph/MovieGraphDataTypes.h"
+#include "Graph/MovieGraphLinearTimeStep.h"
+#include "Graph/MovieGraphDefaultRenderer.h"
+#include "Graph/MovieGraphDefaultDataCaching.h"
+#include "Graph/MovieGraphSequenceTimeRangeBuilder.h"
+#include "Graph/MovieGraphPipeline.h"
+
+FMovieGraphInitConfig::FMovieGraphInitConfig()
+{
+	TimeStepClass = UMovieGraphLinearTimeStep::StaticClass();
+	RendererClass = UMovieGraphDefaultRenderer::StaticClass();
+	TimeRangeBuilderClass = UMovieGraphSequenceTimeRangeBuilder::StaticClass();
+	DataCachingClass = UMovieGraphDefaultDataCaching::StaticClass();
+}
+
+UMovieGraphPipeline* UMovieGraphTimeStepBase::GetOwningGraph() const
+{ 
+	return GetTypedOuter<UMovieGraphPipeline>();
+}
+
+UMovieGraphPipeline* UMovieGraphRendererBase::GetOwningGraph() const
+{ 
+	return GetTypedOuter<UMovieGraphPipeline>();
+}
+
+UMovieGraphPipeline* UMovieGraphTimeRangeBuilderBase::GetOwningGraph() const
+{
+	return GetTypedOuter<UMovieGraphPipeline>();
+}
+
+UMovieGraphPipeline* UMovieGraphDataCachingBase::GetOwningGraph() const
+{
+	return GetTypedOuter<UMovieGraphPipeline>();
+}
\ No newline at end of file
diff --git a/Engine/Plugins/MovieScene/MovieRenderPipeline/Source/MovieRenderPipelineCore/Private/Graph/MovieGraphDefaultDataCaching.cpp b/Engine/Plugins/MovieScene/MovieRenderPipeline/Source/MovieRenderPipelineCore/Private/Graph/MovieGraphDefaultDataCaching.cpp
new file mode 100644
--- /dev/null
+++ b/Engine/Plugins/MovieScene/MovieRenderPipeline/Source/MovieRenderPipelineCore/Private/Graph/MovieGraphDefaultDataCaching.cpp
@@ -0,0 +1,104 @@
+// Copyright Epic Games, Inc. All Rights Reserved.
+
+#include "Graph/MovieGraphDefaultDataCaching.h"
+#include "Graph/MovieGraphPipeline.h"
+#include "LevelSequence.h"
+#include "LevelSequenceActor.h"
+#include "LevelSequencePlayer.h"
+#include "MoviePipelineQueue.h"
+#include "CoreGlobals.h"
+#include "EngineUtils.h"
+#include "MovieRenderPipelineCoreModule.h"
+#include "MoviePipelineBlueprintLibrary.h"
+#include "Engine/GameViewportClient.h"
+
+void UMovieGraphDefaultDataCaching::CacheDataPreJob(const FMovieGraphInitConfig& InInitConfig)
+{
+	// Turn off screen messages as some forms are drawn directly to final render targets
+	// which will polute final frames.
+	GAreScreenMessagesEnabled = false;
+
+	// Turn off player viewport rendering to avoid overhead of an extra render.
+	if (UGameViewportClient* Viewport = GetWorld()->GetGameViewport())
+	{
+		Viewport->bDisableWorldRendering = !InInitConfig.bRenderViewport;
+	}
+
+	ULevelSequence* RootSequence = Cast<ULevelSequence>(GetOwningGraph()->GetCurrentJob()->Sequence.TryLoad());
+	if (RootSequence)
+	{
+		CacheLevelSequenceData(RootSequence);
+	}
+}
+
+void UMovieGraphDefaultDataCaching::RestoreCachedDataPostJob()
+{
+	GAreScreenMessagesEnabled = true;
+	if (UGameViewportClient* Viewport = GetWorld()->GetGameViewport())
+	{
+		Viewport->bDisableWorldRendering = false;
+	}
+}
+
+void UMovieGraphDefaultDataCaching::UpdateShotList()
+{
+	ULevelSequence* RootSequence = Cast<ULevelSequence>(GetOwningGraph()->GetCurrentJob()->Sequence.TryLoad());
+	if (RootSequence)
+	{
+		bool bShotsChanged = false;
+		UMoviePipelineBlueprintLibrary::UpdateJobShotListFromSequence(RootSequence, GetOwningGraph()->GetCurrentJob(), bShotsChanged);
+	}
+}
+
+void UMovieGraphDefaultDataCaching::CacheLevelSequenceData(ULevelSequence* InSequence)
+{
+	// There is a reasonable chance that there exists a Level Sequence Actor in the world already set up to play this sequence.
+	ALevelSequenceActor* ExistingActor = nullptr;
+
+	for (auto It = TActorIterator<ALevelSequenceActor>(GetWorld()); It; ++It)
+	{
+		// Iterate through all of them in the event someone has multiple copies in the world on accident.
+		if (It->GetSequence() == InSequence)
+		{
+			// Found it!
+			ExistingActor = *It;
+
+			// Stop it from playing if it's already playing.
+			if (ExistingActor->GetSequencePlayer())
+			{
+				ExistingActor->GetSequencePlayer()->Stop();
+			}
+		}
+	}
+
+	LevelSequenceActor = ExistingActor;
+	if (!LevelSequenceActor)
+	{
+		// Spawn a new level sequence
+		LevelSequenceActor = GetWorld()->SpawnActor<ALevelSequenceActor>();
+		check(LevelSequenceActor);
+	}
+
+	// Enforce settings.
+	LevelSequenceActor->PlaybackSettings.LoopCount.Value = 0;
+	LevelSequenceActor->PlaybackSettings.bAutoPlay = false;
+	LevelSequenceActor->PlaybackSettings.bPauseAtEnd = true;
+	LevelSequenceActor->PlaybackSettings.bRestoreState = true;
+
+	// Ensure the (possibly new) Level Sequence Actor uses our sequence
+	LevelSequenceActor->SetSequence(InSequence);
+
+	// LevelSequenceActor->GetSequencePlayer()->SetTimeController(CustomSequenceTimeController);
+	LevelSequenceActor->GetSequencePlayer()->Stop();
+
+	LevelSequenceActor->GetSequencePlayer()->OnSequenceUpdated().AddUObject(this, &UMovieGraphDefaultDataCaching::OnSequenceEvaluated);
+
+}
+
+void UMovieGraphDefaultDataCaching::OnSequenceEvaluated(const UMovieSceneSequencePlayer& Player, FFrameTime CurrentTime, FFrameTime PreviousTime)
+{
+	// This callback exists for logging purposes. DO NOT HINGE LOGIC ON THIS CALLBACK
+	// because this may get called multiple times per frame and may be the result of
+	// a seek operation which is reverted before a frame is even rendered.
+	UE_LOG(LogMovieRenderPipeline, VeryVerbose, TEXT("[GFrameCounter: %d] Sequence Evaluated. CurrentTime: %s PreviousTime: %s"), GFrameCounter, *LexToString(CurrentTime), *LexToString(PreviousTime));
+}
diff --git a/Engine/Plugins/MovieScene/MovieRenderPipeline/Source/MovieRenderPipelineCore/Private/Graph/MovieGraphDefaultRenderer.cpp b/Engine/Plugins/MovieScene/MovieRenderPipeline/Source/MovieRenderPipelineCore/Private/Graph/MovieGraphDefaultRenderer.cpp
new file mode 100644
--- /dev/null
+++ b/Engine/Plugins/MovieScene/MovieRenderPipeline/Source/MovieRenderPipelineCore/Private/Graph/MovieGraphDefaultRenderer.cpp
@@ -0,0 +1,20 @@
+// Copyright Epic Games, Inc. All Rights Reserved.
+
+#include "Graph/MovieGraphDefaultRenderer.h"
+#include "Graph/MovieGraphPipeline.h"
+
+void UMovieGraphDefaultRenderer::Render(const FMovieGraphTimeStepData& InTimeStepData)
+{
+	UE_LOG(LogTemp, Warning, TEXT("MovieGraphRenderer: Render"));
+}
+
+void UMovieGraphDefaultRenderer::SetupRenderingPipelineForShot(UMoviePipelineExecutorShot* InShot)
+{
+	UE_LOG(LogTemp, Warning, TEXT("MovieGraphRenderer: SetupForShot"));
+
+}
+
+void UMovieGraphDefaultRenderer::TeardownRenderingPipelineForShot(UMoviePipelineExecutorShot* InShot)
+{
+	UE_LOG(LogTemp, Warning, TEXT("MovieGraphRenderer: TeardownForShot"));
+}
\ No newline at end of file
diff --git a/Engine/Plugins/MovieScene/MovieRenderPipeline/Source/MovieRenderPipelineCore/Private/Graph/MovieGraphLinearTimeStep.cpp b/Engine/Plugins/MovieScene/MovieRenderPipeline/Source/MovieRenderPipelineCore/Private/Graph/MovieGraphLinearTimeStep.cpp
new file mode 100644
--- /dev/null
+++ b/Engine/Plugins/MovieScene/MovieRenderPipeline/Source/MovieRenderPipelineCore/Private/Graph/MovieGraphLinearTimeStep.cpp
@@ -0,0 +1,8 @@
+// Copyright Epic Games, Inc. All Rights Reserved.
+
+#include "Graph/MovieGraphLinearTimeStep.h"
+
+void UMovieGraphLinearTimeStep::CalculateTimeStep(FMovieGraphTimeStepData& OutTimeData)
+{
+
+}
\ No newline at end of file
diff --git a/Engine/Plugins/MovieScene/MovieRenderPipeline/Source/MovieRenderPipelineCore/Private/Graph/MovieGraphPipeline.cpp b/Engine/Plugins/MovieScene/MovieRenderPipeline/Source/MovieRenderPipelineCore/Private/Graph/MovieGraphPipeline.cpp
new file mode 100644
--- /dev/null
+++ b/Engine/Plugins/MovieScene/MovieRenderPipeline/Source/MovieRenderPipelineCore/Private/Graph/MovieGraphPipeline.cpp
@@ -0,0 +1,695 @@
+// Copyright Epic Games, Inc. All Rights Reserved.
+
+#include "Graph/MovieGraphPipeline.h"
+#include "MovieRenderPipelineCoreModule.h"
+#include "Misc/CoreDelegates.h"
+#include "MoviePipelineQueue.h"
+#include "Graph/MovieGraphDataTypes.h"
+#include "MovieScene.h"
+#include "RenderingThread.h"
+
+// Temp
+#include "LevelSequence.h"
+
+UMovieGraphPipeline::UMovieGraphPipeline()
+	: CurrentShotIndex(-1)
+	, bIsTransitioningState(false)
+	, PipelineState(EMovieRenderPipelineState::Uninitialized)
+{
+}
+
+void UMovieGraphPipeline::Initialize(UMoviePipelineExecutorJob* InJob, const FMovieGraphInitConfig& InitConfig)
+{
+	UE_LOG(LogMovieRenderPipeline, Log, TEXT("Initializing MovieGraph Render"));
+	if (!ensureAlwaysMsgf(InJob, TEXT("MovieGraph cannot be initialized with a null job. Aborting.")))
+	{
+		//Shutdown(true);
+		return;
+	}
+
+	if (!ensureAlwaysMsgf(InJob->GetGraphConfig(), TEXT("MoviePipeline cannot be initialized with a null job configuration. Make sure you've created a Graph Config for this job (or use the regular UMoviePipeline instead of UMovieGraphPipeline). Aborting.")))
+	{
+		//Shutdown(true);
+		return;
+	}
+
+	// ToDo: If we cache presets into the job, update the job config here. I don't think we want to do that though
+	// due to the recursive nature of linked graphs.
+	// But we could (should?) at least copy the primary preset asset into the primary config slot.
+
+	if (!ensureAlwaysMsgf(PipelineState == EMovieRenderPipelineState::Uninitialized, TEXT("Pipeline cannot be reused. Create a new pipeline to execute a job.")))
+	{
+		Shutdown(true);
+		return;
+	}
+
+	// Ensure this object has the World as part of its Outer (so that it has context to spawn things)
+	if (!ensureAlwaysMsgf(GetWorld(), TEXT("Pipeline does not contain the world as an outer.")))
+	{
+		Shutdown(true);
+		return;
+	}
+
+
+	// Register any additional engine callbacks needed.
+	{
+		// Called before the Custom Timestep is updated. This gives us a chance to calculate
+		// what we want the frame to look like and then cache that information so that the
+		// Custom Timestep doesn't have to perform its own logic.
+		FCoreDelegates::OnBeginFrame.AddUObject(this, &UMovieGraphPipeline::OnEngineTickBeginFrame);
+		// Called at the end of the frame after everything has been ticked and rendered for the frame.
+		FCoreDelegates::OnEndFrame.AddUObject(this, &UMovieGraphPipeline::OnEngineTickEndFrame);
+	}
+
+	// Create instances of our different classes from the InitConfig
+	GraphTimeStepInstance = NewObject<UMovieGraphTimeStepBase>(this, InitConfig.TimeStepClass);
+	GraphRendererInstance = NewObject<UMovieGraphRendererBase>(this, InitConfig.RendererClass);
+	GraphTimeRangeBuilderInstance = NewObject<UMovieGraphTimeRangeBuilderBase>(this, InitConfig.TimeRangeBuilderClass);
+	GraphDataCachingInstance = NewObject<UMovieGraphDataCachingBase>(this, InitConfig.DataCachingClass);
+	
+	CurrentJob = InJob;
+	CurrentShotIndex = 0;
+	
+	// Now that we've created our various systems, we will start using them. First thing we do is cache data about
+	// the world, job, player viewport, etc, before we make any modifications. These will be restored at the end
+	// of the render.
+	GraphDataCachingInstance->CacheDataPreJob(InitConfig);
+
+	// Construct a debug UI and bind it to this instance.
+	// LoadDebugWidget();
+	// SetupAudioRendering();
+
+	// Update our list of shots from our data source, and then
+	// create our list of active shots, so we don't try to render
+	// a shot the user has deactivated.
+	BuildShotListFromDataSource();
+
+	//for (ULevelStreaming* Level : GetWorld()->GetStreamingLevels())
+	//{
+	//	UClass* StreamingClass = Level->GetClass();
+	//
+	//	if (StreamingClass == ULevelStreamingDynamic::StaticClass())
+	//	{
+	//		const FString NonPrefixedLevelName = UWorld::StripPIEPrefixFromPackageName(Level->GetWorldAssetPackageName(), GetWorld()->StreamingLevelsPrefix);
+	//		UE_LOG(LogMovieRenderPipeline, Warning, TEXT("Sub-level '%s' is set to blueprint streaming and will not be visible during a render unless a Sequencer Visibility Track controls its visibility or you have written other code to handle loading it."),
+	//			*NonPrefixedLevelName);
+	//	}
+	//}
+
+	// ToDo: Print some information, Job Name, Level Sequence, etc.
+	UE_LOG(LogMovieRenderPipeline, Log, TEXT("Finished initializing MovieGraph Render"));
+
+
+	// If the shot mask entirely disabled everything we'll transition directly to finish as there is no work to do.
+	if (ActiveShotList.Num() == 0)
+	{
+		// We have to transition twice as Uninitialized -> n state is a no-op, so the second tick will take us to Finished which shuts down.
+		UE_LOG(LogMovieRenderPipeline, Warning, TEXT("MovieGraph No shots detected to render. Either all outside playback range, or disabled via shot mask, bailing."));
+	
+		TransitionToState(EMovieRenderPipelineState::Export);
+		TransitionToState(EMovieRenderPipelineState::Finished);
+	}
+	else
+	{
+		TransitionToState(EMovieRenderPipelineState::ProducingFrames);
+	}
+}
+
+void UMovieGraphPipeline::BuildShotListFromDataSource()
+{
+	// Synchronize our shot list with our target data source. New shots will be added and outdated shots removed.
+	// Shots that are already in the list will be updated but their enable flag will be respected. 
+	GraphDataCachingInstance->UpdateShotList();
+
+	for (UMoviePipelineExecutorShot* Shot : GetCurrentJob()->ShotInfo)
+	{
+		Shot->ShotInfo.CurrentTickInRoot = Shot->ShotInfo.TotalOutputRangeRoot.GetLowerBoundValue();
+	}
+
+
+	// The active shot-list is a subset of the whole shot-list; The ShotInfo contains information about every range it detected to render
+	// but if the user has turned the shot off in the UI then we don't want to render it.
+	ActiveShotList.Empty();
+	for (UMoviePipelineExecutorShot* Shot : GetCurrentJob()->ShotInfo)
+	{
+		if (Shot->ShouldRender())
+		{
+			ActiveShotList.Add(Shot);
+		}
+	}
+}
+
+void UMovieGraphPipeline::OnEngineTickBeginFrame()
+{
+	LLM_SCOPE_BYNAME(TEXT("MovieGraphBeginFrame"));
+
+
+	UE_LOG(LogMovieRenderPipeline, VeryVerbose, TEXT("MovieGraph OnEngineTickBeginFrame"));
+	switch (PipelineState)
+	{
+	case EMovieRenderPipelineState::Uninitialized:
+		// We shouldn't register this delegate until we're initialized.
+		check(false);
+		break;
+	case EMovieRenderPipelineState::ProducingFrames:
+		TickProducingFrames();
+		break;
+	case EMovieRenderPipelineState::Finalize:
+		// Don't flush the finalize to keep the UI responsive.
+		TickFinalizeOutputContainers(false);
+		break;
+	case EMovieRenderPipelineState::Export:
+		// Don't flush the export to keep the UI responsive.
+		TickPostFinalizeExport(false);
+		break;
+	}
+}
+
+void UMovieGraphPipeline::TickProducingFrames()
+{
+	check(GraphTimeStepInstance);
+
+	// The callback for this function does not get registered until Initialization has been called, which sets
+	// the state to Render. If it's not, we have a initialization order/flow issue!
+	//check(PipelineState == EMovieRenderPipelineState::ProducingFrames);
+
+	// We should not be calling this once we have completed all the shots.
+	//check(CurrentShotIndex >= 0 && CurrentShotIndex < ActiveShotList.Num());
+	//
+	//ProcessOutstandingFutures();
+	//
+	//if (bShutdownRequested)
+	//{
+	//	UE_LOG(LogMovieRenderPipeline, Log, TEXT("MovieGraph TickProductingFrames: Async Shutdown Requested, abandoning remaining work and moving to Finalize."));
+	//	TransitionToState(EMovieRenderPipelineState::Finalize);
+	//	return;
+	//}
+
+		// When start up we want to override the engine's Custom Timestep with our own.
+	// This gives us the ability to completely control the engine tick/delta time before the frame
+	// is started so that we don't have to always be thinking of delta times one frame ahead. We need
+	// to do this only once we're ready to set the timestep though, as Initialize can be called as
+	// a result of a OnBeginFrame, meaning that Initialize is called on the frame before TickProducingFrames
+	// so there would be one frame where it used the custom timestep (after initialize) before TPF was called.
+	//if (GEngine->GetCustomTimeStep() != CustomTimeStep)
+	//{
+	//	CachedPrevCustomTimeStep = GEngine->GetCustomTimeStep();
+	//	GEngine->SetCustomTimeStep(CustomTimeStep);
+	//}
+
+
+	FMovieGraphTimeStepData TimeStepData;
+
+	GraphTimeStepInstance->ResetPerTickData();
+
+	UMoviePipelineExecutorShot* CurrentCameraCut = ActiveShotList[CurrentShotIndex];
+	if (CurrentCameraCut->ShotInfo.State == EMovieRenderShotState::Uninitialized)
+	{
+		UE_LOG(LogMovieRenderPipeline, Log, TEXT("MovieGraph Initializing Camera Cut [%d/%d] in [%s] %s."),
+			CurrentShotIndex + 1, ActiveShotList.Num(), *CurrentCameraCut->OuterName, *CurrentCameraCut->InnerName);
+
+		SetupShot(CurrentCameraCut);
+		GraphTimeStepInstance->InitializeShot(CurrentCameraCut, /*Out*/ TimeStepData);
+		GraphDataCachingInstance->InitializeShot(CurrentCameraCut, /*In*/ TimeStepData);
+
+		// We can safely fall through to the below states as they're OK to process the same frame we set up.
+		UE_LOG(LogMovieRenderPipeline, Log, TEXT("MovieGraph Finished initializing Camera Cut [%d/%d] in [%s] %s."),
+			CurrentShotIndex + 1, ActiveShotList.Num(), *CurrentCameraCut->OuterName, *CurrentCameraCut->InnerName);
+
+		// Temp...
+		CurrentCameraCut->ShotInfo.State = EMovieRenderShotState::Rendering;
+	}
+
+	if (CurrentCameraCut->ShotInfo.State == EMovieRenderShotState::Rendering)
+	{
+		ULevelSequence* CurrentSequence = Cast<ULevelSequence>(GetCurrentJob()->Sequence.TryLoad());
+		FFrameRate TickResolution = CurrentSequence->GetMovieScene()->GetTickResolution();
+		FFrameRate FrameRate = CurrentSequence->GetMovieScene()->GetDisplayRate();
+		
+		FFrameTime TicksPerOutputFrame = FFrameRate::TransformTime(FFrameTime(FFrameNumber(1)), FrameRate, TickResolution);
+		CurrentCameraCut->ShotInfo.CurrentTickInRoot += TicksPerOutputFrame.GetFrame();
+
+		if (CurrentCameraCut->ShotInfo.CurrentTickInRoot >= CurrentCameraCut->ShotInfo.TotalOutputRangeRoot.GetUpperBoundValue())
+		{
+			CurrentCameraCut->ShotInfo.State = EMovieRenderShotState::Finished;
+			TeardownShot(CurrentCameraCut);
+			return;
+		}
+	}
+
+	// const MoviePipeline::FFrameConstantMetrics FrameMetrics = CalculateShotFrameMetrics(CurrentCameraCut);
+	// CachedOutputState.ResetPerFrameData();
+
+
+	// Public API:
+	// Which Output Frame are we on (overall)
+
+	// Output State needs to know:
+	// ShotIndex
+	// does output state need sto be an interface? stocastic sampling makes
+	// temporal sample count/temporal index not have any meaning. or maybe
+	// we just re-use it and have some fields be dummies? This is used for
+	// the UI (for progress tracking), and filenames, and debug file output names.
+
+}
+
+
+void UMovieGraphPipeline::TickFinalizeOutputContainers(const bool bInForceFinish)
+{
+	// Tick all containers until they all report that they have finalized.
+	bool bAllContainsFinishedProcessing = true;
+
+	//do
+	//{
+	//	bAllContainsFinishedProcessing = true;
+	//
+	//	// Ask the containers if they're all done processing.
+	//	for (UMoviePipelineOutputBase* Container : GetPipelinePrimaryConfig()->GetOutputContainers())
+	//	{
+	//		bAllContainsFinishedProcessing &= Container->HasFinishedProcessing();
+	//	}
+	//
+	//	// If we aren't forcing a finish, early out after one loop to keep
+	//	// the editor/ui responsive.
+	//	if (!bInForceFinish || bAllContainsFinishedProcessing)
+	//	{
+	//		break;
+	//	}
+	//
+	//	// If they've reached here, they're forcing them to finish so we'll sleep for a touch to give
+	//	// everyone a chance to actually do work before asking them if they're done.
+	//	FPlatformProcess::Sleep(1.f);
+	//
+	//} while (true);
+
+	// If an output container is still working, we'll early out to keep the UI responsive.
+	// If they've forced a finish this will have to be true before we can reach this block.
+	if (!bAllContainsFinishedProcessing)
+	{
+		return;
+	}
+
+	//TArray<UMoviePipelineOutputBase*> Settings = GetPipelinePrimaryConfig()->GetOutputContainers();
+	//Algo::SortBy(Settings, [](const UMoviePipelineOutputBase* Setting) { return Setting->GetPriority(); });
+	//for (UMoviePipelineOutputBase* Container : Settings)
+	//{
+	//	// All containers have finished processing, final shutdown.
+	//	Container->Finalize();
+	//}
+
+	TransitionToState(EMovieRenderPipelineState::Export);
+}
+
+void UMovieGraphPipeline::TickPostFinalizeExport(const bool bInForceFinish)
+{
+	// This step assumes you have produced data and filled the data structures.
+	check(PipelineState == EMovieRenderPipelineState::Export);
+	UE_LOG(LogMovieRenderPipeline, Verbose, TEXT("[%d] PostFinalize Export (Start)."), GFrameCounter);
+
+	// ToDo: Loop through any extensions (such as XML export) and let them export using all of the
+	// data that was generated during this run such as containers, output names and lengths.
+	// Tick all containers until they all report that they have finalized.
+	bool bAllContainsFinishedProcessing = true;
+
+	//do
+	//{
+	//	bAllContainsFinishedProcessing = true;
+	//
+	//	// Ask the containers if they're all done processing.
+	//	for (UMoviePipelineSetting* Setting : GetPipelinePrimaryConfig()->GetAllSettings())
+	//	{
+	//		bAllContainsFinishedProcessing &= Setting->HasFinishedExporting();
+	//	}
+	//
+	//	// If we aren't forcing a finish, early out after one loop to keep
+	//	// the editor/ui responsive.
+	//	if (!bInForceFinish || bAllContainsFinishedProcessing)
+	//	{
+	//		break;
+	//	}
+	//
+	//	// If they've reached here, they're forcing them to finish so we'll sleep for a touch to give
+	//	// everyone a chance to actually do work before asking them if they're done.
+	//	FPlatformProcess::Sleep(1.f);
+	//
+	//} while (true);
+
+	UE_LOG(LogMovieRenderPipeline, Verbose, TEXT("[%d] PostFinalize Export (End)."), GFrameCounter);
+
+	// If an output container is still working, we'll early out to keep the UI responsive.
+	// If they've forced a finish this will have to be true before we can reach this block.
+	if (!bAllContainsFinishedProcessing)
+	{
+		return;
+	}
+
+	TransitionToState(EMovieRenderPipelineState::Finished);
+}
+
+
+void UMovieGraphPipeline::SetupShot(UMoviePipelineExecutorShot* InShot)
+{
+	// Set the new shot as the active shot. This enables the specified shot section and disables all other shot sections.
+	// SetSoloShot(InShot);
+
+	// Loop through just our primary settings and let them know which shot we're about to start.
+	//TArray<UMoviePipelineSetting*> Settings = GetPipelinePrimaryConfig()->GetAllSettings();
+	//Algo::SortBy(Settings, [](const UMoviePipelineSetting* Setting) { return Setting->GetPriority(); });
+	//for (UMoviePipelineSetting* Setting : Settings)
+	//{
+	//	Setting->OnSetupForShot(InShot);
+	//}
+	//
+	//if (InShot->GetShotOverrideConfiguration() != nullptr)
+	//{
+	//	// Any shot-specific overrides haven't had first time initialization. So we'll do that now.
+	//	TArray<UMoviePipelineSetting*> ShotSettings = InShot->GetShotOverrideConfiguration()->GetUserSettings();
+	//	Algo::SortBy(ShotSettings, [](const UMoviePipelineSetting* Setting) { return Setting->GetPriority(); });
+	//	for (UMoviePipelineSetting* Setting : ShotSettings)
+	//	{
+	//		Setting->OnMoviePipelineInitialized(this);
+	//	}
+	//}
+
+	// Setup required rendering architecture for all passes in this shot.
+	GraphRendererInstance->SetupRenderingPipelineForShot(InShot);
+}
+
+void UMovieGraphPipeline::TeardownShot(UMoviePipelineExecutorShot* InShot)
+{
+	// Teardown any rendering architecture for this shot. This needs to happen first because it'll flush outstanding rendering commands
+	GraphRendererInstance->TeardownRenderingPipelineForShot(InShot);
+
+	// some other stuff
+
+	CurrentShotIndex++;
+
+	// Check to see if this was the last shot in the Pipeline, otherwise on the next
+	// tick the new shot will be initialized and processed.
+	if (CurrentShotIndex >= ActiveShotList.Num())
+	{
+		UE_LOG(LogMovieRenderPipeline, Log, TEXT("MovieGraph Finished rendering last shot. Moving to Finalize to finish writing items to disk."));
+		TransitionToState(EMovieRenderPipelineState::Finalize);
+	}
+}
+
+void UMovieGraphPipeline::OnEngineTickEndFrame()
+{
+	LLM_SCOPE_BYNAME(TEXT("MovieGraphEndFrame"));
+
+	// Unfortunately, since we can't control when our Initialization function is called
+	// we can end up in a situation where this callback is registered but the matching
+	// OnEngineTickBeginFrame() hasn't been called for that given engine tick. Instead of
+	// changing this registration to hang off of the end of the first OnEngineTickBeginFrame()
+	// we instead just early out here if that hasn't actually been called once. This decision
+	// is designed to minimize places where callbacks are registered and where flow changes.
+	//if (!bHasRunBeginFrameOnce)
+	//{
+	//	return;
+	//}
+
+	UE_LOG(LogMovieRenderPipeline, VeryVerbose, TEXT("MovieGraph OnEngineTickEndFrame (Start)"));
+
+	// ProcessAudioTick();
+	RenderFrame();
+
+	UE_LOG(LogMovieRenderPipeline, VeryVerbose, TEXT("MovieGraph OnEngineTickEndFrame (End)"));
+}
+
+void UMovieGraphPipeline::RenderFrame()
+{
+	check(GraphRendererInstance);
+	FMovieGraphTimeStepData TimeStepData;
+
+	GraphRendererInstance->Render(TimeStepData);
+}
+
+void UMovieGraphPipeline::RequestShutdownImpl(bool bIsError)
+{
+	// It's possible for a previous call to RequestionShutdown to have set an error before this call that may not
+	// We don't want to unset a previously set error state
+	if (bIsError)
+	{
+		bShutdownSetErrorFlag = true;
+	}
+
+	// The user has requested a shutdown, it will be read the next available chance and possibly acted on.
+	bShutdownRequested = true;
+	switch (PipelineState)
+	{
+		// It is valid to call Shutdown at any point during these two states.
+	case EMovieRenderPipelineState::Uninitialized:
+	case EMovieRenderPipelineState::ProducingFrames:
+		break;
+		// You can call Shutdown during these two, but they won't do anything as we're already shutting down at that point.
+	case EMovieRenderPipelineState::Finalize:
+	case EMovieRenderPipelineState::Export:
+		UE_LOG(LogMovieRenderPipeline, Log, TEXT("MovieGraph Async Shutdown Requested, ignoring due to already being on the way to shutdown."));
+		break;
+	}
+}
+
+void UMovieGraphPipeline::ShutdownImpl(bool bIsError)
+{
+	check(IsInGameThread());
+
+	// We flag this so you can check if the shutdown was requested even when we do a stall-stop.
+	bShutdownRequested = true;
+
+	// It's possible for a previous call to RequestShutdown to have set an error before this call and
+	// we don't want to just blow away the error flag even if this function was then called normally.
+	// (ie: Don't accidnetally unset error state)
+	if (bIsError)
+	{
+		bShutdownSetErrorFlag = true;
+	}
+
+	// This is a blocking operation which abandons any outstanding work to be submitted but finishes
+	// the existing work already processed.
+	if (PipelineState == EMovieRenderPipelineState::Uninitialized)
+	{
+		// If initialize is not called, no need to do anything.
+		return;
+	}
+
+	if (PipelineState == EMovieRenderPipelineState::Finished)
+	{
+		// If already shut down, no need to do anything.
+		return;
+	}
+
+	if (PipelineState == EMovieRenderPipelineState::ProducingFrames)
+	{
+		// Teardown the currently active shot (if there is one). This will flush any outstanding rendering
+		// work that has already submitted - it cannot be canceled, so we may as well execute it and save the results.
+		TransitionToState(EMovieRenderPipelineState::Finalize);
+
+		// Abandon the current frame. When using temporal sampling we may had canceled mid-frame, so the rendering
+		// commands were never submitted, thus the output builder will still be expecting a frame to come in.
+		//if (CachedOutputState.TemporalSampleCount > 1)
+		//{
+		//	OutputBuilder->AbandonOutstandingWork();
+		//}
+	}
+
+	if (PipelineState == EMovieRenderPipelineState::Finalize)
+	{
+		// We were either in the middle of writing frames to disk, or we have moved to Finalize as a result of the above block.
+		// Tick output containers until they report they have finished writing to disk. This is a blocking operation. 
+		// Finalize automatically switches our state to Export so no need to manually transition afterwards.
+		//TickFinalizeOutputContainers(true);
+	}
+
+	if (PipelineState == EMovieRenderPipelineState::Export)
+	{
+		// All frames have been written to disk but we're doing a post-export step (such as encoding). Flush this operation as well.
+		// Export automatically switches our state to Finished so no need to manually transition afterwards.
+		//TickPostFinalizeExport(true);
+	}
+}
+
+void UMovieGraphPipeline::TransitionToState(const EMovieRenderPipelineState InNewState)
+{
+	// No re-entrancy. This isn't an error as tearing down a shot may try to move to
+	// Finalize on its own, but we don't want that.
+	if (bIsTransitioningState)
+	{
+		return;
+	}
+
+	TGuardValue<bool> StateTransitionGuard(bIsTransitioningState, true);
+
+	bool bInvalidTransition = true;
+	switch (PipelineState)
+	{
+	case EMovieRenderPipelineState::Uninitialized:
+		PipelineState = InNewState;
+		bInvalidTransition = false;
+		break;
+	case EMovieRenderPipelineState::ProducingFrames:
+		if (InNewState == EMovieRenderPipelineState::Finalize)
+		{
+			bInvalidTransition = false;
+
+			// If we had naturally finished the last shot before doing this transition it will have
+			// already been torn down, so this only catches mid-shot transitions to ensure teardown.
+			if (CurrentShotIndex < ActiveShotList.Num())
+			{
+				// Ensures all in-flight work for that shot is handled.
+				TeardownShot(ActiveShotList[CurrentShotIndex]);
+			}
+
+			// Unregister our OnEngineTickEndFrame delegate. We don't unregister BeginFrame as we need
+			// to continue to call it to allow ticking the Finalization stage.
+			// FCoreDelegates::OnEndFrame.RemoveAll(this);
+
+			// Reset the Custom Timestep because we don't care how long the engine takes now
+			// GEngine->SetCustomTimeStep(CachedPrevCustomTimeStep);
+
+			// Ensure all frames have been processed by the GPU and sent to the Output Merger
+			FlushRenderingCommands();
+
+			// And then make sure all frames are sent to the Output Containers before we finalize.
+			// ProcessOutstandingFinishedFrames();
+
+			// PreviewTexture = nullptr;
+
+			// This is called once notifying output containers that all frames that will be submitted have been submitted.
+			PipelineState = EMovieRenderPipelineState::Finalize;
+			//BeginFinalize();
+		}
+		break;
+	case EMovieRenderPipelineState::Finalize:
+		if (InNewState == EMovieRenderPipelineState::Export)
+		{
+			bInvalidTransition = false;
+
+			// This is called once notifying our export step that they can begin the export.
+			PipelineState = EMovieRenderPipelineState::Export;
+
+			// Restore the sequence so that the export processes can operate on the original sequence. 
+			// This is also done in the finished state because it's not guaranteed that the Export state 
+			// will be set when the render is canceled early
+			// LevelSequenceActor->GetSequencePlayer()->Stop();
+			// RestoreTargetSequenceToOriginalState();
+
+			// Ensure all of our Futures have been converted to the GeneratedOutputData. This has to happen
+			// after finalize finishes, because the futures won't be available until actually written to disk.
+			// ProcessOutstandingFutures();
+
+			//BeginExport();
+		}
+		break;
+	case EMovieRenderPipelineState::Export:
+		if (InNewState == EMovieRenderPipelineState::Finished)
+		{
+			bInvalidTransition = false;
+			PipelineState = EMovieRenderPipelineState::Finished;
+
+			// Uninitialize our primary config settings. Reverse sorted so settings that cached values restore correctly.
+			//TArray<UMoviePipelineSetting*> Settings = GetPipelinePrimaryConfig()->GetAllSettings();
+			//Algo::SortBy(Settings, [](const UMoviePipelineSetting* Setting) { return Setting->GetPriority(); }, TLess<int32>());
+			//for (UMoviePipelineSetting* Setting : Settings)
+			//{
+			//	Setting->OnMoviePipelineShutdown(this);
+			//}
+
+			// Restore any custom Time Step that may have been set before.
+			//GEngine->SetCustomTimeStep(CachedPrevCustomTimeStep);
+
+			// Ensure our delegates don't get called anymore as we're going to become null soon.
+			FCoreDelegates::OnBeginFrame.RemoveAll(this);
+			FCoreDelegates::OnEndFrame.RemoveAll(this);
+
+			//if (DebugWidget)
+			//{
+			//	DebugWidget->RemoveFromParent();
+			//	DebugWidget = nullptr;
+			//}
+
+			//TArray<UMoviePipelineOutputBase*> ContainerSettings = GetPipelinePrimaryConfig()->GetOutputContainers();
+			//Algo::SortBy(ContainerSettings, [](const UMoviePipelineOutputBase* Setting) { return Setting->GetPriority(); });
+			//for (UMoviePipelineOutputBase* Setting : ContainerSettings)
+			//{
+			//	Setting->OnPipelineFinished();
+			//}
+			//
+			//TeardownAudioRendering();
+			//LevelSequenceActor->GetSequencePlayer()->Stop();
+			//RestoreTargetSequenceToOriginalState();
+			//
+			//if (UGameViewportClient* Viewport = GetWorld()->GetGameViewport())
+			//{
+			//	Viewport->bDisableWorldRendering = false;
+			//}
+
+			// Because the render target pool is shared, if you had a high-resolution render in editor the entire gbuffer
+			// has been resized up to match the new maximum extent. This console command will reset the size of the pool
+			// and cause it to re-allocate at the currrent size on the next render request, which is likely to be the size
+			// of the PIE window (720p) or the Viewport itself.
+			//UKismetSystemLibrary::ExecuteConsoleCommand(GetWorld(), TEXT("r.ResetRenderTargetsExtent"), nullptr);
+			//CustomTimeStep->RestoreCachedWorldSettings();
+
+			//GAreScreenMessagesEnabled = bPrevGScreenMessagesEnabled;
+
+			//UE_LOG(LogMovieRenderPipeline, Log, TEXT("Movie Pipeline completed. Duration: %s"), *(FDateTime::UtcNow() - InitializationTime).ToString());
+
+			//UMoviePipelineDebugSettings* DebugSetting = FindOrAddSettingForShot<UMoviePipelineDebugSettings>(nullptr);
+			//if (DebugSetting)
+			//{
+			//	if (DebugSetting->bCaptureUnrealInsightsTrace)
+			//	{
+			//		StopUnrealInsightsCapture();
+			//	}
+			//}
+
+			OnMoviePipelineFinishedImpl();
+		}
+		break;
+	}
+
+	if (!ensureAlwaysMsgf(!bInvalidTransition, TEXT("MovieGraph An invalid transition was requested (from: %d to: %d), ignoring transition request."),
+		PipelineState, InNewState))
+	{
+		return;
+	}
+}
+
+void UMovieGraphPipeline::OnMoviePipelineFinishedImpl()
+{
+	// Broadcast to both Native and Python/BP
+	//
+	// Generate a params struct containing the data generated by this job.
+	FMoviePipelineOutputData Params;
+	//Params.Pipeline = this;
+	//Params.Job = GetCurrentJob();
+	//Params.bSuccess = !bFatalError;
+	//Params.ShotData = GeneratedShotOutputData;
+	//
+	//UE_LOG(LogMovieRenderPipelineIO, Verbose, TEXT("Files written to disk for entire sequence:"));
+	//PrintVerboseLogForFiles(GeneratedShotOutputData);
+	//UE_LOG(LogMovieRenderPipelineIO, Verbose, TEXT("Completed outputting files written to disk."));
+	//
+	OnMoviePipelineWorkFinishedDelegateNative.Broadcast(Params);
+	// OnMoviePipelineWorkFinishedDelegate.Broadcast(Params);
+}
+
+//void UMovieGraphPipeline::PrintVerboseLogForFiles(const TArray<FMoviePipelineShotOutputData>& InOutputData) const
+//{
+//	for (const FMoviePipelineShotOutputData& OutputData : InOutputData)
+//	{
+//		const UMoviePipelineExecutorShot* Shot = OutputData.Shot.Get();
+//		if (Shot)
+//		{
+//			UE_LOG(LogMovieRenderPipelineIO, Verbose, TEXT("Shot: %s [%s]"), *Shot->OuterName, *Shot->InnerName);
+//		}
+//		for (const TPair<FMoviePipelinePassIdentifier, FMoviePipelineRenderPassOutputData>& Pair : OutputData.RenderPassData)
+//		{
+//			UE_LOG(LogMovieRenderPipelineIO, Verbose, TEXT("Render Pass: %s"), *Pair.Key.Name);
+//			for (const FString& FilePath : Pair.Value.FilePaths)
+//			{
+//				UE_LOG(LogMovieRenderPipelineIO, Verbose, TEXT("\t\t%s"), *FilePath);
+//			}
+//		}
+//	}
+//}
diff --git a/Engine/Plugins/MovieScene/MovieRenderPipeline/Source/MovieRenderPipelineCore/Private/MoviePipelineInProcessExecutor.cpp b/Engine/Plugins/MovieScene/MovieRenderPipeline/Source/MovieRenderPipelineCore/Private/MoviePipelineInProcessExecutor.cpp
--- a/Engine/Plugins/MovieScene/MovieRenderPipeline/Source/MovieRenderPipelineCore/Private/MoviePipelineInProcessExecutor.cpp
+++ b/Engine/Plugins/MovieScene/MovieRenderPipeline/Source/MovieRenderPipelineCore/Private/MoviePipelineInProcessExecutor.cpp
@@ -161,7 +161,7 @@ void UMoviePipelineInProcessExecutor::OnMapLoadFinished(UWorld* NewWorld)
 		}
 		else if (UMovieGraphPipeline* PipelineAsGraph = Cast<UMovieGraphPipeline>(ActiveMoviePipeline))
 		{
-			// PipelineAsGraph->Initialize(Queue->GetJobs()[CurrentPipelineIndex], FMovieGraphInitConfig());
+			PipelineAsGraph->Initialize(Queue->GetJobs()[CurrentPipelineIndex], FMovieGraphInitConfig());
 		}
 		RemainingInitializationFrames = -1;
 	}
diff --git a/Engine/Plugins/MovieScene/MovieRenderPipeline/Source/MovieRenderPipelineCore/Public/Graph/MovieGraphDataTypes.h b/Engine/Plugins/MovieScene/MovieRenderPipeline/Source/MovieRenderPipelineCore/Public/Graph/MovieGraphDataTypes.h
new file mode 100644
--- /dev/null
+++ b/Engine/Plugins/MovieScene/MovieRenderPipeline/Source/MovieRenderPipelineCore/Public/Graph/MovieGraphDataTypes.h
@@ -0,0 +1,103 @@
+// Copyright Epic Games, Inc. All Rights Reserved.
+#pragma once
+
+#include "UObject/Object.h"
+#include "Templates/SubclassOf.h"
+#include "MovieGraphDataTypes.generated.h"
+
+// Forward Declares
+class UMovieGraphPipeline;
+class UMovieGraphTimeStepBase;
+class UMovieGraphRendererBase;
+class UMovieGraphTimeRangeBuilderBase;
+class UMovieGraphDataCachingBase;
+class UMoviePipelineExecutorJob;
+class UMoviePipelineExecutorShot;
+
+USTRUCT(BlueprintType)
+struct MOVIERENDERPIPELINECORE_API FMovieGraphInitConfig
+{
+	GENERATED_BODY()
+	
+	FMovieGraphInitConfig();
+	
+	UPROPERTY(BlueprintReadWrite, EditAnywhere, Category = "Movie Graph")
+	TSubclassOf<UMovieGraphTimeStepBase> TimeStepClass;
+	
+	UPROPERTY(BlueprintReadWrite, EditAnywhere, Category = "Movie Graph")
+	TSubclassOf<UMovieGraphRendererBase> RendererClass;
+
+	UPROPERTY(BlueprintReadWrite, EditAnywhere, Category = "Movie Graph")
+	TSubclassOf<UMovieGraphTimeRangeBuilderBase> TimeRangeBuilderClass;
+
+	UPROPERTY(BlueprintReadWrite, EditAnywhere, Category = "Movie Graph")
+	TSubclassOf<UMovieGraphDataCachingBase> DataCachingClass;
+
+	UPROPERTY(BlueprintReadWrite, EditAnywhere, Category = "Movie Graph")
+	bool bRenderViewport;
+};
+
+
+USTRUCT(BlueprintType)
+struct MOVIERENDERPIPELINECORE_API FMovieGraphTimeStepData
+{
+	GENERATED_BODY()
+};
+
+UCLASS(BlueprintType, Abstract)
+class MOVIERENDERPIPELINECORE_API UMovieGraphTimeStepBase : public UObject
+{
+	GENERATED_BODY()
+public:
+	virtual void CalculateTimeStep(FMovieGraphTimeStepData& OutTimeData)
+	{}
+	virtual void ResetPerTickData() {}
+	virtual void InitializeShot(UMoviePipelineExecutorShot* InShot, FMovieGraphTimeStepData& OutTimeStepData) {}
+
+protected:
+	UMovieGraphPipeline* GetOwningGraph() const;
+};
+
+UCLASS(BlueprintType, Abstract)
+class MOVIERENDERPIPELINECORE_API UMovieGraphRendererBase : public UObject
+{
+	GENERATED_BODY()
+public:
+	virtual void Render(const FMovieGraphTimeStepData& InTimeData) {}
+	virtual void SetupRenderingPipelineForShot(UMoviePipelineExecutorShot* InShot) {}
+	virtual void TeardownRenderingPipelineForShot(UMoviePipelineExecutorShot* InShot) {}
+
+
+protected:
+	UMovieGraphPipeline* GetOwningGraph() const;
+};
+
+UCLASS(BlueprintType, Abstract)
+class MOVIERENDERPIPELINECORE_API UMovieGraphTimeRangeBuilderBase : public UObject
+{
+	GENERATED_BODY()
+public:
+	virtual void BuildTimeRanges()
+	{
+
+	}
+
+
+protected:
+	UMovieGraphPipeline* GetOwningGraph() const;
+};
+
+UCLASS(BlueprintType, Abstract)
+class MOVIERENDERPIPELINECORE_API UMovieGraphDataCachingBase : public UObject
+{
+	GENERATED_BODY()
+public:
+	virtual void CacheDataPreJob(const FMovieGraphInitConfig& InInitConfig) {}
+	virtual void RestoreCachedDataPostJob() {}
+	virtual void UpdateShotList() {}
+	virtual void InitializeShot(UMoviePipelineExecutorShot* InShot, const FMovieGraphTimeStepData& InTimeStepData) {}
+
+protected:
+	UMovieGraphPipeline* GetOwningGraph() const;
+};
+
diff --git a/Engine/Plugins/MovieScene/MovieRenderPipeline/Source/MovieRenderPipelineCore/Public/Graph/MovieGraphDefaultDataCaching.h b/Engine/Plugins/MovieScene/MovieRenderPipeline/Source/MovieRenderPipelineCore/Public/Graph/MovieGraphDefaultDataCaching.h
new file mode 100644
--- /dev/null
+++ b/Engine/Plugins/MovieScene/MovieRenderPipeline/Source/MovieRenderPipelineCore/Public/Graph/MovieGraphDefaultDataCaching.h
@@ -0,0 +1,32 @@
+// Copyright Epic Games, Inc. All Rights Reserved.
+#pragma once
+
+#include "Graph/MovieGraphDataTypes.h"
+#include "Misc/FrameTime.h"
+#include "MovieGraphDefaultDataCaching.generated.h"
+
+// Forward Declares
+class ULevelSequence;
+class ALevelSequenceActor;
+class UMovieSceneSequencePlayer;
+
+/**
+*
+*/
+UCLASS(BlueprintType)
+class MOVIERENDERPIPELINECORE_API UMovieGraphDefaultDataCaching : public UMovieGraphDataCachingBase
+{
+	GENERATED_BODY()
+public:
+	virtual void CacheDataPreJob(const FMovieGraphInitConfig& InInitConfig) override;
+	virtual void RestoreCachedDataPostJob() override;
+	virtual void UpdateShotList() override;
+
+protected:
+	void CacheLevelSequenceData(ULevelSequence* InSequence);
+	void OnSequenceEvaluated(const UMovieSceneSequencePlayer& Player, FFrameTime CurrentTime, FFrameTime PreviousTime);
+
+protected:
+	UPROPERTY(Transient)
+	TObjectPtr<ALevelSequenceActor> LevelSequenceActor;
+};
\ No newline at end of file
diff --git a/Engine/Plugins/MovieScene/MovieRenderPipeline/Source/MovieRenderPipelineCore/Public/Graph/MovieGraphDefaultRenderer.h b/Engine/Plugins/MovieScene/MovieRenderPipeline/Source/MovieRenderPipelineCore/Public/Graph/MovieGraphDefaultRenderer.h
new file mode 100644
--- /dev/null
+++ b/Engine/Plugins/MovieScene/MovieRenderPipeline/Source/MovieRenderPipelineCore/Public/Graph/MovieGraphDefaultRenderer.h
@@ -0,0 +1,18 @@
+// Copyright Epic Games, Inc. All Rights Reserved.
+#pragma once
+
+#include "Graph/MovieGraphDataTypes.h"
+#include "MovieGraphDefaultRenderer.generated.h"
+
+/**
+* This class 
+*/
+UCLASS(BlueprintType)
+class MOVIERENDERPIPELINECORE_API UMovieGraphDefaultRenderer : public UMovieGraphRendererBase
+{
+	GENERATED_BODY()
+public:
+	virtual void Render(const FMovieGraphTimeStepData& InTimeData) override;
+	virtual void SetupRenderingPipelineForShot(UMoviePipelineExecutorShot* InShot) override;
+	virtual void TeardownRenderingPipelineForShot(UMoviePipelineExecutorShot* InShot) override;
+};
\ No newline at end of file
diff --git a/Engine/Plugins/MovieScene/MovieRenderPipeline/Source/MovieRenderPipelineCore/Public/Graph/MovieGraphLinearTimeStep.h b/Engine/Plugins/MovieScene/MovieRenderPipeline/Source/MovieRenderPipelineCore/Public/Graph/MovieGraphLinearTimeStep.h
new file mode 100644
--- /dev/null
+++ b/Engine/Plugins/MovieScene/MovieRenderPipeline/Source/MovieRenderPipelineCore/Public/Graph/MovieGraphLinearTimeStep.h
@@ -0,0 +1,16 @@
+// Copyright Epic Games, Inc. All Rights Reserved.
+#pragma once
+
+#include "Graph/MovieGraphDataTypes.h"
+#include "MovieGraphLinearTimeStep.generated.h"
+
+/**
+* This class 
+*/
+UCLASS(BlueprintType)
+class MOVIERENDERPIPELINECORE_API UMovieGraphLinearTimeStep : public UMovieGraphTimeStepBase
+{
+	GENERATED_BODY()
+public:
+	virtual void CalculateTimeStep(FMovieGraphTimeStepData& OutTimeData) override;
+};
\ No newline at end of file
diff --git a/Engine/Plugins/MovieScene/MovieRenderPipeline/Source/MovieRenderPipelineCore/Public/Graph/MovieGraphPipeline.h b/Engine/Plugins/MovieScene/MovieRenderPipeline/Source/MovieRenderPipelineCore/Public/Graph/MovieGraphPipeline.h
--- a/Engine/Plugins/MovieScene/MovieRenderPipeline/Source/MovieRenderPipelineCore/Public/Graph/MovieGraphPipeline.h
+++ b/Engine/Plugins/MovieScene/MovieRenderPipeline/Source/MovieRenderPipelineCore/Public/Graph/MovieGraphPipeline.h
@@ -3,6 +3,8 @@
 
 #include "UObject/Object.h"
 #include "MoviePipelineBase.h"
+#include "MovieGraphDataTypes.h"
+#include "MovieRenderPipelineDataTypes.h"
 #include "MovieGraphPipeline.generated.h"
 
 // Forward Declares
@@ -13,4 +15,73 @@ UCLASS(BlueprintType)
 class MOVIERENDERPIPELINECORE_API UMovieGraphPipeline : public UMoviePipelineBase
 {
 	GENERATED_BODY()
+
+public:
+	UMovieGraphPipeline();
+
+	UFUNCTION(BlueprintCallable, Category = "Movie Graph")
+	void Initialize(UMoviePipelineExecutorJob* InJob, const FMovieGraphInitConfig& InitConfig);
+
+	UFUNCTION(BlueprintCallable, Category = "Movie Graph")
+	UMoviePipelineExecutorJob* GetCurrentJob() const { return CurrentJob; }
+
+
+
+protected:
+	UFUNCTION(BlueprintCallable, Category = "Movie Render Pipeline")
+	virtual void OnMoviePipelineFinishedImpl();
+
+protected:
+	virtual void OnEngineTickBeginFrame();
+	virtual void OnEngineTickEndFrame();
+	virtual void RenderFrame();
+	virtual void BuildShotListFromDataSource();
+	virtual void SetupShot(UMoviePipelineExecutorShot* InShot);
+	virtual void TeardownShot(UMoviePipelineExecutorShot* InShot);
+
+	virtual void TickProducingFrames();
+	virtual void TickPostFinalizeExport(const bool bInForceFinish);
+	virtual void TickFinalizeOutputContainers(const bool bInForceFinish);
+	virtual void TransitionToState(const EMovieRenderPipelineState InNewState);
+
+	// UMoviePipelineBase Interface
+	virtual void RequestShutdownImpl(bool bIsError) override;
+	virtual void ShutdownImpl(bool bIsError) override;
+	virtual bool IsShutdownRequestedImpl() const override { return bShutdownRequested; }
+	virtual EMovieRenderPipelineState GetPipelineStateImpl() const override { return PipelineState; }
+	// ~UMoviePipelineBase Interface
+
+protected:
+	UPROPERTY(Transient)
+	TObjectPtr<UMovieGraphTimeStepBase> GraphTimeStepInstance;
+
+	UPROPERTY(Transient)
+	TObjectPtr<UMovieGraphRendererBase> GraphRendererInstance;
+
+	UPROPERTY(Transient)
+	TObjectPtr<UMovieGraphTimeRangeBuilderBase> GraphTimeRangeBuilderInstance;
+
+	UPROPERTY(Transient)
+	TObjectPtr<UMovieGraphDataCachingBase> GraphDataCachingInstance;
+
+protected:
+	UPROPERTY(Transient)
+	TObjectPtr<UMoviePipelineExecutorJob> CurrentJob;
+
+	UPROPERTY(Transient)
+	TArray<TObjectPtr<UMoviePipelineExecutorShot>> ActiveShotList;
+
+	int32 CurrentShotIndex;
+
+	/** True if we're in a TransitionToState call. Used to prevent reentrancy. */
+	bool bIsTransitioningState;
+
+	/** True if RequestShutdown() was called. At the start of the next frame we will stop producing frames and start shutting down. */
+	FThreadSafeBool bShutdownRequested;
+
+	/** Set to true during Shutdown/RequestShutdown if we are shutting down due to an error. */
+	FThreadSafeBool bShutdownSetErrorFlag;
+
+	/** Which step of the rendering process is the graph currently in. */
+	EMovieRenderPipelineState PipelineState;
 };
\ No newline at end of file
diff --git a/Engine/Plugins/MovieScene/MovieRenderPipeline/Source/MovieRenderPipelineCore/Public/Graph/MovieGraphSequenceTimeRangeBuilder.h b/Engine/Plugins/MovieScene/MovieRenderPipeline/Source/MovieRenderPipelineCore/Public/Graph/MovieGraphSequenceTimeRangeBuilder.h
new file mode 100644
--- /dev/null
+++ b/Engine/Plugins/MovieScene/MovieRenderPipeline/Source/MovieRenderPipelineCore/Public/Graph/MovieGraphSequenceTimeRangeBuilder.h
@@ -0,0 +1,19 @@
+// Copyright Epic Games, Inc. All Rights Reserved.
+#pragma once
+
+#include "Graph/MovieGraphDataTypes.h"
+#include "MovieGraphSequenceTimeRangeBuilder.generated.h"
+
+// Forward Declares
+class ULevelSequence;
+class ALevelSequenceActor;
+
+/**
+*
+*/
+UCLASS(BlueprintType)
+class MOVIERENDERPIPELINECORE_API UMovieGraphSequenceTimeRangeBuilder : public UMovieGraphTimeRangeBuilderBase
+{
+	GENERATED_BODY()
+
+};
\ No newline at end of file
diff --git a/Engine/Plugins/MovieScene/MovieRenderPipeline/Source/MovieRenderPipelineEditor/Private/MoviePipelinePIEExecutor.cpp b/Engine/Plugins/MovieScene/MovieRenderPipeline/Source/MovieRenderPipelineEditor/Private/MoviePipelinePIEExecutor.cpp
--- a/Engine/Plugins/MovieScene/MovieRenderPipeline/Source/MovieRenderPipelineEditor/Private/MoviePipelinePIEExecutor.cpp
+++ b/Engine/Plugins/MovieScene/MovieRenderPipeline/Source/MovieRenderPipelineEditor/Private/MoviePipelinePIEExecutor.cpp
@@ -263,7 +263,7 @@ void UMoviePipelinePIEExecutor::OnPIEStartupFinished(bool)
 		}
 		else if (UMovieGraphPipeline* PipelineAsGraph = Cast<UMovieGraphPipeline>(ActiveMoviePipeline))
 		{
-			// PipelineAsGraph->Initialize(Queue->GetJobs()[CurrentPipelineIndex], FMovieGraphInitConfig());
+			PipelineAsGraph->Initialize(Queue->GetJobs()[CurrentPipelineIndex], FMovieGraphInitConfig());
 		}
 		RemainingInitializationFrames = -1;
 	}
