MoviePipeline: Fixed an uninitialized struct member error.

diff --git a/Engine/Plugins/MovieScene/MovieRenderPipeline/Source/MovieRenderPipelineCore/Private/Graph/MovieGraphDataTypes.cpp b/Engine/Plugins/MovieScene/MovieRenderPipeline/Source/MovieRenderPipelineCore/Private/Graph/MovieGraphDataTypes.cpp
--- a/Engine/Plugins/MovieScene/MovieRenderPipeline/Source/MovieRenderPipelineCore/Private/Graph/MovieGraphDataTypes.cpp
+++ b/Engine/Plugins/MovieScene/MovieRenderPipeline/Source/MovieRenderPipelineCore/Private/Graph/MovieGraphDataTypes.cpp
@@ -13,6 +13,7 @@ FMovieGraphInitConfig::FMovieGraphInitConfig()
 	RendererClass = UMovieGraphDefaultRenderer::StaticClass();
 	TimeRangeBuilderClass = UMovieGraphSequenceTimeRangeBuilder::StaticClass();
 	DataCachingClass = UMovieGraphDefaultDataCaching::StaticClass();
+	bRenderViewport = false;
 }
 
 UMovieGraphPipeline* UMovieGraphTimeStepBase::GetOwningGraph() const
