MoviePipeline: Added support for variables inside the graph.

This initial pass enables variables to be created on the graph in the Members panel. Once created, variable getters can be created in the graph's context menu.

Also did a bunch of miscellaneous supporting updates:
* Added a context menu to the Members panel
* Enabled deletion of nodes in the graph (currently only exposed via the Members panel and not the graph itself)
* Selection events across the graph and Members panel now cooperate properly
* Added basic support for the data model ("runtime" classes) to communicate updates to the editor equivalents so the graph can react
* Lots of cleanup

To be done in future checkins:
* Extensive details panel customizations for variables, which will be driven by the final design
* Variables are currently hardcoded to be backed by floats; they need to support arbitrary data types, though (TVariant?)

diff --git a/Engine/Plugins/MovieScene/MovieRenderPipeline/Source/MovieRenderPipelineCore/Private/Graph/MovieGraphConfig.cpp b/Engine/Plugins/MovieScene/MovieRenderPipeline/Source/MovieRenderPipelineCore/Private/Graph/MovieGraphConfig.cpp
--- a/Engine/Plugins/MovieScene/MovieRenderPipeline/Source/MovieRenderPipelineCore/Private/Graph/MovieGraphConfig.cpp
+++ b/Engine/Plugins/MovieScene/MovieRenderPipeline/Source/MovieRenderPipelineCore/Private/Graph/MovieGraphConfig.cpp
@@ -1,9 +1,22 @@
 // Copyright Epic Games, Inc. All Rights Reserved.
 
 #include "Graph/MovieGraphConfig.h"
+
+#include "Algo/Transform.h"
 #include "Graph/MovieGraphEdge.h"
+#include "MovieGraphUtils.h"
 #include "MovieRenderPipelineCoreModule.h"
 
+#define LOCTEXT_NAMESPACE "MovieGraphConfig"
+
+#if WITH_EDITOR
+void UMovieGraphVariable::PostEditChangeProperty(FPropertyChangedEvent& PropertyChangedEvent)
+{
+	Super::PostEditChangeProperty(PropertyChangedEvent);
+
+	OnMovieGraphVariableChangedDelegate.Broadcast(this);
+}
+#endif
 
 UMovieGraphConfig::UMovieGraphConfig()
 {
@@ -36,6 +49,13 @@ void UMovieGraphConfig::TraverseTest(UMovieGraphNode* InNode)
 	}
 }
 
+void UMovieGraphConfig::OnVariableUpdated(UMovieGraphVariable* UpdatedVariable)
+{
+#if WITH_EDITOR
+	OnGraphChangedDelegate.Broadcast();
+#endif
+}
+
 UMovieGraphNode* UMovieGraphConfig::FindTraversalStartForContext(const FMovieGraphTraversalContext& InContext) const
 {
 	return OutputNode;
@@ -221,3 +241,108 @@ bool UMovieGraphConfig::RemoveOutboundEdges(UMovieGraphNode* InNode, const FName
 	return bChanged;
 }
 
+bool UMovieGraphConfig::RemoveNode(UMovieGraphNode* InNode)
+{
+	if (!InNode)
+	{
+		UE_LOG(LogMovieRenderPipeline, Error, TEXT("RemoveNode: Invalid Node"));
+		return false;
+	}
+
+	RemoveAllInboundEdges(InNode);
+	RemoveAllOutboundEdges(InNode);
+
+	return AllNodes.RemoveSingle(InNode) == 1;
+}
+
+UMovieGraphVariable* UMovieGraphConfig::AddVariable()
+{
+	using namespace UE::MoviePipeline::RenderGraph;
+	
+	UMovieGraphVariable* NewVariable = NewObject<UMovieGraphVariable>(this, NAME_None, RF_Transactional);
+	Variables.Add(NewVariable);
+
+	NewVariable->Type = EMovieGraphVariableType::Float;
+	NewVariable->SetGuid(FGuid::NewGuid());
+
+#if WITH_EDITOR
+	NewVariable->OnMovieGraphVariableChangedDelegate.AddUObject(this, &UMovieGraphConfig::OnVariableUpdated);
+#endif
+
+	// Generate and set a unique name
+	TArray<FString> ExistingVariableNames;
+	Algo::Transform(GetVariables(), ExistingVariableNames, [](const UMovieGraphVariable* Variable) { return Variable->Name; });
+	NewVariable->Name = GetUniqueName(ExistingVariableNames, "Variable");
+
+	return NewVariable;
+}
+
+UMovieGraphVariable* UMovieGraphConfig::GetVariableByGuid(const FGuid& InGuid) const
+{
+	for (const TObjectPtr<UMovieGraphVariable> Variable : Variables)
+	{
+		if (Variable->GetGuid() == InGuid)
+		{
+			return Variable;
+		}
+	}
+
+	return nullptr;
+}
+
+TArray<UMovieGraphVariable*> UMovieGraphConfig::GetVariables() const
+{
+	return Variables;
+}
+
+void UMovieGraphConfig::DeleteMember(UObject* MemberToDelete)
+{
+	if (!MemberToDelete)
+	{
+		return;
+	}
+
+	if (UMovieGraphVariable* GraphVariableToDelete = Cast<UMovieGraphVariable>(MemberToDelete))
+	{
+		DeleteVariableMember(GraphVariableToDelete);
+	}
+}
+
+void UMovieGraphConfig::DeleteVariableMember(UMovieGraphVariable* VariableMemberToDelete)
+{
+	// Find all accessor nodes using this graph variable
+	TArray<TObjectPtr<UMovieGraphNode>> NodesToRemove =
+		AllNodes.FilterByPredicate([VariableMemberToDelete](const TObjectPtr<UMovieGraphNode>& GraphNode)
+		{
+			if (const UMovieGraphVariableNode* VariableNode = Cast<UMovieGraphVariableNode>(GraphNode))
+			{
+				const UMovieGraphVariable* GraphVariable = VariableNode->GetVariable();
+				if (GraphVariable && (GraphVariable->GetGuid() == VariableMemberToDelete->GetGuid()))
+				{
+					return true;
+				}
+			}
+
+			return false;
+		});
+
+	// Remove accessor nodes
+	TArray<UMovieGraphNode*> RemovedNodes;
+	for (const TObjectPtr<UMovieGraphNode>& NodeToRemove : NodesToRemove)
+	{
+		const FGuid& NodeGuid = NodeToRemove->GetGuid();
+		if (RemoveNode(NodeToRemove.Get()))
+		{
+			RemovedNodes.Add(NodeToRemove.Get());
+		}
+	}
+
+	// Remove this variable from the variables tracked by the graph
+	Variables.RemoveSingle(VariableMemberToDelete);
+
+#if WITH_EDITOR
+	OnGraphNodesDeletedDelegate.Broadcast(RemovedNodes);
+#endif
+}
+
+#undef LOCTEXT_NAMESPACE
diff --git a/Engine/Plugins/MovieScene/MovieRenderPipeline/Source/MovieRenderPipelineCore/Private/Graph/MovieGraphNode.cpp b/Engine/Plugins/MovieScene/MovieRenderPipeline/Source/MovieRenderPipelineCore/Private/Graph/MovieGraphNode.cpp
--- a/Engine/Plugins/MovieScene/MovieRenderPipeline/Source/MovieRenderPipelineCore/Private/Graph/MovieGraphNode.cpp
+++ b/Engine/Plugins/MovieScene/MovieRenderPipeline/Source/MovieRenderPipelineCore/Private/Graph/MovieGraphNode.cpp
@@ -185,3 +185,56 @@ UMovieGraphPin* UMovieGraphNode::GetOutputPin(const FName& Label) const
 
 	return nullptr;
 }
+
+TArray<FMovieGraphPinProperties> UMovieGraphVariableNode::GetOutputPinProperties() const
+{
+	TArray<FMovieGraphPinProperties> Properties;
+	
+	if (GraphVariable)
+	{
+		Properties.Add(FMovieGraphPinProperties(FName(GraphVariable->Name), false));
+	}
+	else
+	{
+		Properties.Add(FMovieGraphPinProperties(TEXT("Unknown"), false));
+	}
+	
+	return Properties;
+}
+
+void UMovieGraphVariableNode::SetVariable(UMovieGraphVariable* InVariable)
+{
+	if (InVariable)
+	{
+		GraphVariable = InVariable;
+
+		// Update the output pin to reflect the new variable, and update the pin whenever the variable changes
+		// (eg, when the variable is renamed)
+		UpdateOutputPin(GraphVariable);
+
+#if WITH_EDITOR
+		GraphVariable->OnMovieGraphVariableChangedDelegate.AddUObject(this, &UMovieGraphVariableNode::UpdateOutputPin);
+#endif
+	}
+}
+
+#if WITH_EDITOR
+FText UMovieGraphVariableNode::GetMenuDescription() const
+{
+	return GraphVariable ? FText::FromString(GraphVariable->Name) : FText();
+}
+	
+FText UMovieGraphVariableNode::GetMenuCategory() const
+{
+	return NSLOCTEXT("MovieGraphNode", "VariableNode_Category", "Variables");
+}
+#endif // WITH_EDITOR
+
+void UMovieGraphVariableNode::UpdateOutputPin(UMovieGraphVariable* ChangedVariable)
+{
+	if (!OutputPins.IsEmpty() && ChangedVariable)
+	{
+		// Update the output pin w/ the name of the variable
+		OutputPins[0]->Properties.Label = FName(ChangedVariable->Name);
+	}
+}
\ No newline at end of file
diff --git a/Engine/Plugins/MovieScene/MovieRenderPipeline/Source/MovieRenderPipelineCore/Private/Graph/MovieGraphUtils.cpp b/Engine/Plugins/MovieScene/MovieRenderPipeline/Source/MovieRenderPipelineCore/Private/Graph/MovieGraphUtils.cpp
new file mode 100644
--- /dev/null
+++ b/Engine/Plugins/MovieScene/MovieRenderPipeline/Source/MovieRenderPipelineCore/Private/Graph/MovieGraphUtils.cpp
@@ -0,0 +1,20 @@
+﻿// Copyright Epic Games, Inc. All Rights Reserved.
+
+#include "MovieGraphUtils.h"
+
+namespace UE::MoviePipeline::RenderGraph
+{
+	FString GetUniqueName(const TArray<FString>& InExistingNames, const FString& InBaseName)
+	{
+		int32 Postfix = 0;
+		FString NewName = InBaseName;
+
+		while (InExistingNames.Contains(NewName))
+		{
+			Postfix++;
+			NewName = FString::Format(TEXT("{0} {1}"), {InBaseName, Postfix});
+		}
+
+		return NewName;
+	}
+}
\ No newline at end of file
diff --git a/Engine/Plugins/MovieScene/MovieRenderPipeline/Source/MovieRenderPipelineCore/Private/Graph/MovieGraphUtils.h b/Engine/Plugins/MovieScene/MovieRenderPipeline/Source/MovieRenderPipelineCore/Private/Graph/MovieGraphUtils.h
new file mode 100644
--- /dev/null
+++ b/Engine/Plugins/MovieScene/MovieRenderPipeline/Source/MovieRenderPipelineCore/Private/Graph/MovieGraphUtils.h
@@ -0,0 +1,15 @@
+﻿// Copyright Epic Games, Inc. All Rights Reserved.
+
+#pragma once
+
+#include "CoreMinimal.h"
+
+namespace UE::MoviePipeline::RenderGraph
+{
+	/**
+	 * Generate a unique name given a set of existing names and the desired base name. The base name will
+	 * be given a postfix value if it conflicts with an existing name (eg, if the base name is "Foo" but
+	 * there's already an existing name "Foo", the generated name would be "Foo 1").
+	 */
+	FString GetUniqueName(const TArray<FString>& InExistingNames, const FString& InBaseName);
+}
diff --git a/Engine/Plugins/MovieScene/MovieRenderPipeline/Source/MovieRenderPipelineCore/Public/Graph/MovieGraphConfig.h b/Engine/Plugins/MovieScene/MovieRenderPipeline/Source/MovieRenderPipelineCore/Public/Graph/MovieGraphConfig.h
--- a/Engine/Plugins/MovieScene/MovieRenderPipeline/Source/MovieRenderPipelineCore/Public/Graph/MovieGraphConfig.h
+++ b/Engine/Plugins/MovieScene/MovieRenderPipeline/Source/MovieRenderPipelineCore/Public/Graph/MovieGraphConfig.h
@@ -4,6 +4,7 @@
 #include "CoreMinimal.h"
 #include "EdGraph/EdGraph.h"
 #include "Graph/MovieGraphNode.h"
+
 #include "MovieGraphConfig.generated.h"
 
 // Forward Declare
@@ -21,6 +22,79 @@ struct FMovieGraphTraversalContext
 	FString RenderLayerName;
 };
 
+/** Types of UMovieGraphVariable that can be used in the graph. */
+UENUM()
+enum class EMovieGraphVariableType : uint8
+{
+	Bool,
+	Float,
+	Int,
+	String
+};
+
+#if WITH_EDITOR
+	DECLARE_MULTICAST_DELEGATE_OneParam(FOnMovieGraphVariableChanged, UMovieGraphVariable*);
+#endif 
+
+/**
+ * A variable that can be created and used inside the graph. These variables can be controlled at
+ * the job level.
+ */
+UCLASS(BlueprintType)
+class UMovieGraphVariable : public UObject
+{
+	GENERATED_BODY()
+
+public:
+	UMovieGraphVariable() = default;
+
+	/** Gets the GUID that uniquely identifies this variable. */
+	const FGuid& GetGuid() const { return Guid; }
+
+	/** Sets the GUID that uniquely identifies this variable. */
+	void SetGuid(const FGuid& InGuid) { Guid = InGuid; }
+
+public:
+#if WITH_EDITOR
+	FOnMovieGraphVariableChanged OnMovieGraphVariableChangedDelegate;
+#endif
+
+	/** The type of data stored in this variable. */
+	UPROPERTY(EditAnywhere, Category = "General")
+	EMovieGraphVariableType Type = EMovieGraphVariableType::Float;
+
+	// TODO: Need a details customization that validates whether or not the name is valid/unique
+	/** The name of this variable, which is user-facing. */
+	UPROPERTY(EditAnywhere, Category = "General")
+	FString Name;
+
+	/** The optional description of this variable, which is user-facing. */
+	UPROPERTY(EditAnywhere, Category = "General")
+	FString Description;
+
+	// TODO: This needs to eventually be capable of storing multiple types. Using TVariant looks promising, but it will
+	// still require extensive details panel customizations
+	/** The default value of this variable. */
+	UPROPERTY(EditAnywhere, Category = "General")
+	float Default = 0.f;
+
+	//~ Begin UObject overrides
+#if WITH_EDITOR
+	virtual void PostEditChangeProperty(FPropertyChangedEvent& PropertyChangedEvent) override;
+#endif
+	//~ End UObject overrides
+
+private:
+	/** A GUID that uniquely identifies this variable within its graph. */
+	UPROPERTY()
+	FGuid Guid;
+};
+
+#if WITH_EDITOR
+	DECLARE_MULTICAST_DELEGATE(FOnMovieGraphChanged);
+	DECLARE_MULTICAST_DELEGATE_OneParam(FOnMovieGraphNodesDeleted, TArray<UMovieGraphNode*>);
+#endif // WITH_EDITOR
+
 /**
 * This is the runtime representation of the UMoviePipelineEdGraph which contains the actual strongly
 * typed graph network that is read by the MoviePipeline. There is an editor-only representation of
@@ -33,7 +107,7 @@ class MOVIERENDERPIPELINECORE_API UMovieGraphConfig : public UObject
 
 public:
 	UMovieGraphConfig();
-public:
+
 	bool AddLabeledEdge(UMovieGraphNode* FromNode, const FName& FromPinLabel, UMovieGraphNode* ToNode, const FName& ToPinLabel);
 	bool RemoveEdge(UMovieGraphNode* FromNode, const FName& FromPinName, UMovieGraphNode* ToNode, const FName& ToPinName);
 	bool RemoveAllInboundEdges(UMovieGraphNode* InNode);
@@ -41,13 +115,37 @@ public:
 	bool RemoveInboundEdges(UMovieGraphNode* InNode, const FName& InPinName);
 	bool RemoveOutboundEdges(UMovieGraphNode* InNode, const FName& InPinName);
 
+	/** Removes the specified node from the graph. */
+	bool RemoveNode(UMovieGraphNode* InNode);
+
 	UMovieGraphNode* GetInputNode() const { return InputNode; }
 	UMovieGraphNode* GetOutputNode() const { return OutputNode; }
 	const TArray<TObjectPtr<UMovieGraphNode>>& GetNodes() const { return AllNodes; }
 
+	/** Adds a new variable with default values to the graph. Returns the new variable on success, else nullptr. */
+	UMovieGraphVariable* AddVariable();
+
+	/** Gets the variable in the graph with the specified GUID, else nullptr if one could not be found. */
+	UMovieGraphVariable* GetVariableByGuid(const FGuid& InGuid) const;
+
+	/** Get all variables that have been defined on the graph. */
+	TArray<UMovieGraphVariable*> GetVariables() const;
+
+	/** Remove the specified member (input, output, variable) from the graph. */
+	void DeleteMember(UObject* MemberToDelete);
+
+	/** Remove the specified variable member from the graph. */
+	void DeleteVariableMember(UMovieGraphVariable* VariableMemberToDelete);
+
 protected:
 	UMovieGraphNode* FindTraversalStartForContext(const FMovieGraphTraversalContext& InContext) const;
+
 public:
+#if WITH_EDITOR
+	FOnMovieGraphChanged OnGraphChangedDelegate;
+	FOnMovieGraphNodesDeleted OnGraphNodesDeletedDelegate;
+#endif
+	
 	UPROPERTY()
 	TArray<TObjectPtr<UMovieGraphNode>> AllNodes;
 
@@ -73,13 +171,21 @@ public:
 		T* RuntimeNode = NewObject<T>(this, PipelineGraphNodeClass, NAME_None, RF_Transactional);
 		RuntimeNode->UpdateDynamicProperties();
 		RuntimeNode->UpdatePins();
+		RuntimeNode->Guid = FGuid::NewGuid();
 #if WITH_EDITOR
 		AllNodes.Add(RuntimeNode);
-#endif // WITH_EDITORONLY_DATA
+#endif
 		return RuntimeNode;
 	}
 
 	void TraversalTest();
 	void TraverseTest(UMovieGraphNode* InNode);
 
+private:
+	void OnVariableUpdated(UMovieGraphVariable* UpdatedVariable);
+
+private:
+	/** All variables which have been defined on the graph. */
+	UPROPERTY()
+	TArray<TObjectPtr<UMovieGraphVariable>> Variables;
 };
\ No newline at end of file
diff --git a/Engine/Plugins/MovieScene/MovieRenderPipeline/Source/MovieRenderPipelineCore/Public/Graph/MovieGraphNode.h b/Engine/Plugins/MovieScene/MovieRenderPipeline/Source/MovieRenderPipelineCore/Public/Graph/MovieGraphNode.h
--- a/Engine/Plugins/MovieScene/MovieRenderPipeline/Source/MovieRenderPipelineCore/Public/Graph/MovieGraphNode.h
+++ b/Engine/Plugins/MovieScene/MovieRenderPipeline/Source/MovieRenderPipelineCore/Public/Graph/MovieGraphNode.h
@@ -10,6 +10,7 @@
 
 // Forward Declares
 class UMovieGraphPin;
+class UMovieGraphVariable;
 
 #if WITH_EDITOR
 class UEdGraphNode;
@@ -64,6 +65,9 @@ public:
 	UMovieGraphPin* GetInputPin(const FName& InPinLabel) const;
 	UMovieGraphPin* GetOutputPin(const FName& InPinLabel) const;
 
+	/** Gets the GUID which uniquely identifies this node. */
+	const FGuid& GetGuid() const { return Guid; }
+
 #if WITH_EDITOR
 	int32 GetNodePosX() const { return NodePosX; }
 	int32 GetNodePosY() const { return NodePosY; }
@@ -111,6 +115,10 @@ protected:
 	UPROPERTY()
 	int32 NodePosY = 0;
 #endif
+
+	/** A GUID which uniquely identifies this node. */
+	UPROPERTY()
+	FGuid Guid;
 };
 
 // Dummy test nodes
@@ -229,4 +237,40 @@ public:
 		Properties.Add(FMovieGraphPinProperties(TEXT("Input"), false));
 		return Properties;
 	}
+};
+
+/** A node which gets the value of a variable which has been defined on the graph. */
+UCLASS()
+class MOVIERENDERPIPELINECORE_API UMovieGraphVariableNode : public UMovieGraphNode
+{
+	GENERATED_BODY()
+
+public:
+	UMovieGraphVariableNode() = default;
+
+	virtual TArray<FMovieGraphPinProperties> GetOutputPinProperties() const override;
+
+	/** Gets the variable that this node represents. */
+	UMovieGraphVariable* GetVariable() const { return GraphVariable; }
+
+	/** Sets the variable that this node represents. */
+	void SetVariable(UMovieGraphVariable* InVariable);
+
+#if WITH_EDITOR
+	virtual FText GetMenuDescription() const override;
+	virtual FText GetMenuCategory() const override;
+#endif
+
+private:
+	/** Updates the output pin on the node to match the provided variable. */
+	void UpdateOutputPin(UMovieGraphVariable* ChangedVariable);
+
+private:
+	/** The underlying graph variable this node represents. */
+	UPROPERTY()
+	TObjectPtr<UMovieGraphVariable> GraphVariable = nullptr;
+
+	/** The properties for the output pin on this node. */
+	UPROPERTY(Transient)
+	FMovieGraphPinProperties OutputPin;
 };
\ No newline at end of file
diff --git a/Engine/Plugins/MovieScene/MovieRenderPipeline/Source/MovieRenderPipelineEditor/Private/Graph/MovieEdGraph.cpp b/Engine/Plugins/MovieScene/MovieRenderPipeline/Source/MovieRenderPipelineEditor/Private/Graph/MovieEdGraph.cpp
--- a/Engine/Plugins/MovieScene/MovieRenderPipeline/Source/MovieRenderPipelineEditor/Private/Graph/MovieEdGraph.cpp
+++ b/Engine/Plugins/MovieScene/MovieRenderPipeline/Source/MovieRenderPipelineEditor/Private/Graph/MovieEdGraph.cpp
@@ -7,9 +7,15 @@
 #include "Graph/MovieGraphNode.h"
 #include "MovieEdGraphOutputNode.h"
 #include "MovieEdGraphInputNode.h"
+#include "MovieEdGraphVariableNode.h"
 #include "MovieRenderPipelineCoreModule.h"
 #include "EdGraph/EdGraphPin.h"
 
+template UMoviePipelineEdGraphNodeBase* UMoviePipelineEdGraph::CreateNodeFromRuntimeNode<UMoviePipelineEdGraphNodeInput>(UMovieGraphNode* InRuntimeNode);
+template UMoviePipelineEdGraphNodeBase* UMoviePipelineEdGraph::CreateNodeFromRuntimeNode<UMoviePipelineEdGraphNodeOutput>(UMovieGraphNode* InRuntimeNode);
+template UMoviePipelineEdGraphNodeBase* UMoviePipelineEdGraph::CreateNodeFromRuntimeNode<UMoviePipelineEdGraphVariableNode>(UMovieGraphNode* InRuntimeNode);
+template UMoviePipelineEdGraphNodeBase* UMoviePipelineEdGraph::CreateNodeFromRuntimeNode<UMoviePipelineEdGraphNode>(UMovieGraphNode* InRuntimeNode);
+
 UMovieGraphConfig* UMoviePipelineEdGraph::GetPipelineGraph() const
 {
 	return CastChecked<UMovieGraphConfig>(GetOuter());
@@ -20,37 +26,31 @@ void UMoviePipelineEdGraph::InitFromRuntimeGraph(UMovieGraphConfig* InGraph)
 	// Don't allow reinitialization of an existing graph
 	check(InGraph && !bInitialized);
 
-	const bool bSelectNewNode = false;
 	TMap<UMovieGraphNode*, UMoviePipelineEdGraphNodeBase*> NodeLookup;
 
 	// Input
 	{
 		UMovieGraphNode* InputNode = InGraph->GetInputNode();
-		FGraphNodeCreator<UMoviePipelineEdGraphNodeInput> InputNodeCreator(*this);
-		UMoviePipelineEdGraphNodeInput* InputGraphNode = InputNodeCreator.CreateNode(bSelectNewNode);
-		InputGraphNode->Construct(InputNode);
-		InputNodeCreator.Finalize();
-		NodeLookup.Add(InputNode, InputGraphNode);
+		NodeLookup.Add(InputNode, CreateNodeFromRuntimeNode<UMoviePipelineEdGraphNodeInput>(InputNode));
 	}
 
 	// Output
 	{
 		UMovieGraphNode* OutputNode = InGraph->GetOutputNode();
-		FGraphNodeCreator<UMoviePipelineEdGraphNodeOutput> OutputNodeCreator(*this);
-		UMoviePipelineEdGraphNodeOutput* OutputGraphNode = OutputNodeCreator.CreateNode(bSelectNewNode);
-		OutputGraphNode->Construct(OutputNode);
-		OutputNodeCreator.Finalize();
-		NodeLookup.Add(OutputNode, OutputGraphNode);
+		NodeLookup.Add(OutputNode, CreateNodeFromRuntimeNode<UMoviePipelineEdGraphNodeOutput>(OutputNode));
 	}
 
 	// Create the rest of the nodes in the graph
 	for (const TObjectPtr<UMovieGraphNode> RuntimeNode : InGraph->GetNodes())
 	{
-		FGraphNodeCreator<UMoviePipelineEdGraphNode> NodeCreator(*this);
-		UMoviePipelineEdGraphNode* GraphNode = NodeCreator.CreateNode(bSelectNewNode);
-		GraphNode->Construct(RuntimeNode);
-		NodeCreator.Finalize();
-		NodeLookup.Add(RuntimeNode, GraphNode);
+		if (RuntimeNode->IsA<UMovieGraphVariableNode>())
+		{
+			NodeLookup.Add(RuntimeNode, CreateNodeFromRuntimeNode<UMoviePipelineEdGraphVariableNode>(RuntimeNode));
+		}
+		else
+		{
+			NodeLookup.Add(RuntimeNode, CreateNodeFromRuntimeNode<UMoviePipelineEdGraphNode>(RuntimeNode));
+		}
 	}
 
 	// Now that we've added an Editor Graph representation for every node in the graph, link
@@ -62,6 +62,9 @@ void UMoviePipelineEdGraph::InitFromRuntimeGraph(UMovieGraphConfig* InGraph)
 		CreateLinks(Pair.Value, bCreateInboundLinks, bCreateOutboundLinks, NodeLookup);
 	}
 
+	InGraph->OnGraphChangedDelegate.AddUObject(this, &UMoviePipelineEdGraph::OnGraphConfigChanged);
+	InGraph->OnGraphNodesDeletedDelegate.AddUObject(this, &UMoviePipelineEdGraph::OnGraphNodesDeleted);
+
 	bInitialized = true;
 }
 
@@ -83,6 +86,11 @@ void UMoviePipelineEdGraph::CreateLinks(UMoviePipelineEdGraphNodeBase* InGraphNo
 	CreateLinks(InGraphNode, bCreateInboundLinks, bCreateOutboundLinks, RuntimeNodeToEdNodeMap);
 }
 
+void UMoviePipelineEdGraph::ReconstructGraph()
+{
+	OnGraphConfigChanged();
+}
+
 void UMoviePipelineEdGraph::CreateLinks(UMoviePipelineEdGraphNodeBase* InGraphNode, bool bCreateInboundLinks, bool bCreateOutboundLinks,
 	const TMap<UMovieGraphNode*, UMoviePipelineEdGraphNodeBase*>& RuntimeNodeToEdNodeMap)
 {
@@ -137,4 +145,53 @@ void UMoviePipelineEdGraph::CreateLinks(UMoviePipelineEdGraphNodeBase* InGraphNo
 	{
 		CreateLinks(RuntimeNode->GetOutputPins(), EEdGraphPinDirection::EGPD_Output);
 	}
-}
\ No newline at end of file
+}
+
+void UMoviePipelineEdGraph::OnGraphConfigChanged()
+{
+	// TODO: Optimize this. Ideally we can target specific changes to the graph rather than rebuilding everything.
+	// However, this isn't strictly necessary unless rebuilding becomes a performance bottleneck.
+	for (const TObjectPtr<UEdGraphNode> Node: Nodes)
+	{
+		Node->ReconstructNode();
+	}
+}
+
+void UMoviePipelineEdGraph::OnGraphNodesDeleted(TArray<UMovieGraphNode*> DeletedNodes)
+{
+	TArray<TObjectPtr<UEdGraphNode>> EdGraphNodesToDelete;
+	
+	for (const TObjectPtr<UEdGraphNode> Node : Nodes)
+	{
+		if (!Node)
+		{
+			continue;
+		}
+
+		if (const UMoviePipelineEdGraphNodeBase* GraphNode = Cast<UMoviePipelineEdGraphNodeBase>(Node.Get()))
+		{
+			if (DeletedNodes.Contains(GraphNode->GetRuntimeNode()))
+			{
+				EdGraphNodesToDelete.Add(Node);
+			}
+		}
+	}
+
+	for (TObjectPtr<UEdGraphNode>& EdNodeToDelete : EdGraphNodesToDelete)
+	{
+		RemoveNode(EdNodeToDelete);
+	}
+}
+
+template <typename T>
+UMoviePipelineEdGraphNodeBase* UMoviePipelineEdGraph::CreateNodeFromRuntimeNode(UMovieGraphNode* InRuntimeNode)
+{
+	const bool bSelectNewNode = false;
+	
+	FGraphNodeCreator<T> NodeCreator(*this);
+	UMoviePipelineEdGraphNodeBase* NewNode = NodeCreator.CreateNode(bSelectNewNode);
+	NewNode->Construct(InRuntimeNode);
+	NodeCreator.Finalize();
+
+	return NewNode;
+}
diff --git a/Engine/Plugins/MovieScene/MovieRenderPipeline/Source/MovieRenderPipelineEditor/Private/Graph/MovieEdGraph.h b/Engine/Plugins/MovieScene/MovieRenderPipeline/Source/MovieRenderPipelineEditor/Private/Graph/MovieEdGraph.h
--- a/Engine/Plugins/MovieScene/MovieRenderPipeline/Source/MovieRenderPipelineEditor/Private/Graph/MovieEdGraph.h
+++ b/Engine/Plugins/MovieScene/MovieRenderPipeline/Source/MovieRenderPipelineEditor/Private/Graph/MovieEdGraph.h
@@ -30,9 +30,27 @@ public:
 	/** Creates the links/edges between nodes in the graph */
 	void CreateLinks(UMoviePipelineEdGraphNodeBase* InGraphNode, bool bCreateInboundLinks, bool bCreateOutboundLinks);
 
+	/** Completely reconstruct the graph. This is a big hammer and should be used sparingly. */
+	void ReconstructGraph();
+
 protected:
 	void CreateLinks(UMoviePipelineEdGraphNodeBase* InGraphNode, bool bCreateInboundLinks, bool bCreateOutboundLinks,
 		const TMap<UMovieGraphNode*, UMoviePipelineEdGraphNodeBase*>& RuntimeNodeToEdNodeMap);
 
+	/**
+	 * Handler which deals with changes made to the runtime graph. This will most likely be replaced with more specific
+	 * handlers in the future when the graph is more fully-featured.
+	 */
+	void OnGraphConfigChanged();
+
+	/** Handler which deals with deleted nodes in the runtime graph. The GUIDs provided are those of the deleted runtime nodes. */
+	void OnGraphNodesDeleted(TArray<UMovieGraphNode*> DeletedNodes);
+
+protected:
 	bool bInitialized = false;
+
+private:
+	/** Create a new editor node of type T from the given runtime node. */
+	template <typename T>
+	UMoviePipelineEdGraphNodeBase* CreateNodeFromRuntimeNode(UMovieGraphNode* InRuntimeNode);
 };
\ No newline at end of file
diff --git a/Engine/Plugins/MovieScene/MovieRenderPipeline/Source/MovieRenderPipelineEditor/Private/Graph/MovieEdGraphNode.h b/Engine/Plugins/MovieScene/MovieRenderPipeline/Source/MovieRenderPipelineEditor/Private/Graph/MovieEdGraphNode.h
--- a/Engine/Plugins/MovieScene/MovieRenderPipeline/Source/MovieRenderPipelineEditor/Private/Graph/MovieEdGraphNode.h
+++ b/Engine/Plugins/MovieScene/MovieRenderPipeline/Source/MovieRenderPipelineEditor/Private/Graph/MovieEdGraphNode.h
@@ -41,6 +41,7 @@ protected:
 	void UpdatePosition();
 
 protected:
+	/** The runtime node that this editor node represents. */
 	UPROPERTY()
 	TObjectPtr<UMovieGraphNode> RuntimeNode;
 };
diff --git a/Engine/Plugins/MovieScene/MovieRenderPipeline/Source/MovieRenderPipelineEditor/Private/Graph/MovieEdGraphVariableNode.cpp b/Engine/Plugins/MovieScene/MovieRenderPipeline/Source/MovieRenderPipelineEditor/Private/Graph/MovieEdGraphVariableNode.cpp
new file mode 100644
--- /dev/null
+++ b/Engine/Plugins/MovieScene/MovieRenderPipeline/Source/MovieRenderPipelineEditor/Private/Graph/MovieEdGraphVariableNode.cpp
@@ -0,0 +1,29 @@
+﻿// Copyright Epic Games, Inc. All Rights Reserved.
+
+#include "Graph/MovieEdGraphVariableNode.h"
+
+#include "EdGraph/EdGraphPin.h"
+#include "Graph/MovieGraphConfig.h"
+
+#define LOCTEXT_NAMESPACE "MoviePipelineGraph"
+
+FText UMoviePipelineEdGraphVariableNode::GetNodeTitle(ENodeTitleType::Type TitleType) const
+{
+	return LOCTEXT("GetVariableNodeTitle", "Get Variable");
+}
+
+void UMoviePipelineEdGraphVariableNode::AllocateDefaultPins()
+{
+	if (const UMovieGraphVariableNode* VariableNode = Cast<UMovieGraphVariableNode>(RuntimeNode))
+	{
+		FEdGraphPinType EdPinType;
+		EdPinType.ResetToDefaults();
+		
+		EdPinType.PinCategory = FName("Outputs");
+		EdPinType.PinSubCategory = FName();
+		
+		CreatePin(EGPD_Output, EdPinType, FName(VariableNode->GetVariable()->Name));
+	}
+}
+
+#undef LOCTEXT_NAMESPACE
\ No newline at end of file
diff --git a/Engine/Plugins/MovieScene/MovieRenderPipeline/Source/MovieRenderPipelineEditor/Private/Graph/MovieEdGraphVariableNode.h b/Engine/Plugins/MovieScene/MovieRenderPipeline/Source/MovieRenderPipelineEditor/Private/Graph/MovieEdGraphVariableNode.h
new file mode 100644
--- /dev/null
+++ b/Engine/Plugins/MovieScene/MovieRenderPipeline/Source/MovieRenderPipelineEditor/Private/Graph/MovieEdGraphVariableNode.h
@@ -0,0 +1,20 @@
+﻿// Copyright Epic Games, Inc. All Rights Reserved.
+
+#pragma once
+
+#include "Graph/MovieEdGraphNode.h"
+
+#include "MovieEdGraphVariableNode.generated.h"
+
+/** An editor node which represents an accessor for a variable defined on the graph. */
+UCLASS()
+class UMoviePipelineEdGraphVariableNode : public UMoviePipelineEdGraphNodeBase
+{
+	GENERATED_BODY()
+
+public:
+	// ~Begin UEdGraphNode interface
+	virtual FText GetNodeTitle(ENodeTitleType::Type TitleType) const override;
+	virtual void AllocateDefaultPins() override;
+	// ~End UEdGraphNode interface
+};
\ No newline at end of file
diff --git a/Engine/Plugins/MovieScene/MovieRenderPipeline/Source/MovieRenderPipelineEditor/Private/Graph/MovieGraphAssetToolkit.cpp b/Engine/Plugins/MovieScene/MovieRenderPipeline/Source/MovieRenderPipelineEditor/Private/Graph/MovieGraphAssetToolkit.cpp
--- a/Engine/Plugins/MovieScene/MovieRenderPipeline/Source/MovieRenderPipelineEditor/Private/Graph/MovieGraphAssetToolkit.cpp
+++ b/Engine/Plugins/MovieScene/MovieRenderPipeline/Source/MovieRenderPipelineEditor/Private/Graph/MovieGraphAssetToolkit.cpp
@@ -2,7 +2,9 @@
 
 #include "MovieGraphAssetToolkit.h"
 
+#include "Framework/Commands/GenericCommands.h"
 #include "Graph/MovieGraphConfig.h"
+#include "MovieGraphSchema.h"
 #include "PropertyEditorModule.h"
 #include "SMovieGraphMembersTabContent.h"
 #include "ToolMenus.h"
@@ -23,6 +25,11 @@ static TAutoConsoleVariable<bool> CVarMoviePipelineEnableRenderGraph(
 	TEXT("Determines if the Render Graph feature is enabled in the UI. This is a highly experimental feature and is not ready for use.")
 );
 
+FMovieGraphAssetToolkit::FMovieGraphAssetToolkit()
+	: bIsInternalSelectionChange(false)
+{
+}
+
 void FMovieGraphAssetToolkit::RegisterTabSpawners(const TSharedRef<FTabManager>& InTabManager)
 {
 	WorkspaceMenuCategory = InTabManager->AddLocalWorkspaceMenuCategory(
@@ -99,6 +106,8 @@ void FMovieGraphAssetToolkit::InitMovieGraphAssetToolkit(const EToolkitMode::Typ
 	const bool bCreateDefaultStandaloneMenu = true;
 	const bool bCreateDefaultToolbar = true;
 	InitAssetEditor(Mode, InitToolkitHost, AppIdentifier, Layout, bCreateDefaultStandaloneMenu, bCreateDefaultToolbar, InitGraph);
+
+	BindGraphCommands();
 }
 
 TSharedRef<SDockTab> FMovieGraphAssetToolkit::SpawnTab_RenderGraphEditor(const FSpawnTabArgs& Args)
@@ -114,7 +123,22 @@ TSharedRef<SDockTab> FMovieGraphAssetToolkit::SpawnTab_RenderGraphEditor(const F
 				.Graph(InitialGraph)
 				.OnGraphSelectionChanged_Lambda([this](const TArray<UObject*>& NewSelection) -> void
 				{
-					SelectedGraphObjectsDetailsWidget->SetObjects(NewSelection);
+					if (bIsInternalSelectionChange)
+					{
+						return;
+					}
+					
+					// Reset selection in the Members panel
+					if (MembersTabContent.IsValid())
+					{
+						TGuardValue<bool> GuardSelectionChange(bIsInternalSelectionChange, true);
+						MembersTabContent->ClearSelection();
+					}
+
+					if (SelectedGraphObjectsDetailsWidget.IsValid())
+					{
+						SelectedGraphObjectsDetailsWidget->SetObjects(NewSelection);
+					}
 				})
 			]
 		];
@@ -128,7 +152,38 @@ TSharedRef<SDockTab> FMovieGraphAssetToolkit::SpawnTab_RenderGraphMembers(const
 		.TabColorScale(GetTabColorScale())
 		.Label(LOCTEXT("MembersTab_Title", "Members"))
 		[
-			SNew(SMovieGraphMembersTabContent)
+			SAssignNew(MembersTabContent, SMovieGraphMembersTabContent)
+			.Editor(SharedThis(this))
+			.Graph(InitialGraph)
+			.OnActionSelected_Lambda([this](const TArray<TSharedPtr<FEdGraphSchemaAction>>& Selection, ESelectInfo::Type Type)
+			{
+				if (bIsInternalSelectionChange)
+				{
+					return;
+				}
+				
+				// Reset selection in the graph
+				if (MovieGraphWidget.IsValid())
+				{
+					TGuardValue<bool> GuardSelectionChange(bIsInternalSelectionChange, true);
+					MovieGraphWidget->ClearGraphSelection();
+				}
+				
+				TArray<UObject*> SelectedObjects;
+				for (const TSharedPtr<FEdGraphSchemaAction>& SelectedAction : Selection)
+				{
+					const FMovieGraphSchemaAction* GraphAction = static_cast<FMovieGraphSchemaAction*>(SelectedAction.Get());
+					if (GraphAction->ActionTarget)
+					{
+						SelectedObjects.Add(GraphAction->ActionTarget);
+					}
+				}
+
+				if (SelectedGraphObjectsDetailsWidget.IsValid())
+				{
+					SelectedGraphObjectsDetailsWidget->SetObjects(SelectedObjects);
+				}
+			})
 		];
 
 	return NewDockTab;
@@ -153,6 +208,31 @@ TSharedRef<SDockTab> FMovieGraphAssetToolkit::SpawnTab_RenderGraphDetails(const
 		];
 }
 
+void FMovieGraphAssetToolkit::BindGraphCommands()
+{
+	ToolkitCommands->MapAction(FGenericCommands::Get().Delete,
+		FExecuteAction::CreateSP(this, &FMovieGraphAssetToolkit::DeleteSelectedMembers),
+		FCanExecuteAction::CreateSP(this, &FMovieGraphAssetToolkit::CanDeleteSelectedMembers));
+}
+
+void FMovieGraphAssetToolkit::DeleteSelectedMembers()
+{
+	if (MembersTabContent.IsValid())
+	{
+		MembersTabContent->DeleteSelectedMembers();
+	}
+}
+
+bool FMovieGraphAssetToolkit::CanDeleteSelectedMembers()
+{
+	if (MembersTabContent.IsValid())
+	{
+		return MembersTabContent->CanDeleteSelectedMembers();
+	}
+
+	return false;
+}
+
 FName FMovieGraphAssetToolkit::GetToolkitFName() const
 {
 	return FName("MovieGraphEditor");
diff --git a/Engine/Plugins/MovieScene/MovieRenderPipeline/Source/MovieRenderPipelineEditor/Private/Graph/MovieGraphAssetToolkit.h b/Engine/Plugins/MovieScene/MovieRenderPipeline/Source/MovieRenderPipelineEditor/Private/Graph/MovieGraphAssetToolkit.h
--- a/Engine/Plugins/MovieScene/MovieRenderPipeline/Source/MovieRenderPipelineEditor/Private/Graph/MovieGraphAssetToolkit.h
+++ b/Engine/Plugins/MovieScene/MovieRenderPipeline/Source/MovieRenderPipelineEditor/Private/Graph/MovieGraphAssetToolkit.h
@@ -8,7 +8,7 @@
 class FMovieGraphAssetToolkit :  public FAssetEditorToolkit
 {
 public:
-	FMovieGraphAssetToolkit() {}
+	FMovieGraphAssetToolkit();
 	virtual ~FMovieGraphAssetToolkit() override {}
 
 	//~ Begin FAssetEditorToolkit interface
@@ -31,6 +31,15 @@ private:
 	TSharedRef<SDockTab> SpawnTab_RenderGraphMembers(const FSpawnTabArgs& Args);
 	TSharedRef<SDockTab> SpawnTab_RenderGraphDetails(const FSpawnTabArgs& Args);
 
+	/** Register commands used by the editor. */
+	void BindGraphCommands();
+
+	/** Deletes the member(s) which are currently selected. */
+	void DeleteSelectedMembers();
+
+	/** Determines if the selected member(s) can be deleted. */
+	bool CanDeleteSelectedMembers();
+
 private:
 	/** The details panel for the selected object(s) in the graph */
 	TSharedPtr<IDetailsView> SelectedGraphObjectsDetailsWidget;
@@ -38,9 +47,15 @@ private:
 	/** The widget that contains the main node graph */
 	TSharedPtr<class SMoviePipelineGraphPanel> MovieGraphWidget;
 
+	/** The widget which encapsulates all of the content in the Members tab */
+	TSharedPtr<class SMovieGraphMembersTabContent> MembersTabContent;
+
 	/** The graph that the editor was initialized with */
 	TObjectPtr<UMovieGraphConfig> InitialGraph;
 
+	/** Flag which is set when a selection change is not triggered by the user */
+	bool bIsInternalSelectionChange;
+
 	static const FName AppIdentifier;
 	static const FName GraphTabId;
 	static const FName DetailsTabId;
diff --git a/Engine/Plugins/MovieScene/MovieRenderPipeline/Source/MovieRenderPipelineEditor/Private/Graph/MovieGraphSchema.cpp b/Engine/Plugins/MovieScene/MovieRenderPipeline/Source/MovieRenderPipelineEditor/Private/Graph/MovieGraphSchema.cpp
--- a/Engine/Plugins/MovieScene/MovieRenderPipeline/Source/MovieRenderPipelineEditor/Private/Graph/MovieGraphSchema.cpp
+++ b/Engine/Plugins/MovieScene/MovieRenderPipeline/Source/MovieRenderPipelineEditor/Private/Graph/MovieGraphSchema.cpp
@@ -9,6 +9,7 @@
 #include "UObject/UObjectIterator.h"
 #include "ScopedTransaction.h"
 #include "GraphEditor.h"
+#include "MovieEdGraphVariableNode.h"
 
 TArray<UClass*> UMovieGraphSchema::MoviePipelineNodeClasses;
 
@@ -67,24 +68,17 @@ void UMovieGraphSchema::GetGraphContextActions(FGraphContextMenuBuilder& Context
 {
 	InitMoviePipelineNodeClasses();
 
-	/* {
-		FText Category = NSLOCTEXT("MoveiPipeline", "Category", "Cat1");
-		FText MenuDesc = NSLOCTEXT("MoveiPipeline", "MenuDesc", "New Render Layer");
-		FText Description = NSLOCTEXT("MoveiPipeline", "Description", "Description");
-
-		TSharedPtr<FMovieGraphSchemaAction_NewNativeElement> NewAction(new FMovieGraphSchemaAction_NewNativeElement(Category, MenuDesc, Description, 0));
-		NewAction->NodeClass = UMovieGraphRenderLayerNode::StaticClass();
-		ContextMenuBuilder.AddAction(NewAction);
+	const UMoviePipelineEdGraph* Graph = Cast<UMoviePipelineEdGraph>(ContextMenuBuilder.CurrentGraph);
+	if (!Graph)
+	{
+		return;
 	}
 
+	const UMovieGraphConfig* RuntimeGraph = Graph->GetPipelineGraph();
+	if (!RuntimeGraph)
 	{
-		FText Category = NSLOCTEXT("MoveiPipeline", "Category", "Cat1");
-		FText MenuDesc = NSLOCTEXT("MoveiPipeline", "MenuDesc", "New Collection");
-		FText Description = NSLOCTEXT("MoveiPipeline", "Description", "Description");
-
-		NewAction->NodeClass = UMovieGraphCollectionNode::StaticClass();
-		ContextMenuBuilder.AddAction(NewAction);
-	}*/
+		return;
+	}
 
 	for (UClass* PipelineNodeClass : MoviePipelineNodeClasses)
 	{
@@ -93,14 +87,33 @@ void UMovieGraphSchema::GetGraphContextActions(FGraphContextMenuBuilder& Context
 		{
 			const FText Name = PipelineNode->GetMenuDescription();
 			const FText Category = PipelineNode->GetMenuCategory();
-			const FText Tooltip = NSLOCTEXT("MoveiPipeline", "Description", "Placeholder Tooltip");
+			const FText Tooltip = LOCTEXT("CreateNode_Tooltip", "Create a node of this type.");
 
-			TSharedPtr<FMovieGraphSchemaAction_NewNativeElement> NewAction(new FMovieGraphSchemaAction_NewNativeElement(Category, Name, Tooltip, 0));
+			if (PipelineNodeClass == UMovieGraphVariableNode::StaticClass())
+			{
+				// Add variable actions separately
+				continue;
+			}
+			
+			TSharedPtr<FMovieGraphSchemaAction> NewAction = MakeShared<FMovieGraphSchemaAction_NewNode>(Category, Name, Tooltip); 
 			NewAction->NodeClass = PipelineNodeClass;
 
 			ContextMenuBuilder.AddAction(NewAction);
 		}
 	}
+
+	// Create an accessor node action for each variable the graph has
+	for (const UMovieGraphVariable* Variable : RuntimeGraph->GetVariables())
+	{
+		const FText Name = FText::FromString(Variable->Name);
+		const FText Category = LOCTEXT("CreateVariable_Category", "Variables");
+		const FText Tooltip = LOCTEXT("CreateVariable_Tooltip", "Create an accessor node for this variable.");
+		
+		TSharedPtr<FMovieGraphSchemaAction> NewAction = MakeShared<FMovieGraphSchemaAction_NewVariableNode>(Category, Name, Variable->GetGuid(), Tooltip);
+		NewAction->NodeClass = UMovieGraphVariableNode::StaticClass();
+		
+		ContextMenuBuilder.AddAction(NewAction);
+	}
 }
 
 const FPinConnectionResponse UMovieGraphSchema::CanCreateConnection(const UEdGraphPin* PinA, const UEdGraphPin* PinB) const
@@ -234,9 +247,13 @@ FLinearColor UMovieGraphSchema::GetPinTypeColor(const FEdGraphPinType& PinType)
 	//return GetDefault<UEdGraphSchema_K2>()->GetPinTypeColor(PinType);
 }
 
+FMovieGraphSchemaAction_NewNode::FMovieGraphSchemaAction_NewNode(FText InNodeCategory, FText InDisplayName, FText InToolTip)
+	: FMovieGraphSchemaAction(MoveTemp(InNodeCategory), MoveTemp(InDisplayName), MoveTemp(InToolTip), 0)
+{
+	
+}
 
-
-UEdGraphNode* FMovieGraphSchemaAction_NewNativeElement::PerformAction(UEdGraph* ParentGraph, UEdGraphPin* FromPin, const FVector2D Location, bool bSelectNewNode)
+UEdGraphNode* FMovieGraphSchemaAction_NewNode::PerformAction(UEdGraph* ParentGraph, UEdGraphPin* FromPin, const FVector2D Location, bool bSelectNewNode)
 {
 	UMovieGraphConfig* RuntimeGraph = CastChecked<UMoviePipelineEdGraph>(ParentGraph)->GetPipelineGraph();
 	const FScopedTransaction Transaction(LOCTEXT("GraphEditor_NewNode", "Create Pipeline Graph Node."));
@@ -256,4 +273,35 @@ UEdGraphNode* FMovieGraphSchemaAction_NewNativeElement::PerformAction(UEdGraph*
 	return GraphNode;
 }
 
+FMovieGraphSchemaAction_NewVariableNode::FMovieGraphSchemaAction_NewVariableNode(FText InNodeCategory, FText InDisplayName, FGuid InVariableGuid, FText InToolTip)
+	: FMovieGraphSchemaAction(MoveTemp(InNodeCategory), MoveTemp(InDisplayName), MoveTemp(InToolTip), 0)
+	, VariableGuid(InVariableGuid)
+{
+	
+}
+
+UEdGraphNode* FMovieGraphSchemaAction_NewVariableNode::PerformAction(UEdGraph* ParentGraph, UEdGraphPin* FromPin, const FVector2D Location, bool bSelectNewNode)
+{
+	UMovieGraphConfig* RuntimeGraph = CastChecked<UMoviePipelineEdGraph>(ParentGraph)->GetPipelineGraph();
+	const FScopedTransaction Transaction(LOCTEXT("GraphEditor_NewVariableNode", "Add New Variable Accessor Node"));
+	RuntimeGraph->Modify();
+
+	UMovieGraphNode* RuntimeNode = RuntimeGraph->ConstructRuntimeNode<UMovieGraphNode>(NodeClass);
+	if (UMovieGraphVariableNode* VariableNode = Cast<UMovieGraphVariableNode>(RuntimeNode))
+	{
+		VariableNode->SetVariable(RuntimeGraph->GetVariableByGuid(VariableGuid));
+	}
+
+	// Now create the variable node
+	FGraphNodeCreator<UMoviePipelineEdGraphVariableNode> NodeCreator(*ParentGraph);
+	UMoviePipelineEdGraphVariableNode* GraphNode = NodeCreator.CreateNode(bSelectNewNode);
+	GraphNode->Construct(RuntimeNode);
+	GraphNode->NodePosX = Location.X;
+	GraphNode->NodePosY = Location.Y;
+
+	// Finalize generates a guid, calls a post-place callback, and allocates default pins if needed
+	NodeCreator.Finalize();
+	return GraphNode;
+}
+
 #undef LOCTEXT_NAMESPACE // "MoviePipelineGraphSchema"
diff --git a/Engine/Plugins/MovieScene/MovieRenderPipeline/Source/MovieRenderPipelineEditor/Private/Graph/MovieGraphSchema.h b/Engine/Plugins/MovieScene/MovieRenderPipeline/Source/MovieRenderPipelineEditor/Private/Graph/MovieGraphSchema.h
--- a/Engine/Plugins/MovieScene/MovieRenderPipeline/Source/MovieRenderPipelineEditor/Private/Graph/MovieGraphSchema.h
+++ b/Engine/Plugins/MovieScene/MovieRenderPipeline/Source/MovieRenderPipelineEditor/Private/Graph/MovieGraphSchema.h
@@ -38,27 +38,82 @@ private:
 	
 };
 
-
+/** Base class for schema actions in the graph. */
 USTRUCT()
-struct FMovieGraphSchemaAction_NewNativeElement : public FEdGraphSchemaAction
+struct FMovieGraphSchemaAction : public FEdGraphSchemaAction
 {
 	GENERATED_BODY()
-	
+
 	// Inherit the base class's constructors
 	using FEdGraphSchemaAction::FEdGraphSchemaAction;
+
+	virtual ~FMovieGraphSchemaAction() override = default;
 	
-	// Simple type info
 	static FName StaticGetTypeId()
 	{
-		static FName Type("FMovieGraphSchemaAction_NewNativeElement");
+		static FName Type("FMovieGraphSchemaAction");
 		return Type;
 	}
-	
+
+	/** The object the action relates to. */
+	UPROPERTY()
+	TObjectPtr<UObject> ActionTarget = nullptr;
+
 	UPROPERTY()
 	TSubclassOf<UMovieGraphNode> NodeClass;
+};
 
+/** Schema action for creating a new node in the graph. */
+USTRUCT()
+struct FMovieGraphSchemaAction_NewNode : public FMovieGraphSchemaAction
+{
+	GENERATED_BODY()
+
+	FMovieGraphSchemaAction_NewNode()
+		: FMovieGraphSchemaAction()
+	{}
+
+	FMovieGraphSchemaAction_NewNode(FText InNodeCategory, FText InDisplayName, FText InToolTip);
+
+	virtual ~FMovieGraphSchemaAction_NewNode() override = default;
+	
+	static FName StaticGetTypeId()
+	{
+		static FName Type("FMovieGraphSchemaAction_NewNode");
+		return Type;
+	}
+	
 	//~ Begin FEdGraphSchemaAction Interface
-	virtual FName GetTypeId() const override { return StaticGetTypeId(); }
 	virtual UEdGraphNode* PerformAction(UEdGraph* ParentGraph, UEdGraphPin* FromPin, const FVector2D Location, bool bSelectNewNode = true) override;
 	//~ End FEdGraphSchemaAction Interface
+};
+
+/** Schema action for creating a new variable node in the graph. */
+USTRUCT()
+struct FMovieGraphSchemaAction_NewVariableNode : public FMovieGraphSchemaAction
+{
+	GENERATED_BODY()
+	
+	static FName StaticGetTypeId()
+	{
+		static FName Type("FMovieGraphSchemaAction_NewVariableNode");
+		return Type;
+	}
+
+	FMovieGraphSchemaAction_NewVariableNode()
+		: FMovieGraphSchemaAction()
+	{}
+
+	FMovieGraphSchemaAction_NewVariableNode(FText InNodeCategory, FText InDisplayName, FGuid InVariableGuid, FText InToolTip);
+
+	virtual ~FMovieGraphSchemaAction_NewVariableNode() override = default;
+	
+	//~ Begin FEdGraphSchemaAction Interface
+	virtual UEdGraphNode* PerformAction(UEdGraph* ParentGraph, UEdGraphPin* FromPin, const FVector2D Location, bool bSelectNewNode = true) override;
+	//~ End FEdGraphSchemaAction Interface
+
+private:
+	/** GUID of the runtime variable this action relates to. */
+	UPROPERTY()
+	FGuid VariableGuid;
 };
\ No newline at end of file
diff --git a/Engine/Plugins/MovieScene/MovieRenderPipeline/Source/MovieRenderPipelineEditor/Private/Graph/SMovieGraphMembersTabContent.cpp b/Engine/Plugins/MovieScene/MovieRenderPipeline/Source/MovieRenderPipelineEditor/Private/Graph/SMovieGraphMembersTabContent.cpp
--- a/Engine/Plugins/MovieScene/MovieRenderPipeline/Source/MovieRenderPipelineEditor/Private/Graph/SMovieGraphMembersTabContent.cpp
+++ b/Engine/Plugins/MovieScene/MovieRenderPipeline/Source/MovieRenderPipelineEditor/Private/Graph/SMovieGraphMembersTabContent.cpp
@@ -3,7 +3,11 @@
 #include "SMovieGraphMembersTabContent.h"
 
 #include "EdGraph/EdGraphSchema.h"
+#include "Framework/Commands/GenericCommands.h"
+#include "Graph/MovieGraphConfig.h"
+#include "Graph/MovieGraphSchema.h"
 #include "SGraphActionMenu.h"
+#include "Toolkits/AssetEditorToolkit.h"
 
 #define LOCTEXT_NAMESPACE "SMoviePipelineMembersTabContent"
 
@@ -17,11 +21,17 @@ const TArray<FText> SMovieGraphMembersTabContent::ActionMenuSectionNames
 
 void SMovieGraphMembersTabContent::Construct(const FArguments& InArgs)
 {
+	EditorToolkit = InArgs._Editor;
+	CurrentGraph = InArgs._Graph;
+	OnActionSelected = InArgs._OnActionSelected;
+	
 	ChildSlot
 	[
-		SNew(SGraphActionMenu)
+		SAssignNew(ActionMenu, SGraphActionMenu)
+		.OnActionSelected(OnActionSelected)
 		.AutoExpandActionMenu(true)
 		.OnCollectStaticSections(this, &SMovieGraphMembersTabContent::CollectStaticSections)
+		.OnContextMenuOpening(this, &SMovieGraphMembersTabContent::OnContextMenuOpening)
 		.OnGetSectionTitle(this, &SMovieGraphMembersTabContent::GetSectionTitle)
 		.OnGetSectionWidget(this, &SMovieGraphMembersTabContent::GetSectionWidget)
 		.UseSectionStyling(true)
@@ -29,35 +39,92 @@ void SMovieGraphMembersTabContent::Construct(const FArguments& InArgs)
 	];
 }
 
+void SMovieGraphMembersTabContent::ClearSelection() const
+{
+	if (ActionMenu.IsValid())
+	{
+		ActionMenu->SelectItemByName(NAME_None);
+	}
+}
+
+void SMovieGraphMembersTabContent::DeleteSelectedMembers() const
+{
+	if (!ActionMenu.IsValid() || !CurrentGraph)
+	{
+		return;
+	}
+
+	TArray<TSharedPtr<FEdGraphSchemaAction>> SelectedActions;
+	ActionMenu->GetSelectedActions(SelectedActions);
+	for (TSharedPtr<FEdGraphSchemaAction> SelectedAction : SelectedActions)
+	{
+		if (const FMovieGraphSchemaAction* MovieGraphAction = static_cast<FMovieGraphSchemaAction*>(SelectedAction.Get()))
+		{
+			CurrentGraph->DeleteMember(MovieGraphAction->ActionTarget.Get());
+		}
+	}
+
+	RefreshActions();
+}
+
+bool SMovieGraphMembersTabContent::CanDeleteSelectedMembers() const
+{
+	if (!ActionMenu.IsValid())
+	{
+		return false;
+	}
+
+	// There may need to be some more extensive validation here in the future
+	return true;
+}
+
 void SMovieGraphMembersTabContent::CollectAllActions(FGraphActionListBuilderBase& OutAllActions)
 {
+	if (!CurrentGraph)
+	{
+		return;
+	}
+	
 	// TODO: These are placeholder actions and need to be dynamically populated by what is in the graph
+
+	FGraphActionMenuBuilder ActionMenuBuilder;
 	
 	const FText InputActionDesc = LOCTEXT("InputAction", "This is an input");
 	const FText InputActionCategory;
 	const FText InputActionDescription = LOCTEXT("InputActionTooltip", "This is an input tooltip.");
 	const FText InputActionKeywords;
-	const int32 InputActionSectionID = 1;
-	const TSharedPtr<FEdGraphSchemaAction> InputAction(new FEdGraphSchemaAction(InputActionCategory, InputActionDesc, InputActionDescription, 0, InputActionKeywords, InputActionSectionID));
+	const int32 InputActionSectionID = static_cast<int32>(EActionSection::Inputs);
+	const TSharedPtr<FMovieGraphSchemaAction> InputAction(new FMovieGraphSchemaAction(InputActionCategory, InputActionDesc, InputActionDescription, 0, InputActionKeywords, InputActionSectionID));
 
 	const FText OutputActionDesc = LOCTEXT("OutputAction", "This is an output");
 	const FText OutputActionCategory;
 	const FText OutputActionDescription = LOCTEXT("OutputActionTooltip", "This is an output tooltip.");
 	const FText OutputActionKeywords;
-	const int32 OutputActionSectionID = 2;
-	const TSharedPtr<FEdGraphSchemaAction> OutputAction(new FEdGraphSchemaAction(OutputActionCategory, OutputActionDesc, OutputActionDescription, 0, OutputActionKeywords, OutputActionSectionID));
-
-	const FText VariableActionDesc = LOCTEXT("VariableAction", "This is a variable");
-	const FText VariableActionCategory;
-	const FText VariableActionDescription = LOCTEXT("VariableActionTooltip", "This is a variable tooltip.");
-	const FText VariableActionKeywords;
-	const int32 VariableActionSectionID = 3;
-	const TSharedPtr<FEdGraphSchemaAction> VariableAction(new FEdGraphSchemaAction(VariableActionCategory, VariableActionDesc, VariableActionDescription, 0, VariableActionKeywords, VariableActionSectionID));
+	const int32 OutputActionSectionID = static_cast<int32>(EActionSection::Outputs);
+	const TSharedPtr<FMovieGraphSchemaAction> OutputAction(new FMovieGraphSchemaAction(OutputActionCategory, OutputActionDesc, OutputActionDescription, 0, OutputActionKeywords, OutputActionSectionID));
+
+	for (UMovieGraphVariable* Variable : CurrentGraph->GetVariables())
+	{
+		if (!Variable)
+		{
+			continue;
+		}
+		
+		const FText VariableActionDesc = FText::FromString(Variable->Name);
+		const FText VariableActionCategory;
+		const FText VariableActionTooltip;
+		const FText VariableActionKeywords;
+		const int32 VariableActionSectionID = static_cast<int32>(EActionSection::Variables);
+		const TSharedPtr<FMovieGraphSchemaAction> VariableAction(new FMovieGraphSchemaAction(VariableActionCategory, VariableActionDesc, VariableActionTooltip, 0, VariableActionKeywords, VariableActionSectionID));
+		VariableAction->ActionTarget = Variable;
+		ActionMenuBuilder.AddAction(VariableAction);
+
+		// Update actions when a variable is renamed
+		Variable->OnMovieGraphVariableChangedDelegate.AddSP(this, &SMovieGraphMembersTabContent::RefreshActions);
+	}
 	
-	FGraphActionMenuBuilder ActionMenuBuilder;
 	ActionMenuBuilder.AddAction(InputAction);
 	ActionMenuBuilder.AddAction(OutputAction);
-	ActionMenuBuilder.AddAction(VariableAction);
 	
 	OutAllActions.Append(ActionMenuBuilder);
 }
@@ -106,6 +173,7 @@ TSharedRef<SWidget> SMovieGraphMembersTabContent::GetSectionWidget(TSharedRef<SW
 	return
 		SNew(SButton)
 		.ButtonStyle(FAppStyle::Get(), "SimpleButton")
+		.OnClicked(this, &SMovieGraphMembersTabContent::OnAddButtonClickedOnSection, InSectionID)
 		.ContentPadding(FMargin(1, 0))
 		.ToolTipText(ButtonTooltip)
 		[
@@ -115,4 +183,46 @@ TSharedRef<SWidget> SMovieGraphMembersTabContent::GetSectionWidget(TSharedRef<SW
 		];
 }
 
+TSharedPtr<SWidget> SMovieGraphMembersTabContent::OnContextMenuOpening()
+{
+	if (!ActionMenu.IsValid())
+	{
+		return nullptr;
+	}
+
+	TArray<TSharedPtr<FEdGraphSchemaAction>> SelectedActions;
+	ActionMenu->GetSelectedActions(SelectedActions);
+	if (SelectedActions.IsEmpty())
+	{
+		return nullptr;
+	}
+
+	FMenuBuilder MenuBuilder(true, EditorToolkit->GetToolkitCommands());
+	MenuBuilder.AddMenuEntry(FGenericCommands::Get().Delete);
+	
+	return MenuBuilder.MakeWidget();
+}
+
+FReply SMovieGraphMembersTabContent::OnAddButtonClickedOnSection(const int32 InSectionID)
+{
+	const EActionSection Section = static_cast<EActionSection>(InSectionID);
+	
+	if (Section == EActionSection::Variables)
+	{
+		CurrentGraph->AddVariable();
+	}
+
+	RefreshActions();
+
+	return FReply::Handled();
+}
+
+void SMovieGraphMembersTabContent::RefreshActions(UMovieGraphVariable* UpdatedVariable) const
+{
+	// Currently the entire action menu is refreshed rather than a specific action being targeted
+	
+	const bool bPreserveExpansion = true;
+	ActionMenu->RefreshAllActions(bPreserveExpansion);
+}
+
 #undef LOCTEXT_NAMESPACE
\ No newline at end of file
diff --git a/Engine/Plugins/MovieScene/MovieRenderPipeline/Source/MovieRenderPipelineEditor/Private/Graph/SMovieGraphMembersTabContent.h b/Engine/Plugins/MovieScene/MovieRenderPipeline/Source/MovieRenderPipelineEditor/Private/Graph/SMovieGraphMembersTabContent.h
--- a/Engine/Plugins/MovieScene/MovieRenderPipeline/Source/MovieRenderPipelineEditor/Private/Graph/SMovieGraphMembersTabContent.h
+++ b/Engine/Plugins/MovieScene/MovieRenderPipeline/Source/MovieRenderPipelineEditor/Private/Graph/SMovieGraphMembersTabContent.h
@@ -5,6 +5,9 @@
 #include "Widgets/Input/SButton.h"
 #include "Widgets/SCompoundWidget.h"
 
+class FAssetEditorToolkit;
+class SGraphActionMenu;
+struct FEdGraphSchemaAction;
 struct FGraphActionListBuilderBase;
 
 /**
@@ -13,11 +16,36 @@ struct FGraphActionListBuilderBase;
 class SMovieGraphMembersTabContent : public SCompoundWidget
 {
 public:
-	SLATE_BEGIN_ARGS(SMovieGraphMembersTabContent) {}
+	DECLARE_DELEGATE_TwoParams(FOnActionSelected, const TArray<TSharedPtr<FEdGraphSchemaAction>>&, ESelectInfo::Type);
+	
+	SLATE_BEGIN_ARGS(SMovieGraphMembersTabContent)
+		: _Graph(nullptr)
+		, _Editor(nullptr)
+	
+		{}
+
+		/** An event which is triggered when an action is selected. */
+		SLATE_EVENT(FOnActionSelected, OnActionSelected)
+		
+		/** The graph that is currently displayed. */
+		SLATE_ARGUMENT(class UMovieGraphConfig*, Graph)
+
+		/** The editor that is displaying this widget. */
+		SLATE_ARGUMENT(TSharedPtr<FAssetEditorToolkit>, Editor)
+	
 	SLATE_END_ARGS();
 	
 	void Construct(const FArguments& InArgs);
 
+	/** Resets the selected members in the UI. */
+	void ClearSelection() const;
+
+	/** Deletes the member(s) which are currently selected from the graph and the UI. */
+	void DeleteSelectedMembers() const;
+
+	/** Determines if all selected member(s) can be deleted. */
+	bool CanDeleteSelectedMembers() const;
+
 private:
 	/** The section identifier in the action widget. */
 	enum class EActionSection : uint8
@@ -38,4 +66,26 @@ private:
 	void CollectStaticSections(TArray<int32>& StaticSectionIDs);
 	FText GetSectionTitle(int32 InSectionID);
 	TSharedRef<SWidget> GetSectionWidget(TSharedRef<SWidget> RowWidget, int32 InSectionID);
+
+	/** Handler which deals with populating the context menu. */
+	TSharedPtr<SWidget> OnContextMenuOpening();
+
+	/** Handler which deals with clicking the add button for an action section. */
+	FReply OnAddButtonClickedOnSection(const int32 InSectionID);
+
+	/** Refresh/regenerate the action menu when the given variable is updated. */
+	void RefreshActions(class UMovieGraphVariable* UpdatedVariable = nullptr) const;
+
+private:
+	/** The editor that this widget is associated with. */
+	TSharedPtr<FAssetEditorToolkit> EditorToolkit;
+
+	/** The action menu displayed in the UI which allows for creation/manipulation of graph members (eg, variables). */
+	TSharedPtr<SGraphActionMenu> ActionMenu;
+
+	/** The runtime graph that this UI gets/sets data on. */
+	TObjectPtr<UMovieGraphConfig> CurrentGraph;
+	
+	/** Delegate to call when an action is selected */
+	FOnActionSelected OnActionSelected;
 };
diff --git a/Engine/Plugins/MovieScene/MovieRenderPipeline/Source/MovieRenderPipelineEditor/Private/Widgets/Graph/SMoviePipelineGraphPanel.cpp b/Engine/Plugins/MovieScene/MovieRenderPipeline/Source/MovieRenderPipelineEditor/Private/Widgets/Graph/SMoviePipelineGraphPanel.cpp
--- a/Engine/Plugins/MovieScene/MovieRenderPipeline/Source/MovieRenderPipelineEditor/Private/Widgets/Graph/SMoviePipelineGraphPanel.cpp
+++ b/Engine/Plugins/MovieScene/MovieRenderPipeline/Source/MovieRenderPipelineEditor/Private/Widgets/Graph/SMoviePipelineGraphPanel.cpp
@@ -104,20 +104,20 @@ void SMoviePipelineGraphPanel::Construct(const FArguments& InArgs)
 	const UEdGraphSchema* Schema = GraphToEdit->GetSchema();
 	Schema->CreateDefaultNodesForGraph(*GraphToEdit);
 
-	TSharedRef<SWidget> GraphEditor =
-		SNew(SGraphEditor)
-		.IsEditable(true)
-		// .Appearance(this, &FTestEditor::GetGraphAppearance)
-		// .TitleBar(TitleBarWidget)
-		.GraphToEdit(GraphToEdit)
-		.GraphEvents(InEvents);
-
 	ChildSlot
 	[
-		GraphEditor
+		SAssignNew(GraphEditorWidget, SGraphEditor)
+		.IsEditable(true)
+		.GraphToEdit(GraphToEdit)
+		.GraphEvents(InEvents)
 	];
 }
 
+void SMoviePipelineGraphPanel::ClearGraphSelection() const
+{
+	GraphEditorWidget->ClearSelectionSet();
+}
+
 UE_ENABLE_OPTIMIZATION_SHIP
 FReply SMoviePipelineGraphPanel::OnDebugButtonClicked()
 {
@@ -135,7 +135,15 @@ void SMoviePipelineGraphPanel::OnSelectedNodesChanged(const TSet<UObject*>& NewS
 		{
 			if (UMovieGraphNode* GraphNode = NodeBase->GetRuntimeNode())
 			{
-				SelectedObjects.Add(GraphNode);
+				// For variable nodes, select the underlying variable instead. Otherwise, just select the runtime node.
+				if (const UMovieGraphVariableNode* VariableNode = Cast<UMovieGraphVariableNode>(GraphNode))
+				{
+					SelectedObjects.Add(VariableNode->GetVariable());
+				}
+				else
+				{
+					SelectedObjects.Add(GraphNode);
+				}
 			}
 		}
 	}
diff --git a/Engine/Plugins/MovieScene/MovieRenderPipeline/Source/MovieRenderPipelineEditor/Private/Widgets/Graph/SMoviePipelineGraphPanel.h b/Engine/Plugins/MovieScene/MovieRenderPipeline/Source/MovieRenderPipelineEditor/Private/Widgets/Graph/SMoviePipelineGraphPanel.h
--- a/Engine/Plugins/MovieScene/MovieRenderPipeline/Source/MovieRenderPipelineEditor/Private/Widgets/Graph/SMoviePipelineGraphPanel.h
+++ b/Engine/Plugins/MovieScene/MovieRenderPipeline/Source/MovieRenderPipelineEditor/Private/Widgets/Graph/SMoviePipelineGraphPanel.h
@@ -6,6 +6,7 @@
 #include "Widgets/SCompoundWidget.h"
 
 class UMoviePipelineConfigBase;
+class SGraphEditor;
 class SMoviePipelineQueueEditor;
 class SWindow;
 class UMoviePipelineExecutorJob;
@@ -35,6 +36,9 @@ public:
 
 	void Construct(const FArguments& InArgs);
 
+	/** Clears the selection in the graph widget being displayed. */
+	void ClearGraphSelection() const;
+
 private:
 	FReply OnRenderLocalRequested();
 	bool IsRenderLocalEnabled() const;
@@ -72,6 +76,9 @@ private:
 	/** The main movie pipeline queue editor widget */
 	TSharedPtr<SMoviePipelineQueueEditor> PipelineQueueEditorWidget;
 
+	/** The graph widget displayed in this panel. */
+	TSharedPtr<SGraphEditor> GraphEditorWidget;
+
 	TWeakPtr<SWindow> WeakEditorWindow;
 	
 	int32 NumSelectedJobs;
