Moved FConsoleVariablesDetailsCustomization registration from the SMoviePipelineConfigEditor details panel to a general registration in module startup

diff --git a/Engine/Plugins/MovieScene/MovieRenderPipeline/Source/MovieRenderPipelineEditor/Private/MovieRenderPipelineEditorModule.cpp b/Engine/Plugins/MovieScene/MovieRenderPipeline/Source/MovieRenderPipelineEditor/Private/MovieRenderPipelineEditorModule.cpp
--- a/Engine/Plugins/MovieScene/MovieRenderPipeline/Source/MovieRenderPipelineEditor/Private/MovieRenderPipelineEditorModule.cpp
+++ b/Engine/Plugins/MovieScene/MovieRenderPipeline/Source/MovieRenderPipelineEditor/Private/MovieRenderPipelineEditorModule.cpp
@@ -29,6 +29,9 @@
 #include "MoviePipelineEditorBlueprintLibrary.h"
 #include "MovieSceneSection.h"
 #include "Sections/MovieSceneCinematicShotSection.h"
+#include "MoviePipelineConsoleVariableSetting.h"
+#include "PropertyEditorModule.h"
+#include "Customizations/ConsoleVariableCustomization.h"
 
 #define LOCTEXT_NAMESPACE "FMovieRenderPipelineEditorModule"
 
@@ -187,6 +190,19 @@ void FMovieRenderPipelineEditorModule::UnregisterMovieRenderer()
 	}
 }
 
+void FMovieRenderPipelineEditorModule::RegisterTypeCustomizations()
+{
+	FPropertyEditorModule& PropertyModule = FModuleManager::LoadModuleChecked<FPropertyEditorModule>("PropertyEditor");
+	PropertyModule.RegisterCustomPropertyTypeLayout(FMoviePipelineConsoleVariableEntry::StaticStruct()->GetFName(),
+		FOnGetPropertyTypeCustomizationInstance::CreateStatic(FConsoleVariablesDetailsCustomization::MakeInstance));
+}
+
+void FMovieRenderPipelineEditorModule::UnregisterTypeCustomizations()
+{
+	FPropertyEditorModule& PropertyModule = FModuleManager::LoadModuleChecked<FPropertyEditorModule>("PropertyEditor");
+	PropertyModule.UnregisterCustomPropertyTypeLayout(FMoviePipelineConsoleVariableEntry::StaticStruct()->GetFName());
+}
+
 void FMovieRenderPipelineEditorModule::StartupModule()
 {
 	// Initialize our custom style
@@ -195,12 +211,14 @@ void FMovieRenderPipelineEditorModule::StartupModule()
 
 	RegisterTabImpl();
 	RegisterSettings();
+	RegisterTypeCustomizations();
 	RegisterMovieRenderer();
 }
 
 void FMovieRenderPipelineEditorModule::ShutdownModule()
 {
 	UnregisterMovieRenderer();
+	UnregisterTypeCustomizations();
 	UnregisterSettings();
 	FMoviePipelineCommands::Unregister();
 }
diff --git a/Engine/Plugins/MovieScene/MovieRenderPipeline/Source/MovieRenderPipelineEditor/Private/MovieRenderPipelineEditorModule.h b/Engine/Plugins/MovieScene/MovieRenderPipeline/Source/MovieRenderPipelineEditor/Private/MovieRenderPipelineEditorModule.h
--- a/Engine/Plugins/MovieScene/MovieRenderPipeline/Source/MovieRenderPipelineEditor/Private/MovieRenderPipelineEditorModule.h
+++ b/Engine/Plugins/MovieScene/MovieRenderPipeline/Source/MovieRenderPipelineEditor/Private/MovieRenderPipelineEditorModule.h
@@ -20,5 +20,8 @@ private:
 	void RegisterMovieRenderer();
 	void UnregisterMovieRenderer();
 
+	void RegisterTypeCustomizations();
+	void UnregisterTypeCustomizations();
+
 	FDelegateHandle MovieRendererDelegate;
 };
diff --git a/Engine/Plugins/MovieScene/MovieRenderPipeline/Source/MovieRenderPipelineEditor/Private/Widgets/SMoviePipelineConfigEditor.cpp b/Engine/Plugins/MovieScene/MovieRenderPipeline/Source/MovieRenderPipelineEditor/Private/Widgets/SMoviePipelineConfigEditor.cpp
--- a/Engine/Plugins/MovieScene/MovieRenderPipeline/Source/MovieRenderPipelineEditor/Private/Widgets/SMoviePipelineConfigEditor.cpp
+++ b/Engine/Plugins/MovieScene/MovieRenderPipeline/Source/MovieRenderPipelineEditor/Private/Widgets/SMoviePipelineConfigEditor.cpp
@@ -735,7 +735,6 @@ void SMoviePipelineConfigEditor::UpdateDetails()
 		else
 		{
 			TSharedRef<IDetailsView> Details = PropertyEditorModule.CreateDetailView(DetailsViewArgs);
-			Details->RegisterInstancedCustomPropertyTypeLayout(FMoviePipelineConsoleVariableEntry::StaticStruct()->GetFName(), FOnGetPropertyTypeCustomizationInstance::CreateStatic(FConsoleVariablesDetailsCustomization::MakeInstance));
 			Details->RegisterInstancedCustomPropertyTypeLayout("FrameNumber", FOnGetPropertyTypeCustomizationInstance::CreateStatic(&FFrameNumberDetailsCustomization::MakeInstance, NumericTypeInterface));
 			Details->RegisterInstancedCustomPropertyLayout(UMoviePipelineOutputSetting::StaticClass(), FOnGetDetailCustomizationInstance::CreateStatic(FOutputFormatDetailsCustomization::MakeInstance, SharedThis(this)));
 			Details->RegisterInstancedCustomPropertyLayout(UMoviePipelineConsoleVariableSetting::StaticClass(), FOnGetDetailCustomizationInstance::CreateStatic(FConsoleVariablesSettingDetailsCustomization::MakeInstance));
