MoviePipeline: Auto-wire newly created nodes, and support deleting nodes from the graph.

diff --git a/Engine/Plugins/MovieScene/MovieRenderPipeline/Source/MovieRenderPipelineCore/Private/Graph/MovieGraphConfig.cpp b/Engine/Plugins/MovieScene/MovieRenderPipeline/Source/MovieRenderPipelineCore/Private/Graph/MovieGraphConfig.cpp
--- a/Engine/Plugins/MovieScene/MovieRenderPipeline/Source/MovieRenderPipelineCore/Private/Graph/MovieGraphConfig.cpp
+++ b/Engine/Plugins/MovieScene/MovieRenderPipeline/Source/MovieRenderPipelineCore/Private/Graph/MovieGraphConfig.cpp
@@ -241,6 +241,16 @@ bool UMovieGraphConfig::RemoveOutboundEdges(UMovieGraphNode* InNode, const FName
 	return bChanged;
 }
 
+bool UMovieGraphConfig::RemoveNodes(TArray<UMovieGraphNode*> InNodes)
+{
+	bool bChanged = false;
+	for (UMovieGraphNode* Node : InNodes)
+	{
+		bChanged |= RemoveNode(Node);
+	}
+	return bChanged;
+}
+
 bool UMovieGraphConfig::RemoveNode(UMovieGraphNode* InNode)
 {
 	if (!InNode)
@@ -252,6 +262,12 @@ bool UMovieGraphConfig::RemoveNode(UMovieGraphNode* InNode)
 	RemoveAllInboundEdges(InNode);
 	RemoveAllOutboundEdges(InNode);
 
+#if WITH_EDITOR
+	TArray<UMovieGraphNode*> RemovedNodes;
+	RemovedNodes.Add(InNode);
+	OnGraphNodesDeletedDelegate.Broadcast(RemovedNodes);
+#endif
+
 	return AllNodes.RemoveSingle(InNode) == 1;
 }
 
@@ -326,7 +342,7 @@ void UMovieGraphConfig::DeleteVariableMember(UMovieGraphVariable* VariableMember
 			return false;
 		});
 
-	// Remove accessor nodes
+	// Remove accessor nodes (which broadcasts our node changed delegates)
 	TArray<UMovieGraphNode*> RemovedNodes;
 	for (const TObjectPtr<UMovieGraphNode>& NodeToRemove : NodesToRemove)
 	{
@@ -340,9 +356,7 @@ void UMovieGraphConfig::DeleteVariableMember(UMovieGraphVariable* VariableMember
 	// Remove this variable from the variables tracked by the graph
 	Variables.RemoveSingle(VariableMemberToDelete);
 
-#if WITH_EDITOR
-	OnGraphNodesDeletedDelegate.Broadcast(RemovedNodes);
-#endif
 }
 
+
 #undef LOCTEXT_NAMESPACE
diff --git a/Engine/Plugins/MovieScene/MovieRenderPipeline/Source/MovieRenderPipelineCore/Public/Graph/MovieGraphConfig.h b/Engine/Plugins/MovieScene/MovieRenderPipeline/Source/MovieRenderPipelineCore/Public/Graph/MovieGraphConfig.h
--- a/Engine/Plugins/MovieScene/MovieRenderPipeline/Source/MovieRenderPipelineCore/Public/Graph/MovieGraphConfig.h
+++ b/Engine/Plugins/MovieScene/MovieRenderPipeline/Source/MovieRenderPipelineCore/Public/Graph/MovieGraphConfig.h
@@ -117,6 +117,8 @@ public:
 
 	/** Removes the specified node from the graph. */
 	bool RemoveNode(UMovieGraphNode* InNode);
+	bool RemoveNodes(TArray<UMovieGraphNode*> InNodes);
+
 
 	UMovieGraphNode* GetInputNode() const { return InputNode; }
 	UMovieGraphNode* GetOutputNode() const { return OutputNode; }
diff --git a/Engine/Plugins/MovieScene/MovieRenderPipeline/Source/MovieRenderPipelineCore/Public/Graph/MovieGraphNode.h b/Engine/Plugins/MovieScene/MovieRenderPipeline/Source/MovieRenderPipelineCore/Public/Graph/MovieGraphNode.h
--- a/Engine/Plugins/MovieScene/MovieRenderPipeline/Source/MovieRenderPipelineCore/Public/Graph/MovieGraphNode.h
+++ b/Engine/Plugins/MovieScene/MovieRenderPipeline/Source/MovieRenderPipelineCore/Public/Graph/MovieGraphNode.h
@@ -23,7 +23,7 @@ DECLARE_MULTICAST_DELEGATE_OneParam(FOnMovieGraphNodeChanged, UMovieGraphNode*);
 * In the editor, each node in the network will have an editor-only representation too 
 * which contains data about it's visual position in the graph, comments, etc.
 */
-UCLASS()
+UCLASS(Abstract)
 class MOVIERENDERPIPELINECORE_API UMovieGraphNode : public UObject
 {
 	GENERATED_BODY()
@@ -88,8 +88,8 @@ public:
 #endif
 
 #if WITH_EDITOR
-	virtual FText GetMenuDescription() const { return FText(); }
-	virtual FText GetMenuCategory() const { return FText(); }
+	virtual FText GetMenuDescription() const PURE_VIRTUAL(UMovieGraphNode::GetMenuDescription, return FText(););
+	virtual FText GetMenuCategory() const PURE_VIRTUAL(UMovieGraphNode::GetMenuCategory, return FText(); );
 #endif
 
 protected:
@@ -223,6 +223,11 @@ public:
 		Properties.Add(FMovieGraphPinProperties(TEXT("Output"), false));
 		return Properties;
 	}
+
+#if WITH_EDITOR
+	virtual FText GetMenuDescription() const override { return NSLOCTEXT("MovieGraphNodes", "OutputNode_Description", "Output"); }
+	virtual FText GetMenuCategory() const override { return NSLOCTEXT("MovieGraphNodes", "OutputNode_Category", "Input/Output"); }
+#endif
 };
 
 UCLASS()
@@ -237,6 +242,12 @@ public:
 		Properties.Add(FMovieGraphPinProperties(TEXT("Input"), false));
 		return Properties;
 	}
+
+
+#if WITH_EDITOR
+	virtual FText GetMenuDescription() const override { return NSLOCTEXT("MovieGraphNodes", "InputNode_Description", "Input"); }
+	virtual FText GetMenuCategory() const override { return NSLOCTEXT("MovieGraphNodes", "InputNode_Category", "Input/Output"); }
+#endif
 };
 
 /** A node which gets the value of a variable which has been defined on the graph. */
diff --git a/Engine/Plugins/MovieScene/MovieRenderPipeline/Source/MovieRenderPipelineEditor/Private/Graph/AssetDefinition_MovieGraphConfig.cpp b/Engine/Plugins/MovieScene/MovieRenderPipeline/Source/MovieRenderPipelineEditor/Private/Graph/AssetDefinition_MovieGraphConfig.cpp
--- a/Engine/Plugins/MovieScene/MovieRenderPipeline/Source/MovieRenderPipelineEditor/Private/Graph/AssetDefinition_MovieGraphConfig.cpp
+++ b/Engine/Plugins/MovieScene/MovieRenderPipeline/Source/MovieRenderPipelineEditor/Private/Graph/AssetDefinition_MovieGraphConfig.cpp
@@ -15,7 +15,7 @@ EAssetCommandResult UAssetDefinition_MovieGraphConfig::OpenAssets(const FAssetOp
 	for (UMovieGraphConfig* GraphToEdit : OpenArgs.LoadObjects<UMovieGraphConfig>())
 	{
 		TSharedRef<FMovieGraphAssetToolkit> MovieGraphEditor = MakeShared<FMovieGraphAssetToolkit>();
-		MovieGraphEditor->InitMovieGraphAssetToolkit(EToolkitMode::Standalone, TSharedPtr<IToolkitHost>(), GraphToEdit);
+		MovieGraphEditor->InitMovieGraphAssetToolkit(EToolkitMode::Standalone, OpenArgs.ToolkitHost, GraphToEdit);
 	}	
 
 	return EAssetCommandResult::Handled;
diff --git a/Engine/Plugins/MovieScene/MovieRenderPipeline/Source/MovieRenderPipelineEditor/Private/Graph/MovieEdGraphNode.cpp b/Engine/Plugins/MovieScene/MovieRenderPipeline/Source/MovieRenderPipelineEditor/Private/Graph/MovieEdGraphNode.cpp
--- a/Engine/Plugins/MovieScene/MovieRenderPipeline/Source/MovieRenderPipelineEditor/Private/Graph/MovieEdGraphNode.cpp
+++ b/Engine/Plugins/MovieScene/MovieRenderPipeline/Source/MovieRenderPipelineEditor/Private/Graph/MovieEdGraphNode.cpp
@@ -9,6 +9,7 @@
 #include "MovieEdGraph.h"
 #include "PropertyBag.h"
 #include "ToolMenu.h"
+#include "EdGraph/EdGraphSchema.h"
 
 #define LOCTEXT_NAMESPACE "MoviePipelineEdGraphNodeBase"
 
@@ -202,6 +203,43 @@ void UMoviePipelineEdGraphNodeBase::CreatePins(const TArray<UMovieGraphPin*>& In
 	}
 }
 
+void UMoviePipelineEdGraphNodeBase::AutowireNewNode(UEdGraphPin* FromPin)
+{
+	if (RuntimeNode == nullptr || FromPin == nullptr)
+	{
+		return;
+	}
+
+	const bool bFromPinIsInput = FromPin->Direction == EEdGraphPinDirection::EGPD_Input;
+	const TArray<TObjectPtr<UMovieGraphPin>>& OtherPinsList = bFromPinIsInput ? RuntimeNode->GetOutputPins() : RuntimeNode->GetInputPins();
+
+	// Try to connect to the first compatible pin
+	bool bDidAutoconnect = false;
+	for (const TObjectPtr<UMovieGraphPin>& OtherPin : OtherPinsList)
+	{
+		check(OtherPin);
+
+		const FName& OtherPinName = OtherPin->Properties.Label;
+		UEdGraphPin* ToPin = FindPinChecked(OtherPinName, bFromPinIsInput ? EEdGraphPinDirection::EGPD_Output : EEdGraphPinDirection::EGPD_Input);
+		if (ToPin && GetSchema()->TryCreateConnection(FromPin, ToPin))
+		{
+			// Connection succeeded. Notify our other node that their connections changed.
+			if (ToPin->GetOwningNode())
+			{
+				ToPin->GetOwningNode()->NodeConnectionListChanged();
+			}
+			bDidAutoconnect = true;
+			break;
+		}
+	}
+
+	// Notify ourself of the connection list changing too.
+	if (bDidAutoconnect)
+	{
+		NodeConnectionListChanged();
+	}
+}
+
 void UMoviePipelineEdGraphNodeBase::OnRuntimeNodeChanged(UMovieGraphNode* InChangedNode)
 {
 	if (InChangedNode == GetRuntimeNode())
diff --git a/Engine/Plugins/MovieScene/MovieRenderPipeline/Source/MovieRenderPipelineEditor/Private/Graph/MovieEdGraphNode.h b/Engine/Plugins/MovieScene/MovieRenderPipeline/Source/MovieRenderPipelineEditor/Private/Graph/MovieEdGraphNode.h
--- a/Engine/Plugins/MovieScene/MovieRenderPipeline/Source/MovieRenderPipelineEditor/Private/Graph/MovieEdGraphNode.h
+++ b/Engine/Plugins/MovieScene/MovieRenderPipeline/Source/MovieRenderPipelineEditor/Private/Graph/MovieEdGraphNode.h
@@ -22,6 +22,7 @@ public:
 
 	//~ Begin UEdGraphNode Interface
 	virtual void ReconstructNode() override;
+	virtual void AutowireNewNode(UEdGraphPin* FromPin) override;
 	//~ End UEdGraphNode Interface
 	
 	//~ Begin UObject interface
diff --git a/Engine/Plugins/MovieScene/MovieRenderPipeline/Source/MovieRenderPipelineEditor/Private/Graph/MovieGraphSchema.cpp b/Engine/Plugins/MovieScene/MovieRenderPipeline/Source/MovieRenderPipelineEditor/Private/Graph/MovieGraphSchema.cpp
--- a/Engine/Plugins/MovieScene/MovieRenderPipeline/Source/MovieRenderPipelineEditor/Private/Graph/MovieGraphSchema.cpp
+++ b/Engine/Plugins/MovieScene/MovieRenderPipeline/Source/MovieRenderPipelineEditor/Private/Graph/MovieGraphSchema.cpp
@@ -55,7 +55,7 @@ void UMovieGraphSchema::InitMoviePipelineNodeClasses()
 	for (TObjectIterator<UClass> It; It; ++It)
 	{
 		if (It->IsChildOf(UMovieGraphNode::StaticClass())
-			&& !It->HasAnyClassFlags(CLASS_Abstract))
+			&& !It->HasAnyClassFlags(CLASS_Abstract | CLASS_Deprecated | CLASS_NewerVersionExists))
 		{
 			MoviePipelineNodeClasses.Add(*It);
 		}
@@ -83,17 +83,26 @@ void UMovieGraphSchema::GetGraphContextActions(FGraphContextMenuBuilder& Context
 	for (UClass* PipelineNodeClass : MoviePipelineNodeClasses)
 	{
 		const UMovieGraphNode* PipelineNode = PipelineNodeClass->GetDefaultObject<UMovieGraphNode>();
-		if (!ContextMenuBuilder.FromPin || ContextMenuBuilder.FromPin->Direction == EGPD_Input)
+		if (PipelineNodeClass == UMovieGraphVariableNode::StaticClass())
+		{
+			// Add variable actions separately
+			continue;
+		}
+		if(PipelineNodeClass == UMovieGraphInputNode::StaticClass() ||
+			PipelineNodeClass == UMovieGraphOutputNode::StaticClass())
+		{
+			// Can't place Input and Output nodes manually.
+			continue;
+		}
+
+		// This can be used to sort whether or not an option shows up. For now there's no restrictions
+		// on where nodes can be made, but eventually we might check which branch they're on (if from pin)
+		// to filter out incompatible nodes.
+		// if (!ContextMenuBuilder.FromPin || ContextMenuBuilder.FromPin->Direction == EGPD_Input)
 		{
 			const FText Name = PipelineNode->GetMenuDescription();
 			const FText Category = PipelineNode->GetMenuCategory();
 			const FText Tooltip = LOCTEXT("CreateNode_Tooltip", "Create a node of this type.");
-
-			if (PipelineNodeClass == UMovieGraphVariableNode::StaticClass())
-			{
-				// Add variable actions separately
-				continue;
-			}
 			
 			TSharedPtr<FMovieGraphSchemaAction> NewAction = MakeShared<FMovieGraphSchemaAction_NewNode>(Category, Name, Tooltip); 
 			NewAction->NodeClass = PipelineNodeClass;
@@ -207,7 +216,7 @@ void UMovieGraphSchema::BreakSinglePinLink(UEdGraphPin* SourcePin, UEdGraphPin*
 
 FLinearColor UMovieGraphSchema::GetPinTypeColor(const FEdGraphPinType& PinType) const
 {
-	return FLinearColor::Red;
+	return FLinearColor::Green;
 	/*const FName& TypeName = PinType.PinCategory;
 	if (TypeName == UEdGraphSchema_K2::PC_Struct)
 	{
@@ -268,8 +277,14 @@ UEdGraphNode* FMovieGraphSchemaAction_NewNode::PerformAction(UEdGraph* ParentGra
 	GraphNode->NodePosX = Location.X;
 	GraphNode->NodePosY = Location.Y;
 
+
 	// Finalize generates a guid, calls a post-place callback, and allocates default pins if needed
 	NodeCreator.Finalize();
+
+	if (FromPin)
+	{
+		GraphNode->AutowireNewNode(FromPin);
+	}
 	return GraphNode;
 }
 
@@ -298,9 +313,16 @@ UEdGraphNode* FMovieGraphSchemaAction_NewVariableNode::PerformAction(UEdGraph* P
 	GraphNode->Construct(RuntimeNode);
 	GraphNode->NodePosX = Location.X;
 	GraphNode->NodePosY = Location.Y;
-
+	
 	// Finalize generates a guid, calls a post-place callback, and allocates default pins if needed
 	NodeCreator.Finalize();
+
+	if (FromPin)
+	{
+		GraphNode->AutowireNewNode(FromPin);
+	}
+
+
 	return GraphNode;
 }
 
diff --git a/Engine/Plugins/MovieScene/MovieRenderPipeline/Source/MovieRenderPipelineEditor/Private/Graph/SMovieGraphMembersTabContent.cpp b/Engine/Plugins/MovieScene/MovieRenderPipeline/Source/MovieRenderPipelineEditor/Private/Graph/SMovieGraphMembersTabContent.cpp
--- a/Engine/Plugins/MovieScene/MovieRenderPipeline/Source/MovieRenderPipelineEditor/Private/Graph/SMovieGraphMembersTabContent.cpp
+++ b/Engine/Plugins/MovieScene/MovieRenderPipeline/Source/MovieRenderPipelineEditor/Private/Graph/SMovieGraphMembersTabContent.cpp
@@ -197,7 +197,13 @@ TSharedPtr<SWidget> SMovieGraphMembersTabContent::OnContextMenuOpening()
 		return nullptr;
 	}
 
-	FMenuBuilder MenuBuilder(true, EditorToolkit->GetToolkitCommands());
+	TSharedPtr< FAssetEditorToolkit> PinnedToolkit = EditorToolkit.Pin();
+	if(!PinnedToolkit.IsValid())
+	{
+		return nullptr;
+	}
+
+	FMenuBuilder MenuBuilder(true, PinnedToolkit->GetToolkitCommands());
 	MenuBuilder.AddMenuEntry(FGenericCommands::Get().Delete);
 	
 	return MenuBuilder.MakeWidget();
diff --git a/Engine/Plugins/MovieScene/MovieRenderPipeline/Source/MovieRenderPipelineEditor/Private/Graph/SMovieGraphMembersTabContent.h b/Engine/Plugins/MovieScene/MovieRenderPipeline/Source/MovieRenderPipelineEditor/Private/Graph/SMovieGraphMembersTabContent.h
--- a/Engine/Plugins/MovieScene/MovieRenderPipeline/Source/MovieRenderPipelineEditor/Private/Graph/SMovieGraphMembersTabContent.h
+++ b/Engine/Plugins/MovieScene/MovieRenderPipeline/Source/MovieRenderPipelineEditor/Private/Graph/SMovieGraphMembersTabContent.h
@@ -1,4 +1,4 @@
-﻿// Copyright Epic Games, Inc. All Rights Reserved.
+// Copyright Epic Games, Inc. All Rights Reserved.
 
 #pragma once
 
@@ -78,7 +78,7 @@ private:
 
 private:
 	/** The editor that this widget is associated with. */
-	TSharedPtr<FAssetEditorToolkit> EditorToolkit;
+	TWeakPtr<FAssetEditorToolkit> EditorToolkit;
 
 	/** The action menu displayed in the UI which allows for creation/manipulation of graph members (eg, variables). */
 	TSharedPtr<SGraphActionMenu> ActionMenu;
diff --git a/Engine/Plugins/MovieScene/MovieRenderPipeline/Source/MovieRenderPipelineEditor/Private/Widgets/Graph/SMoviePipelineGraphPanel.cpp b/Engine/Plugins/MovieScene/MovieRenderPipeline/Source/MovieRenderPipelineEditor/Private/Widgets/Graph/SMoviePipelineGraphPanel.cpp
--- a/Engine/Plugins/MovieScene/MovieRenderPipeline/Source/MovieRenderPipelineEditor/Private/Widgets/Graph/SMoviePipelineGraphPanel.cpp
+++ b/Engine/Plugins/MovieScene/MovieRenderPipeline/Source/MovieRenderPipelineEditor/Private/Widgets/Graph/SMoviePipelineGraphPanel.cpp
@@ -30,6 +30,7 @@
 #include "Widgets/Images/SImage.h"
 #include "Framework/MultiBox/MultiBoxBuilder.h"
 #include "SPrimaryButton.h"
+#include "Framework/Commands/GenericCommands.h"
 
 // ContentBrowser Includes
 #include "IContentBrowserSingleton.h"
@@ -103,16 +104,92 @@ void SMoviePipelineGraphPanel::Construct(const FArguments& InArgs)
 
 	const UEdGraphSchema* Schema = GraphToEdit->GetSchema();
 	Schema->CreateDefaultNodesForGraph(*GraphToEdit);
+	MakeEditorCommands();
 
 	ChildSlot
 	[
 		SAssignNew(GraphEditorWidget, SGraphEditor)
 		.IsEditable(true)
 		.GraphToEdit(GraphToEdit)
+		.AdditionalCommands(GraphEditorCommands)
 		.GraphEvents(InEvents)
 	];
 }
 
+void SMoviePipelineGraphPanel::MakeEditorCommands()
+{
+	GraphEditorCommands = MakeShareable(new FUICommandList);
+	{
+		// Editing commands
+		GraphEditorCommands->MapAction(FGenericCommands::Get().SelectAll,
+			FExecuteAction::CreateSP(this, &SMoviePipelineGraphPanel::SelectAllNodes),
+			FCanExecuteAction::CreateSP(this, &SMoviePipelineGraphPanel::CanSelectAllNodes)
+		);
+
+		GraphEditorCommands->MapAction(FGenericCommands::Get().Delete,
+			FExecuteAction::CreateSP(this, &SMoviePipelineGraphPanel::DeleteSelectedNodes),
+			FCanExecuteAction::CreateSP(this, &SMoviePipelineGraphPanel::CanDeleteSelectedNodes)
+		);
+	}
+}
+
+
+void SMoviePipelineGraphPanel::SelectAllNodes()
+{
+	GraphEditorWidget->SelectAllNodes();
+}
+
+
+bool SMoviePipelineGraphPanel::CanSelectAllNodes() const
+{
+	return GraphEditorWidget.IsValid();
+}
+
+bool SMoviePipelineGraphPanel::CanDeleteSelectedNodes() const
+{
+	if (GraphEditorWidget->GetSelectedNodes().IsEmpty())
+	{
+		return false;
+	}
+
+	for (UObject* Object : GraphEditorWidget->GetSelectedNodes())
+	{
+		UEdGraphNode* GraphNode = Cast<UEdGraphNode>(Object);
+		if (GraphNode && !GraphNode->CanUserDeleteNode())
+		{
+			return false;
+		}
+	}
+
+	return true;
+}
+
+void SMoviePipelineGraphPanel::DeleteSelectedNodes()
+{
+	TArray<UMovieGraphNode*> NodesToDelete;
+	for (UObject* Object : GraphEditorWidget->GetSelectedNodes())
+	{
+		UMoviePipelineEdGraphNodeBase* GraphNode = Cast<UMoviePipelineEdGraphNodeBase>(Object);
+		if (GraphNode && GraphNode->CanUserDeleteNode())
+		{
+			NodesToDelete.Add(GraphNode->GetRuntimeNode());
+		}
+	}
+
+	if (NodesToDelete.IsEmpty())
+	{
+		return;
+	}
+
+	UMoviePipelineEdGraph* Graph = Cast<UMoviePipelineEdGraph>(GraphEditorWidget->GetCurrentGraph());
+	if (Graph)
+	{
+		Graph->GetPipelineGraph()->RemoveNodes(NodesToDelete);
+	}
+
+	ClearGraphSelection();
+}
+
 void SMoviePipelineGraphPanel::ClearGraphSelection() const
 {
 	GraphEditorWidget->ClearSelectionSet();
@@ -570,4 +647,5 @@ void SMoviePipelineGraphPanel::OnImportSavedQueueAssest(const FAssetData& InPres
 	}
 }
 
+
 #undef LOCTEXT_NAMESPACE // SMoviePipelineGraphPanel
\ No newline at end of file
diff --git a/Engine/Plugins/MovieScene/MovieRenderPipeline/Source/MovieRenderPipelineEditor/Private/Widgets/Graph/SMoviePipelineGraphPanel.h b/Engine/Plugins/MovieScene/MovieRenderPipeline/Source/MovieRenderPipelineEditor/Private/Widgets/Graph/SMoviePipelineGraphPanel.h
--- a/Engine/Plugins/MovieScene/MovieRenderPipeline/Source/MovieRenderPipelineEditor/Private/Widgets/Graph/SMoviePipelineGraphPanel.h
+++ b/Engine/Plugins/MovieScene/MovieRenderPipeline/Source/MovieRenderPipelineEditor/Private/Widgets/Graph/SMoviePipelineGraphPanel.h
@@ -13,6 +13,8 @@ class UMoviePipelineExecutorJob;
 class UMoviePipelineExecutorShot;
 class IDetailsView;
 struct FAssetData;
+class FUICommandList;
+class UMovieGraphNode;
 
 /**
  * Outermost widget that is used for adding and removing jobs from the Movie Pipeline Queue Subsystem.
@@ -54,6 +56,7 @@ private:
 	void OnConfigWindowClosed();
 
 	void OnSelectionChanged(const TArray<UMoviePipelineExecutorJob*>& InSelectedJobs);
+	TArray<UMovieGraphNode*> GetSelectedModelNodes() const;
 
 	TSharedRef<SWidget> OnGenerateSavedQueuesMenu();
 	bool OpenSaveDialog(const FString& InDefaultPath, const FString& InNewNameSuggestion, FString& OutPackageName);
@@ -70,8 +73,12 @@ private:
 private:
 	/** Allocates a transient preset so that the user can use the pipeline without saving it to an asset first. */
 	//UMoviePipelineConfigBase* AllocateTransientPreset();
+	void MakeEditorCommands();
 
-
+	void SelectAllNodes();
+	bool CanSelectAllNodes() const;
+	bool CanDeleteSelectedNodes() const;
+	void DeleteSelectedNodes();
 private:
 	/** The main movie pipeline queue editor widget */
 	TSharedPtr<SMoviePipelineQueueEditor> PipelineQueueEditorWidget;
@@ -84,4 +91,5 @@ private:
 	int32 NumSelectedJobs;
 
 	FOnGraphSelectionChanged OnGraphSelectionChangedEvent;
+	TSharedPtr<FUICommandList> GraphEditorCommands;
 };
\ No newline at end of file
