MoviePipeline: Expose graph variables at the job level.

Jobs now store "variable assignments", which can override graph variables on the graph that the job is using (if any). These are exposed in the job's details pane via a new customization. Variable assignments can be enabled/disabled, and are currently restricted to only floats (just like graph variables).

diff --git a/Engine/Plugins/MovieScene/MovieRenderPipeline/Source/MovieRenderPipelineCore/Private/Graph/MovieGraphConfig.cpp b/Engine/Plugins/MovieScene/MovieRenderPipeline/Source/MovieRenderPipelineCore/Private/Graph/MovieGraphConfig.cpp
--- a/Engine/Plugins/MovieScene/MovieRenderPipeline/Source/MovieRenderPipelineCore/Private/Graph/MovieGraphConfig.cpp
+++ b/Engine/Plugins/MovieScene/MovieRenderPipeline/Source/MovieRenderPipelineCore/Private/Graph/MovieGraphConfig.cpp
@@ -53,6 +53,7 @@ void UMovieGraphConfig::OnVariableUpdated(UMovieGraphVariable* UpdatedVariable)
 {
 #if WITH_EDITOR
 	OnGraphChangedDelegate.Broadcast();
+	OnGraphVariablesChangedDelegate.Broadcast();
 #endif
 }
 
@@ -281,15 +282,16 @@ UMovieGraphVariable* UMovieGraphConfig::AddVariable()
 	NewVariable->Type = EMovieGraphVariableType::Float;
 	NewVariable->SetGuid(FGuid::NewGuid());
 
-#if WITH_EDITOR
-	NewVariable->OnMovieGraphVariableChangedDelegate.AddUObject(this, &UMovieGraphConfig::OnVariableUpdated);
-#endif
-
 	// Generate and set a unique name
 	TArray<FString> ExistingVariableNames;
 	Algo::Transform(GetVariables(), ExistingVariableNames, [](const UMovieGraphVariable* Variable) { return Variable->Name; });
 	NewVariable->Name = GetUniqueName(ExistingVariableNames, "Variable");
 
+#if WITH_EDITOR
+	NewVariable->OnMovieGraphVariableChangedDelegate.AddUObject(this, &UMovieGraphConfig::OnVariableUpdated);
+	OnGraphVariablesChangedDelegate.Broadcast();
+#endif
+
 	return NewVariable;
 }
 
@@ -346,7 +348,6 @@ void UMovieGraphConfig::DeleteVariableMember(UMovieGraphVariable* VariableMember
 	TArray<UMovieGraphNode*> RemovedNodes;
 	for (const TObjectPtr<UMovieGraphNode>& NodeToRemove : NodesToRemove)
 	{
-		const FGuid& NodeGuid = NodeToRemove->GetGuid();
 		if (RemoveNode(NodeToRemove.Get()))
 		{
 			RemovedNodes.Add(NodeToRemove.Get());
@@ -356,7 +357,9 @@ void UMovieGraphConfig::DeleteVariableMember(UMovieGraphVariable* VariableMember
 	// Remove this variable from the variables tracked by the graph
 	Variables.RemoveSingle(VariableMemberToDelete);
 
+#if WITH_EDITOR
+	OnGraphVariablesChangedDelegate.Broadcast();
+#endif
 }
 
-
 #undef LOCTEXT_NAMESPACE
diff --git a/Engine/Plugins/MovieScene/MovieRenderPipeline/Source/MovieRenderPipelineCore/Private/MoviePipelineQueue.cpp b/Engine/Plugins/MovieScene/MovieRenderPipeline/Source/MovieRenderPipelineCore/Private/MoviePipelineQueue.cpp
--- a/Engine/Plugins/MovieScene/MovieRenderPipeline/Source/MovieRenderPipelineCore/Private/MoviePipelineQueue.cpp
+++ b/Engine/Plugins/MovieScene/MovieRenderPipeline/Source/MovieRenderPipelineCore/Private/MoviePipelineQueue.cpp
@@ -195,6 +195,68 @@ void UMoviePipelineExecutorJob::PostEditChangeProperty(FPropertyChangedEvent& Pr
 }
 #endif
 
+void UMoviePipelineExecutorJob::PreSave(FObjectPreSaveContext ObjectSaveContext)
+{
+	Super::PreSave(ObjectSaveContext);
+
+	CleanUpInvalidVariableAssignments();
+}
+
+bool UMoviePipelineExecutorJob::AddVariableAssignment(const UMovieGraphVariable* InGraphVariable)
+{
+	return InGraphVariable && (FindOrAddVariableAssignment(InGraphVariable) != nullptr);
+}
+
+bool UMoviePipelineExecutorJob::SetVariableAssignmentValue(const UMovieGraphVariable* InGraphVariable, float NewValue)
+{
+	const bool bAddIfNotExists = true;
+	if (FMoviePipelineGraphVariableValueAssignment* Assignment = FindOrAddVariableAssignment(InGraphVariable, bAddIfNotExists))
+	{
+		Modify();
+		Assignment->VariableValue = NewValue;
+		return true;
+	}
+
+	return false;
+}
+
+bool UMoviePipelineExecutorJob::SetVariableAssignmentEnableState(const UMovieGraphVariable* InGraphVariable, bool bIsEnabled)
+{
+	const bool bAddIfNotExists = true;
+	if (FMoviePipelineGraphVariableValueAssignment* Assignment = FindOrAddVariableAssignment(InGraphVariable, bAddIfNotExists))
+	{
+		Modify();
+		Assignment->bIsValueSet = bIsEnabled;
+		return true;
+	}
+
+	return false;
+}
+
+bool UMoviePipelineExecutorJob::GetVariableAssignmentValue(const UMovieGraphVariable* InGraphVariable, float& OutValue)
+{
+	const bool bAddIfNotExists = false;
+	if (const FMoviePipelineGraphVariableValueAssignment* Assignment = FindOrAddVariableAssignment(InGraphVariable, bAddIfNotExists))
+	{
+		OutValue = Assignment->VariableValue;
+		return true;
+	}
+
+	return false;
+}
+
+bool UMoviePipelineExecutorJob::GetVariableAssignmentEnableState(const UMovieGraphVariable* InGraphVariable, bool& bOutIsEnabled)
+{
+	const bool bAddIfNotExists = false;
+	if (const FMoviePipelineGraphVariableValueAssignment* Assignment = FindOrAddVariableAssignment(InGraphVariable, bAddIfNotExists))
+	{
+		bOutIsEnabled = Assignment->bIsValueSet;
+		return true;
+	}
+
+	return false;
+}
+
 void UMoviePipelineExecutorJob::SetSequence(FSoftObjectPath InSequence)
 {
 	Sequence = InSequence;
@@ -243,6 +305,47 @@ void UMoviePipelineExecutorJob::OnDuplicated_Implementation()
 	SetConsumed(false);
 }
 
+FMoviePipelineGraphVariableValueAssignment* UMoviePipelineExecutorJob::FindOrAddVariableAssignment(const UMovieGraphVariable* InGraphVariable, bool bAddIfNotExists)
+{
+	if (InGraphVariable)
+	{
+		for (FMoviePipelineGraphVariableValueAssignment& Assignment : GraphVariableAssignments)
+		{
+			if (Assignment.VariableGuid == InGraphVariable->GetGuid())
+			{
+				return &Assignment;
+			}
+		}
+
+		// Didn't find an existing one; add a new assignment
+		if (bAddIfNotExists)
+		{
+			Modify();
+			
+			FMoviePipelineGraphVariableValueAssignment NewAssignment;
+			NewAssignment.VariableGuid = InGraphVariable->GetGuid();
+			return &GraphVariableAssignments.Add_GetRef(MoveTemp(NewAssignment));
+		}
+	}
+
+	return nullptr;
+}
+
+void UMoviePipelineExecutorJob::CleanUpInvalidVariableAssignments()
+{
+	// TODO: This should get the graph preset OR config
+	const UMovieGraphConfig* Graph = GetGraphPreset();
+	if (!Graph)
+	{
+		return;
+	}
+	
+	GraphVariableAssignments.RemoveAll([Graph](const FMoviePipelineGraphVariableValueAssignment& Assignment)
+	{
+		return Graph->GetVariableByGuid(Assignment.VariableGuid) == nullptr;
+	});
+}
+
 UMoviePipelineShotConfig* UMoviePipelineExecutorShot::AllocateNewShotOverrideConfig(TSubclassOf<UMoviePipelineShotConfig> InConfigType)
 {
 	if (!ensureAlwaysMsgf(InConfigType, TEXT("Failed to specify a Config Type. Use the default in project setting or UMoviePipelineShotConfig.")))
diff --git a/Engine/Plugins/MovieScene/MovieRenderPipeline/Source/MovieRenderPipelineCore/Public/Graph/MovieGraphConfig.h b/Engine/Plugins/MovieScene/MovieRenderPipeline/Source/MovieRenderPipelineCore/Public/Graph/MovieGraphConfig.h
--- a/Engine/Plugins/MovieScene/MovieRenderPipeline/Source/MovieRenderPipelineCore/Public/Graph/MovieGraphConfig.h
+++ b/Engine/Plugins/MovieScene/MovieRenderPipeline/Source/MovieRenderPipelineCore/Public/Graph/MovieGraphConfig.h
@@ -92,6 +92,7 @@ private:
 
 #if WITH_EDITOR
 	DECLARE_MULTICAST_DELEGATE(FOnMovieGraphChanged);
+	DECLARE_MULTICAST_DELEGATE(FOnMovieGraphVariablesChanged);
 	DECLARE_MULTICAST_DELEGATE_OneParam(FOnMovieGraphNodesDeleted, TArray<UMovieGraphNode*>);
 #endif // WITH_EDITOR
 
@@ -145,6 +146,7 @@ protected:
 public:
 #if WITH_EDITOR
 	FOnMovieGraphChanged OnGraphChangedDelegate;
+	FOnMovieGraphVariablesChanged OnGraphVariablesChangedDelegate;
 	FOnMovieGraphNodesDeleted OnGraphNodesDeletedDelegate;
 #endif
 	
diff --git a/Engine/Plugins/MovieScene/MovieRenderPipeline/Source/MovieRenderPipelineCore/Public/MoviePipelineQueue.h b/Engine/Plugins/MovieScene/MovieRenderPipeline/Source/MovieRenderPipelineCore/Public/MoviePipelineQueue.h
--- a/Engine/Plugins/MovieScene/MovieRenderPipeline/Source/MovieRenderPipelineCore/Public/MoviePipelineQueue.h
+++ b/Engine/Plugins/MovieScene/MovieRenderPipeline/Source/MovieRenderPipelineCore/Public/MoviePipelineQueue.h
@@ -250,6 +250,25 @@ private:
 	TSoftObjectPtr<UMovieGraphConfig> GraphPreset;
 };
 
+/** An override to a graph variable's value from the job level. */
+USTRUCT(BlueprintType)
+struct FMoviePipelineGraphVariableValueAssignment
+{
+	GENERATED_BODY()
+	
+	/** The GUID of the variable that is having its value changed. */
+	UPROPERTY()
+	FGuid VariableGuid;
+
+	/** The new value of the variable, if one was set. */
+	UPROPERTY()
+	float VariableValue = 0.f;
+
+	/** Whether a value override has been set. */
+	UPROPERTY()
+	bool bIsValueSet = false;
+};
+
 /**
 * A particular job within the Queue
 */
@@ -475,6 +494,41 @@ public:
 		GraphConfig = InGraphConfig;
 	}
 
+	/**
+	 * Add a variable assignment for the graph. Returns true if the variable was added or it already exists,
+	 * else false.
+	 */
+	UFUNCTION(BlueprintCallable, Category = "Movie Render Pipeline|Experimental")
+	bool AddVariableAssignment(const UMovieGraphVariable* InGraphVariable);
+
+	/**
+	 * Updates an existing variable assignment for the provided graph variable to a new value, or adds a new assignment
+	 * and updates its value. Returns true on success, else false.
+	 */
+	UFUNCTION(BlueprintCallable, Category = "Movie Render Pipeline|Experimental")
+	bool SetVariableAssignmentValue(const UMovieGraphVariable* InGraphVariable, float NewValue);
+
+	/**
+	 * Updates an existing variable assignment for the provided graph variable to a new enable state, or adds a new
+	 * assignment and updates its enable state. Returns true on success, else false.
+	 */
+	UFUNCTION(BlueprintCallable, Category = "Movie Render Pipeline|Experimental")
+	bool SetVariableAssignmentEnableState(const UMovieGraphVariable* InGraphVariable, bool bIsEnabled);
+
+	/**
+	 * Gets the value of the variable assignment for the provided graph variable. The value is provided via OutValue.
+	 * Returns true if a value was set on the variable and OutValue was changed, else false.
+	 */
+	UFUNCTION(BlueprintCallable, Category = "Movie Render Pipeline|Experimental")
+	bool GetVariableAssignmentValue(const UMovieGraphVariable* InGraphVariable, float& OutValue);
+
+	/**
+	 * Gets the enable state of the variable assignment for the provided graph variable. The enable state is provided
+	 * via bOutIsEnabled. Returns true if an enable state was set on the variable and bOutIsEnabled was changed, else false.
+	 */
+	UFUNCTION(BlueprintCallable, Category = "Movie Render Pipeline|Experimental")
+	bool GetVariableAssignmentEnableState(const UMovieGraphVariable* InGraphVariable, bool& bOutIsEnabled);
+
 	UFUNCTION(BlueprintSetter, Category = "Movie Render Pipeline")
 	void SetSequence(FSoftObjectPath InSequence);
 
@@ -483,6 +537,7 @@ public:
 #if WITH_EDITOR
 	void PostEditChangeProperty(FPropertyChangedEvent& PropertyChangedEvent);
 #endif
+	virtual void PreSave(FObjectPreSaveContext ObjectSaveContext) override;
 	// ~UObject Interface
 
 protected:
@@ -504,6 +559,19 @@ protected:
 	virtual bool IsEnabled_Implementation() const { return bEnabled; }
 	// ~UMoviePipelineExecutorJob Interface
 
+	// Note: This is not exposed to the public API because the underlying struct should remain private.
+	/**
+	 * Finds a variable assignment for the provided variable, or adds one if one does not already exist (and
+	 * bAddIfNotExists is set to true). Returns nullptr if the operation failed. 
+	 */
+	FMoviePipelineGraphVariableValueAssignment* FindOrAddVariableAssignment(const UMovieGraphVariable* InGraphVariable, bool bAddIfNotExists = true);
+
+	/**
+	 * Removes variable assignments for variables which no longer exist on the graph that the job currently uses. Old
+	 * and invalid assignments are not necessarily harmful, but they can bloat the asset.
+	 */
+	void CleanUpInvalidVariableAssignments();
+
 public:
 	/** (Optional) Name of the job. Shown on the default burn-in. */
 	UPROPERTY(BlueprintReadWrite, EditAnywhere, Category = "Movie Render Pipeline")
@@ -539,10 +607,13 @@ public:
 private:
 	UPROPERTY(Transient)
 	FString StatusMessage;
+	
 	UPROPERTY(Transient)
 	float StatusProgress;
+	
 	UPROPERTY(Transient)
 	bool bIsConsumed;
+
 private:
 	/** 
 	*/
@@ -565,6 +636,10 @@ private:
 	/** The graph-based configuration preset that this job is using. Can be nullptr. */
 	UPROPERTY()
 	TSoftObjectPtr<UMovieGraphConfig> GraphPreset;
+
+	/** Values that the job has set on variables for the graph it is using (if any). */
+	UPROPERTY()
+	TArray<FMoviePipelineGraphVariableValueAssignment> GraphVariableAssignments;
 };
 
 /**
diff --git a/Engine/Plugins/MovieScene/MovieRenderPipeline/Source/MovieRenderPipelineEditor/Private/Customizations/JobCustomization.h b/Engine/Plugins/MovieScene/MovieRenderPipeline/Source/MovieRenderPipelineEditor/Private/Customizations/JobCustomization.h
new file mode 100644
--- /dev/null
+++ b/Engine/Plugins/MovieScene/MovieRenderPipeline/Source/MovieRenderPipelineEditor/Private/Customizations/JobCustomization.h
@@ -0,0 +1,186 @@
+﻿// Copyright Epic Games, Inc. All Rights Reserved.
+
+#pragma once
+
+#include "DetailCategoryBuilder.h"
+#include "DetailLayoutBuilder.h"
+#include "IDetailCustomization.h"
+#include "MoviePipelineQueue.h"
+#include "PropertyCustomizationHelpers.h"
+#include "PropertyHandle.h"
+#include "Widgets/Input/SNumericEntryBox.h"
+#include "Widgets/Text/STextBlock.h"
+
+#define LOCTEXT_NAMESPACE "MoviePipelineEditor"
+
+/** Customize how properties for a job appear in the details panel. */
+class FJobDetailsCustomization : public IDetailCustomization
+{
+public:
+	static TSharedRef<IDetailCustomization> MakeInstance()
+	{
+		return MakeShared<FJobDetailsCustomization>();
+	}
+
+protected:
+	//~ Begin IDetailCustomization interface
+	virtual void CustomizeDetails(const TSharedPtr<IDetailLayoutBuilder>& DetailBuilder) override
+	{
+		WeakDetailBuilder = DetailBuilder;
+		CustomizeDetails(*DetailBuilder);
+	}
+
+	virtual void CustomizeDetails(IDetailLayoutBuilder& DetailBuilder) override
+	{
+		// Add a new "Graph Variables" category. Set as "Uncommon" priority to push variables down below the other properties.
+		IDetailCategoryBuilder& GraphVariablesCategory = DetailBuilder.EditCategory(
+"GraphVariables", LOCTEXT("GraphVariablesCategory", "Graph Variables"), ECategoryPriority::Uncommon);
+
+		TArray<TWeakObjectPtr<UObject>> ObjectsBeingCustomized;
+		DetailBuilder.GetObjectsBeingCustomized(ObjectsBeingCustomized);
+
+		// Only display variables if there is one job selected
+		if (ObjectsBeingCustomized.Num() != 1)
+		{
+			return;
+		}
+
+		// Get the displayed job
+		UMoviePipelineExecutorJob* CustomizedJob = Cast<UMoviePipelineExecutorJob>(ObjectsBeingCustomized[0].Get());
+		if (!CustomizedJob)
+		{
+			return;
+		}
+
+		UMovieGraphConfig* GraphConfig = GetGraphConfig(CustomizedJob);
+		if (!GraphConfig)
+		{
+			return;
+		}
+
+		for (const UMovieGraphVariable* Variable : GraphConfig->GetVariables())
+		{
+			AddVariableAssignmentRow(GraphVariablesCategory, DetailBuilder, Variable, CustomizedJob);
+		}
+
+		// Refresh the variable listing if anything changes with the variables
+		GraphConfig->OnGraphVariablesChangedDelegate.AddSP(this, &FJobDetailsCustomization::Refresh);
+	}
+	//~ End IDetailCustomization interface
+
+	void Refresh() const
+	{
+		// Don't keep the details builder alive just for the sake of this Refresh() call
+		if (IDetailLayoutBuilder* DetailLayoutBuilder = WeakDetailBuilder.Pin().Get())
+		{
+			DetailLayoutBuilder->ForceRefreshDetails();
+		}
+	}
+
+	void AddVariableAssignmentRow(IDetailCategoryBuilder& DetailCategoryBuilder, const IDetailLayoutBuilder& DetailBuilder,
+		const UMovieGraphVariable* Variable, UMoviePipelineExecutorJob* CustomizedJob)
+	{
+		FDetailWidgetRow& NewRow = DetailCategoryBuilder.AddCustomRow(FText::FromString(Variable->Name));
+		NewRow.WholeRowWidget
+		[
+			SNew(SHorizontalBox)
+
+			+ SHorizontalBox::Slot()
+			.AutoWidth()
+			.VAlign(VAlign_Center)
+			.Padding(0, 0, 5, 0)
+			[
+				SNew(SCheckBox)
+				.IsChecked(this, &FJobDetailsCustomization::IsVariableOverrideChecked, Variable, CustomizedJob)
+				.OnCheckStateChanged(this, &FJobDetailsCustomization::OnVariableOverrideEnableStateChanged, Variable, CustomizedJob)
+			]
+				
+			+ SHorizontalBox::Slot()
+			.FillWidth(0.5f)
+			.VAlign(VAlign_Center)
+			.Padding(0, 0, 5, 0)
+			[
+				SNew(STextBlock)
+				.Text(FText::FromString(Variable->Name))
+				.IsEnabled(this, &FJobDetailsCustomization::IsVariableOverrideEnabled, Variable, CustomizedJob)
+				.Font(DetailBuilder.GetDetailFont())
+			]
+
+			+ SHorizontalBox::Slot()
+			.FillWidth(0.5f)
+			.VAlign(VAlign_Center)
+			.Padding(0, 0, 1, 0)
+			[
+				SNew(SNumericEntryBox<float>)
+				.Value(this, &FJobDetailsCustomization::GetVariableValue, Variable, CustomizedJob)
+				.IsEnabled(this, &FJobDetailsCustomization::IsVariableOverrideEnabled, Variable, CustomizedJob)
+				.Font(DetailBuilder.GetDetailFont())
+				.OnValueCommitted(this, &FJobDetailsCustomization::OnVariableValueCommitted, Variable, CustomizedJob)
+			]
+		];
+	}
+
+	UMovieGraphConfig* GetGraphConfig(const UMoviePipelineExecutorJob* InJob) const
+	{
+		// TODO: This should get the config OR preset
+		return InJob->GetGraphPreset();
+	}
+	
+	ECheckBoxState IsVariableOverrideChecked(const UMovieGraphVariable* InVariable, UMoviePipelineExecutorJob* InJob) const
+	{
+		if (InJob && InVariable)
+		{
+			bool bIsEnabled;
+			if (InJob->GetVariableAssignmentEnableState(InVariable, bIsEnabled))
+			{
+				return bIsEnabled ? ECheckBoxState::Checked : ECheckBoxState::Unchecked;
+			}
+		}
+
+		return ECheckBoxState::Unchecked;
+	}
+
+	bool IsVariableOverrideEnabled(const UMovieGraphVariable* InVariable, UMoviePipelineExecutorJob* InJob) const
+	{
+		return IsVariableOverrideChecked(InVariable, InJob) == ECheckBoxState::Checked;
+	}
+
+	void OnVariableOverrideEnableStateChanged(ECheckBoxState CheckBoxState, const UMovieGraphVariable* InVariable, UMoviePipelineExecutorJob* InJob) const
+	{
+		if (InJob && InVariable)
+		{
+			InJob->SetVariableAssignmentEnableState(InVariable, CheckBoxState == ECheckBoxState::Checked);
+		}
+	}
+
+	TOptional<float> GetVariableValue(const UMovieGraphVariable* InVariable, UMoviePipelineExecutorJob* InJob) const
+	{
+		if (InJob && InVariable)
+		{
+			// Get the value of the variable from the job, if any
+			float VariableValue;
+			if (InJob->GetVariableAssignmentValue(InVariable, VariableValue))
+			{
+				return VariableValue;
+			}
+
+			// Fall back to the default value of the variable
+			return InVariable->Default;
+		}
+
+		return TOptional<float>();
+	}
+
+	void OnVariableValueCommitted(float NewValue, ETextCommit::Type Arg, const UMovieGraphVariable* Variable, UMoviePipelineExecutorJob* InJob) const
+	{
+		if (InJob)
+		{
+			InJob->SetVariableAssignmentValue(Variable, NewValue);
+		}
+	}
+
+private:
+	TWeakPtr<IDetailLayoutBuilder> WeakDetailBuilder;
+};
+
+#undef LOCTEXT_NAMESPACE
\ No newline at end of file
diff --git a/Engine/Plugins/MovieScene/MovieRenderPipeline/Source/MovieRenderPipelineEditor/Private/Widgets/SMoviePipelineQueuePanel.cpp b/Engine/Plugins/MovieScene/MovieRenderPipeline/Source/MovieRenderPipelineEditor/Private/Widgets/SMoviePipelineQueuePanel.cpp
--- a/Engine/Plugins/MovieScene/MovieRenderPipeline/Source/MovieRenderPipelineEditor/Private/Widgets/SMoviePipelineQueuePanel.cpp
+++ b/Engine/Plugins/MovieScene/MovieRenderPipeline/Source/MovieRenderPipelineEditor/Private/Widgets/SMoviePipelineQueuePanel.cpp
@@ -2,6 +2,7 @@
 
 // Movie Pipeline Includes
 #include "Widgets/SMoviePipelineQueuePanel.h"
+#include "Customizations/JobCustomization.h"
 #include "Widgets/MoviePipelineWidgetConstants.h"
 #include "SMoviePipelineQueueEditor.h"
 #include "SMoviePipelineConfigPanel.h"
@@ -67,6 +68,10 @@ void SMoviePipelineQueuePanel::Construct(const FArguments& InArgs)
 
 	JobDetailsPanelWidget = PropertyEditorModule.CreateDetailView(DetailsViewArgs);
 
+	JobDetailsPanelWidget->RegisterInstancedCustomPropertyLayout(
+		UMoviePipelineExecutorJob::StaticClass(),
+		FOnGetDetailCustomizationInstance::CreateStatic(&FJobDetailsCustomization::MakeInstance));
+
 	// Create the child widgets that need to know about our pipeline
 	PipelineQueueEditorWidget = SNew(SMoviePipelineQueueEditor)
 		.OnEditConfigRequested(this, &SMoviePipelineQueuePanel::OnEditJobConfigRequested)
