Modified ConsoleVariableCustomization to support usage outside of UMoviePipelineConsoleVariableSetting

diff --git a/Engine/Plugins/MovieScene/MovieRenderPipeline/Source/MovieRenderPipelineEditor/Private/Customizations/ConsoleVariableCustomization.h b/Engine/Plugins/MovieScene/MovieRenderPipeline/Source/MovieRenderPipelineEditor/Private/Customizations/ConsoleVariableCustomization.h
--- a/Engine/Plugins/MovieScene/MovieRenderPipeline/Source/MovieRenderPipelineEditor/Private/Customizations/ConsoleVariableCustomization.h
+++ b/Engine/Plugins/MovieScene/MovieRenderPipeline/Source/MovieRenderPipelineEditor/Private/Customizations/ConsoleVariableCustomization.h
@@ -122,16 +122,19 @@ protected:
 		{
 			return nullptr;
 		}
-		
-		UMoviePipelineConsoleVariableSetting* Setting = GetCVarSettingFromHandle(PropertyHandle);
-		if (!Setting)
+
+		TArray<UObject *> OuterObjects;
+		PropertyHandle->GetOuterObjects(OuterObjects);
+
+		if (OuterObjects.Num() != 1)
 		{
 			return nullptr;
 		}
-		
-		// Get the cvar entry associated w/ this property
-		const int32 EntryIndex = PropertyHandle->GetIndexInArray();
-		return Setting->GetCVarAtIndex(EntryIndex);
+
+		// Get the pointer to the ConsoleVariable struct from the first outer object
+		return reinterpret_cast<FMoviePipelineConsoleVariableEntry*>(
+			PropertyHandle->GetValueBaseAddress(reinterpret_cast<uint8*>(OuterObjects[0]))
+		);
 	}
 
 	ECheckBoxState GetIsEnabledCheckBoxState(TSharedRef<IPropertyHandle> PropertyHandle) const
@@ -217,14 +220,23 @@ protected:
 		{
 			// Get the non-override value of this cvar
 			const UMoviePipelineConsoleVariableSetting* CVarSetting = GetCVarSettingFromHandle(ParentPropertyHandle);
+			float CVarFloatValue = 0.f;
+
+			// If we are in a settings object then use that to get the disabled value
 			if (CVarSetting)
 			{
-				float CVarFloatValue = 0.f;
-				if (FDefaultValueHelper::ParseFloat(CVarSetting->ResolveDisabledValue(*CVarEntry), CVarFloatValue))
+				FDefaultValueHelper::ParseFloat(CVarSetting->ResolveDisabledValue(*CVarEntry), CVarFloatValue);
+			}
+			// Fall back to the startup value of the cvar if we're not in a settings object
+			else
+			{
+				const TSharedPtr<FConsoleVariablesEditorCommandInfo> CommandInfo = CVarEntry->CommandInfo.Pin();
+				if (CommandInfo.IsValid())
 				{
-					return CVarFloatValue;
+					FDefaultValueHelper::ParseFloat(CommandInfo->StartupValueAsString, CVarFloatValue);
 				}
 			}
+			return CVarFloatValue;
 		}
 		
 		float Value;
@@ -256,7 +268,7 @@ protected:
 		
 		const UMoviePipelineConsoleVariableSetting* Setting = GetCVarSettingFromHandle(StructPropertyHandle);
 		const FMoviePipelineConsoleVariableEntry* CVarEntry = GetEntryForHandle(StructPropertyHandle);
-		if (!Setting || !CVarEntry)
+		if (!CVarEntry)
 		{
 			return FText();
 		}
@@ -268,7 +280,7 @@ protected:
 		}
 		
 		const FString CustomValue = FString::SanitizeFloat(CVarEntry->Value);
-		const FString PresetValue = Setting->ResolvePresetValue(CVarEntry->Name);
+		const FString PresetValue = Setting ? Setting->ResolvePresetValue(CVarEntry->Name) : FString();
 		const FString StartupValue = CommandInfo->StartupValueAsString;
 		
 		return FText::Format(ToolTipFormatText,
