MoviePipeline: Fixed the UI warning you about losing unsaved changes to Queues when there weren't actually any changes.

diff --git a/Engine/Plugins/MovieScene/MovieRenderPipeline/Source/MovieRenderPipelineEditor/Private/Widgets/SMoviePipelineQueuePanel.cpp b/Engine/Plugins/MovieScene/MovieRenderPipeline/Source/MovieRenderPipelineEditor/Private/Widgets/SMoviePipelineQueuePanel.cpp
--- a/Engine/Plugins/MovieScene/MovieRenderPipeline/Source/MovieRenderPipelineEditor/Private/Widgets/SMoviePipelineQueuePanel.cpp
+++ b/Engine/Plugins/MovieScene/MovieRenderPipeline/Source/MovieRenderPipelineEditor/Private/Widgets/SMoviePipelineQueuePanel.cpp
@@ -601,15 +601,14 @@ TSharedRef<SWidget> SMoviePipelineQueuePanel::OnGenerateSavedQueuesMenu()
 
 FText SMoviePipelineQueuePanel::GetQueueMenuButtonText() const
 {
-	const FString QueueName = GetQueueOriginName();
-	if (!QueueName.IsEmpty())
+	FText QueueName = FText::FromString(GetQueueOriginName());
+	if (QueueName.IsEmpty())
 	{
-		return FText::Format(FText::FromString(TEXT("{0}{1}")),
-			FText::FromString(QueueName),
-			FText::FromString(IsQueueDirty() ? TEXT(" *") : TEXT("")));
+		QueueName = LOCTEXT("QueueSaveMenuUnsavedConfig_Text", "Unsaved Queue");
 	}
-	
-	return LOCTEXT("QueueSaveMenuUnsavedConfig_Text", "Unsaved Queue *");
+	return FText::Format(FText::FromString(TEXT("{0}{1}")),
+		QueueName,
+		FText::FromString(IsQueueDirty() ? TEXT(" *") : TEXT("")));
 }
 
 bool SMoviePipelineQueuePanel::OpenSaveDialog(const FString& InDefaultPath, const FString& InNewNameSuggestion, FString& OutPackageName)
@@ -815,10 +814,22 @@ bool SMoviePipelineQueuePanel::IsQueueDirty() const
 	check(Subsystem);
 
 	const UMoviePipelineQueue* Queue = Subsystem->GetQueue();
+	if(!Queue)
+	{
+		return false;
+	}
 
 	// The queue is considered dirty if the current queue has no origin (ie, it has never been saved) or it has been
-	// modified since it was loaded
-	return !GetQueueOrigin() || (Queue && Queue->IsDirty());
+	// modified since it was loaded. We skip the no-origin check if there are no jobs (to avoid triggering on an
+	// empty, first-load queue), but we don't skip the IsDirty check because a queue that had jobs and then
+	// removed them will be dirty.
+	bool bShouldCheckQueueOrigin = Queue->GetJobs().Num() > 0;
+	bool bHasNoQueueOrigin = false;
+	if (bShouldCheckQueueOrigin)
+	{
+		bHasNoQueueOrigin = !GetQueueOrigin();
+	}
+	return bHasNoQueueOrigin || (Queue->IsDirty());
 }
 
 UMoviePipelineQueue* SMoviePipelineQueuePanel::GetQueueOrigin() const
