MoviePipeline: Stubbed out some of the functionality required for very basic graph traversal, and very basic render submission.

diff --git a/Engine/Plugins/MovieScene/MovieRenderPipeline/Source/MovieRenderPipelineCore/Private/Graph/MovieGraphConfig.cpp b/Engine/Plugins/MovieScene/MovieRenderPipeline/Source/MovieRenderPipelineCore/Private/Graph/MovieGraphConfig.cpp
--- a/Engine/Plugins/MovieScene/MovieRenderPipeline/Source/MovieRenderPipelineCore/Private/Graph/MovieGraphConfig.cpp
+++ b/Engine/Plugins/MovieScene/MovieRenderPipeline/Source/MovieRenderPipelineCore/Private/Graph/MovieGraphConfig.cpp
@@ -30,23 +30,59 @@ UMovieGraphConfig::UMovieGraphConfig()
 	}
 }
 
-void UMovieGraphConfig::TraverseTest(UMovieGraphNode* InNode)
+void UMovieGraphConfig::TraverseGraphRecursive(UMovieGraphNode* InNode, TSubclassOf<UMovieGraphNode> InClassType, const FMovieGraphTraversalContext& InContext, TArray<UMovieGraphNode*>& OutNodes) const
 {
-	UE_LOG(LogTemp, Warning, TEXT("Visiting Node: %s"), *InNode->GetFName().ToString());
+	check(InNode);
+	
+	if (InNode->IsA(InClassType))
+	{
+		OutNodes.Add(InNode);
+	}
 
-	for (UMovieGraphPin* Pin : InNode->InputPins)
+	// Loop through the input pins on each node
+	for (UMovieGraphPin* Pin : InNode->GetInputPins())
 	{
-		for (UMovieGraphEdge* Edge : Pin->Edges)
+		// ToDo: Pins should have a type (property, branch line, etc.) and we can mask following them.
+
+		bool bFollowPath = true;
+		// If this is an output node we mask their pins against the traversal context
+		if (UMovieGraphOutputNode* NodeAsOutputNode = Cast<UMovieGraphOutputNode>(InNode))
+		{
+			if (InContext.RootBranch.BranchName != NAME_None)
+			{
+				bFollowPath = Pin->Properties.Label == InContext.RootBranch.BranchName;
+				UE_LOG(LogMovieRenderPipeline, Log, TEXT("Traversing Graph with Root Branch: %s (bMatch: %d)"), *InContext.RootBranch.BranchName.ToString(), bFollowPath);
+			}
+		}
+
+		if (bFollowPath)
 		{
-			if (UMovieGraphPin* OtherPin = Edge->GetOtherPin(Pin))
+			for (UMovieGraphEdge* Edge : Pin->Edges)
 			{
-				if (OtherPin->Node != InNode)
+				if (UMovieGraphPin* OtherPin = Edge->GetOtherPin(Pin))
 				{
-					TraverseTest(OtherPin->Node);
+					// ToDo: This needs to be upgraded to a full cyclic detection so we don't get stuck in an infinite
+					// loop through some badly malformed graph created via scripting/etc.
+					if (OtherPin->Node != InNode)
+					{
+						TraverseGraphRecursive(OtherPin->Node, InClassType, InContext, OutNodes);
+					}
 				}
 			}
 		}
 	}
+
+}
+
+TArray<UMovieGraphNode*> UMovieGraphConfig::TraverseGraph(TSubclassOf<UMovieGraphNode> InClassType, const FMovieGraphTraversalContext& InContext) const
+{
+	TArray<UMovieGraphNode*> OutNodes;
+
+	if (OutputNode)
+	{
+		TraverseGraphRecursive(OutputNode, InClassType, InContext, /*Out*/OutNodes);
+	}
+	return OutNodes;
 }
 
 void UMovieGraphConfig::OnVariableUpdated(UMovieGraphVariable* UpdatedVariable)
@@ -57,29 +93,6 @@ void UMovieGraphConfig::OnVariableUpdated(UMovieGraphVariable* UpdatedVariable)
 #endif
 }
 
-UMovieGraphNode* UMovieGraphConfig::FindTraversalStartForContext(const FMovieGraphTraversalContext& InContext) const
-{
-	return OutputNode;
-}
-
-void UMovieGraphConfig::TraversalTest()
-{
-	FMovieGraphTraversalContext Context;
-	Context.ShotName = TEXT("Test");
-	Context.RenderLayerName = TEXT("Test2");
-
-	UMovieGraphNode* StartNode = FindTraversalStartForContext(Context);
-	if (StartNode)
-	{
-		UE_LOG(LogTemp, Warning, TEXT("Traversing!"));
-		TraverseTest(StartNode);
-	}
-	else
-	{
-		UE_LOG(LogTemp, Error, TEXT("Could not find traversal entry point for context"));
-	}
-}
-
 bool UMovieGraphConfig::AddLabeledEdge(UMovieGraphNode* FromNode, const FName& FromPinLabel, UMovieGraphNode* ToNode, const FName& ToPinLabel)
 {
 	if (!FromNode || !ToNode)
@@ -362,4 +375,19 @@ void UMovieGraphConfig::DeleteVariableMember(UMovieGraphVariable* VariableMember
 #endif
 }
 
+TArray<FMovieGraphBranch> UMovieGraphConfig::GetOutputBranches() const
+{
+	TArray<FMovieGraphBranch> Branches;
+	if (OutputNode)
+	{
+		for (UMovieGraphPin* Pin : OutputNode->GetInputPins())
+		{
+			FMovieGraphBranch& NewBranch = Branches.AddDefaulted_GetRef();
+			NewBranch.BranchName = Pin->Properties.Label;
+		}
+	}
+
+	return Branches;
+}
+
 #undef LOCTEXT_NAMESPACE
diff --git a/Engine/Plugins/MovieScene/MovieRenderPipeline/Source/MovieRenderPipelineCore/Private/Graph/MovieGraphDefaultRenderer.cpp b/Engine/Plugins/MovieScene/MovieRenderPipeline/Source/MovieRenderPipelineCore/Private/Graph/MovieGraphDefaultRenderer.cpp
--- a/Engine/Plugins/MovieScene/MovieRenderPipeline/Source/MovieRenderPipelineCore/Private/Graph/MovieGraphDefaultRenderer.cpp
+++ b/Engine/Plugins/MovieScene/MovieRenderPipeline/Source/MovieRenderPipelineCore/Private/Graph/MovieGraphDefaultRenderer.cpp
@@ -2,19 +2,152 @@
 
 #include "Graph/MovieGraphDefaultRenderer.h"
 #include "Graph/MovieGraphPipeline.h"
+#include "Graph/MovieGraphConfig.h"
+#include "Graph/Nodes/MovieGraphRenderLayerNode.h"
+#include "Graph/Nodes/MovieGraphRenderPassNode.h"
+#include "MovieRenderPipelineCoreModule.h"
+#include "RenderingThread.h"
+#include "Engine/TextureRenderTarget2D.h"
+#include "UObject/Package.h"
+
+void UMovieGraphDefaultRenderer::SetupRenderingPipelineForShot(UMoviePipelineExecutorShot* InShot)
+{
+	// Iterate through the graph config and look for Render Layers.
+	UMovieGraphConfig* RootGraph = GetOwningGraph()->GetRootGraphForShot(InShot);
+
+	struct FMovieGraphPass
+	{
+		TSubclassOf<UMovieGraphRenderPassNode> ClassType;
+		TArray<FString> LayerNames;
+	};
+
+	TArray<FMovieGraphPass> OutputPasses;
+
+	// Start by getting our root set of branches we should follow
+	TArray<FMovieGraphBranch> GraphBranches = RootGraph->GetOutputBranches();
+	for (const FMovieGraphBranch& Branch : GraphBranches)
+	{
+		// We follow each branch looking for Render Layer nodes to figure out what render layer this should be.
+		FMovieGraphTraversalContext TraversalContext = GetOwningGraph()->GetTraversalContextForShot(InShot);
+		TraversalContext.RootBranch = Branch;
+
+		FString RenderLayerName = TEXT("Unnamed_Layer");
+		UMovieGraphRenderLayerNode* FirstRenderLayerNode = RootGraph->IterateGraphForClass<UMovieGraphRenderLayerNode>(TraversalContext);
+		if (FirstRenderLayerNode)
+		{
+			RenderLayerName = FirstRenderLayerNode->GetRenderLayerName();;
+		}
+
+		// Now we need to figure out which renderers are on this branch.
+		TArray<UMovieGraphRenderPassNode*> Renderers = RootGraph->IterateGraphForClassAll<UMovieGraphRenderPassNode>(TraversalContext);
+		if (FirstRenderLayerNode && Renderers.Num() == 0)
+		{
+			UE_LOG(LogMovieRenderPipeline, Warning, TEXT("Found RenderLayer: \"%s\" but no Renderers defined."), *RenderLayerName);
+		}
+
+		for (UMovieGraphRenderPassNode* RenderPassNode : Renderers)
+		{
+			FString RendererName = RenderPassNode->GetRendererName();
+
+			FMovieGraphPass* ExistingPass = OutputPasses.FindByPredicate([RenderPassNode](const FMovieGraphPass& Pass)
+				{ return Pass.ClassType == RenderPassNode->GetClass(); });
+			
+			if (!ExistingPass)
+			{
+				ExistingPass = &OutputPasses.AddDefaulted_GetRef();
+				ExistingPass->ClassType = RenderPassNode->GetClass();
+			}
+			ExistingPass->LayerNames.AddUnique(RenderLayerName);
+		}
+	}
+
+	UE_LOG(LogMovieRenderPipeline, Warning, TEXT("Found: %d Render Passes:"), OutputPasses.Num());
+	int32 TotalLayerCount = 0;
+	for (const FMovieGraphPass& Pass : OutputPasses)
+	{
+		// ToDo: This should probably come from the Renderers themselves, as they can internally produce multiple
+		// renders (such as ObjectID passes).
+		TotalLayerCount += Pass.LayerNames.Num();
+		
+		FMovieGraphRenderPassSetupData SetupData;
+		SetupData.Renderer = this;
+		UE_LOG(LogMovieRenderPipeline, Warning, TEXT("\tRenderer Class: %s"), *Pass.ClassType->GetName());
+		for (const FString& LayerName : Pass.LayerNames)
+		{
+			FMovieGraphRenderPassLayerData& LayerData = SetupData.Layers.AddDefaulted_GetRef();
+			LayerData.LayerName = LayerName;
+
+			UE_LOG(LogMovieRenderPipeline, Warning, TEXT("\t\tLayer Name: %s"), *LayerName);
+		}
+
+		UMovieGraphRenderPassNode* RenderPassCDO = Pass.ClassType->GetDefaultObject<UMovieGraphRenderPassNode>();
+		RenderPassCDO->Setup(SetupData);
+
+		RenderPassesInUse.Add(RenderPassCDO);
+	}
+
+	UE_LOG(LogMovieRenderPipeline, Log, TEXT("Finished initializing %d Render Passes (with %d total layers)."), OutputPasses.Num(), TotalLayerCount);
+}
+
+void UMovieGraphDefaultRenderer::TeardownRenderingPipelineForShot(UMoviePipelineExecutorShot* InShot)
+{
+	FlushRenderingCommands();
+
+	for (TObjectPtr< UMovieGraphRenderPassNode> RenderPass : RenderPassesInUse)
+	{
+		RenderPass->Teardown();
+	}
+
+	RenderPassesInUse.Reset();
+
+	// ToDo: This could probably be preserved across shots to avoid allocations
+	PooledViewRenderTargets.Reset();
+}
+
+
+void UMovieGraphDefaultRenderer::AddReferencedObjects(UObject* InThis, FReferenceCollector& Collector)
+{
+	Super::AddReferencedObjects(InThis, Collector);
+	UMovieGraphDefaultRenderer* This = CastChecked<UMovieGraphDefaultRenderer>(InThis);
+
+	// Can't be a const& due to AddStableReference API
+	for (TPair<FRenderTargetInitParams, TObjectPtr<UTextureRenderTarget2D>>& KVP : This->PooledViewRenderTargets)
+	{
+		Collector.AddStableReference(&KVP.Value);
+	}
+}
+
 
 void UMovieGraphDefaultRenderer::Render(const FMovieGraphTimeStepData& InTimeStepData)
 {
-	UE_LOG(LogTemp, Warning, TEXT("MovieGraphRenderer: Render"));
+	for (TObjectPtr<UMovieGraphRenderPassNode> RenderPass : RenderPassesInUse)
+	{
+		RenderPass->Render();
+	}
 }
 
-void UMovieGraphDefaultRenderer::SetupRenderingPipelineForShot(UMoviePipelineExecutorShot* InShot)
+
+UTextureRenderTarget2D* UMovieGraphDefaultRenderer::GetOrCreateViewRenderTarget(const FRenderTargetInitParams& InInitParams)
 {
-	UE_LOG(LogTemp, Warning, TEXT("MovieGraphRenderer: SetupForShot"));
+	if (const TObjectPtr<UTextureRenderTarget2D>* ExistViewRenderTarget = PooledViewRenderTargets.Find(InInitParams))
+	{
+		return *ExistViewRenderTarget;
+	}
 
+	const TObjectPtr<UTextureRenderTarget2D> NewViewRenderTarget = CreateViewRenderTarget(InInitParams);
+	PooledViewRenderTargets.Emplace(InInitParams, NewViewRenderTarget);
+
+	return NewViewRenderTarget.Get();
 }
 
-void UMovieGraphDefaultRenderer::TeardownRenderingPipelineForShot(UMoviePipelineExecutorShot* InShot)
+TObjectPtr<UTextureRenderTarget2D> UMovieGraphDefaultRenderer::CreateViewRenderTarget(const FRenderTargetInitParams& InInitParams) const
 {
-	UE_LOG(LogTemp, Warning, TEXT("MovieGraphRenderer: TeardownForShot"));
-}
\ No newline at end of file
+	TObjectPtr<UTextureRenderTarget2D> NewTarget = NewObject<UTextureRenderTarget2D>(GetTransientPackage());
+	NewTarget->ClearColor = FLinearColor(0.0f, 0.0f, 0.0f, 0.0f);
+	NewTarget->TargetGamma = InInitParams.TargetGamma;
+	NewTarget->InitCustomFormat(InInitParams.Size.X, InInitParams.Size.Y, InInitParams.PixelFormat, false);
+	int32 ResourceSizeBytes = NewTarget->GetResourceSizeBytes(EResourceSizeMode::Type::EstimatedTotal);
+	UE_LOG(LogMovieRenderPipeline, Log, TEXT("Allocated a View Render Target sized: (%d, %d), Bytes: %d"), InInitParams.Size.X, InInitParams.Size.Y, ResourceSizeBytes);
+
+	return NewTarget;
+}
diff --git a/Engine/Plugins/MovieScene/MovieRenderPipeline/Source/MovieRenderPipelineCore/Private/Graph/MovieGraphPipeline.cpp b/Engine/Plugins/MovieScene/MovieRenderPipeline/Source/MovieRenderPipelineCore/Private/Graph/MovieGraphPipeline.cpp
--- a/Engine/Plugins/MovieScene/MovieRenderPipeline/Source/MovieRenderPipelineCore/Private/Graph/MovieGraphPipeline.cpp
+++ b/Engine/Plugins/MovieScene/MovieRenderPipeline/Source/MovieRenderPipelineCore/Private/Graph/MovieGraphPipeline.cpp
@@ -541,7 +541,7 @@ void UMovieGraphPipeline::TransitionToState(const EMovieRenderPipelineState InNe
 
 			// Unregister our OnEngineTickEndFrame delegate. We don't unregister BeginFrame as we need
 			// to continue to call it to allow ticking the Finalization stage.
-			// FCoreDelegates::OnEndFrame.RemoveAll(this);
+			FCoreDelegates::OnEndFrame.RemoveAll(this);
 
 			// Reset the Custom Timestep because we don't care how long the engine takes now
 			// GEngine->SetCustomTimeStep(CachedPrevCustomTimeStep);
@@ -693,3 +693,44 @@ void UMovieGraphPipeline::OnMoviePipelineFinishedImpl()
 //		}
 //	}
 //}
+
+UMovieGraphConfig* UMovieGraphPipeline::GetRootGraphForShot(UMoviePipelineExecutorShot* InShot) const
+{
+	if (GetCurrentJob())
+	{
+		for (UMoviePipelineExecutorShot* Shot : GetCurrentJob()->ShotInfo)
+		{
+			if (Shot == InShot)
+			{
+				if (Shot->GetGraphPreset())
+				{
+					return Shot->GetGraphPreset();
+				}
+				else if (Shot->GetGraphConfig())
+				{
+					return Shot->GetGraphConfig();
+				}
+			}
+
+		}
+
+		// If the shot hasn't overwritten the preset then we return the root one for the whole job.
+		if (GetCurrentJob()->GetGraphPreset())
+		{
+			return GetCurrentJob()->GetGraphPreset();
+		}
+		else if (GetCurrentJob()->GetGraphConfig())
+		{
+			return GetCurrentJob()->GetGraphConfig();
+		}
+	}
+
+	return nullptr;
+}
+
+FMovieGraphTraversalContext UMovieGraphPipeline::GetTraversalContextForShot(UMoviePipelineExecutorShot* InShot) const
+{
+	// ToDo: Fill out shot index/current shot/current frame/job/etc.
+	return FMovieGraphTraversalContext();
+}
+
diff --git a/Engine/Plugins/MovieScene/MovieRenderPipeline/Source/MovieRenderPipelineCore/Private/Graph/Nodes/MovieGraphDeferredRenderPassNode.cpp b/Engine/Plugins/MovieScene/MovieRenderPipeline/Source/MovieRenderPipelineCore/Private/Graph/Nodes/MovieGraphDeferredRenderPassNode.cpp
new file mode 100644
--- /dev/null
+++ b/Engine/Plugins/MovieScene/MovieRenderPipeline/Source/MovieRenderPipelineCore/Private/Graph/Nodes/MovieGraphDeferredRenderPassNode.cpp
@@ -0,0 +1,382 @@
+// Copyright Epic Games, Inc. All Rights Reserved.
+
+#include "Graph/Nodes/MovieGraphDeferredRenderPassNode.h"
+#include "Graph/MovieGraphDataTypes.h"
+#include "Graph/MovieGraphDefaultRenderer.h"
+#include "MovieRenderPipelineCoreModule.h"
+
+#include "EngineModule.h"
+#include "SceneManagement.h"
+#include "CanvasTypes.h"
+#include "Engine/TextureRenderTarget2D.h"
+#include "LegacyScreenPercentageDriver.h"
+#include "SceneViewExtensionContext.h"
+#include "OpenColorIODisplayExtension.h"
+#include "TextureResource.h"
+
+void UMovieGraphDeferredRenderPassNode::SetupImpl(const FMovieGraphRenderPassSetupData& InSetupData)
+{
+	// To make the implementation simpler, we make one instance of FMovieGraphDeferredRenderPas
+	// per camera, and per render layer. These objects can pull from common pools to share state,
+	// which gives us a better overview of how many resources are being used by MRQ.
+	for (const FMovieGraphRenderPassLayerData& LayerData : InSetupData.Layers)
+	{
+		TUniquePtr<FMovieGraphDeferredRenderPass> RendererInstance = MakeUnique<FMovieGraphDeferredRenderPass>();
+		RendererInstance->Setup(InSetupData.Renderer, LayerData);
+		CurrentInstances.Add(MoveTemp(RendererInstance));
+	}
+}
+
+void UMovieGraphDeferredRenderPassNode::TeardownImpl()
+{
+	// We don't need to flush the rendering commands as we assume the MovieGraph
+	// Renderer has already done that once, so all data for all passes should
+	// have been submitted to the GPU (and subsequently read back) by now.
+	for (TUniquePtr<FMovieGraphDeferredRenderPass>& Instance : CurrentInstances)
+	{
+		Instance->Teardown();
+	}
+	CurrentInstances.Reset();
+}
+
+
+void UMovieGraphDeferredRenderPassNode::RenderImpl()
+{
+	for (TUniquePtr<FMovieGraphDeferredRenderPass>& Instance : CurrentInstances)
+	{
+		Instance->Render();
+	}
+}
+
+
+void UMovieGraphDeferredRenderPassNode::AddReferencedObjects(UObject* InThis, FReferenceCollector& Collector)
+{
+	Super::AddReferencedObjects(InThis, Collector);
+	UMovieGraphDeferredRenderPassNode* This = CastChecked<UMovieGraphDeferredRenderPassNode>(InThis);
+	for (TUniquePtr<FMovieGraphDeferredRenderPass>& Instance : This->CurrentInstances)
+	{
+		Instance->AddReferencedObjects(Collector);
+	}
+}
+
+void UMovieGraphDeferredRenderPassNode::FMovieGraphDeferredRenderPass::Setup(TWeakObjectPtr<class UMovieGraphDefaultRenderer> InRenderer, const FMovieGraphRenderPassLayerData& InLayer)
+{
+	LayerData = InLayer;
+	Renderer = InRenderer;
+
+	const FIntPoint OutputCameraResolution = FIntPoint(1920, 1080);// Renderer->GetEffectiveOutputResolution(LayerData);
+
+	// Figure out how big each sub-region (tile) is.
+	// const int32 TileSize = Graph->FindSetting(LayerData.LayerName, "highres.tileSize");
+	//int32 TileCount = 1; // ToDo: We should do tile sizes (ie: 1024x512) instead
+	//FIntPoint BackbufferResolution = FIntPoint(
+	//	FMath::CeilToInt((float)OutputCameraResolution.X / (float)TileCount),
+	//	FMath::CeilToInt((float)OutputCameraResolution.Y / (float)TileCount));
+
+	// BackbufferResolution = HighResSettings->CalculatePaddedBackbufferSize(BackbufferResolution);
+	
+	// Create a view render target for our given resolution. This is require for the FCanvas/FSceneViewAPI. We pool these.
+	// Renderer->GetOrCreateViewRenderTarget(BackbufferResolution);
+	// CreateSurfaceQueueImpl(BackbufferResolution);
+
+	SceneViewState.Allocate(Renderer->GetWorld()->GetFeatureLevel());
+
+}
+
+void UMovieGraphDeferredRenderPassNode::FMovieGraphDeferredRenderPass::Teardown()
+{
+	FSceneViewStateInterface* Ref = SceneViewState.GetReference();
+	if (Ref)
+	{
+		Ref->ClearMIDPool();
+	}
+	SceneViewState.Destroy();
+}
+
+void UMovieGraphDeferredRenderPassNode::FMovieGraphDeferredRenderPass::AddReferencedObjects(FReferenceCollector& Collector)
+{
+	FSceneViewStateInterface* Ref = SceneViewState.GetReference();
+	if (Ref)
+	{
+		Ref->AddReferencedObjects(Collector);
+	}
+}
+
+void UMovieGraphDeferredRenderPassNode::FMovieGraphDeferredRenderPass::Render()
+{
+
+	// OCIO: Since this is a manually created Render target we don't need Gamma to be applied.
+	// We use this render target to render to via a display extension that utilizes Display Gamma
+	// which has a default value of 2.2 (DefaultDisplayGamma), therefore we need to set Gamma on this render target to 2.2 to cancel out any unwanted effects.
+	
+	UMovieGraphDefaultRenderer::FRenderTargetInitParams RenderTargetInitParams;
+	RenderTargetInitParams.Size = FIntPoint(1920, 1080);
+	RenderTargetInitParams.TargetGamma = FOpenColorIODisplayExtension::DefaultDisplayGamma;
+	RenderTargetInitParams.PixelFormat = EPixelFormat::PF_FloatRGBA;
+
+	UTextureRenderTarget2D* RenderTarget = Renderer->GetOrCreateViewRenderTarget(RenderTargetInitParams);
+
+	FViewFamilyContextInitData InitData;
+	InitData.RenderTarget = RenderTarget->GameThread_GetRenderTargetResource();
+	InitData.World = Renderer->GetWorld();
+
+	APlayerController* LocalPlayerController = Renderer->GetWorld()->GetFirstPlayerController();
+	if (LocalPlayerController && LocalPlayerController->PlayerCameraManager)
+	{
+		InitData.MinimalViewInfo = LocalPlayerController->PlayerCameraManager->GetCameraCacheView();
+		InitData.ViewActor = LocalPlayerController->GetViewTarget();
+	}
+	else
+	{
+		InitData.MinimalViewInfo = FMinimalViewInfo();
+		InitData.ViewActor = nullptr;
+		UE_LOG(LogMovieRenderPipeline, Error, TEXT("Failed to find Local Player Controller/Camera Manager to get viewpoint!"));
+	}
+
+	// ToDo:
+	InitData.TimeData.FrameDeltaTime = 1 / 24.f;
+	InitData.TimeData.TimeDilation = 1.0f;
+	InitData.TimeData.WorldSeconds = 0.f;
+	InitData.TimeData.MotionBlurFraction = 1 / 24.f;
+	InitData.SceneCaptureSource = ESceneCaptureSource::SCS_FinalToneCurveHDR;
+	InitData.bWorldIsPaused = false;
+	InitData.GlobalScreenPercentageFraction = FLegacyScreenPercentageDriver::GetCVarResolutionFraction();
+	InitData.OverscanFraction = 0.f;
+	InitData.FrameIndex = 1;
+	InitData.bCameraCut = false;
+	InitData.AntiAliasingMethod = EAntiAliasingMethod::AAM_None;
+	InitData.SceneViewStateReference = SceneViewState.GetReference();
+
+
+	// Allocate the view we want to render from, then a Family for the view to live in.
+	// We apply a lot of MRQ-specific overrides to the ViewFamily and View, so do those next.
+	// Allocating the Scene View automatically 
+	TSharedRef<FSceneViewFamilyContext> ViewFamily = AllocateSceneViewFamilyContext(InitData);
+	AllocateSceneView(ViewFamily, InitData);
+	ApplyMoviePipelineOverridesToViewFamily(ViewFamily, InitData);
+
+	
+	FRenderTarget* RenderTargetResource = RenderTarget->GameThread_GetRenderTargetResource();
+	check(RenderTarget);
+
+	FHitProxyConsumer* HitProxyConsumer = nullptr;
+	const float DPIScale = 1.0f;
+	FCanvas Canvas = FCanvas(RenderTargetResource, HitProxyConsumer, Renderer->GetWorld(), Renderer->GetWorld()->GetFeatureLevel(), FCanvas::CDM_DeferDrawing, DPIScale);
+
+	// Submit the renderer to be rendered
+	GetRendererModule().BeginRenderingViewFamily(&Canvas, ViewFamily.ToSharedPtr().Get());
+
+	// Readback + Accumulate.
+	// PostRendererSubmission(InOutSampleState, PassIdentifierForCurrentCamera, GetOutputFileSortingOrder(), Canvas);
+}
+
+TSharedRef<FSceneViewFamilyContext> UMovieGraphDeferredRenderPassNode::FMovieGraphDeferredRenderPass::AllocateSceneViewFamilyContext(const FViewFamilyContextInitData& InInitData)
+{
+	FEngineShowFlags ShowFlags = FEngineShowFlags(EShowFlagInitMode::ESFIM_Game);
+	EViewModeIndex ViewModeIndex = EViewModeIndex::VMI_Lit;
+
+	const bool bIsPerspective = InInitData.MinimalViewInfo.ProjectionMode == ECameraProjectionMode::Type::Perspective;
+
+	// Allow the Engine Showflag system to override our engine showflags, based on our view mode index.
+	// This is required for certain debug view modes (to have matching show flags set for rendering).
+	ApplyViewMode(/*In*/ ViewModeIndex, bIsPerspective, /*InOut*/ShowFlags);
+
+	// And then we have to let another system override them again (based on cvars, etc.)
+	EngineShowFlagOverride(ESFIM_Game, ViewModeIndex, ShowFlags, false);
+
+	TSharedRef<FSceneViewFamilyContext> OutViewFamily = MakeShared<FSceneViewFamilyContext>(FSceneViewFamily::ConstructionValues(
+		InInitData.RenderTarget,
+		InInitData.World->Scene,
+		ShowFlags)
+		.SetTime(FGameTime::CreateUndilated(InInitData.TimeData.WorldSeconds, InInitData.TimeData.FrameDeltaTime))
+		.SetRealtimeUpdate(true));
+
+	// Used to specify if the Tone Curve is being applied or not to our Linear Output data
+	OutViewFamily->SceneCaptureSource = InInitData.SceneCaptureSource;
+	OutViewFamily->bWorldIsPaused = InInitData.bWorldIsPaused;
+	OutViewFamily->ViewMode = ViewModeIndex;
+	OutViewFamily->bOverrideVirtualTextureThrottle = true;
+
+	// ToDo: Let settings modify the ScreenPercentageInterface so third party screen percentages are supported.
+
+	// If UMoviePipelineViewFamilySetting never set a Screen percentage interface we fallback to default.
+	if (OutViewFamily->GetScreenPercentageInterface() == nullptr)
+	{
+		OutViewFamily->SetScreenPercentageInterface(new FLegacyScreenPercentageDriver(*OutViewFamily, InInitData.GlobalScreenPercentageFraction));
+	}
+
+	return OutViewFamily;
+}
+
+
+FSceneView* UMovieGraphDeferredRenderPassNode::FMovieGraphDeferredRenderPass::AllocateSceneView(TSharedPtr<FSceneViewFamilyContext> InViewFamilyContext, FViewFamilyContextInitData& InInitData) const
+{
+	FSceneViewInitOptions ViewInitOptions;
+	ViewInitOptions.ViewFamily = InViewFamilyContext.Get();
+	ViewInitOptions.ViewOrigin = InInitData.MinimalViewInfo.Location;
+	ViewInitOptions.SetViewRectangle(FIntRect(FIntPoint(0, 0), FIntPoint(InInitData.RenderTarget->GetSizeXY().X, InInitData.RenderTarget->GetSizeXY().Y)));
+	ViewInitOptions.ViewRotationMatrix = FInverseRotationMatrix(InInitData.MinimalViewInfo.Rotation);
+	ViewInitOptions.ViewActor = InInitData.ViewActor;
+
+	// Rotate the view 90 degrees to match the rest of the engine.
+	ViewInitOptions.ViewRotationMatrix = ViewInitOptions.ViewRotationMatrix * FMatrix(
+		FPlane(0, 0, 1, 0),
+		FPlane(1, 0, 0, 0),
+		FPlane(0, 1, 0, 0),
+		FPlane(0, 0, 0, 1));
+	float ViewFOV = InInitData.MinimalViewInfo.FOV;
+
+	// Inflate our FOV to support the overscan 
+	ViewFOV = 2.0f * FMath::RadiansToDegrees(FMath::Atan((1.0f + InInitData.OverscanFraction) * FMath::Tan(FMath::DegreesToRadians(ViewFOV * 0.5f))));
+	float DofSensorScale = 1.0f;
+
+	// ToDo: This isn't great but this appears to be how the system works... ideally we could fetch this from the camera itself
+	const EAspectRatioAxisConstraint AspectRatioAxisConstraint = GetDefault<ULocalPlayer>()->AspectRatioAxisConstraint;
+	FIntRect ViewportRect = FIntRect(FIntPoint(0, 0), FIntPoint(InInitData.RenderTarget->GetSizeXY().X, InInitData.RenderTarget->GetSizeXY().Y));
+	FMinimalViewInfo::CalculateProjectionMatrixGivenViewRectangle(InInitData.MinimalViewInfo, AspectRatioAxisConstraint, ViewportRect, /*InOut*/ ViewInitOptions);
+
+	// ToDo: High Res Tiling, Overscan support, letterboxing
+	ViewInitOptions.SceneViewStateInterface = InInitData.SceneViewStateReference;
+	ViewInitOptions.FOV = ViewFOV;
+	ViewInitOptions.DesiredFOV = ViewFOV;
+
+	FSceneView* View = new FSceneView(ViewInitOptions);
+	InViewFamilyContext->Views.Add(View);
+
+	//View->ViewLocation = CameraInfo.ViewInfo.Location;
+	//View->ViewRotation = CameraInfo.ViewInfo.Rotation;
+	// Override previous/current view transforms so that tiled renders don't use the wrong occlusion/motion blur information.
+	//View->PreviousViewTransform = CameraInfo.ViewInfo.PreviousViewTransform;
+
+	View->StartFinalPostprocessSettings(View->ViewLocation);
+	// BlendPostProcessSettings(View, InOutSampleState, OptPayload);
+
+	// Scaling sensor size inversely with the the projection matrix [0][0] should physically
+	// cause the circle of confusion to be unchanged.
+	View->FinalPostProcessSettings.DepthOfFieldSensorWidth *= DofSensorScale;
+	// Modify the 'center' of the lens to be offset for high-res tiling, helps some effects (vignette) etc. still work.
+	// View->LensPrincipalPointOffsetScale = (FVector4f)CalculatePrinciplePointOffsetForTiling(InOutSampleState); // LWC_TODO: precision loss. CalculatePrinciplePointOffsetForTiling() could return float, it's normalized?
+	View->EndFinalPostprocessSettings(ViewInitOptions);
+	return View;
+}
+
+
+void UMovieGraphDeferredRenderPassNode::FMovieGraphDeferredRenderPass::ApplyMoviePipelineOverridesToViewFamily(TSharedRef<FSceneViewFamilyContext> InOutFamily, const FViewFamilyContextInitData& InInitData)
+{
+	/*// A third set of overrides required to properly configure views to match the given showflags/etc.
+	// ToDo: There's now five(!) identical implementations of this, we should unify them.
+	// SetupViewForViewModeOverride(View);
+
+	// Override the view's FrameIndex to be based on our progress through the sequence. This greatly increases
+	// determinism with things like TAA.
+	InInitData.View->OverrideFrameIndexValue = InInitData.FrameIndex;
+	InInitData.View->bCameraCut = InInitData.bCameraCut;
+	InInitData.View->bIsOfflineRender = true;
+	InInitData.View->AntiAliasingMethod = InInitData.AntiAliasingMethod;
+
+	// Add any view extensions that were added to the scene to this View too
+	// OutViewFamily->ViewExtensions.Append(GEngine->ViewExtensions->GatherActiveExtensions(FSceneViewExtensionContext(GetWorld()->Scene)));
+	// ToDo: Allow render passes to add additional view extensions not added to the scene. (AddViewExtensions(*OutViewFamily, InOutSampleState))
+
+	// Set up the view extensions with this view family
+	//for (auto ViewExt : OutViewFamily->ViewExtensions)
+	//{
+	//	ViewExt->SetupViewFamily(*OutViewFamily.Get());
+	//}
+
+	// Set up the view
+	//for (int ViewExt = 0; ViewExt < OutViewFamily->ViewExtensions.Num(); ViewExt++)
+	//{
+	//	OutViewFamily->ViewExtensions[ViewExt]->SetupView(*OutViewFamily.Get(), *View);
+	//}
+
+	// Override the Motion Blur settings since these are controlled by the movie pipeline.
+	{
+		// FFrameRate OutputFrameRate = GetPipeline()->GetPipelinePrimaryConfig()->GetEffectiveFrameRate(GetPipeline()->GetTargetSequence());
+		FFrameRate OutputFrameRate = FFrameRate(24, 1); // ToDo, get this from config.
+
+		// We need to inversly scale the target FPS by time dilation to counteract slowmo. If scaling isn't applied then motion blur length
+		// stays the same length despite the smaller delta time and the blur ends up too long.
+		View->FinalPostProcessSettings.MotionBlurTargetFPS = FMath::RoundToInt(OutputFrameRate.AsDecimal() / FMath::Max(SMALL_NUMBER, InInitData.TimeData.TimeDilation));
+		View->FinalPostProcessSettings.MotionBlurAmount = InInitData.TimeData.MotionBlurFraction;
+		View->FinalPostProcessSettings.MotionBlurMax = 100.f;
+		View->FinalPostProcessSettings.bOverride_MotionBlurAmount = true;
+		View->FinalPostProcessSettings.bOverride_MotionBlurTargetFPS = true;
+		View->FinalPostProcessSettings.bOverride_MotionBlurMax = true;
+
+		// Skip the whole pass if they don't want motion blur.
+		if (FMath::IsNearlyZero(InInitData.TimeData.MotionBlurFraction))
+		{
+			OutViewFamily->EngineShowFlags.SetMotionBlur(false);
+		}
+	}
+
+	// Warn the user for invalid setting combinations / enforce hardware limitations
+	{
+
+		// Locked Exposure
+		const bool bAutoExposureAllowed = IsAutoExposureAllowed(InOutSampleState);
+		{
+			// If the rendering pass doesn't allow autoexposure and they dont' have manual exposure set up, warn.
+			if (!bAutoExposureAllowed && (View->FinalPostProcessSettings.AutoExposureMethod != EAutoExposureMethod::AEM_Manual))
+			{
+				// Skip warning if the project setting is disabled though, as exposure will be forced off in the renderer anyways.
+				const URendererSettings* RenderSettings = GetDefault<URendererSettings>();
+				if (RenderSettings->bDefaultFeatureAutoExposure != false)
+				{
+					UE_LOG(LogMovieRenderPipeline, Warning, TEXT("Camera Auto Exposure Method not supported by one or more render passes. Change the Auto Exposure Method to Manual!"));
+					View->FinalPostProcessSettings.AutoExposureMethod = EAutoExposureMethod::AEM_Manual;
+				}
+			}
+		}
+
+		// Orthographic cameras don't support anti-aliasing outside the path tracer (other than FXAA)
+		const bool bIsOrthographicCamera = !View->IsPerspectiveProjection();
+		if (bIsOrthographicCamera)
+		{
+			bool bIsSupportedAAMethod = View->AntiAliasingMethod == EAntiAliasingMethod::AAM_FXAA;
+			bool bIsPathTracer = OutViewFamily->EngineShowFlags.PathTracing;
+			bool bWarnJitters = InOutSampleState.ProjectionMatrixJitterAmount.SquaredLength() > SMALL_NUMBER;
+			if ((!bIsPathTracer && !bIsSupportedAAMethod) || bWarnJitters)
+			{
+				UE_LOG(LogMovieRenderPipeline, Warning, TEXT("Orthographic Cameras are only supported with PathTracer or Deferred with FXAA Anti-Aliasing"));
+			}
+		}
+
+		{
+			bool bMethodWasUnsupported = false;
+			if (View->AntiAliasingMethod == AAM_TemporalAA && !SupportsGen4TAA(View->GetShaderPlatform()))
+			{
+				UE_LOG(LogMovieRenderPipeline, Error, TEXT("TAA was requested but this hardware does not support it."));
+				bMethodWasUnsupported = true;
+			}
+			else if (View->AntiAliasingMethod == AAM_TSR && !SupportsTSR(View->GetShaderPlatform()))
+			{
+				UE_LOG(LogMovieRenderPipeline, Error, TEXT("TSR was requested but this hardware does not support it."));
+				bMethodWasUnsupported = true;
+			}
+
+			if (bMethodWasUnsupported)
+			{
+				View->AntiAliasingMethod = AAM_None;
+			}
+		}
+	}
+
+
+	// Anti Aliasing
+	{
+		// If we're not using TAA, TSR, or Path Tracing we will apply the View Matrix projection jitter. Normally TAA sets this
+		// inside FSceneRenderer::PreVisibilityFrameSetup. Path Tracing does its own anti-aliasing internally.
+		bool bApplyProjectionJitter = !bIsOrthographicCamera
+			&& !OutViewFamily->EngineShowFlags.PathTracing
+			&& !IsTemporalAccumulationBasedMethod(View->AntiAliasingMethod);
+		if (bApplyProjectionJitter)
+		{
+			View->ViewMatrices.HackAddTemporalAAProjectionJitter(InOutSampleState.ProjectionMatrixJitterAmount);
+		}
+	}
+
+	return OutViewFamily;*/
+}
+
diff --git a/Engine/Plugins/MovieScene/MovieRenderPipeline/Source/MovieRenderPipelineCore/Public/Graph/MovieGraphConfig.h b/Engine/Plugins/MovieScene/MovieRenderPipeline/Source/MovieRenderPipelineCore/Public/Graph/MovieGraphConfig.h
--- a/Engine/Plugins/MovieScene/MovieRenderPipeline/Source/MovieRenderPipelineCore/Public/Graph/MovieGraphConfig.h
+++ b/Engine/Plugins/MovieScene/MovieRenderPipeline/Source/MovieRenderPipelineCore/Public/Graph/MovieGraphConfig.h
@@ -10,15 +10,29 @@
 // Forward Declare
 class UMovieGraphNode;
 
+USTRUCT(BlueprintType)
+struct FMovieGraphBranch
+{
+	GENERATED_BODY()
+
+public:
+	UPROPERTY(BlueprintReadWrite, EditAnywhere, Category = "MovieGraph")
+	FName BranchName;
+};
+
+
 USTRUCT(BlueprintType)
 struct FMovieGraphTraversalContext
 {
 	GENERATED_BODY();
 
-	UPROPERTY(BlueprintReadWrite, EditAnywhere, Category = "Blah")
+	UPROPERTY(BlueprintReadWrite, EditAnywhere, Category = "MovieGraph")
+	FMovieGraphBranch RootBranch;
+
+	UPROPERTY(BlueprintReadWrite, EditAnywhere, Category = "MovieGraph")
 	FString ShotName;
 
-	UPROPERTY(BlueprintReadWrite, EditAnywhere, Category = "Blah")
+	UPROPERTY(BlueprintReadWrite, EditAnywhere, Category = "MovieGraph")
 	FString RenderLayerName;
 };
 
@@ -140,8 +154,38 @@ public:
 	/** Remove the specified variable member from the graph. */
 	void DeleteVariableMember(UMovieGraphVariable* VariableMemberToDelete);
 
+	/** Returns only the names of the root branches in the Output Node, with no depth information. */
+	TArray<FMovieGraphBranch> GetOutputBranches() const;
+
+	template<typename NodeType>
+	NodeType* IterateGraphForClass(const FMovieGraphTraversalContext& InContext) const
+	{
+		TArray<NodeType*> AllSettings = IterateGraphForClassAll<NodeType>(InContext);
+		if (AllSettings.Num() > 0)
+		{
+			return AllSettings[0];
+		}
+		return nullptr;
+	}
+	
+	template<typename NodeType>
+	TArray<NodeType*> IterateGraphForClassAll(const FMovieGraphTraversalContext& InContext) const
+	{
+		TArray<NodeType*> TypedNodes;
+		const TArray<UMovieGraphNode*> FoundNodes = TraverseGraph(NodeType::StaticClass(), InContext);
+		for (UMovieGraphNode* Node : FoundNodes)
+		{
+			TypedNodes.Add(CastChecked<NodeType>(Node));
+		}
+
+		return TypedNodes;
+	}
+
+	TArray<UMovieGraphNode*> TraverseGraph(TSubclassOf<UMovieGraphNode> InClassType, const FMovieGraphTraversalContext& InContext) const;
+
+
 protected:
-	UMovieGraphNode* FindTraversalStartForContext(const FMovieGraphTraversalContext& InContext) const;
+	void TraverseGraphRecursive(UMovieGraphNode* InNode, TSubclassOf<UMovieGraphNode> InClassType, const FMovieGraphTraversalContext& InContext, TArray<UMovieGraphNode*>& OutNodes) const;
 
 public:
 #if WITH_EDITOR
@@ -182,9 +226,6 @@ public:
 		return RuntimeNode;
 	}
 
-	void TraversalTest();
-	void TraverseTest(UMovieGraphNode* InNode);
-
 private:
 	void OnVariableUpdated(UMovieGraphVariable* UpdatedVariable);
 
diff --git a/Engine/Plugins/MovieScene/MovieRenderPipeline/Source/MovieRenderPipelineCore/Public/Graph/MovieGraphDataTypes.h b/Engine/Plugins/MovieScene/MovieRenderPipeline/Source/MovieRenderPipelineCore/Public/Graph/MovieGraphDataTypes.h
--- a/Engine/Plugins/MovieScene/MovieRenderPipeline/Source/MovieRenderPipelineCore/Public/Graph/MovieGraphDataTypes.h
+++ b/Engine/Plugins/MovieScene/MovieRenderPipeline/Source/MovieRenderPipelineCore/Public/Graph/MovieGraphDataTypes.h
@@ -101,3 +101,13 @@ protected:
 	UMovieGraphPipeline* GetOwningGraph() const;
 };
 
+struct FMovieGraphRenderPassLayerData
+{
+	FString LayerName;
+};
+
+struct FMovieGraphRenderPassSetupData
+{
+	TWeakObjectPtr<class UMovieGraphDefaultRenderer> Renderer;
+	TArray<FMovieGraphRenderPassLayerData> Layers;
+};
\ No newline at end of file
diff --git a/Engine/Plugins/MovieScene/MovieRenderPipeline/Source/MovieRenderPipelineCore/Public/Graph/MovieGraphDefaultRenderer.h b/Engine/Plugins/MovieScene/MovieRenderPipeline/Source/MovieRenderPipelineCore/Public/Graph/MovieGraphDefaultRenderer.h
--- a/Engine/Plugins/MovieScene/MovieRenderPipeline/Source/MovieRenderPipelineCore/Public/Graph/MovieGraphDefaultRenderer.h
+++ b/Engine/Plugins/MovieScene/MovieRenderPipeline/Source/MovieRenderPipelineCore/Public/Graph/MovieGraphDefaultRenderer.h
@@ -4,6 +4,13 @@
 #include "Graph/MovieGraphDataTypes.h"
 #include "MovieGraphDefaultRenderer.generated.h"
 
+// Forward Declares
+class UMovieGraphRenderPassNode;
+class UTextureRenderTarget2D;
+enum EPixelFormat : uint8;
+
+
+
 /**
 * This class 
 */
@@ -11,8 +18,44 @@ UCLASS(BlueprintType)
 class MOVIERENDERPIPELINECORE_API UMovieGraphDefaultRenderer : public UMovieGraphRendererBase
 {
 	GENERATED_BODY()
+
+public:
+	struct FRenderTargetInitParams
+	{
+		FIntPoint Size;
+		float TargetGamma;
+		EPixelFormat PixelFormat;
+
+		bool operator == (const FRenderTargetInitParams& InRHS) const
+		{
+			return Size == InRHS.Size && TargetGamma == InRHS.TargetGamma && PixelFormat == InRHS.PixelFormat;
+		}
+
+		bool operator != (const FRenderTargetInitParams& InRHS) const
+		{
+			return !(*this == InRHS);
+		}
+
+		friend uint32 GetTypeHash(FRenderTargetInitParams Params)
+		{
+			return HashCombineFast(GetTypeHash(Params.Size), HashCombineFast(GetTypeHash(Params.TargetGamma), GetTypeHash(Params.PixelFormat)));
+		}
+	};
+
 public:
 	virtual void Render(const FMovieGraphTimeStepData& InTimeData) override;
 	virtual void SetupRenderingPipelineForShot(UMoviePipelineExecutorShot* InShot) override;
 	virtual void TeardownRenderingPipelineForShot(UMoviePipelineExecutorShot* InShot) override;
+
+	static void AddReferencedObjects(UObject* InThis, FReferenceCollector& Collector);
+	UTextureRenderTarget2D* GetOrCreateViewRenderTarget(const FRenderTargetInitParams& InInitParams);
+
+protected:
+	TObjectPtr<UTextureRenderTarget2D> CreateViewRenderTarget(const FRenderTargetInitParams& InInitParams) const;
+
+protected:
+	UPROPERTY(Transient)
+	TArray<TObjectPtr<UMovieGraphRenderPassNode>> RenderPassesInUse;
+
+	TMap<FRenderTargetInitParams, TObjectPtr<UTextureRenderTarget2D>> PooledViewRenderTargets;
 };
\ No newline at end of file
diff --git a/Engine/Plugins/MovieScene/MovieRenderPipeline/Source/MovieRenderPipelineCore/Public/Graph/MovieGraphNode.h b/Engine/Plugins/MovieScene/MovieRenderPipeline/Source/MovieRenderPipelineCore/Public/Graph/MovieGraphNode.h
--- a/Engine/Plugins/MovieScene/MovieRenderPipeline/Source/MovieRenderPipelineCore/Public/Graph/MovieGraphNode.h
+++ b/Engine/Plugins/MovieScene/MovieRenderPipeline/Source/MovieRenderPipelineCore/Public/Graph/MovieGraphNode.h
@@ -172,44 +172,6 @@ public:
 #endif
 };
 
-UCLASS()
-class UMoviePipelineRenderLayerNode : public UMovieGraphNode
-{
-	GENERATED_BODY()
-	
-public:
-	virtual TArray<FMovieGraphPinProperties> GetInputPinProperties() const override
-	{
-		TArray<FMovieGraphPinProperties> Properties;
-		Properties.Add(FMovieGraphPinProperties(TEXT("Test Input 3"), false));
-		Properties.Add(FMovieGraphPinProperties(TEXT("Test Input 4"), false));
-		Properties.Add(FMovieGraphPinProperties(TEXT("Test Input 5"), false));
-		Properties.Add(FMovieGraphPinProperties(TEXT("Test Input 6"), false));
-		return Properties;
-	}
-
-	virtual TArray<FMovieGraphPinProperties> GetOutputPinProperties() const override
-	{
-		TArray<FMovieGraphPinProperties> Properties;
-		Properties.Add(FMovieGraphPinProperties(TEXT("Test Output 2"), false));
-		return Properties;
-	}
-
-#if WITH_EDITOR
-	virtual FText GetMenuDescription() const override
-	{
-		return NSLOCTEXT("debug", "collection nodename2", "Deferred Render Layer");
-	}
-	
-	virtual FText GetMenuCategory() const override
-	{
-		return NSLOCTEXT("debug", "collection cat", "Rendering");
-	}
-#endif
-	UPROPERTY(BlueprintReadWrite, EditAnywhere, Category = "Settings2")
-	FString LayerName;
-};
-
 
 UCLASS()
 class MOVIERENDERPIPELINECORE_API UMovieGraphOutputNode : public UMovieGraphNode
diff --git a/Engine/Plugins/MovieScene/MovieRenderPipeline/Source/MovieRenderPipelineCore/Public/Graph/MovieGraphPipeline.h b/Engine/Plugins/MovieScene/MovieRenderPipeline/Source/MovieRenderPipelineCore/Public/Graph/MovieGraphPipeline.h
--- a/Engine/Plugins/MovieScene/MovieRenderPipeline/Source/MovieRenderPipelineCore/Public/Graph/MovieGraphPipeline.h
+++ b/Engine/Plugins/MovieScene/MovieRenderPipeline/Source/MovieRenderPipelineCore/Public/Graph/MovieGraphPipeline.h
@@ -4,6 +4,7 @@
 #include "UObject/Object.h"
 #include "MoviePipelineBase.h"
 #include "MovieGraphDataTypes.h"
+#include "MovieGraphConfig.h"
 #include "MovieRenderPipelineDataTypes.h"
 #include "MovieGraphPipeline.generated.h"
 
@@ -26,7 +27,9 @@ public:
 	UMoviePipelineExecutorJob* GetCurrentJob() const { return CurrentJob; }
 
 
-
+public:
+	UMovieGraphConfig* GetRootGraphForShot(UMoviePipelineExecutorShot* InShot) const;
+	FMovieGraphTraversalContext GetTraversalContextForShot(UMoviePipelineExecutorShot* InShot) const;
 protected:
 	UFUNCTION(BlueprintCallable, Category = "Movie Render Pipeline")
 	virtual void OnMoviePipelineFinishedImpl();
diff --git a/Engine/Plugins/MovieScene/MovieRenderPipeline/Source/MovieRenderPipelineCore/Public/Graph/Nodes/MovieGraphDeferredRenderPassNode.h b/Engine/Plugins/MovieScene/MovieRenderPipeline/Source/MovieRenderPipelineCore/Public/Graph/Nodes/MovieGraphDeferredRenderPassNode.h
new file mode 100644
--- /dev/null
+++ b/Engine/Plugins/MovieScene/MovieRenderPipeline/Source/MovieRenderPipelineCore/Public/Graph/Nodes/MovieGraphDeferredRenderPassNode.h
@@ -0,0 +1,94 @@
+// Copyright Epic Games, Inc. All Rights Reserved.
+
+#pragma once
+#include "Graph/Nodes/MovieGraphRenderPassNode.h"
+#include "Graph/MovieGraphDataTypes.h"
+#include "SceneTypes.h"
+#include "Camera/CameraTypes.h"
+#include "MovieRenderPipelineDataTypes.h"
+#include "Engine/EngineTypes.h"
+#include "MovieGraphDeferredRenderPassNode.generated.h"
+
+// Forward Declares
+class UMovieGraphDefaultRenderer;
+struct FMovieGraphRenderPassSetupData;
+struct FMovieGraphRenderPassLayerData;
+
+// For FViewFamilyContextInitData
+class FRenderTarget;
+class UWorld;
+class FSceneView;
+class AActor;
+class FSceneViewFamilyContext;
+
+
+UCLASS()
+class UMovieGraphDeferredRenderPassNode : public UMovieGraphRenderPassNode
+{
+	GENERATED_BODY()
+public:
+#if WITH_EDITOR
+	virtual FText GetMenuDescription() const override
+	{
+		return NSLOCTEXT("MovieGraphNodes", "DeferredRenderPassGraphNode_Description", "Deferred Renderer");
+	}
+#endif
+protected:
+	// UMovieGraphRenderPassNode Interface
+	virtual FString GetRendererNameImpl() const override { return TEXT("Deferred"); }
+	virtual void SetupImpl(const FMovieGraphRenderPassSetupData& InSetupData) override;
+	virtual void TeardownImpl() override;
+	virtual void RenderImpl() override;
+	// ~UMovieGraphRenderPassNode Interface
+
+	static void AddReferencedObjects(UObject* InThis, FReferenceCollector& Collector);
+
+	struct FViewFamilyContextInitData
+	{
+		class FRenderTarget* RenderTarget;
+		class UWorld* World;
+		FMoviePipelineFrameOutputState::FTimeData TimeData;
+		ESceneCaptureSource SceneCaptureSource;
+		bool bWorldIsPaused;
+		float GlobalScreenPercentageFraction;
+		float OverscanFraction;
+		int32 FrameIndex;
+		bool bCameraCut;
+		EAntiAliasingMethod AntiAliasingMethod;
+
+		// Ownership of this will be passed to the FSceneViewFamilyContext
+		FSceneView* View;
+
+		FSceneViewStateInterface* SceneViewStateReference;
+
+		// Camera Setup
+		FMinimalViewInfo MinimalViewInfo;
+		class AActor* ViewActor;
+	};
+
+	struct FMovieGraphDeferredRenderPass
+	{
+	public:
+		void Setup(TWeakObjectPtr<UMovieGraphDefaultRenderer> InRenderer, const FMovieGraphRenderPassLayerData& InLayer);
+		void Teardown();
+		void Render();
+		void AddReferencedObjects(FReferenceCollector& Collector);
+
+	protected:
+		TSharedRef<FSceneViewFamilyContext> AllocateSceneViewFamilyContext(const FViewFamilyContextInitData& InInitData);
+		FSceneView* AllocateSceneView(TSharedPtr<FSceneViewFamilyContext> InViewFamilyContext, FViewFamilyContextInitData& InInitData) const;
+		void ApplyMoviePipelineOverridesToViewFamily(TSharedRef<FSceneViewFamilyContext> InOutFamily, const FViewFamilyContextInitData& InInitData);
+
+
+
+	protected:
+		FMovieGraphRenderPassLayerData LayerData;
+
+		// Scene View history used by the renderer 
+		FSceneViewStateReference SceneViewState;
+
+		TWeakObjectPtr<class UMovieGraphDefaultRenderer> Renderer;
+	};
+
+	TArray<TUniquePtr<FMovieGraphDeferredRenderPass>> CurrentInstances;
+};
\ No newline at end of file
diff --git a/Engine/Plugins/MovieScene/MovieRenderPipeline/Source/MovieRenderPipelineCore/Public/Graph/Nodes/MovieGraphRenderLayerNode.h b/Engine/Plugins/MovieScene/MovieRenderPipeline/Source/MovieRenderPipelineCore/Public/Graph/Nodes/MovieGraphRenderLayerNode.h
new file mode 100644
--- /dev/null
+++ b/Engine/Plugins/MovieScene/MovieRenderPipeline/Source/MovieRenderPipelineCore/Public/Graph/Nodes/MovieGraphRenderLayerNode.h
@@ -0,0 +1,62 @@
+// Copyright Epic Games, Inc. All Rights Reserved.
+
+#pragma once
+#include "Graph/MovieGraphNode.h"
+#include "MovieGraphRenderLayerNode.generated.h"
+
+
+UCLASS()
+class UMovieGraphRenderLayerNode : public UMovieGraphNode
+{
+	GENERATED_BODY()
+public:
+	UMovieGraphRenderLayerNode()
+	{
+		LayerName = TEXT("beauty");
+	}
+
+	FString GetRenderLayerName() const { return LayerName; }
+
+	virtual TArray<FMovieGraphPinProperties> GetInputPinProperties() const override
+	{
+		TArray<FMovieGraphPinProperties> Properties;
+		Properties.Add(FMovieGraphPinProperties(TEXT("Input"), false));
+		return Properties;
+	}
+
+	virtual TArray<FMovieGraphPinProperties> GetOutputPinProperties() const override
+	{
+		TArray<FMovieGraphPinProperties> Properties;
+		Properties.Add(FMovieGraphPinProperties(TEXT("Output"), false));
+		return Properties;
+	}
+
+	virtual TArray<FPropertyBagPropertyDesc> GetDynamicPropertyDescriptions() const
+	{
+		TArray<FPropertyBagPropertyDesc> Properties;
+		FPropertyBagPropertyDesc FloatEditConProperty = FPropertyBagPropertyDesc("bOverride_TestPropName", EPropertyBagPropertyType::Bool);
+		FPropertyBagPropertyDesc FloatProperty = FPropertyBagPropertyDesc("TestPropName", EPropertyBagPropertyType::Float);
+#if WITH_EDITOR
+		FloatEditConProperty.MetaData.Add(FPropertyBagPropertyDescMetaData("InlineEditConditionToggle", "true"));
+		FloatProperty.MetaData.Add(FPropertyBagPropertyDescMetaData("EditCondition", "bOverride_TestPropName"));
+#endif
+		
+		Properties.Add(FloatEditConProperty);
+		Properties.Add(FloatProperty);
+		return Properties;
+	}
+#if WITH_EDITOR
+	virtual FText GetMenuDescription() const override
+	{
+		return NSLOCTEXT("MovieGraphNodes", "RenderLayerGraphNode_Description", "Render Layer");
+	}
+	
+	virtual FText GetMenuCategory() const override
+	{
+		return NSLOCTEXT("MovieGraphNodes", "RenderLayerGraphNode_Category", "Rendering");
+	}
+#endif
+
+	UPROPERTY(BlueprintReadWrite, EditAnywhere, Category = "Settings")
+	FString LayerName;
+};
\ No newline at end of file
diff --git a/Engine/Plugins/MovieScene/MovieRenderPipeline/Source/MovieRenderPipelineCore/Public/Graph/Nodes/MovieGraphRenderPassNode.h b/Engine/Plugins/MovieScene/MovieRenderPipeline/Source/MovieRenderPipelineCore/Public/Graph/Nodes/MovieGraphRenderPassNode.h
new file mode 100644
--- /dev/null
+++ b/Engine/Plugins/MovieScene/MovieRenderPipeline/Source/MovieRenderPipelineCore/Public/Graph/Nodes/MovieGraphRenderPassNode.h
@@ -0,0 +1,52 @@
+// Copyright Epic Games, Inc. All Rights Reserved.
+
+#pragma once
+#include "Graph/MovieGraphNode.h"
+#include "MovieGraphRenderPassNode.generated.h"
+
+// Forward Declare
+struct FMovieGraphRenderPassSetupData;
+
+UCLASS(Abstract)
+class UMovieGraphRenderPassNode : public UMovieGraphNode
+{
+	GENERATED_BODY()
+public:
+	UMovieGraphRenderPassNode()
+	{
+	}
+
+	// UMovieGraphRenderPassNode interface
+	FString GetRendererName() const { return GetRendererNameImpl(); }
+	void Setup(const FMovieGraphRenderPassSetupData& InSetupData) { SetupImpl(InSetupData); }
+	void Teardown() { TeardownImpl(); }
+	void Render() { RenderImpl(); }
+	// ~UMovieGraphRenderPassNode interface
+
+	virtual TArray<FMovieGraphPinProperties> GetInputPinProperties() const override
+	{
+		TArray<FMovieGraphPinProperties> Properties;
+		Properties.Add(FMovieGraphPinProperties(TEXT("Input"), false));
+		return Properties;
+	}
+
+	virtual TArray<FMovieGraphPinProperties> GetOutputPinProperties() const override
+	{
+		TArray<FMovieGraphPinProperties> Properties;
+		Properties.Add(FMovieGraphPinProperties(TEXT("Output"), false));
+		return Properties;
+	}
+
+#if WITH_EDITOR
+	virtual FText GetMenuCategory() const override
+	{
+		return NSLOCTEXT("MovieGraphNodes", "RenderPassGraphNode_Category", "Rendering");
+	}
+#endif
+
+protected:
+	virtual FString GetRendererNameImpl() const { return TEXT("UnnamedRenderPass"); }
+	virtual void SetupImpl(const FMovieGraphRenderPassSetupData& InSetupData) {}
+	virtual void TeardownImpl() {}
+	virtual void RenderImpl() {}
+};
\ No newline at end of file
diff --git a/Engine/Plugins/MovieScene/MovieRenderPipeline/Source/MovieRenderPipelineEditor/Private/Widgets/Graph/SMoviePipelineGraphPanel.cpp b/Engine/Plugins/MovieScene/MovieRenderPipeline/Source/MovieRenderPipelineEditor/Private/Widgets/Graph/SMoviePipelineGraphPanel.cpp
--- a/Engine/Plugins/MovieScene/MovieRenderPipeline/Source/MovieRenderPipelineEditor/Private/Widgets/Graph/SMoviePipelineGraphPanel.cpp
+++ b/Engine/Plugins/MovieScene/MovieRenderPipeline/Source/MovieRenderPipelineEditor/Private/Widgets/Graph/SMoviePipelineGraphPanel.cpp
@@ -196,11 +196,6 @@ void SMoviePipelineGraphPanel::ClearGraphSelection() const
 }
 
 UE_ENABLE_OPTIMIZATION_SHIP
-FReply SMoviePipelineGraphPanel::OnDebugButtonClicked()
-{
-	CurrentGraph->TraversalTest();
-	return FReply::Handled();
-}
 
 void SMoviePipelineGraphPanel::OnSelectedNodesChanged(const TSet<UObject*>& NewSelection)
 {
diff --git a/Engine/Plugins/MovieScene/MovieRenderPipeline/Source/MovieRenderPipelineEditor/Private/Widgets/Graph/SMoviePipelineGraphPanel.h b/Engine/Plugins/MovieScene/MovieRenderPipeline/Source/MovieRenderPipelineEditor/Private/Widgets/Graph/SMoviePipelineGraphPanel.h
--- a/Engine/Plugins/MovieScene/MovieRenderPipeline/Source/MovieRenderPipelineEditor/Private/Widgets/Graph/SMoviePipelineGraphPanel.h
+++ b/Engine/Plugins/MovieScene/MovieRenderPipeline/Source/MovieRenderPipelineEditor/Private/Widgets/Graph/SMoviePipelineGraphPanel.h
@@ -68,7 +68,6 @@ private:
 	void OnNodeDoubleClicked(class UEdGraphNode* Node);
 	void OnNodeTitleCommitted(const FText& NewText, ETextCommit::Type CommitInfo, UEdGraphNode* NodeBeingChanged);
 	
-	FReply OnDebugButtonClicked();
 	TObjectPtr<class UMovieGraphConfig> CurrentGraph;
 private:
 	/** Allocates a transient preset so that the user can use the pipeline without saving it to an asset first. */
