MoviePipeline: Renamed SMoviePipelineGraphPanel to SMovieGraphConfigPanel.

diff --git a/Engine/Plugins/MovieScene/MovieRenderPipeline/Source/MovieRenderPipelineEditor/Private/Graph/MovieGraphAssetToolkit.cpp b/Engine/Plugins/MovieScene/MovieRenderPipeline/Source/MovieRenderPipelineEditor/Private/Graph/MovieGraphAssetToolkit.cpp
--- a/Engine/Plugins/MovieScene/MovieRenderPipeline/Source/MovieRenderPipelineEditor/Private/Graph/MovieGraphAssetToolkit.cpp
+++ b/Engine/Plugins/MovieScene/MovieRenderPipeline/Source/MovieRenderPipelineEditor/Private/Graph/MovieGraphAssetToolkit.cpp
@@ -9,7 +9,7 @@
 #include "SMovieGraphMembersTabContent.h"
 #include "ToolMenus.h"
 #include "Widgets/Docking/SDockTab.h"
-#include "Widgets/Graph/SMoviePipelineGraphPanel.h"
+#include "Widgets/Graph/SMovieGraphConfigPanel.h"
 
 #define LOCTEXT_NAMESPACE "MovieGraphAssetToolkit"
 
diff --git a/Engine/Plugins/MovieScene/MovieRenderPipeline/Source/MovieRenderPipelineEditor/Private/Widgets/Graph/SMovieGraphConfigPanel.cpp b/Engine/Plugins/MovieScene/MovieRenderPipeline/Source/MovieRenderPipelineEditor/Private/Widgets/Graph/SMovieGraphConfigPanel.cpp
new file mode 100644
--- /dev/null
+++ b/Engine/Plugins/MovieScene/MovieRenderPipeline/Source/MovieRenderPipelineEditor/Private/Widgets/Graph/SMovieGraphConfigPanel.cpp
@@ -0,0 +1,646 @@
+// Copyright Epic Games, Inc. All Rights Reserved.
+
+// Movie Pipeline Includes
+#include "Widgets/Graph/SMovieGraphConfigPanel.h"
+#include "Widgets/MoviePipelineWidgetConstants.h"
+#include "Widgets/SMoviePipelineQueueEditor.h"
+#include "Widgets/SMoviePipelineConfigPanel.h"
+#include "MovieRenderPipelineSettings.h"
+#include "MovieRenderPipelineStyle.h"
+#include "MoviePipelineBlueprintLibrary.h"
+#include "Sections/MovieSceneCinematicShotSection.h"
+#include "MoviePipelineQueue.h"
+#include "MoviePipelinePrimaryConfig.h"
+#include "MoviePipelineQueueSubsystem.h"
+
+// Slate Includes
+#include "Widgets/SBoxPanel.h"
+#include "Widgets/Input/SComboBox.h"
+#include "Widgets/Layout/SBorder.h"
+#include "Widgets/Text/STextBlock.h"
+#include "Widgets/Input/SButton.h"
+#include "Widgets/SWindow.h"
+#include "Styling/AppStyle.h"
+#include "Framework/Application/SlateApplication.h"
+#include "Framework/Notifications/NotificationManager.h"
+#include "Widgets/Layout/SSplitter.h"
+#include "Widgets/Layout/SWidgetSwitcher.h"
+#include "Widgets/Notifications/SNotificationList.h"
+#include "Styling/SlateIconFinder.h"
+#include "Widgets/Images/SImage.h"
+#include "Framework/MultiBox/MultiBoxBuilder.h"
+#include "SPrimaryButton.h"
+#include "Framework/Commands/GenericCommands.h"
+
+// ContentBrowser Includes
+#include "IContentBrowserSingleton.h"
+#include "ContentBrowserModule.h"
+
+// Misc
+#include "Editor.h"
+#include "PropertyEditorModule.h"
+#include "IDetailsView.h"
+#include "Modules/ModuleManager.h"
+
+// UnrealEd Includes
+#include "ScopedTransaction.h"
+#include "AssetRegistry/AssetData.h"
+#include "AssetRegistry/AssetRegistryModule.h"
+#include "FileHelpers.h"
+#include "AssetToolsModule.h"
+#include "Misc/FileHelper.h"
+
+#include "Graph/MovieGraphConfig.h"
+#include "Graph/MovieEdGraph.h"
+#include "Graph/MovieEdGraphNode.h"
+#include "Graph/MovieGraphSchema.h"
+#include "Kismet2/BlueprintEditorUtils.h"
+#include "GraphEditor.h"
+
+#define LOCTEXT_NAMESPACE "SMoviePipelineGraphPanel"
+
+UE_DISABLE_OPTIMIZATION_SHIP
+void SMoviePipelineGraphPanel::Construct(const FArguments& InArgs)
+{
+	OnGraphSelectionChangedEvent = InArgs._OnGraphSelectionChanged;
+
+	// Create the child widgets that need to know about our pipeline
+	PipelineQueueEditorWidget = SNew(SMoviePipelineQueueEditor)
+		.OnEditConfigRequested(this, &SMoviePipelineGraphPanel::OnEditJobConfigRequested)
+		.OnPresetChosen(this, &SMoviePipelineGraphPanel::OnJobPresetChosen)
+		.OnJobSelectionChanged(this, &SMoviePipelineGraphPanel::OnSelectionChanged);
+
+
+	{
+		// Automatically select the first job in the queue
+		UMoviePipelineQueueSubsystem* Subsystem = GEditor->GetEditorSubsystem<UMoviePipelineQueueSubsystem>();
+		check(Subsystem);
+
+		TArray<UMoviePipelineExecutorJob*> Jobs;
+		if (Subsystem->GetQueue()->GetJobs().Num() > 0)
+		{
+			Jobs.Add(Subsystem->GetQueue()->GetJobs()[0]);
+		}
+
+		// Go through the UI so it updates the UI selection too and then this will loop back
+		// around to OnSelectionChanged to update ourself.
+		PipelineQueueEditorWidget->SetSelectedJobs(Jobs);
+	}
+
+	SGraphEditor::FGraphEditorEvents InEvents;
+	InEvents.OnSelectionChanged = SGraphEditor::FOnSelectionChanged::CreateSP(this, &SMoviePipelineGraphPanel::OnSelectedNodesChanged);
+	InEvents.OnNodeDoubleClicked = FSingleNodeEvent::CreateSP(this, &SMoviePipelineGraphPanel::OnNodeDoubleClicked);
+	InEvents.OnTextCommitted = FOnNodeTextCommitted::CreateSP(this, &SMoviePipelineGraphPanel::OnNodeTitleCommitted);
+
+	CurrentGraph = InArgs._Graph;
+	
+	UMoviePipelineEdGraph* GraphToEdit = Cast<UMoviePipelineEdGraph>(FBlueprintEditorUtils::CreateNewGraph(CurrentGraph, TEXT("MoviePipelineEdGraph"), UMoviePipelineEdGraph::StaticClass(), UMovieGraphSchema::StaticClass()));
+
+	// Probably not ideal.. USoundCue has a CreateGraph() node that does this (#if WITH_EDITOR) but then requires an interface
+	// for the editor half to avoid the circular dependency.
+	CurrentGraph->PipelineEdGraph = GraphToEdit;
+	GraphToEdit->InitFromRuntimeGraph(CurrentGraph);
+
+
+	const UEdGraphSchema* Schema = GraphToEdit->GetSchema();
+	Schema->CreateDefaultNodesForGraph(*GraphToEdit);
+	MakeEditorCommands();
+
+	ChildSlot
+	[
+		SAssignNew(GraphEditorWidget, SGraphEditor)
+		.IsEditable(true)
+		.GraphToEdit(GraphToEdit)
+		.AdditionalCommands(GraphEditorCommands)
+		.GraphEvents(InEvents)
+	];
+}
+
+void SMoviePipelineGraphPanel::MakeEditorCommands()
+{
+	GraphEditorCommands = MakeShareable(new FUICommandList);
+	{
+		// Editing commands
+		GraphEditorCommands->MapAction(FGenericCommands::Get().SelectAll,
+			FExecuteAction::CreateSP(this, &SMoviePipelineGraphPanel::SelectAllNodes),
+			FCanExecuteAction::CreateSP(this, &SMoviePipelineGraphPanel::CanSelectAllNodes)
+		);
+
+		GraphEditorCommands->MapAction(FGenericCommands::Get().Delete,
+			FExecuteAction::CreateSP(this, &SMoviePipelineGraphPanel::DeleteSelectedNodes),
+			FCanExecuteAction::CreateSP(this, &SMoviePipelineGraphPanel::CanDeleteSelectedNodes)
+		);
+	}
+}
+
+
+void SMoviePipelineGraphPanel::SelectAllNodes()
+{
+	GraphEditorWidget->SelectAllNodes();
+}
+
+
+bool SMoviePipelineGraphPanel::CanSelectAllNodes() const
+{
+	return GraphEditorWidget.IsValid();
+}
+
+bool SMoviePipelineGraphPanel::CanDeleteSelectedNodes() const
+{
+	if (GraphEditorWidget->GetSelectedNodes().IsEmpty())
+	{
+		return false;
+	}
+
+	for (UObject* Object : GraphEditorWidget->GetSelectedNodes())
+	{
+		UEdGraphNode* GraphNode = Cast<UEdGraphNode>(Object);
+		if (GraphNode && !GraphNode->CanUserDeleteNode())
+		{
+			return false;
+		}
+	}
+
+	return true;
+}
+
+void SMoviePipelineGraphPanel::DeleteSelectedNodes()
+{
+	TArray<UMovieGraphNode*> NodesToDelete;
+	for (UObject* Object : GraphEditorWidget->GetSelectedNodes())
+	{
+		UMoviePipelineEdGraphNodeBase* GraphNode = Cast<UMoviePipelineEdGraphNodeBase>(Object);
+		if (GraphNode && GraphNode->CanUserDeleteNode())
+		{
+			NodesToDelete.Add(GraphNode->GetRuntimeNode());
+		}
+	}
+
+	if (NodesToDelete.IsEmpty())
+	{
+		return;
+	}
+
+	UMoviePipelineEdGraph* Graph = Cast<UMoviePipelineEdGraph>(GraphEditorWidget->GetCurrentGraph());
+	if (Graph)
+	{
+		Graph->GetPipelineGraph()->RemoveNodes(NodesToDelete);
+	}
+
+	ClearGraphSelection();
+}
+
+void SMoviePipelineGraphPanel::ClearGraphSelection() const
+{
+	GraphEditorWidget->ClearSelectionSet();
+}
+
+UE_ENABLE_OPTIMIZATION_SHIP
+
+void SMoviePipelineGraphPanel::OnSelectedNodesChanged(const TSet<UObject*>& NewSelection)
+{
+	TArray<UObject*> SelectedObjects;
+	for (UObject* Obj : NewSelection)
+	{
+		// The selected objects will be editor objects. Get the underlying runtime object for each editor object.
+		if (const UMoviePipelineEdGraphNodeBase* NodeBase = Cast<UMoviePipelineEdGraphNodeBase>(Obj))
+		{
+			if (UMovieGraphNode* GraphNode = NodeBase->GetRuntimeNode())
+			{
+				// For variable nodes, select the underlying variable instead. Otherwise, just select the runtime node.
+				if (const UMovieGraphVariableNode* VariableNode = Cast<UMovieGraphVariableNode>(GraphNode))
+				{
+					SelectedObjects.Add(VariableNode->GetVariable());
+				}
+				else
+				{
+					SelectedObjects.Add(GraphNode);
+				}
+			}
+		}
+	}
+
+	OnGraphSelectionChangedEvent.ExecuteIfBound(SelectedObjects);
+}
+
+void SMoviePipelineGraphPanel::OnNodeDoubleClicked(class UEdGraphNode* Node)
+{
+}
+
+void SMoviePipelineGraphPanel::OnNodeTitleCommitted(const FText& NewText, ETextCommit::Type CommitInfo, UEdGraphNode* NodeBeingChanged)
+{
+	if (NodeBeingChanged)
+	{
+		const FScopedTransaction Transaction(NSLOCTEXT("K2_RenameNode", "RenameNode", "Rename Node"));
+		NodeBeingChanged->Modify();
+		NodeBeingChanged->OnRenameNode(NewText.ToString());
+	}
+}
+
+FReply SMoviePipelineGraphPanel::OnRenderLocalRequested()
+{
+	UMoviePipelineQueueSubsystem* Subsystem = GEditor->GetEditorSubsystem<UMoviePipelineQueueSubsystem>();
+	check(Subsystem);
+
+	const UMovieRenderPipelineProjectSettings* ProjectSettings = GetDefault<UMovieRenderPipelineProjectSettings>();
+	TSubclassOf<UMoviePipelineExecutorBase> ExecutorClass = ProjectSettings->DefaultLocalExecutor.TryLoadClass<UMoviePipelineExecutorBase>();
+
+	// OnRenderLocalRequested should only get called if IsRenderLocalEnabled() returns true, meaning there's a valid class.
+	check(ExecutorClass != nullptr);
+	Subsystem->RenderQueueWithExecutor(ExecutorClass);
+	return FReply::Handled();
+}
+
+bool SMoviePipelineGraphPanel::IsRenderLocalEnabled() const
+{
+	UMoviePipelineQueueSubsystem* Subsystem = GEditor->GetEditorSubsystem<UMoviePipelineQueueSubsystem>();
+	check(Subsystem);
+
+	const UMovieRenderPipelineProjectSettings* ProjectSettings = GetDefault<UMovieRenderPipelineProjectSettings>();
+	const bool bHasExecutor = ProjectSettings->DefaultLocalExecutor.TryLoadClass<UMoviePipelineExecutorBase>() != nullptr;
+	const bool bNotRendering = !Subsystem->IsRendering();
+	const bool bConfigWindowIsOpen = WeakEditorWindow.IsValid();
+
+	bool bAtLeastOneJobAvailable = false;
+	for (UMoviePipelineExecutorJob* Job : Subsystem->GetQueue()->GetJobs())
+	{
+		if (!Job->IsConsumed() && Job->IsEnabled())
+		{
+			bAtLeastOneJobAvailable = true;
+			break;
+		}
+	}
+
+	const bool bWorldIsActive = GEditor->IsPlaySessionInProgress();
+	return bHasExecutor && bNotRendering && bAtLeastOneJobAvailable && !bWorldIsActive && !bConfigWindowIsOpen;
+}
+
+FReply SMoviePipelineGraphPanel::OnRenderRemoteRequested()
+{
+	UMoviePipelineQueueSubsystem* Subsystem = GEditor->GetEditorSubsystem<UMoviePipelineQueueSubsystem>();
+	check(Subsystem);
+
+	const UMovieRenderPipelineProjectSettings* ProjectSettings = GetDefault<UMovieRenderPipelineProjectSettings>();
+	TSubclassOf<UMoviePipelineExecutorBase> ExecutorClass = ProjectSettings->DefaultRemoteExecutor.TryLoadClass<UMoviePipelineExecutorBase>();
+
+	// OnRenderRemoteRequested should only get called if IsRenderRemoteEnabled() returns true, meaning there's a valid class.
+	check(ExecutorClass != nullptr);
+
+	Subsystem->RenderQueueWithExecutor(ExecutorClass);
+	return FReply::Handled();
+}
+
+bool SMoviePipelineGraphPanel::IsRenderRemoteEnabled() const
+{
+	UMoviePipelineQueueSubsystem* Subsystem = GEditor->GetEditorSubsystem<UMoviePipelineQueueSubsystem>();
+	check(Subsystem);
+
+	const UMovieRenderPipelineProjectSettings* ProjectSettings = GetDefault<UMovieRenderPipelineProjectSettings>();
+	const bool bHasExecutor = ProjectSettings->DefaultRemoteExecutor.TryLoadClass<UMoviePipelineExecutorBase>() != nullptr;
+	const bool bNotRendering = !Subsystem->IsRendering();
+	const bool bConfigWindowIsOpen = WeakEditorWindow.IsValid();
+
+	bool bAtLeastOneJobAvailable = false;
+	for (UMoviePipelineExecutorJob* Job : Subsystem->GetQueue()->GetJobs())
+	{
+		if (!Job->IsConsumed() && Job->IsEnabled())
+		{
+			bAtLeastOneJobAvailable = true;
+			break;
+		}
+	}
+
+	return bHasExecutor && bNotRendering && bAtLeastOneJobAvailable && !bConfigWindowIsOpen;
+}
+
+void SMoviePipelineGraphPanel::OnJobPresetChosen(TWeakObjectPtr<UMoviePipelineExecutorJob> InJob, TWeakObjectPtr<UMoviePipelineExecutorShot> InShot)
+{
+	UMovieRenderPipelineProjectSettings* ProjectSettings = GetMutableDefault<UMovieRenderPipelineProjectSettings>();
+	if (!InShot.IsValid())
+	{
+		// Store the preset so the next job they make will use it.
+		ProjectSettings->LastPresetOrigin = InJob->GetPresetOrigin();
+	}
+	ProjectSettings->SaveConfig();
+}
+
+void SMoviePipelineGraphPanel::OnEditJobConfigRequested(TWeakObjectPtr<UMoviePipelineExecutorJob> InJob, TWeakObjectPtr<UMoviePipelineExecutorShot> InShot)
+{
+	// Only allow one editor open at once for now.
+	if (WeakEditorWindow.IsValid())
+	{
+		FWidgetPath ExistingWindowPath;
+		if (FSlateApplication::Get().FindPathToWidget(WeakEditorWindow.Pin().ToSharedRef(), ExistingWindowPath, EVisibility::All))
+		{
+			WeakEditorWindow.Pin()->BringToFront();
+			FSlateApplication::Get().SetAllUserFocus(ExistingWindowPath, EFocusCause::SetDirectly);
+		}
+
+		return;
+	}
+
+	TSubclassOf<UMoviePipelineConfigBase> ConfigType;
+	UMoviePipelineConfigBase* BasePreset = nullptr;
+	UMoviePipelineConfigBase* BaseConfig = nullptr;
+	if (InShot.IsValid())
+	{
+		ConfigType = UMoviePipelineShotConfig::StaticClass();
+		BasePreset = InShot->GetShotOverridePresetOrigin();
+		BaseConfig = InShot->GetShotOverrideConfiguration();
+	}
+	else
+	{
+		ConfigType = UMoviePipelinePrimaryConfig::StaticClass();
+		BasePreset = InJob->GetPresetOrigin();
+		BaseConfig = InJob->GetConfiguration();
+	}
+
+	TSharedRef<SWindow> EditorWindow =
+		SNew(SWindow)
+		.ClientSize(FVector2D(700, 600));
+
+	TSharedRef<SMoviePipelineConfigPanel> ConfigEditorPanel =
+		SNew(SMoviePipelineConfigPanel, ConfigType)
+		.Job(InJob)
+		.Shot(InShot)
+		.OnConfigurationModified(this, &SMoviePipelineGraphPanel::OnConfigUpdatedForJob)
+		.OnConfigurationSetToPreset(this, &SMoviePipelineGraphPanel::OnConfigUpdatedForJobToPreset)
+		.BasePreset(BasePreset)
+		.BaseConfig(BaseConfig);
+
+	EditorWindow->SetContent(ConfigEditorPanel);
+
+
+	TSharedPtr<SWindow> ParentWindow = FSlateApplication::Get().FindWidgetWindow(AsShared());
+	if (ParentWindow.IsValid())
+	{
+		FSlateApplication::Get().AddWindowAsNativeChild(EditorWindow, ParentWindow.ToSharedRef());
+	}
+
+	WeakEditorWindow = EditorWindow;
+}
+
+void SMoviePipelineGraphPanel::OnConfigWindowClosed()
+{
+	if (WeakEditorWindow.IsValid())
+	{
+		WeakEditorWindow.Pin()->RequestDestroyWindow();
+	}
+}
+
+void SMoviePipelineGraphPanel::OnConfigUpdatedForJob(TWeakObjectPtr<UMoviePipelineExecutorJob> InJob, TWeakObjectPtr<UMoviePipelineExecutorShot> InShot, UMoviePipelineConfigBase* InConfig)
+{
+	if (InJob.IsValid())
+	{
+		if (InShot.IsValid())
+		{
+			if (UMoviePipelineShotConfig* ShotConfig = Cast<UMoviePipelineShotConfig>(InConfig))
+			{
+				InShot->SetShotOverrideConfiguration(ShotConfig);
+			}
+		}
+		else
+		{
+			if (UMoviePipelinePrimaryConfig* PrimaryConfig = Cast<UMoviePipelinePrimaryConfig>(InConfig))
+			{
+				InJob->SetConfiguration(PrimaryConfig);
+			}
+		}
+	}
+
+	OnConfigWindowClosed();
+}
+
+void SMoviePipelineGraphPanel::OnConfigUpdatedForJobToPreset(TWeakObjectPtr<UMoviePipelineExecutorJob> InJob, TWeakObjectPtr<UMoviePipelineExecutorShot> InShot, UMoviePipelineConfigBase* InConfig)
+{
+	if (InJob.IsValid())
+	{
+		if (InShot.IsValid())
+		{
+			if (UMoviePipelineShotConfig* ShotConfig = Cast<UMoviePipelineShotConfig>(InConfig))
+			{
+				InShot->SetShotOverridePresetOrigin(ShotConfig);
+			}
+		}
+		else
+		{
+			if (UMoviePipelinePrimaryConfig* PrimaryConfig = Cast<UMoviePipelinePrimaryConfig>(InConfig))
+			{
+				InJob->SetPresetOrigin(PrimaryConfig);
+			}
+		}
+	}
+
+	// Store the preset they used as the last set one
+	OnJobPresetChosen(InJob, InShot);
+
+	OnConfigWindowClosed();
+}
+
+void SMoviePipelineGraphPanel::OnSelectionChanged(const TArray<UMoviePipelineExecutorJob*>& InSelectedJobs)
+{
+	TArray<UObject*> Jobs;
+	for (UMoviePipelineExecutorJob* Job : InSelectedJobs)
+	{
+		Jobs.Add(Job);
+	}
+	
+	NumSelectedJobs = InSelectedJobs.Num();
+}
+
+TSharedRef<SWidget> SMoviePipelineGraphPanel::OnGenerateSavedQueuesMenu()
+{
+	FMenuBuilder MenuBuilder(true, nullptr);
+
+	IContentBrowserSingleton& ContentBrowser = FModuleManager::LoadModuleChecked<FContentBrowserModule>("ContentBrowser").Get();
+
+	MenuBuilder.AddMenuEntry(
+		LOCTEXT("SaveAsQueue_Text", "Save As Asset"),
+		LOCTEXT("SaveAsQueue_Tip", "Save the current configuration as a new preset that can be shared between multiple jobs, or imported later as the base of a new configuration."),
+		FSlateIcon(FAppStyle::Get().GetStyleSetName(), "AssetEditor.SaveAsset"),
+		FUIAction(FExecuteAction::CreateSP(this, &SMoviePipelineGraphPanel::OnSaveAsAsset))
+	);
+
+	FAssetPickerConfig AssetPickerConfig;
+	{
+		AssetPickerConfig.SelectionMode = ESelectionMode::Single;
+		AssetPickerConfig.InitialAssetViewType = EAssetViewType::Column;
+		AssetPickerConfig.bFocusSearchBoxWhenOpened = true;
+		AssetPickerConfig.bAllowNullSelection = false;
+		AssetPickerConfig.bShowBottomToolbar = true;
+		AssetPickerConfig.bAutohideSearchBar = false;
+		AssetPickerConfig.bAllowDragging = false;
+		AssetPickerConfig.bCanShowClasses = false;
+		AssetPickerConfig.bShowPathInColumnView = true;
+		AssetPickerConfig.bShowTypeInColumnView = false;
+		AssetPickerConfig.bSortByPathInColumnView = false;
+
+		AssetPickerConfig.AssetShowWarningText = LOCTEXT("NoQueueAssets_Warning", "No Queues Found");
+		AssetPickerConfig.Filter.ClassPaths.Add(UMoviePipelineQueue::StaticClass()->GetClassPathName());
+		AssetPickerConfig.Filter.bRecursiveClasses = true;
+		AssetPickerConfig.OnAssetSelected = FOnAssetSelected::CreateSP(this, &SMoviePipelineGraphPanel::OnImportSavedQueueAssest);
+	}
+
+	MenuBuilder.BeginSection(NAME_None, LOCTEXT("LoadQueue_MenuSection", "Load Queue"));
+	{
+		TSharedRef<SWidget> PresetPicker = SNew(SBox)
+			.MinDesiredWidth(400.f)
+			.MinDesiredHeight(400.f)
+			[
+				ContentBrowser.CreateAssetPicker(AssetPickerConfig)
+			];
+
+		MenuBuilder.AddWidget(PresetPicker, FText(), true, false);
+	}
+	MenuBuilder.EndSection();
+
+	return MenuBuilder.MakeWidget();
+}
+
+bool SMoviePipelineGraphPanel::OpenSaveDialog(const FString& InDefaultPath, const FString& InNewNameSuggestion, FString& OutPackageName)
+{
+	FSaveAssetDialogConfig SaveAssetDialogConfig;
+	{
+		SaveAssetDialogConfig.DefaultPath = InDefaultPath;
+		SaveAssetDialogConfig.DefaultAssetName = InNewNameSuggestion;
+		SaveAssetDialogConfig.AssetClassNames.Add(UMoviePipelineQueue::StaticClass()->GetClassPathName());
+		SaveAssetDialogConfig.ExistingAssetPolicy = ESaveAssetDialogExistingAssetPolicy::AllowButWarn;
+		SaveAssetDialogConfig.DialogTitleOverride = LOCTEXT("SaveQueueAssetDialogTitle", "Save Queue Asset");
+	}
+
+	FContentBrowserModule& ContentBrowserModule = FModuleManager::LoadModuleChecked<FContentBrowserModule>("ContentBrowser");
+	FString SaveObjectPath = ContentBrowserModule.Get().CreateModalSaveAssetDialog(SaveAssetDialogConfig);
+
+	if (!SaveObjectPath.IsEmpty())
+	{
+		OutPackageName = FPackageName::ObjectPathToPackageName(SaveObjectPath);
+		return true;
+	}
+
+	return false;
+}
+
+bool SMoviePipelineGraphPanel::GetSavePresetPackageName(const FString& InExistingName, FString& OutName)
+{
+	UMovieRenderPipelineProjectSettings* ConfigSettings = GetMutableDefault<UMovieRenderPipelineProjectSettings>();
+
+	// determine default package path
+	const FString DefaultSaveDirectory = ConfigSettings->PresetSaveDir.Path;
+
+	FString DialogStartPath;
+	FPackageName::TryConvertFilenameToLongPackageName(DefaultSaveDirectory, DialogStartPath);
+	if (DialogStartPath.IsEmpty())
+	{
+		DialogStartPath = TEXT("/Game");
+	}
+
+	// determine default asset name
+	FString DefaultName = InExistingName;
+
+	FString UniquePackageName;
+	FString UniqueAssetName;
+
+	FAssetToolsModule& AssetToolsModule = FModuleManager::LoadModuleChecked<FAssetToolsModule>("AssetTools");
+	AssetToolsModule.Get().CreateUniqueAssetName(DialogStartPath / DefaultName, TEXT(""), UniquePackageName, UniqueAssetName);
+
+	FString DialogStartName = FPaths::GetCleanFilename(UniqueAssetName);
+
+	FString UserPackageName;
+	FString NewPackageName;
+
+	// get destination for asset
+	bool bFilenameValid = false;
+	while (!bFilenameValid)
+	{
+		if (!OpenSaveDialog(DialogStartPath, DialogStartName, UserPackageName))
+		{
+			return false;
+		}
+
+		NewPackageName = UserPackageName;
+
+		FText OutError;
+		bFilenameValid = FFileHelper::IsFilenameValidForSaving(NewPackageName, OutError);
+	}
+
+	// Update to the last location they saved to so it remembers their settings next time.
+	ConfigSettings->PresetSaveDir.Path = FPackageName::GetLongPackagePath(UserPackageName);
+	ConfigSettings->SaveConfig();
+	OutName = MoveTemp(NewPackageName);
+	return true;
+}
+
+void SMoviePipelineGraphPanel::OnSaveAsAsset()
+{
+	UMoviePipelineQueueSubsystem* Subsystem = GEditor->GetEditorSubsystem<UMoviePipelineQueueSubsystem>();
+	check(Subsystem);
+	UMoviePipelineQueue* CurrentQueue = Subsystem->GetQueue();
+
+	FString PackageName;
+	if (!GetSavePresetPackageName(CurrentQueue->GetName(), PackageName))
+	{
+		return;
+	}
+	
+	// Saving into a new package
+	const FString NewAssetName = FPackageName::GetLongPackageAssetName(PackageName);
+	UPackage* NewPackage = CreatePackage(*PackageName);
+	NewPackage->MarkAsFullyLoaded();
+	UMoviePipelineQueue* DuplicateQueue = DuplicateObject<UMoviePipelineQueue>(CurrentQueue, NewPackage, *NewAssetName);
+
+	if (DuplicateQueue)
+	{
+		DuplicateQueue->SetFlags(RF_Public | RF_Standalone | RF_Transactional);
+
+		FAssetRegistryModule::AssetCreated(DuplicateQueue);
+
+		FEditorFileUtils::EPromptReturnCode PromptReturnCode = FEditorFileUtils::PromptForCheckoutAndSave({ NewPackage }, false, false);
+	}
+}
+
+void SMoviePipelineGraphPanel::OnImportSavedQueueAssest(const FAssetData& InPresetAsset)
+{
+	FSlateApplication::Get().DismissAllMenus();
+
+	UMoviePipelineQueue* SavedQueue = CastChecked<UMoviePipelineQueue>(InPresetAsset.GetAsset());
+	if (SavedQueue)
+	{
+		// Duplicate the queue so we don't start modifying the one in the asset.
+		UMoviePipelineQueueSubsystem* Subsystem = GEditor->GetEditorSubsystem<UMoviePipelineQueueSubsystem>();
+		check(Subsystem);
+
+		Subsystem->GetQueue()->CopyFrom(SavedQueue);
+
+		// Update the shot list in case the stored queue being copied is out of date with the sequence
+		for (UMoviePipelineExecutorJob* Job : Subsystem->GetQueue()->GetJobs())
+		{
+			ULevelSequence* LoadedSequence = Cast<ULevelSequence>(Job->Sequence.TryLoad());
+			if (LoadedSequence)
+			{
+				bool bShotsChanged = false;
+				UMoviePipelineBlueprintLibrary::UpdateJobShotListFromSequence(LoadedSequence, Job, bShotsChanged);
+
+				if (bShotsChanged)
+				{
+					FNotificationInfo Info(LOCTEXT("QueueShotsUpdated", "Shots have changed since the queue was saved, please resave the queue"));
+					Info.ExpireDuration = 5.0f;
+					FSlateNotificationManager::Get().AddNotification(Info)->SetCompletionState(SNotificationItem::CS_Fail);
+				}
+			}
+		}
+
+		// Automatically select the first job in the queue
+		TArray<UMoviePipelineExecutorJob*> Jobs;
+		if (Subsystem->GetQueue()->GetJobs().Num() > 0)
+		{
+			Jobs.Add(Subsystem->GetQueue()->GetJobs()[0]);
+		}
+
+		// Go through the UI so it updates the UI selection too and then this will loop back
+		// around to OnSelectionChanged to update ourself.
+		PipelineQueueEditorWidget->SetSelectedJobs(Jobs);
+	}
+}
+
+
+#undef LOCTEXT_NAMESPACE // SMoviePipelineGraphPanel
\ No newline at end of file
diff --git a/Engine/Plugins/MovieScene/MovieRenderPipeline/Source/MovieRenderPipelineEditor/Private/Widgets/Graph/SMovieGraphConfigPanel.h b/Engine/Plugins/MovieScene/MovieRenderPipeline/Source/MovieRenderPipelineEditor/Private/Widgets/Graph/SMovieGraphConfigPanel.h
new file mode 100644
--- /dev/null
+++ b/Engine/Plugins/MovieScene/MovieRenderPipeline/Source/MovieRenderPipelineEditor/Private/Widgets/Graph/SMovieGraphConfigPanel.h
@@ -0,0 +1,94 @@
+// Copyright Epic Games, Inc. All Rights Reserved.
+
+#pragma once
+
+#include "Widgets/DeclarativeSyntaxSupport.h"
+#include "Widgets/SCompoundWidget.h"
+
+class UMoviePipelineConfigBase;
+class SGraphEditor;
+class SMoviePipelineQueueEditor;
+class SWindow;
+class UMoviePipelineExecutorJob;
+class UMoviePipelineExecutorShot;
+class IDetailsView;
+struct FAssetData;
+class FUICommandList;
+class UMovieGraphNode;
+
+/**
+ * Outermost widget that is used for adding and removing jobs from the Movie Pipeline Queue Subsystem.
+ */
+class SMoviePipelineGraphPanel : public SCompoundWidget
+{
+public:
+	DECLARE_DELEGATE_OneParam(FOnGraphSelectionChanged, TArray<UObject*>);
+	
+	SLATE_BEGIN_ARGS(SMoviePipelineGraphPanel)
+		: _Graph(nullptr)
+
+		{}
+
+		SLATE_EVENT(FOnGraphSelectionChanged, OnGraphSelectionChanged)
+
+		/** The graph that is initially displayed */
+		SLATE_ARGUMENT(class UMovieGraphConfig*, Graph)
+
+	SLATE_END_ARGS()
+
+	void Construct(const FArguments& InArgs);
+
+	/** Clears the selection in the graph widget being displayed. */
+	void ClearGraphSelection() const;
+
+private:
+	FReply OnRenderLocalRequested();
+	bool IsRenderLocalEnabled() const;
+	FReply OnRenderRemoteRequested();
+	bool IsRenderRemoteEnabled() const;
+
+	/** When they want to edit the current configuration for the job */
+	void OnEditJobConfigRequested(TWeakObjectPtr<UMoviePipelineExecutorJob> InJob, TWeakObjectPtr<UMoviePipelineExecutorShot> InShot);
+	/** When an existing preset is chosen for the specified job. */
+	void OnJobPresetChosen(TWeakObjectPtr<UMoviePipelineExecutorJob> InJob, TWeakObjectPtr<UMoviePipelineExecutorShot> InShot);
+	void OnConfigUpdatedForJob(TWeakObjectPtr<UMoviePipelineExecutorJob> InJob, TWeakObjectPtr<UMoviePipelineExecutorShot> InShot, UMoviePipelineConfigBase* InConfig);
+	void OnConfigUpdatedForJobToPreset(TWeakObjectPtr<UMoviePipelineExecutorJob> InJob, TWeakObjectPtr<UMoviePipelineExecutorShot> InShot, UMoviePipelineConfigBase* InConfig);
+	void OnConfigWindowClosed();
+
+	void OnSelectionChanged(const TArray<UMoviePipelineExecutorJob*>& InSelectedJobs);
+	TArray<UMovieGraphNode*> GetSelectedModelNodes() const;
+
+	TSharedRef<SWidget> OnGenerateSavedQueuesMenu();
+	bool OpenSaveDialog(const FString& InDefaultPath, const FString& InNewNameSuggestion, FString& OutPackageName);
+	bool GetSavePresetPackageName(const FString& InExistingName, FString& OutName);
+	void OnSaveAsAsset();
+	void OnImportSavedQueueAssest(const FAssetData& InPresetAsset);
+	
+	void OnSelectedNodesChanged(const TSet<class UObject*>& NewSelection);
+	void OnNodeDoubleClicked(class UEdGraphNode* Node);
+	void OnNodeTitleCommitted(const FText& NewText, ETextCommit::Type CommitInfo, UEdGraphNode* NodeBeingChanged);
+	
+	TObjectPtr<class UMovieGraphConfig> CurrentGraph;
+private:
+	/** Allocates a transient preset so that the user can use the pipeline without saving it to an asset first. */
+	//UMoviePipelineConfigBase* AllocateTransientPreset();
+	void MakeEditorCommands();
+
+	void SelectAllNodes();
+	bool CanSelectAllNodes() const;
+	bool CanDeleteSelectedNodes() const;
+	void DeleteSelectedNodes();
+private:
+	/** The main movie pipeline queue editor widget */
+	TSharedPtr<SMoviePipelineQueueEditor> PipelineQueueEditorWidget;
+
+	/** The graph widget displayed in this panel. */
+	TSharedPtr<SGraphEditor> GraphEditorWidget;
+
+	TWeakPtr<SWindow> WeakEditorWindow;
+	
+	int32 NumSelectedJobs;
+
+	FOnGraphSelectionChanged OnGraphSelectionChangedEvent;
+	TSharedPtr<FUICommandList> GraphEditorCommands;
+};
\ No newline at end of file
diff --git a/Engine/Plugins/MovieScene/MovieRenderPipeline/Source/MovieRenderPipelineEditor/Private/Widgets/Graph/SMoviePipelineGraphPanel.cpp b/Engine/Plugins/MovieScene/MovieRenderPipeline/Source/MovieRenderPipelineEditor/Private/Widgets/Graph/SMoviePipelineGraphPanel.cpp
deleted file mode 100644
diff --git a/Engine/Plugins/MovieScene/MovieRenderPipeline/Source/MovieRenderPipelineEditor/Private/Widgets/Graph/SMoviePipelineGraphPanel.h b/Engine/Plugins/MovieScene/MovieRenderPipeline/Source/MovieRenderPipelineEditor/Private/Widgets/Graph/SMoviePipelineGraphPanel.h
deleted file mode 100644
