CIS: (Hopefully) fixed a Static Analysis warning about uninitialized memory

diff --git a/Engine/Plugins/MovieScene/MovieRenderPipeline/Source/MovieRenderPipelineCore/Public/Graph/MovieGraphDefaultRenderer.h b/Engine/Plugins/MovieScene/MovieRenderPipeline/Source/MovieRenderPipelineCore/Public/Graph/MovieGraphDefaultRenderer.h
--- a/Engine/Plugins/MovieScene/MovieRenderPipeline/Source/MovieRenderPipelineCore/Public/Graph/MovieGraphDefaultRenderer.h
+++ b/Engine/Plugins/MovieScene/MovieRenderPipeline/Source/MovieRenderPipelineCore/Public/Graph/MovieGraphDefaultRenderer.h
@@ -2,14 +2,12 @@
 #pragma once
 
 #include "Graph/MovieGraphDataTypes.h"
+#include "PixelFormat.h"
 #include "MovieGraphDefaultRenderer.generated.h"
 
 // Forward Declares
 class UMovieGraphRenderPassNode;
 class UTextureRenderTarget2D;
-enum EPixelFormat : uint8;
-
-
 
 /**
 * This class 
@@ -22,6 +20,13 @@ class MOVIERENDERPIPELINECORE_API UMovieGraphDefaultRenderer : public UMovieGrap
 public:
 	struct FRenderTargetInitParams
 	{
+		FRenderTargetInitParams()
+			: Size(FIntPoint(0, 0))
+			, TargetGamma(0.f)
+			, PixelFormat(EPixelFormat::PF_Unknown)
+		{
+		}
+
 		FIntPoint Size;
 		float TargetGamma;
 		EPixelFormat PixelFormat;
diff --git a/Engine/Plugins/MovieScene/MovieRenderPipeline/Source/MovieRenderPipelineCore/Public/Graph/Nodes/MovieGraphDeferredRenderPassNode.h b/Engine/Plugins/MovieScene/MovieRenderPipeline/Source/MovieRenderPipelineCore/Public/Graph/Nodes/MovieGraphDeferredRenderPassNode.h
--- a/Engine/Plugins/MovieScene/MovieRenderPipeline/Source/MovieRenderPipelineCore/Public/Graph/Nodes/MovieGraphDeferredRenderPassNode.h
+++ b/Engine/Plugins/MovieScene/MovieRenderPipeline/Source/MovieRenderPipelineCore/Public/Graph/Nodes/MovieGraphDeferredRenderPassNode.h
@@ -45,6 +45,22 @@ protected:
 
 	struct FViewFamilyContextInitData
 	{
+		FViewFamilyContextInitData()
+			: RenderTarget(nullptr)
+			, World(nullptr)
+			, SceneCaptureSource(ESceneCaptureSource::SCS_MAX)
+			, bWorldIsPaused(false)
+			, GlobalScreenPercentageFraction(1.0f)
+			, OverscanFraction(0.f)
+			, FrameIndex(-1)
+			, bCameraCut(false)
+			, AntiAliasingMethod(EAntiAliasingMethod::AAM_None)
+			, View(nullptr)
+			, SceneViewStateReference(nullptr)
+			, ViewActor(nullptr)
+		{
+		}
+
 		class FRenderTarget* RenderTarget;
 		class UWorld* World;
 		FMoviePipelineFrameOutputState::FTimeData TimeData;
