MoviePipeline: Added ability to add input/output members to a graph. Graphs now have three member types: Inputs, Outputs, and Variables.

Adding inputs/outputs is done through the Members panel. Once they're added, they show up on the Inputs or Outputs nodes in the graph.

diff --git a/Engine/Plugins/MovieScene/MovieRenderPipeline/Source/MovieRenderPipelineCore/Private/Graph/MovieGraphConfig.cpp b/Engine/Plugins/MovieScene/MovieRenderPipeline/Source/MovieRenderPipelineCore/Private/Graph/MovieGraphConfig.cpp
--- a/Engine/Plugins/MovieScene/MovieRenderPipeline/Source/MovieRenderPipelineCore/Private/Graph/MovieGraphConfig.cpp
+++ b/Engine/Plugins/MovieScene/MovieRenderPipeline/Source/MovieRenderPipelineCore/Private/Graph/MovieGraphConfig.cpp
@@ -16,6 +16,20 @@ void UMovieGraphVariable::PostEditChangeProperty(FPropertyChangedEvent& Property
 
 	OnMovieGraphVariableChangedDelegate.Broadcast(this);
 }
+
+void UMovieGraphInput::PostEditChangeProperty(FPropertyChangedEvent& PropertyChangedEvent)
+{
+	Super::PostEditChangeProperty(PropertyChangedEvent);
+	
+	OnMovieGraphInputChangedDelegate.Broadcast(this);
+}
+
+void UMovieGraphOutput::PostEditChangeProperty(FPropertyChangedEvent& PropertyChangedEvent)
+{
+	Super::PostEditChangeProperty(PropertyChangedEvent);
+
+	OnMovieGraphOutputChangedDelegate.Broadcast(this);
+}
 #endif
 
 UMovieGraphConfig::UMovieGraphConfig()
@@ -85,14 +99,6 @@ TArray<UMovieGraphNode*> UMovieGraphConfig::TraverseGraph(TSubclassOf<UMovieGrap
 	return OutNodes;
 }
 
-void UMovieGraphConfig::OnVariableUpdated(UMovieGraphVariable* UpdatedVariable)
-{
-#if WITH_EDITOR
-	OnGraphChangedDelegate.Broadcast();
-	OnGraphVariablesChangedDelegate.Broadcast();
-#endif
-}
-
 bool UMovieGraphConfig::AddLabeledEdge(UMovieGraphNode* FromNode, const FName& FromPinLabel, UMovieGraphNode* ToNode, const FName& ToPinLabel)
 {
 	if (!FromNode || !ToNode)
@@ -285,29 +291,68 @@ bool UMovieGraphConfig::RemoveNode(UMovieGraphNode* InNode)
 	return AllNodes.RemoveSingle(InNode) == 1;
 }
 
-UMovieGraphVariable* UMovieGraphConfig::AddVariable()
+template<typename T>
+T* UMovieGraphConfig::AddMember(TArray<TObjectPtr<T>>& MemberArray, const FText& BaseName)
 {
+	static_assert(std::is_base_of_v<UMovieGraphMember, T>, "T is not derived from UMovieGraphMember");
+	
 	using namespace UE::MoviePipeline::RenderGraph;
 	
-	UMovieGraphVariable* NewVariable = NewObject<UMovieGraphVariable>(this, NAME_None, RF_Transactional);
-	Variables.Add(NewVariable);
+	T* NewMember = NewObject<T>(this, NAME_None, RF_Transactional);
+	MemberArray.Add(NewMember);
 
-	NewVariable->Type = EMovieGraphVariableType::Float;
-	NewVariable->SetGuid(FGuid::NewGuid());
+	NewMember->Type = EMovieGraphMemberType::Float;
+	NewMember->SetGuid(FGuid::NewGuid());
 
 	// Generate and set a unique name
-	TArray<FString> ExistingVariableNames;
-	Algo::Transform(GetVariables(), ExistingVariableNames, [](const UMovieGraphVariable* Variable) { return Variable->Name; });
-	NewVariable->Name = GetUniqueName(ExistingVariableNames, "Variable");
+	TArray<FString> ExistingMemberNames;
+	Algo::Transform(MemberArray, ExistingMemberNames, [](const T* Member) { return Member->Name; });
+	NewMember->Name = GetUniqueName(ExistingMemberNames, BaseName.ToString());
+
+	return NewMember;
+}
+
+UMovieGraphVariable* UMovieGraphConfig::AddVariable()
+{
+	static const FText VariableBaseName = LOCTEXT("VariableBaseName", "Variable");
+	
+	UMovieGraphVariable* NewVariable = AddMember(Variables, VariableBaseName);
 
 #if WITH_EDITOR
-	NewVariable->OnMovieGraphVariableChangedDelegate.AddUObject(this, &UMovieGraphConfig::OnVariableUpdated);
 	OnGraphVariablesChangedDelegate.Broadcast();
 #endif
 
 	return NewVariable;
 }
 
+UMovieGraphInput* UMovieGraphConfig::AddInput()
+{
+	static const FText InputBaseName = LOCTEXT("InputBaseName", "Input");
+
+	UMovieGraphInput* NewInput = AddMember(Inputs, InputBaseName);
+	InputNode->UpdatePins();
+	
+#if WITH_EDITOR
+	OnGraphInputAddedDelegate.Broadcast(NewInput);
+#endif
+
+	return NewInput;
+}
+
+UMovieGraphOutput* UMovieGraphConfig::AddOutput()
+{
+	static const FText OutputBaseName = LOCTEXT("OutputBaseName", "Output");
+	
+	UMovieGraphOutput* NewOutput = AddMember(Outputs, OutputBaseName);
+	OutputNode->UpdatePins();
+
+#if WITH_EDITOR
+	OnGraphOutputAddedDelegate.Broadcast(NewOutput);
+#endif
+
+	return NewOutput;
+}
+
 UMovieGraphVariable* UMovieGraphConfig::GetVariableByGuid(const FGuid& InGuid) const
 {
 	for (const TObjectPtr<UMovieGraphVariable> Variable : Variables)
@@ -326,6 +371,16 @@ TArray<UMovieGraphVariable*> UMovieGraphConfig::GetVariables() const
 	return Variables;
 }
 
+TArray<UMovieGraphInput*> UMovieGraphConfig::GetInputs() const
+{
+	return Inputs;
+}
+
+TArray<UMovieGraphOutput*> UMovieGraphConfig::GetOutputs() const
+{
+	return Outputs;
+}
+
 void UMovieGraphConfig::DeleteMember(UObject* MemberToDelete)
 {
 	if (!MemberToDelete)
@@ -337,6 +392,14 @@ void UMovieGraphConfig::DeleteMember(UObject* MemberToDelete)
 	{
 		DeleteVariableMember(GraphVariableToDelete);
 	}
+	else if (UMovieGraphInput* GraphInputToDelete = Cast<UMovieGraphInput>(MemberToDelete))
+	{
+		DeleteInputMember(GraphInputToDelete);
+	}
+	else if (UMovieGraphOutput* GraphOutputToDelete = Cast<UMovieGraphOutput>(MemberToDelete))
+	{
+		DeleteOutputMember(GraphOutputToDelete);
+	}
 }
 
 void UMovieGraphConfig::DeleteVariableMember(UMovieGraphVariable* VariableMemberToDelete)
@@ -390,4 +453,28 @@ TArray<FMovieGraphBranch> UMovieGraphConfig::GetOutputBranches() const
 	return Branches;
 }
 
+void UMovieGraphConfig::DeleteInputMember(UMovieGraphInput* InputMemberToDelete)
+{
+	if (InputMemberToDelete)
+	{
+		Inputs.RemoveSingle(InputMemberToDelete);
+		RemoveOutboundEdges(InputNode, FName(InputMemberToDelete->Name));
+
+		// This calls OnNodeChangedDelegate to update the graph
+		InputNode->UpdatePins();
+	}
+}
+
+void UMovieGraphConfig::DeleteOutputMember(UMovieGraphOutput* OutputMemberToDelete)
+{
+	if (OutputMemberToDelete)
+	{
+		Outputs.RemoveSingle(OutputMemberToDelete);
+		RemoveInboundEdges(OutputNode, FName(OutputMemberToDelete->Name));
+
+		// This calls OnNodeChangedDelegate to update the graph
+		OutputNode->UpdatePins();
+	}
+}
+
 #undef LOCTEXT_NAMESPACE
diff --git a/Engine/Plugins/MovieScene/MovieRenderPipeline/Source/MovieRenderPipelineCore/Private/Graph/MovieGraphNode.cpp b/Engine/Plugins/MovieScene/MovieRenderPipeline/Source/MovieRenderPipelineCore/Private/Graph/MovieGraphNode.cpp
--- a/Engine/Plugins/MovieScene/MovieRenderPipeline/Source/MovieRenderPipelineCore/Private/Graph/MovieGraphNode.cpp
+++ b/Engine/Plugins/MovieScene/MovieRenderPipeline/Source/MovieRenderPipelineCore/Private/Graph/MovieGraphNode.cpp
@@ -186,6 +186,159 @@ UMovieGraphPin* UMovieGraphNode::GetOutputPin(const FName& Label) const
 	return nullptr;
 }
 
+void UMovieGraphNode::PostLoad()
+{
+	Super::PostLoad();
+
+	RegisterDelegates();
+}
+
+UMovieGraphOutputNode::UMovieGraphOutputNode()
+{
+#if WITH_EDITOR
+	if (!HasAnyFlags(RF_ClassDefaultObject))
+	{
+		if (UMovieGraphConfig* Graph = GetGraph())
+		{
+			// Register delegates for new outputs when they're added to the graph
+			Graph->OnGraphOutputAddedDelegate.AddUObject(this, &UMovieGraphOutputNode::RegisterDelegates);
+		}
+	}
+#endif
+}
+
+TArray<FMovieGraphPinProperties> UMovieGraphOutputNode::GetInputPinProperties() const
+{
+	TArray<FMovieGraphPinProperties> Properties;
+
+	if (const UMovieGraphConfig* ParentGraph = GetGraph())
+	{
+		for (const UMovieGraphOutput* Output : ParentGraph->GetOutputs())
+        {
+        	Properties.Add(FMovieGraphPinProperties(FName(Output->Name), false));
+        }
+	}
+	
+	return Properties;
+}
+
+void UMovieGraphOutputNode::RegisterDelegates() const
+{
+	Super::RegisterDelegates();
+	
+	if (const UMovieGraphConfig* Graph = GetGraph())
+	{
+		for (UMovieGraphOutput* OutputMember : Graph->GetOutputs())
+		{
+			RegisterDelegates(OutputMember);
+		}
+	}
+}
+
+void UMovieGraphOutputNode::RegisterDelegates(UMovieGraphOutput* Output) const
+{
+#if WITH_EDITOR
+	if (Output)
+	{
+		Output->OnMovieGraphOutputChangedDelegate.AddUObject(this, &UMovieGraphOutputNode::UpdateExistingPins);
+	}
+#endif
+}
+
+void UMovieGraphOutputNode::UpdateExistingPins(UMovieGraphMember* ChangedOutput) const
+{
+	if (const UMovieGraphConfig* Graph = GetGraph())
+	{
+		const TArray<UMovieGraphOutput*> OutputMembers = Graph->GetOutputs();
+		if (OutputMembers.Num() == InputPins.Num())
+		{
+			for (int32 Index = 0; Index < OutputMembers.Num(); ++Index)
+			{
+				InputPins[Index]->Properties.Label = FName(OutputMembers[Index]->Name);
+			}
+		}
+
+		OnNodeChangedDelegate.Broadcast(this);
+	}
+}
+
+UMovieGraphInputNode::UMovieGraphInputNode()
+{
+#if WITH_EDITOR
+	if (!HasAnyFlags(RF_ClassDefaultObject))
+	{
+		if (UMovieGraphConfig* Graph = GetGraph())
+		{
+			// Register delegates for new inputs when they're added to the graph
+			Graph->OnGraphInputAddedDelegate.AddUObject(this, &UMovieGraphInputNode::RegisterDelegates);
+		}
+	}
+#endif
+}
+
+TArray<FMovieGraphPinProperties> UMovieGraphInputNode::GetOutputPinProperties() const
+{
+	TArray<FMovieGraphPinProperties> Properties;
+
+	if (const UMovieGraphConfig* ParentGraph = GetGraph())
+	{
+		for (const UMovieGraphInput* Input : ParentGraph->GetInputs())
+		{
+			Properties.Add(FMovieGraphPinProperties(FName(Input->Name), false));
+		}
+	}
+	
+	return Properties;
+}
+
+void UMovieGraphInputNode::RegisterDelegates() const
+{
+	Super::RegisterDelegates();
+	
+	if (const UMovieGraphConfig* Graph = GetGraph())
+	{
+		for (UMovieGraphInput* InputMember : Graph->GetInputs())
+		{
+			RegisterDelegates(InputMember);
+		}
+	}
+}
+
+void UMovieGraphInputNode::RegisterDelegates(UMovieGraphInput* Input) const
+{
+#if WITH_EDITOR
+	if (Input)
+	{
+		Input->OnMovieGraphInputChangedDelegate.AddUObject(this, &UMovieGraphInputNode::UpdateExistingPins);
+	}
+#endif
+}
+
+void UMovieGraphInputNode::UpdateExistingPins(UMovieGraphMember* ChangedInput) const
+{
+	if (const UMovieGraphConfig* Graph = GetGraph())
+	{
+		const TArray<UMovieGraphInput*> InputMembers = Graph->GetInputs();
+		if (InputMembers.Num() == OutputPins.Num())
+		{
+			for (int32 Index = 0; Index < InputMembers.Num(); ++Index)
+			{
+				OutputPins[Index]->Properties.Label = FName(InputMembers[Index]->Name);
+			}
+		}
+		
+		OnNodeChangedDelegate.Broadcast(this);
+	}
+}
+
+UMovieGraphVariableNode::UMovieGraphVariableNode()
+{
+	if (!HasAnyFlags(RF_ClassDefaultObject))
+	{
+		RegisterDelegates();
+	}
+}
+
 TArray<FMovieGraphPinProperties> UMovieGraphVariableNode::GetOutputPinProperties() const
 {
 	TArray<FMovieGraphPinProperties> Properties;
@@ -212,9 +365,7 @@ void UMovieGraphVariableNode::SetVariable(UMovieGraphVariable* InVariable)
 		// (eg, when the variable is renamed)
 		UpdateOutputPin(GraphVariable);
 
-#if WITH_EDITOR
-		GraphVariable->OnMovieGraphVariableChangedDelegate.AddUObject(this, &UMovieGraphVariableNode::UpdateOutputPin);
-#endif
+		RegisterDelegates();
 	}
 }
 
@@ -230,11 +381,25 @@ FText UMovieGraphVariableNode::GetMenuCategory() const
 }
 #endif // WITH_EDITOR
 
-void UMovieGraphVariableNode::UpdateOutputPin(UMovieGraphVariable* ChangedVariable)
+void UMovieGraphVariableNode::RegisterDelegates() const
+{
+	Super::RegisterDelegates();
+
+#if WITH_EDITOR
+	if (GraphVariable)
+	{
+		GraphVariable->OnMovieGraphVariableChangedDelegate.AddUObject(this, &UMovieGraphVariableNode::UpdateOutputPin);
+	}
+#endif
+}
+
+void UMovieGraphVariableNode::UpdateOutputPin(UMovieGraphMember* ChangedVariable) const
 {
 	if (!OutputPins.IsEmpty() && ChangedVariable)
 	{
 		// Update the output pin w/ the name of the variable
 		OutputPins[0]->Properties.Label = FName(ChangedVariable->Name);
 	}
+
+	OnNodeChangedDelegate.Broadcast(this);
 }
\ No newline at end of file
diff --git a/Engine/Plugins/MovieScene/MovieRenderPipeline/Source/MovieRenderPipelineCore/Public/Graph/MovieGraphConfig.h b/Engine/Plugins/MovieScene/MovieRenderPipeline/Source/MovieRenderPipelineCore/Public/Graph/MovieGraphConfig.h
--- a/Engine/Plugins/MovieScene/MovieRenderPipeline/Source/MovieRenderPipelineCore/Public/Graph/MovieGraphConfig.h
+++ b/Engine/Plugins/MovieScene/MovieRenderPipeline/Source/MovieRenderPipelineCore/Public/Graph/MovieGraphConfig.h
@@ -36,9 +36,9 @@ struct FMovieGraphTraversalContext
 	FString RenderLayerName;
 };
 
-/** Types of UMovieGraphVariable that can be used in the graph. */
+/** Types of members that can be used in the graph. */
 UENUM()
-enum class EMovieGraphVariableType : uint8
+enum class EMovieGraphMemberType : uint8
 {
 	Bool,
 	Float,
@@ -47,66 +47,120 @@ enum class EMovieGraphVariableType : uint8
 };
 
 #if WITH_EDITOR
-	DECLARE_MULTICAST_DELEGATE_OneParam(FOnMovieGraphVariableChanged, UMovieGraphVariable*);
-#endif 
+	DECLARE_MULTICAST_DELEGATE_OneParam(FOnMovieGraphVariableChanged, class UMovieGraphMember*);
+	DECLARE_MULTICAST_DELEGATE_OneParam(FOnMovieGraphInputChanged, class UMovieGraphMember*);
+	DECLARE_MULTICAST_DELEGATE_OneParam(FOnMovieGraphOutputChanged, class UMovieGraphMember*);
+#endif
 
-/**
- * A variable that can be created and used inside the graph. These variables can be controlled at
- * the job level.
- */
-UCLASS(BlueprintType)
-class UMovieGraphVariable : public UObject
+UCLASS(Abstract)
+class UMovieGraphMember : public UObject
 {
 	GENERATED_BODY()
 
 public:
-	UMovieGraphVariable() = default;
+	UMovieGraphMember() = default;
 
-	/** Gets the GUID that uniquely identifies this variable. */
+	/** Gets the GUID that uniquely identifies this member. */
 	const FGuid& GetGuid() const { return Guid; }
 
-	/** Sets the GUID that uniquely identifies this variable. */
+	/** Sets the GUID that uniquely identifies this member. */
 	void SetGuid(const FGuid& InGuid) { Guid = InGuid; }
 
 public:
-#if WITH_EDITOR
-	FOnMovieGraphVariableChanged OnMovieGraphVariableChangedDelegate;
-#endif
-
-	/** The type of data stored in this variable. */
+	/** The type of data associated with this member. */
 	UPROPERTY(EditAnywhere, Category = "General")
-	EMovieGraphVariableType Type = EMovieGraphVariableType::Float;
+	EMovieGraphMemberType Type = EMovieGraphMemberType::Float;
 
 	// TODO: Need a details customization that validates whether or not the name is valid/unique
-	/** The name of this variable, which is user-facing. */
+	/** The name of this member, which is user-facing. */
 	UPROPERTY(EditAnywhere, Category = "General")
 	FString Name;
 
-	/** The optional description of this variable, which is user-facing. */
+	/** The optional description of this member, which is user-facing. */
 	UPROPERTY(EditAnywhere, Category = "General")
 	FString Description;
 
 	// TODO: This needs to eventually be capable of storing multiple types. Using TVariant looks promising, but it will
 	// still require extensive details panel customizations
-	/** The default value of this variable. */
+	/** The default value of this member. */
 	UPROPERTY(EditAnywhere, Category = "General")
 	float Default = 0.f;
 
+private:
+	/** A GUID that uniquely identifies this member within its graph. */
+	UPROPERTY()
+	FGuid Guid;
+};
+
+/**
+ * A variable that can be created and used inside the graph. These variables can be controlled at
+ * the job level.
+ */
+UCLASS(BlueprintType)
+class UMovieGraphVariable : public UMovieGraphMember
+{
+	GENERATED_BODY()
+
+public:
+	UMovieGraphVariable() = default;
+
+public:
+#if WITH_EDITOR
+	FOnMovieGraphVariableChanged OnMovieGraphVariableChangedDelegate;
+
 	//~ Begin UObject overrides
+	virtual void PostEditChangeProperty(FPropertyChangedEvent& PropertyChangedEvent) override;
+	//~ End UObject overrides
+#endif // WITH_EDITOR
+};
+
+/**
+ * An input exposed on the graph that will be available for nodes to connect to.
+ */
+UCLASS(BlueprintType)
+class UMovieGraphInput : public UMovieGraphMember
+{
+	GENERATED_BODY()
+
+public:
+	UMovieGraphInput() = default;
+
+public:
 #if WITH_EDITOR
+	FOnMovieGraphInputChanged OnMovieGraphInputChangedDelegate;
+
+	//~ Begin UObject overrides
 	virtual void PostEditChangeProperty(FPropertyChangedEvent& PropertyChangedEvent) override;
-#endif
 	//~ End UObject overrides
+#endif
+};
 
-private:
-	/** A GUID that uniquely identifies this variable within its graph. */
-	UPROPERTY()
-	FGuid Guid;
+/**
+ * An output exposed on the graph that will be available for nodes to connect to.
+ */
+UCLASS(BlueprintType)
+class UMovieGraphOutput : public UMovieGraphMember
+{
+	GENERATED_BODY()
+
+public:
+	UMovieGraphOutput() = default;
+
+public:
+#if WITH_EDITOR
+	FOnMovieGraphOutputChanged OnMovieGraphOutputChangedDelegate;
+
+	//~ Begin UObject overrides
+	virtual void PostEditChangeProperty(FPropertyChangedEvent& PropertyChangedEvent) override;
+	//~ End UObject overrides
+#endif
 };
 
 #if WITH_EDITOR
 	DECLARE_MULTICAST_DELEGATE(FOnMovieGraphChanged);
 	DECLARE_MULTICAST_DELEGATE(FOnMovieGraphVariablesChanged);
+	DECLARE_MULTICAST_DELEGATE_OneParam(FOnMovieGraphInputAdded, UMovieGraphInput*);
+	DECLARE_MULTICAST_DELEGATE_OneParam(FOnMovieGraphOutputAdded, UMovieGraphOutput*);
 	DECLARE_MULTICAST_DELEGATE_OneParam(FOnMovieGraphNodesDeleted, TArray<UMovieGraphNode*>);
 #endif // WITH_EDITOR
 
@@ -134,20 +188,31 @@ public:
 	bool RemoveNode(UMovieGraphNode* InNode);
 	bool RemoveNodes(TArray<UMovieGraphNode*> InNodes);
 
-
 	UMovieGraphNode* GetInputNode() const { return InputNode; }
 	UMovieGraphNode* GetOutputNode() const { return OutputNode; }
 	const TArray<TObjectPtr<UMovieGraphNode>>& GetNodes() const { return AllNodes; }
 
-	/** Adds a new variable with default values to the graph. Returns the new variable on success, else nullptr. */
+	/** Adds a new variable member with default values to the graph. Returns the new variable on success, else nullptr. */
 	UMovieGraphVariable* AddVariable();
 
+	/** Adds a new input member to the graph. Returns the new input on success, else nullptr. */
+	UMovieGraphInput* AddInput();
+
+	/** Adds a new output member to the graph. Returns the new output on success, else nullptr. */
+	UMovieGraphOutput* AddOutput();
+
 	/** Gets the variable in the graph with the specified GUID, else nullptr if one could not be found. */
 	UMovieGraphVariable* GetVariableByGuid(const FGuid& InGuid) const;
 
-	/** Get all variables that have been defined on the graph. */
+	/** Gets all variables that have been defined on the graph. */
 	TArray<UMovieGraphVariable*> GetVariables() const;
 
+	/** Gets all inputs that have been defined on the graph. */
+	TArray<UMovieGraphInput*> GetInputs() const;
+
+	/** Gets all outputs that have been defined on the graph. */
+	TArray<UMovieGraphOutput*> GetOutputs() const;
+
 	/** Remove the specified member (input, output, variable) from the graph. */
 	void DeleteMember(UObject* MemberToDelete);
 
@@ -183,6 +248,11 @@ public:
 
 	TArray<UMovieGraphNode*> TraverseGraph(TSubclassOf<UMovieGraphNode> InClassType, const FMovieGraphTraversalContext& InContext) const;
 
+	/** Remove the specified input member from the graph. */
+	void DeleteInputMember(UMovieGraphInput* InputMemberToDelete);
+
+	/** Remove the specified output member from the graph. */
+	void DeleteOutputMember(UMovieGraphOutput* OutputMemberToDelete);
 
 protected:
 	void TraverseGraphRecursive(UMovieGraphNode* InNode, TSubclassOf<UMovieGraphNode> InClassType, const FMovieGraphTraversalContext& InContext, TArray<UMovieGraphNode*>& OutNodes) const;
@@ -191,6 +261,8 @@ public:
 #if WITH_EDITOR
 	FOnMovieGraphChanged OnGraphChangedDelegate;
 	FOnMovieGraphVariablesChanged OnGraphVariablesChangedDelegate;
+	FOnMovieGraphInputAdded OnGraphInputAddedDelegate;
+	FOnMovieGraphOutputAdded OnGraphOutputAddedDelegate;
 	FOnMovieGraphNodesDeleted OnGraphNodesDeletedDelegate;
 #endif
 	
@@ -227,10 +299,20 @@ public:
 	}
 
 private:
-	void OnVariableUpdated(UMovieGraphVariable* UpdatedVariable);
+	/** Add a new member of type T to MemberArray, with a unique name that includes BaseName in it. */
+	template<typename T>
+	T* AddMember(TArray<TObjectPtr<T>>& MemberArray, const FText& BaseName);
 
 private:
 	/** All variables which have been defined on the graph. */
 	UPROPERTY()
 	TArray<TObjectPtr<UMovieGraphVariable>> Variables;
+
+	/** All inputs which have been defined on the graph. */
+	UPROPERTY()
+	TArray<TObjectPtr<UMovieGraphInput>> Inputs;
+
+	/** All outputs which have been defined on the graph. */
+	UPROPERTY()
+	TArray<TObjectPtr<UMovieGraphOutput>> Outputs;
 };
\ No newline at end of file
diff --git a/Engine/Plugins/MovieScene/MovieRenderPipeline/Source/MovieRenderPipelineCore/Public/Graph/MovieGraphNode.h b/Engine/Plugins/MovieScene/MovieRenderPipeline/Source/MovieRenderPipelineCore/Public/Graph/MovieGraphNode.h
--- a/Engine/Plugins/MovieScene/MovieRenderPipeline/Source/MovieRenderPipelineCore/Public/Graph/MovieGraphNode.h
+++ b/Engine/Plugins/MovieScene/MovieRenderPipeline/Source/MovieRenderPipelineCore/Public/Graph/MovieGraphNode.h
@@ -9,6 +9,9 @@
 
 
 // Forward Declares
+class UMovieGraphInput;
+class UMovieGraphMember;
+class UMovieGraphOutput;
 class UMovieGraphPin;
 class UMovieGraphVariable;
 
@@ -16,7 +19,7 @@ class UMovieGraphVariable;
 class UEdGraphNode;
 #endif
 
-DECLARE_MULTICAST_DELEGATE_OneParam(FOnMovieGraphNodeChanged, UMovieGraphNode*);
+DECLARE_MULTICAST_DELEGATE_OneParam(FOnMovieGraphNodeChanged, const UMovieGraphNode*);
 
 /**
 * This is a base class for all nodes that can exist in the UMovieGraphConfig network.
@@ -76,6 +79,10 @@ public:
 	void SetNodePosY(const int32 InNodePosY) { NodePosY = InNodePosY; }
 #endif
 
+	//~ Begin UObject Interface
+	virtual void PostLoad() override;
+	//~ End UObject Interface
+
 public:
 	FOnMovieGraphNodeChanged OnNodeChangedDelegate;
 
@@ -95,6 +102,9 @@ public:
 protected:
 	virtual TArray<FMovieGraphPinProperties> GetExposedDynamicPinProperties() const;
 
+	/** Register any delegates that need to be set up on the node. Called in PostLoad(). */
+	virtual void RegisterDelegates() const { }
+
 protected:
 	UPROPERTY()
 	TArray<TObjectPtr<UMovieGraphPin>> InputPins;
@@ -172,44 +182,56 @@ public:
 #endif
 };
 
-
+/** A graph node which displays all output members available in the graph. */
 UCLASS()
 class MOVIERENDERPIPELINECORE_API UMovieGraphOutputNode : public UMovieGraphNode
 {
 	GENERATED_BODY()
 
 public:
-	virtual TArray<FMovieGraphPinProperties> GetInputPinProperties() const override
-	{
-		TArray<FMovieGraphPinProperties> Properties;
-		Properties.Add(FMovieGraphPinProperties(TEXT("Output"), false));
-		return Properties;
-	}
+	UMovieGraphOutputNode();
+	
+	virtual TArray<FMovieGraphPinProperties> GetInputPinProperties() const override;
 
 #if WITH_EDITOR
 	virtual FText GetMenuDescription() const override { return NSLOCTEXT("MovieGraphNodes", "OutputNode_Description", "Output"); }
 	virtual FText GetMenuCategory() const override { return NSLOCTEXT("MovieGraphNodes", "OutputNode_Category", "Input/Output"); }
 #endif
+
+private:
+	virtual void RegisterDelegates() const override;
+
+	/** Register delegates for the provided output member. */
+	void RegisterDelegates(UMovieGraphOutput* Output) const;
+
+	/** Update data (name, etc) on all existing input pins on this node to reflect the output members on the graph. */
+	void UpdateExistingPins(UMovieGraphMember* ChangedVariable) const;
 };
 
+/** A graph node which displays all input members available in the graph. */
 UCLASS()
 class MOVIERENDERPIPELINECORE_API UMovieGraphInputNode : public UMovieGraphNode
 {
 	GENERATED_BODY()
 
 public:
-	virtual TArray<FMovieGraphPinProperties> GetOutputPinProperties() const override
-	{
-		TArray<FMovieGraphPinProperties> Properties;
-		Properties.Add(FMovieGraphPinProperties(TEXT("Input"), false));
-		return Properties;
-	}
-
+	UMovieGraphInputNode();
+	
+	virtual TArray<FMovieGraphPinProperties> GetOutputPinProperties() const override;
 
 #if WITH_EDITOR
 	virtual FText GetMenuDescription() const override { return NSLOCTEXT("MovieGraphNodes", "InputNode_Description", "Input"); }
 	virtual FText GetMenuCategory() const override { return NSLOCTEXT("MovieGraphNodes", "InputNode_Category", "Input/Output"); }
 #endif
+
+private:
+	virtual void RegisterDelegates() const override;
+
+	/** Register delegates for the provided input member. */
+	void RegisterDelegates(UMovieGraphInput* Input) const;
+
+	/** Update data (name, etc) on all existing output pins on this node to reflect the input members on the graph. */
+	void UpdateExistingPins(UMovieGraphMember* ChangedVariable) const;
 };
 
 /** A node which gets the value of a variable which has been defined on the graph. */
@@ -219,7 +241,7 @@ class MOVIERENDERPIPELINECORE_API UMovieGraphVariableNode : public UMovieGraphNo
 	GENERATED_BODY()
 
 public:
-	UMovieGraphVariableNode() = default;
+	UMovieGraphVariableNode();
 
 	virtual TArray<FMovieGraphPinProperties> GetOutputPinProperties() const override;
 
@@ -235,8 +257,10 @@ public:
 #endif
 
 private:
+	virtual void RegisterDelegates() const override;
+	
 	/** Updates the output pin on the node to match the provided variable. */
-	void UpdateOutputPin(UMovieGraphVariable* ChangedVariable);
+	void UpdateOutputPin(UMovieGraphMember* ChangedVariable) const;
 
 private:
 	/** The underlying graph variable this node represents. */
diff --git a/Engine/Plugins/MovieScene/MovieRenderPipeline/Source/MovieRenderPipelineEditor/Private/Graph/MovieEdGraphInputNode.cpp b/Engine/Plugins/MovieScene/MovieRenderPipeline/Source/MovieRenderPipelineEditor/Private/Graph/MovieEdGraphInputNode.cpp
--- a/Engine/Plugins/MovieScene/MovieRenderPipeline/Source/MovieRenderPipelineEditor/Private/Graph/MovieEdGraphInputNode.cpp
+++ b/Engine/Plugins/MovieScene/MovieRenderPipeline/Source/MovieRenderPipelineEditor/Private/Graph/MovieEdGraphInputNode.cpp
@@ -3,9 +3,11 @@
 #include "Graph/MovieEdGraphInputNode.h"
 
 #define LOCTEXT_NAMESPACE "MoviePipelineGraph"
+
 FText UMoviePipelineEdGraphNodeInput::GetNodeTitle(ENodeTitleType::Type TitleType) const
 {
-	return FText::FromName(TEXT("Input"));
+	static const FText NodeTitle = LOCTEXT("InputsNodeTitle", "Inputs");
+	return NodeTitle;
 }
 
 void UMoviePipelineEdGraphNodeInput::AllocateDefaultPins()
diff --git a/Engine/Plugins/MovieScene/MovieRenderPipeline/Source/MovieRenderPipelineEditor/Private/Graph/MovieEdGraphNode.cpp b/Engine/Plugins/MovieScene/MovieRenderPipeline/Source/MovieRenderPipelineEditor/Private/Graph/MovieEdGraphNode.cpp
--- a/Engine/Plugins/MovieScene/MovieRenderPipeline/Source/MovieRenderPipelineEditor/Private/Graph/MovieEdGraphNode.cpp
+++ b/Engine/Plugins/MovieScene/MovieRenderPipeline/Source/MovieRenderPipelineEditor/Private/Graph/MovieEdGraphNode.cpp
@@ -240,7 +240,7 @@ void UMoviePipelineEdGraphNodeBase::AutowireNewNode(UEdGraphPin* FromPin)
 	}
 }
 
-void UMoviePipelineEdGraphNodeBase::OnRuntimeNodeChanged(UMovieGraphNode* InChangedNode)
+void UMoviePipelineEdGraphNodeBase::OnRuntimeNodeChanged(const UMovieGraphNode* InChangedNode)
 {
 	if (InChangedNode == GetRuntimeNode())
 	{
diff --git a/Engine/Plugins/MovieScene/MovieRenderPipeline/Source/MovieRenderPipelineEditor/Private/Graph/MovieEdGraphNode.h b/Engine/Plugins/MovieScene/MovieRenderPipeline/Source/MovieRenderPipelineEditor/Private/Graph/MovieEdGraphNode.h
--- a/Engine/Plugins/MovieScene/MovieRenderPipeline/Source/MovieRenderPipelineEditor/Private/Graph/MovieEdGraphNode.h
+++ b/Engine/Plugins/MovieScene/MovieRenderPipeline/Source/MovieRenderPipelineEditor/Private/Graph/MovieEdGraphNode.h
@@ -18,7 +18,7 @@ public:
 
 	void Construct(UMovieGraphNode* InRuntimeNode);
 
-	void OnRuntimeNodeChanged(UMovieGraphNode* InChangedNode);
+	void OnRuntimeNodeChanged(const UMovieGraphNode* InChangedNode);
 
 	//~ Begin UEdGraphNode Interface
 	virtual void ReconstructNode() override;
diff --git a/Engine/Plugins/MovieScene/MovieRenderPipeline/Source/MovieRenderPipelineEditor/Private/Graph/MovieEdGraphOutputNode.cpp b/Engine/Plugins/MovieScene/MovieRenderPipeline/Source/MovieRenderPipelineEditor/Private/Graph/MovieEdGraphOutputNode.cpp
--- a/Engine/Plugins/MovieScene/MovieRenderPipeline/Source/MovieRenderPipelineEditor/Private/Graph/MovieEdGraphOutputNode.cpp
+++ b/Engine/Plugins/MovieScene/MovieRenderPipeline/Source/MovieRenderPipelineEditor/Private/Graph/MovieEdGraphOutputNode.cpp
@@ -6,7 +6,8 @@
 
 FText UMoviePipelineEdGraphNodeOutput::GetNodeTitle(ENodeTitleType::Type TitleType) const
 {
-	return FText::FromName(TEXT("Output"));
+	static const FText NodeTitle = LOCTEXT("OutputsNodeTitle", "Outputs");
+	return NodeTitle;
 }
 
 void UMoviePipelineEdGraphNodeOutput::AllocateDefaultPins()
diff --git a/Engine/Plugins/MovieScene/MovieRenderPipeline/Source/MovieRenderPipelineEditor/Private/Graph/SMovieGraphMembersTabContent.cpp b/Engine/Plugins/MovieScene/MovieRenderPipeline/Source/MovieRenderPipelineEditor/Private/Graph/SMovieGraphMembersTabContent.cpp
--- a/Engine/Plugins/MovieScene/MovieRenderPipeline/Source/MovieRenderPipelineEditor/Private/Graph/SMovieGraphMembersTabContent.cpp
+++ b/Engine/Plugins/MovieScene/MovieRenderPipeline/Source/MovieRenderPipelineEditor/Private/Graph/SMovieGraphMembersTabContent.cpp
@@ -30,6 +30,7 @@ void SMovieGraphMembersTabContent::Construct(const FArguments& InArgs)
 		SAssignNew(ActionMenu, SGraphActionMenu)
 		.OnActionSelected(OnActionSelected)
 		.AutoExpandActionMenu(true)
+		.AlphaSortItems(false)
 		.OnCollectStaticSections(this, &SMovieGraphMembersTabContent::CollectStaticSections)
 		.OnContextMenuOpening(this, &SMovieGraphMembersTabContent::OnContextMenuOpening)
 		.OnGetSectionTitle(this, &SMovieGraphMembersTabContent::GetSectionTitle)
@@ -64,7 +65,7 @@ void SMovieGraphMembersTabContent::DeleteSelectedMembers() const
 		}
 	}
 
-	RefreshActions();
+	RefreshMemberActions();
 }
 
 bool SMovieGraphMembersTabContent::CanDeleteSelectedMembers() const
@@ -84,48 +85,55 @@ void SMovieGraphMembersTabContent::CollectAllActions(FGraphActionListBuilderBase
 	{
 		return;
 	}
-	
-	// TODO: These are placeholder actions and need to be dynamically populated by what is in the graph
 
 	FGraphActionMenuBuilder ActionMenuBuilder;
-	
-	const FText InputActionDesc = LOCTEXT("InputAction", "This is an input");
-	const FText InputActionCategory;
-	const FText InputActionDescription = LOCTEXT("InputActionTooltip", "This is an input tooltip.");
-	const FText InputActionKeywords;
-	const int32 InputActionSectionID = static_cast<int32>(EActionSection::Inputs);
-	const TSharedPtr<FMovieGraphSchemaAction> InputAction(new FMovieGraphSchemaAction(InputActionCategory, InputActionDesc, InputActionDescription, 0, InputActionKeywords, InputActionSectionID));
-
-	const FText OutputActionDesc = LOCTEXT("OutputAction", "This is an output");
-	const FText OutputActionCategory;
-	const FText OutputActionDescription = LOCTEXT("OutputActionTooltip", "This is an output tooltip.");
-	const FText OutputActionKeywords;
-	const int32 OutputActionSectionID = static_cast<int32>(EActionSection::Outputs);
-	const TSharedPtr<FMovieGraphSchemaAction> OutputAction(new FMovieGraphSchemaAction(OutputActionCategory, OutputActionDesc, OutputActionDescription, 0, OutputActionKeywords, OutputActionSectionID));
+
+	// Creates a new action in the action menu under a specific section w/ the provided action target
+	auto AddToActionMenu = [&ActionMenuBuilder](UMovieGraphMember* ActionTarget, const EActionSection Section) -> void
+	{
+		const FText InputActionDesc = FText::FromString(ActionTarget->Name);
+		const FText InputActionCategory;
+		const FText InputActionTooltip;
+		const FText InputActionKeywords;
+		const int32 InputActionSectionID = static_cast<int32>(Section);
+		const TSharedPtr<FMovieGraphSchemaAction> InputAction(new FMovieGraphSchemaAction(InputActionCategory, InputActionDesc, InputActionTooltip, 0, InputActionKeywords, InputActionSectionID));
+		InputAction->ActionTarget = ActionTarget;
+		ActionMenuBuilder.AddAction(InputAction);
+	};
+
+	for (UMovieGraphInput* Input : CurrentGraph->GetInputs())
+	{
+		if (Input)
+		{
+			AddToActionMenu(Input, EActionSection::Inputs);
+            
+            // Update actions when an input is updated (renamed, etc)
+            Input->OnMovieGraphInputChangedDelegate.AddSP(this, &SMovieGraphMembersTabContent::RefreshMemberActions);
+		}
+	}
+
+	for (UMovieGraphOutput* Output : CurrentGraph->GetOutputs())
+	{
+		if (Output)
+		{
+			AddToActionMenu(Output, EActionSection::Outputs);
+
+			// Update actions when an output is updated (renamed, etc)
+			Output->OnMovieGraphOutputChangedDelegate.AddSP(this, &SMovieGraphMembersTabContent::RefreshMemberActions);
+		}
+	}
 
 	for (UMovieGraphVariable* Variable : CurrentGraph->GetVariables())
 	{
-		if (!Variable)
+		if (Variable)
 		{
-			continue;
+			AddToActionMenu(Variable, EActionSection::Variables);
+
+			// Update actions when a variable is updated (renamed, etc)
+			Variable->OnMovieGraphVariableChangedDelegate.AddSP(this, &SMovieGraphMembersTabContent::RefreshMemberActions);
 		}
-		
-		const FText VariableActionDesc = FText::FromString(Variable->Name);
-		const FText VariableActionCategory;
-		const FText VariableActionTooltip;
-		const FText VariableActionKeywords;
-		const int32 VariableActionSectionID = static_cast<int32>(EActionSection::Variables);
-		const TSharedPtr<FMovieGraphSchemaAction> VariableAction(new FMovieGraphSchemaAction(VariableActionCategory, VariableActionDesc, VariableActionTooltip, 0, VariableActionKeywords, VariableActionSectionID));
-		VariableAction->ActionTarget = Variable;
-		ActionMenuBuilder.AddAction(VariableAction);
-
-		// Update actions when a variable is renamed
-		Variable->OnMovieGraphVariableChangedDelegate.AddSP(this, &SMovieGraphMembersTabContent::RefreshActions);
 	}
 	
-	ActionMenuBuilder.AddAction(InputAction);
-	ActionMenuBuilder.AddAction(OutputAction);
-	
 	OutAllActions.Append(ActionMenuBuilder);
 }
 
@@ -212,18 +220,26 @@ TSharedPtr<SWidget> SMovieGraphMembersTabContent::OnContextMenuOpening()
 FReply SMovieGraphMembersTabContent::OnAddButtonClickedOnSection(const int32 InSectionID)
 {
 	const EActionSection Section = static_cast<EActionSection>(InSectionID);
-	
-	if (Section == EActionSection::Variables)
+
+	if (Section == EActionSection::Inputs)
+	{
+		CurrentGraph->AddInput();
+	}
+	else if (Section == EActionSection::Outputs)
+	{
+		CurrentGraph->AddOutput();
+	}
+	else if (Section == EActionSection::Variables)
 	{
 		CurrentGraph->AddVariable();
 	}
 
-	RefreshActions();
+	RefreshMemberActions();
 
 	return FReply::Handled();
 }
 
-void SMovieGraphMembersTabContent::RefreshActions(UMovieGraphVariable* UpdatedVariable) const
+void SMovieGraphMembersTabContent::RefreshMemberActions(UMovieGraphMember* UpdatedMember) const
 {
 	// Currently the entire action menu is refreshed rather than a specific action being targeted
 	
diff --git a/Engine/Plugins/MovieScene/MovieRenderPipeline/Source/MovieRenderPipelineEditor/Private/Graph/SMovieGraphMembersTabContent.h b/Engine/Plugins/MovieScene/MovieRenderPipeline/Source/MovieRenderPipelineEditor/Private/Graph/SMovieGraphMembersTabContent.h
--- a/Engine/Plugins/MovieScene/MovieRenderPipeline/Source/MovieRenderPipelineEditor/Private/Graph/SMovieGraphMembersTabContent.h
+++ b/Engine/Plugins/MovieScene/MovieRenderPipeline/Source/MovieRenderPipelineEditor/Private/Graph/SMovieGraphMembersTabContent.h
@@ -73,8 +73,8 @@ private:
 	/** Handler which deals with clicking the add button for an action section. */
 	FReply OnAddButtonClickedOnSection(const int32 InSectionID);
 
-	/** Refresh/regenerate the action menu when the given variable is updated. */
-	void RefreshActions(class UMovieGraphVariable* UpdatedVariable = nullptr) const;
+	/** Refresh/regenerate the action menu when the given member is updated. */
+	void RefreshMemberActions(class UMovieGraphMember* UpdatedMember = nullptr) const;
 
 private:
 	/** The editor that this widget is associated with. */
