MoviePipeline: Added branch pins, as well as pin coloring, to the graph.

diff --git a/Engine/Plugins/MovieScene/MovieRenderPipeline/Source/MovieRenderPipelineCore/Private/Graph/MovieGraphNode.cpp b/Engine/Plugins/MovieScene/MovieRenderPipeline/Source/MovieRenderPipelineCore/Private/Graph/MovieGraphNode.cpp
--- a/Engine/Plugins/MovieScene/MovieRenderPipeline/Source/MovieRenderPipelineCore/Private/Graph/MovieGraphNode.cpp
+++ b/Engine/Plugins/MovieScene/MovieRenderPipeline/Source/MovieRenderPipelineCore/Private/Graph/MovieGraphNode.cpp
@@ -17,7 +17,7 @@ TArray<FMovieGraphPinProperties> UMovieGraphNode::GetExposedDynamicPinProperties
 	{
 		// Note: Currently hardcoded to not allow multiple connections, but this may need to change in the future
 		const bool bInAllowMultipleConnections = false;
-		Properties.Add(FMovieGraphPinProperties(PropertyName, bInAllowMultipleConnections));
+		Properties.Add(FMovieGraphPinProperties(PropertyName, EMovieGraphMemberType::Float, bInAllowMultipleConnections));
 	}
 
 	return Properties;
@@ -215,7 +215,7 @@ TArray<FMovieGraphPinProperties> UMovieGraphOutputNode::GetInputPinProperties()
 	{
 		for (const UMovieGraphOutput* Output : ParentGraph->GetOutputs())
         {
-        	Properties.Add(FMovieGraphPinProperties(FName(Output->Name), false));
+        	Properties.Add(FMovieGraphPinProperties(FName(Output->Name), EMovieGraphMemberType::Branch, false));
         }
 	}
 	
@@ -284,7 +284,7 @@ TArray<FMovieGraphPinProperties> UMovieGraphInputNode::GetOutputPinProperties()
 	{
 		for (const UMovieGraphInput* Input : ParentGraph->GetInputs())
 		{
-			Properties.Add(FMovieGraphPinProperties(FName(Input->Name), false));
+			Properties.Add(FMovieGraphPinProperties(FName(Input->Name), EMovieGraphMemberType::Branch, false));
 		}
 	}
 	
@@ -345,11 +345,11 @@ TArray<FMovieGraphPinProperties> UMovieGraphVariableNode::GetOutputPinProperties
 	
 	if (GraphVariable)
 	{
-		Properties.Add(FMovieGraphPinProperties(FName(GraphVariable->Name), false));
+		Properties.Add(FMovieGraphPinProperties(FName(GraphVariable->Name), GraphVariable->Type, false));
 	}
 	else
 	{
-		Properties.Add(FMovieGraphPinProperties(TEXT("Unknown"), false));
+		Properties.Add(FMovieGraphPinProperties(TEXT("Unknown"), EMovieGraphMemberType::Float, false));
 	}
 	
 	return Properties;
@@ -397,8 +397,9 @@ void UMovieGraphVariableNode::UpdateOutputPin(UMovieGraphMember* ChangedVariable
 {
 	if (!OutputPins.IsEmpty() && ChangedVariable)
 	{
-		// Update the output pin w/ the name of the variable
+		// Update the output pin to reflect the variable data model
 		OutputPins[0]->Properties.Label = FName(ChangedVariable->Name);
+		OutputPins[0]->Properties.Type = ChangedVariable->Type;
 	}
 
 	OnNodeChangedDelegate.Broadcast(this);
diff --git a/Engine/Plugins/MovieScene/MovieRenderPipeline/Source/MovieRenderPipelineCore/Public/Graph/MovieGraphCommon.h b/Engine/Plugins/MovieScene/MovieRenderPipeline/Source/MovieRenderPipelineCore/Public/Graph/MovieGraphCommon.h
new file mode 100644
--- /dev/null
+++ b/Engine/Plugins/MovieScene/MovieRenderPipeline/Source/MovieRenderPipelineCore/Public/Graph/MovieGraphCommon.h
@@ -0,0 +1,17 @@
+﻿// Copyright Epic Games, Inc. All Rights Reserved.
+
+#pragma once
+
+#include "MovieGraphCommon.generated.h"
+
+/** Types of members that can be used in the graph. */
+UENUM()
+enum class EMovieGraphMemberType : uint8
+{
+	Branch,
+	Bool,
+	Float,
+	Int,
+	IntPoint,
+	String
+};
\ No newline at end of file
diff --git a/Engine/Plugins/MovieScene/MovieRenderPipeline/Source/MovieRenderPipelineCore/Public/Graph/MovieGraphConfig.h b/Engine/Plugins/MovieScene/MovieRenderPipeline/Source/MovieRenderPipelineCore/Public/Graph/MovieGraphConfig.h
--- a/Engine/Plugins/MovieScene/MovieRenderPipeline/Source/MovieRenderPipelineCore/Public/Graph/MovieGraphConfig.h
+++ b/Engine/Plugins/MovieScene/MovieRenderPipeline/Source/MovieRenderPipelineCore/Public/Graph/MovieGraphConfig.h
@@ -36,16 +36,6 @@ struct FMovieGraphTraversalContext
 	FString RenderLayerName;
 };
 
-/** Types of members that can be used in the graph. */
-UENUM()
-enum class EMovieGraphMemberType : uint8
-{
-	Bool,
-	Float,
-	Int,
-	String
-};
-
 #if WITH_EDITOR
 	DECLARE_MULTICAST_DELEGATE_OneParam(FOnMovieGraphVariableChanged, class UMovieGraphMember*);
 	DECLARE_MULTICAST_DELEGATE_OneParam(FOnMovieGraphInputChanged, class UMovieGraphMember*);
diff --git a/Engine/Plugins/MovieScene/MovieRenderPipeline/Source/MovieRenderPipelineCore/Public/Graph/MovieGraphNode.h b/Engine/Plugins/MovieScene/MovieRenderPipeline/Source/MovieRenderPipelineCore/Public/Graph/MovieGraphNode.h
--- a/Engine/Plugins/MovieScene/MovieRenderPipeline/Source/MovieRenderPipelineCore/Public/Graph/MovieGraphNode.h
+++ b/Engine/Plugins/MovieScene/MovieRenderPipeline/Source/MovieRenderPipelineCore/Public/Graph/MovieGraphNode.h
@@ -144,14 +144,14 @@ public:
 	virtual TArray<FMovieGraphPinProperties> GetInputPinProperties() const override
 	{
 		TArray<FMovieGraphPinProperties> Properties;
-		Properties.Add(FMovieGraphPinProperties(TEXT("Input"), false));
+		Properties.Add(FMovieGraphPinProperties(NAME_None, EMovieGraphMemberType::Branch, false));
 		return Properties;
 	}
 
 	virtual TArray<FMovieGraphPinProperties> GetOutputPinProperties() const override
 	{
 		TArray<FMovieGraphPinProperties> Properties;
-		Properties.Add(FMovieGraphPinProperties(TEXT("Output"), false));
+		Properties.Add(FMovieGraphPinProperties(NAME_None, EMovieGraphMemberType::Branch, false));
 		return Properties;
 	}
 
diff --git a/Engine/Plugins/MovieScene/MovieRenderPipeline/Source/MovieRenderPipelineCore/Public/Graph/MovieGraphPin.h b/Engine/Plugins/MovieScene/MovieRenderPipeline/Source/MovieRenderPipelineCore/Public/Graph/MovieGraphPin.h
--- a/Engine/Plugins/MovieScene/MovieRenderPipeline/Source/MovieRenderPipelineCore/Public/Graph/MovieGraphPin.h
+++ b/Engine/Plugins/MovieScene/MovieRenderPipeline/Source/MovieRenderPipelineCore/Public/Graph/MovieGraphPin.h
@@ -2,6 +2,7 @@
 
 #pragma once
 #include "CoreMinimal.h"
+#include "MovieGraphCommon.h"
 #include "MovieGraphPin.generated.h"
 
 // Forward Declares
@@ -14,20 +15,25 @@ struct FMovieGraphPinProperties
 	GENERATED_BODY()
 
 	FMovieGraphPinProperties() = default;
-	explicit FMovieGraphPinProperties(const FName& InLabel, bool bInAllowMultipleConnections)
+	explicit FMovieGraphPinProperties(const FName& InLabel, const EMovieGraphMemberType PinType, bool bInAllowMultipleConnections)
 		: Label(InLabel)
+		, Type(PinType)
 		, bAllowMultipleConnections(bInAllowMultipleConnections)
 	{}
 
 	UPROPERTY(BlueprintReadWrite, EditAnywhere, Category = "Settings")
 	FName Label = NAME_None;
 
+	UPROPERTY(BlueprintReadWrite, EditAnywhere, Category = "Settings")
+	EMovieGraphMemberType Type = EMovieGraphMemberType::Float;
+
 	UPROPERTY(BlueprintReadWrite, EditAnywhere, Category = "Settings")
 	bool bAllowMultipleConnections = true;
 
 	bool operator==(const FMovieGraphPinProperties& Other) const
 	{
 		return Label == Other.Label
+			&& Type == Other.Type
 			&& bAllowMultipleConnections == Other.bAllowMultipleConnections;
 	}
 
diff --git a/Engine/Plugins/MovieScene/MovieRenderPipeline/Source/MovieRenderPipelineCore/Public/Graph/Nodes/MovieGraphRenderLayerNode.h b/Engine/Plugins/MovieScene/MovieRenderPipeline/Source/MovieRenderPipelineCore/Public/Graph/Nodes/MovieGraphRenderLayerNode.h
--- a/Engine/Plugins/MovieScene/MovieRenderPipeline/Source/MovieRenderPipelineCore/Public/Graph/Nodes/MovieGraphRenderLayerNode.h
+++ b/Engine/Plugins/MovieScene/MovieRenderPipeline/Source/MovieRenderPipelineCore/Public/Graph/Nodes/MovieGraphRenderLayerNode.h
@@ -20,14 +20,14 @@ public:
 	virtual TArray<FMovieGraphPinProperties> GetInputPinProperties() const override
 	{
 		TArray<FMovieGraphPinProperties> Properties;
-		Properties.Add(FMovieGraphPinProperties(TEXT("Input"), false));
+		Properties.Add(FMovieGraphPinProperties(NAME_None, EMovieGraphMemberType::Branch, false));
 		return Properties;
 	}
 
 	virtual TArray<FMovieGraphPinProperties> GetOutputPinProperties() const override
 	{
 		TArray<FMovieGraphPinProperties> Properties;
-		Properties.Add(FMovieGraphPinProperties(TEXT("Output"), false));
+		Properties.Add(FMovieGraphPinProperties(NAME_None, EMovieGraphMemberType::Branch, false));
 		return Properties;
 	}
 
diff --git a/Engine/Plugins/MovieScene/MovieRenderPipeline/Source/MovieRenderPipelineCore/Public/Graph/Nodes/MovieGraphRenderPassNode.h b/Engine/Plugins/MovieScene/MovieRenderPipeline/Source/MovieRenderPipelineCore/Public/Graph/Nodes/MovieGraphRenderPassNode.h
--- a/Engine/Plugins/MovieScene/MovieRenderPipeline/Source/MovieRenderPipelineCore/Public/Graph/Nodes/MovieGraphRenderPassNode.h
+++ b/Engine/Plugins/MovieScene/MovieRenderPipeline/Source/MovieRenderPipelineCore/Public/Graph/Nodes/MovieGraphRenderPassNode.h
@@ -26,14 +26,14 @@ public:
 	virtual TArray<FMovieGraphPinProperties> GetInputPinProperties() const override
 	{
 		TArray<FMovieGraphPinProperties> Properties;
-		Properties.Add(FMovieGraphPinProperties(TEXT("Input"), false));
+		Properties.Add(FMovieGraphPinProperties(NAME_None, EMovieGraphMemberType::Branch, false));
 		return Properties;
 	}
 
 	virtual TArray<FMovieGraphPinProperties> GetOutputPinProperties() const override
 	{
 		TArray<FMovieGraphPinProperties> Properties;
-		Properties.Add(FMovieGraphPinProperties(TEXT("Output"), false));
+		Properties.Add(FMovieGraphPinProperties(NAME_None, EMovieGraphMemberType::Branch, false));
 		return Properties;
 	}
 
diff --git a/Engine/Plugins/MovieScene/MovieRenderPipeline/Source/MovieRenderPipelineEditor/Private/Graph/MovieEdGraphConnectionPolicy.cpp b/Engine/Plugins/MovieScene/MovieRenderPipeline/Source/MovieRenderPipelineEditor/Private/Graph/MovieEdGraphConnectionPolicy.cpp
new file mode 100644
--- /dev/null
+++ b/Engine/Plugins/MovieScene/MovieRenderPipeline/Source/MovieRenderPipelineEditor/Private/Graph/MovieEdGraphConnectionPolicy.cpp
@@ -0,0 +1,36 @@
+// Copyright Epic Games, Inc. All Rights Reserved.
+
+#include "Graph/MovieEdGraphConnectionPolicy.h"
+
+#include "EdGraph/EdGraphPin.h"
+#include "Graph/MovieEdGraph.h"
+#include "Graph/MovieGraphSchema.h"
+
+FMovieEdGraphConnectionDrawingPolicy::FMovieEdGraphConnectionDrawingPolicy(int32 InBackLayerID, int32 InFrontLayerID, float ZoomFactor,
+	const FSlateRect& InClippingRect, FSlateWindowElementList& InDrawElements, UEdGraph* InGraphObj)
+	: FConnectionDrawingPolicy(InBackLayerID, InFrontLayerID, ZoomFactor, InClippingRect, InDrawElements)
+	, MovieEdGraph(CastChecked<UMoviePipelineEdGraph>(InGraphObj))
+{
+}
+
+void FMovieEdGraphConnectionDrawingPolicy::DetermineWiringStyle(UEdGraphPin* OutputPin, UEdGraphPin* InputPin, /*inout*/ FConnectionParams& Params)
+{
+	if (!MovieEdGraph)
+	{
+		return;
+	}
+
+	const UMovieGraphSchema* Schema = CastChecked<UMovieGraphSchema>(MovieEdGraph->GetSchema());
+	
+	Params.AssociatedPin1 = OutputPin;
+	Params.AssociatedPin2 = InputPin;
+	Params.WireColor = Schema->GetPinTypeColor(OutputPin->PinType);
+
+	if (OutputPin)
+	{
+		if (OutputPin->PinType.PinCategory == UMovieGraphSchema::PC_Branch)
+		{
+			Params.WireThickness = Settings->DefaultExecutionWireThickness;
+		}
+	}
+}
\ No newline at end of file
diff --git a/Engine/Plugins/MovieScene/MovieRenderPipeline/Source/MovieRenderPipelineEditor/Private/Graph/MovieEdGraphConnectionPolicy.h b/Engine/Plugins/MovieScene/MovieRenderPipeline/Source/MovieRenderPipelineEditor/Private/Graph/MovieEdGraphConnectionPolicy.h
new file mode 100644
--- /dev/null
+++ b/Engine/Plugins/MovieScene/MovieRenderPipeline/Source/MovieRenderPipelineEditor/Private/Graph/MovieEdGraphConnectionPolicy.h
@@ -0,0 +1,22 @@
+﻿// Copyright Epic Games, Inc. All Rights Reserved.
+
+#pragma once
+
+#include "ConnectionDrawingPolicy.h"
+
+class UMoviePipelineEdGraph;
+struct FConnectionParams;
+
+/** Connection drawing policy for the render graph. */
+class FMovieEdGraphConnectionDrawingPolicy : public FConnectionDrawingPolicy
+{
+public:
+	FMovieEdGraphConnectionDrawingPolicy(int32 InBackLayerID, int32 InFrontLayerID, float ZoomFactor, const FSlateRect& InClippingRect, FSlateWindowElementList& InDrawElements, UEdGraph* InGraphObj);
+	
+	//~ FMovieGraphConnectionDrawingPolicy interface 
+	virtual void DetermineWiringStyle(UEdGraphPin* OutputPin, UEdGraphPin* InputPin, /*inout*/ FConnectionParams& Params) override;
+	//~ End FMovieGraphConnectionDrawingPolicy interface
+
+private:
+	UMoviePipelineEdGraph* MovieEdGraph;
+};
\ No newline at end of file
diff --git a/Engine/Plugins/MovieScene/MovieRenderPipeline/Source/MovieRenderPipelineEditor/Private/Graph/MovieEdGraphNode.cpp b/Engine/Plugins/MovieScene/MovieRenderPipeline/Source/MovieRenderPipelineEditor/Private/Graph/MovieEdGraphNode.cpp
--- a/Engine/Plugins/MovieScene/MovieRenderPipeline/Source/MovieRenderPipelineEditor/Private/Graph/MovieEdGraphNode.cpp
+++ b/Engine/Plugins/MovieScene/MovieRenderPipeline/Source/MovieRenderPipelineEditor/Private/Graph/MovieEdGraphNode.cpp
@@ -7,6 +7,7 @@
 #include "EdGraph/EdGraphPin.h"
 #include "Misc/TransactionObjectEvent.h"
 #include "MovieEdGraph.h"
+#include "MovieGraphSchema.h"
 #include "PropertyBag.h"
 #include "ToolMenu.h"
 #include "EdGraph/EdGraphSchema.h"
@@ -43,8 +44,34 @@ FEdGraphPinType UMoviePipelineEdGraphNodeBase::GetPinType(const UMovieGraphPin*
 	FEdGraphPinType EdPinType;
 	EdPinType.ResetToDefaults();
 	
-	EdPinType.PinCategory = FName("TestCategory");
-	EdPinType.PinSubCategory = FName("TestSubCategory");
+	EdPinType.PinCategory = NAME_None;
+	EdPinType.PinSubCategory = NAME_None;
+
+	switch (InPin->Properties.Type)
+	{
+	case EMovieGraphMemberType::Branch:
+		EdPinType.PinCategory = UMovieGraphSchema::PC_Branch;
+		break;
+	case EMovieGraphMemberType::Float:
+		EdPinType.PinCategory = UMovieGraphSchema::PC_Float;
+		break;
+	case EMovieGraphMemberType::Int:
+		EdPinType.PinCategory = UMovieGraphSchema::PC_Integer;
+		break;
+	case EMovieGraphMemberType::Bool:
+		EdPinType.PinCategory = UMovieGraphSchema::PC_Boolean;
+		break;
+	case EMovieGraphMemberType::String:
+		EdPinType.PinCategory = UMovieGraphSchema::PC_String;
+		break;
+	case EMovieGraphMemberType::IntPoint:
+		EdPinType.PinCategory = UMovieGraphSchema::PC_IntPoint;
+		break;
+	default:
+		EdPinType.PinCategory = UMovieGraphSchema::PC_Float;
+		break;
+	}
+	
 	return EdPinType;
 }
 
@@ -176,7 +203,7 @@ void UMoviePipelineEdGraphNodeBase::CreatePins(const TArray<UMovieGraphPin*>& In
 			continue;
 		}
 
-		UEdGraphPin* Pin = CreatePin(EEdGraphPinDirection::EGPD_Input, GetPinType(InputPin), InputPin->Properties.Label);
+		UEdGraphPin* Pin = CreatePin(EGPD_Input, GetPinType(InputPin), InputPin->Properties.Label);
 		// Pin->bAdvancedView = InputPin->Properties.bAdvancedPin;
 		bHasAdvancedPin |= Pin->bAdvancedView;
 	}
@@ -188,7 +215,7 @@ void UMoviePipelineEdGraphNodeBase::CreatePins(const TArray<UMovieGraphPin*>& In
 			continue;
 		}
 
-		UEdGraphPin* Pin = CreatePin(EEdGraphPinDirection::EGPD_Output, GetPinType(OutputPin), OutputPin->Properties.Label);
+		UEdGraphPin* Pin = CreatePin(EGPD_Output, GetPinType(OutputPin), OutputPin->Properties.Label);
 		// Pin->bAdvancedView = OutputPin->Properties.bAdvancedPin;
 		bHasAdvancedPin |= Pin->bAdvancedView;
 	}
diff --git a/Engine/Plugins/MovieScene/MovieRenderPipeline/Source/MovieRenderPipelineEditor/Private/Graph/MovieEdGraphVariableNode.cpp b/Engine/Plugins/MovieScene/MovieRenderPipeline/Source/MovieRenderPipelineEditor/Private/Graph/MovieEdGraphVariableNode.cpp
--- a/Engine/Plugins/MovieScene/MovieRenderPipeline/Source/MovieRenderPipelineEditor/Private/Graph/MovieEdGraphVariableNode.cpp
+++ b/Engine/Plugins/MovieScene/MovieRenderPipeline/Source/MovieRenderPipelineEditor/Private/Graph/MovieEdGraphVariableNode.cpp
@@ -16,13 +16,11 @@ void UMoviePipelineEdGraphVariableNode::AllocateDefaultPins()
 {
 	if (const UMovieGraphVariableNode* VariableNode = Cast<UMovieGraphVariableNode>(RuntimeNode))
 	{
-		FEdGraphPinType EdPinType;
-		EdPinType.ResetToDefaults();
-		
-		EdPinType.PinCategory = FName("Outputs");
-		EdPinType.PinSubCategory = FName();
-		
-		CreatePin(EGPD_Output, EdPinType, FName(VariableNode->GetVariable()->Name));
+		const TArray<TObjectPtr<UMovieGraphPin>>& OutputPins = RuntimeNode->GetOutputPins();
+		if (!OutputPins.IsEmpty())
+		{
+			CreatePin(EGPD_Output, GetPinType(OutputPins[0].Get()), FName(VariableNode->GetVariable()->Name));
+		}
 	}
 }
 
diff --git a/Engine/Plugins/MovieScene/MovieRenderPipeline/Source/MovieRenderPipelineEditor/Private/Graph/MovieGraphPinFactory.cpp b/Engine/Plugins/MovieScene/MovieRenderPipeline/Source/MovieRenderPipelineEditor/Private/Graph/MovieGraphPinFactory.cpp
new file mode 100644
--- /dev/null
+++ b/Engine/Plugins/MovieScene/MovieRenderPipeline/Source/MovieRenderPipelineEditor/Private/Graph/MovieGraphPinFactory.cpp
@@ -0,0 +1,16 @@
+﻿// Copyright Epic Games, Inc. All Rights Reserved.
+
+#include "Graph/MovieGraphPinFactory.h"
+
+#include "KismetPins/SGraphPinExec.h"
+#include "MovieGraphSchema.h"
+
+TSharedPtr<SGraphPin> FMovieGraphPanelPinFactory::CreatePin(UEdGraphPin* InPin) const
+{
+	if (InPin->GetSchema()->IsA<UMovieGraphSchema>() && InPin->PinType.PinCategory == UMovieGraphSchema::PC_Branch)
+	{
+		return SNew(SGraphPinExec, InPin);
+	}
+
+	return nullptr;
+}
\ No newline at end of file
diff --git a/Engine/Plugins/MovieScene/MovieRenderPipeline/Source/MovieRenderPipelineEditor/Private/Graph/MovieGraphPinFactory.h b/Engine/Plugins/MovieScene/MovieRenderPipeline/Source/MovieRenderPipelineEditor/Private/Graph/MovieGraphPinFactory.h
new file mode 100644
--- /dev/null
+++ b/Engine/Plugins/MovieScene/MovieRenderPipeline/Source/MovieRenderPipelineEditor/Private/Graph/MovieGraphPinFactory.h
@@ -0,0 +1,15 @@
+﻿// Copyright Epic Games, Inc. All Rights Reserved.
+
+#pragma once
+
+#include "CoreMinimal.h"
+#include "EdGraphUtilities.h"
+
+/** Pin factory for the render graph. */
+class FMovieGraphPanelPinFactory : public FGraphPanelPinFactory
+{
+public:
+	//~ FGraphPanelPinFactory interface
+	virtual TSharedPtr<class SGraphPin> CreatePin(class UEdGraphPin* InPin) const override;
+	//~ End FGraphPanelPinFactory interface
+};
\ No newline at end of file
diff --git a/Engine/Plugins/MovieScene/MovieRenderPipeline/Source/MovieRenderPipelineEditor/Private/Graph/MovieGraphSchema.cpp b/Engine/Plugins/MovieScene/MovieRenderPipeline/Source/MovieRenderPipelineEditor/Private/Graph/MovieGraphSchema.cpp
--- a/Engine/Plugins/MovieScene/MovieRenderPipeline/Source/MovieRenderPipelineEditor/Private/Graph/MovieGraphSchema.cpp
+++ b/Engine/Plugins/MovieScene/MovieRenderPipeline/Source/MovieRenderPipelineEditor/Private/Graph/MovieGraphSchema.cpp
@@ -4,6 +4,7 @@
 #include "Graph/MovieGraphConfig.h"
 #include "Graph/MovieGraphNode.h"
 #include "Graph/MovieEdGraph.h"
+#include "Graph/MovieEdGraphConnectionPolicy.h"
 #include "Graph/MovieEdGraphNode.h"
 #include "UObject/Class.h"
 #include "UObject/UObjectIterator.h"
@@ -15,6 +16,13 @@ TArray<UClass*> UMovieGraphSchema::MoviePipelineNodeClasses;
 
 #define LOCTEXT_NAMESPACE "MoviePipelineGraphSchema"
 
+const FName UMovieGraphSchema::PC_Branch(TEXT("branch"));
+const FName UMovieGraphSchema::PC_Float(TEXT("float"));
+const FName UMovieGraphSchema::PC_Integer(TEXT("integer"));
+const FName UMovieGraphSchema::PC_Boolean(TEXT("boolean"));
+const FName UMovieGraphSchema::PC_String(TEXT("string"));
+const FName UMovieGraphSchema::PC_IntPoint(TEXT("intpoint"));
+
 void UMovieGraphSchema::CreateDefaultNodesForGraph(UEdGraph& Graph) const
 {
 	/*UMovieGraphConfig* RuntimeGraph = CastChecked<UMoviePipelineEdGraph>(&Graph)->GetPipelineGraph();
@@ -132,6 +140,13 @@ const FPinConnectionResponse UMovieGraphSchema::CanCreateConnection(const UEdGra
 	{
 		return FPinConnectionResponse(CONNECT_RESPONSE_DISALLOW, NSLOCTEXT("MoviePipeline", "CircularPinError", "No Circular Connections!"));
 	}
+
+	// Pins need to be the same type
+	if (PinA->PinType.PinCategory != PinB->PinType.PinCategory)
+	{
+		return FPinConnectionResponse(CONNECT_RESPONSE_DISALLOW, NSLOCTEXT("MoviePipeline", "PinTypeMismatchError", "Pin types don't match!"));
+	}
+	
 	// Make sure the pins are not on the same node
 	return FPinConnectionResponse(CONNECT_RESPONSE_MAKE, NSLOCTEXT("MoviePipeline", "PinConnect", "Connect nodes"));
 }
@@ -216,44 +231,46 @@ void UMovieGraphSchema::BreakSinglePinLink(UEdGraphPin* SourcePin, UEdGraphPin*
 
 FLinearColor UMovieGraphSchema::GetPinTypeColor(const FEdGraphPinType& PinType) const
 {
-	return FLinearColor::Green;
-	/*const FName& TypeName = PinType.PinCategory;
-	if (TypeName == UEdGraphSchema_K2::PC_Struct)
+	const UGraphEditorSettings* Settings = GetDefault<UGraphEditorSettings>();
+
+	if (PinType.PinCategory == PC_Branch)
 	{
-		if (UStruct* Struct = Cast<UStruct>(PinType.PinSubCategoryObject))
-		{
-			if (Struct->IsChildOf(FRigVMExecuteContext::StaticStruct()))
-			{
-				return FLinearColor::White;
-			}
-
-			if (Struct->IsChildOf(RigVMTypeUtils::GetWildCardCPPTypeObject()))
-			{
-				return FLinearColor(FVector3f::OneVector * 0.25f);
-			}
-
-			if (Struct == FRigElementKey::StaticStruct() || Struct == FRigElementKeyCollection::StaticStruct())
-			{
-				return FLinearColor(0.0, 0.6588, 0.9490);
-			}
-
-			if (Struct == FRigElementKey::StaticStruct() || Struct == FRigPose::StaticStruct())
-			{
-				return FLinearColor(0.0, 0.3588, 0.5490);
-			}
-
-			// external types can register their own colors, check if there are any
-			if (IControlRigDeveloperModule* Module = FModuleManager::GetModulePtr<IControlRigDeveloperModule>("ControlRigDeveloper"))
-			{
-				if (const FLinearColor* Color = Module->FindPinTypeColor(Struct))
-				{
-					return *Color;
-				}
-			}
-		}
-	}*/
+		return Settings->ExecutionPinTypeColor;
+	}
+
+	if (PinType.PinCategory == PC_Float)
+	{
+		return Settings->FloatPinTypeColor;
+	}
+
+	if (PinType.PinCategory == PC_Integer)
+	{
+		return Settings->IntPinTypeColor;
+	}
+
+	if (PinType.PinCategory == PC_Boolean)
+	{
+		return Settings->BooleanPinTypeColor;
+	}
+
+	if (PinType.PinCategory == PC_String)
+	{
+		return Settings->StringPinTypeColor;
+	}
 
-	//return GetDefault<UEdGraphSchema_K2>()->GetPinTypeColor(PinType);
+	if (PinType.PinCategory == PC_IntPoint)
+	{
+		return Settings->VectorPinTypeColor;
+	}
+	
+	return Settings->DefaultPinTypeColor;
+}
+
+FConnectionDrawingPolicy* UMovieGraphSchema::CreateConnectionDrawingPolicy(int32 InBackLayerID, int32 InFrontLayerID,
+	float InZoomFactor, const FSlateRect& InClippingRect, FSlateWindowElementList& InDrawElements,
+	UEdGraph* InGraphObj) const
+{
+	return new FMovieEdGraphConnectionDrawingPolicy(InBackLayerID, InFrontLayerID, InZoomFactor, InClippingRect, InDrawElements, InGraphObj);
 }
 
 FMovieGraphSchemaAction_NewNode::FMovieGraphSchemaAction_NewNode(FText InNodeCategory, FText InDisplayName, FText InToolTip)
diff --git a/Engine/Plugins/MovieScene/MovieRenderPipeline/Source/MovieRenderPipelineEditor/Private/Graph/MovieGraphSchema.h b/Engine/Plugins/MovieScene/MovieRenderPipeline/Source/MovieRenderPipelineEditor/Private/Graph/MovieGraphSchema.h
--- a/Engine/Plugins/MovieScene/MovieRenderPipeline/Source/MovieRenderPipelineEditor/Private/Graph/MovieGraphSchema.h
+++ b/Engine/Plugins/MovieScene/MovieRenderPipeline/Source/MovieRenderPipelineEditor/Private/Graph/MovieGraphSchema.h
@@ -13,6 +13,8 @@ UCLASS()
 class UMovieGraphSchema : public UEdGraphSchema
 {
 	GENERATED_BODY()
+
+public:
 	//~ Begin EdGraphSchema Interface
 	virtual void CreateDefaultNodesForGraph(UEdGraph& Graph) const override;
 	virtual void GetGraphContextActions(FGraphContextMenuBuilder& ContextMenuBuilder) const override;
@@ -21,6 +23,7 @@ class UMovieGraphSchema : public UEdGraphSchema
 	virtual bool TryCreateConnection(UEdGraphPin* InA, UEdGraphPin* InB) const override;
 	//virtual const FPinConnectionResponse CanMergeNodes(const UEdGraphNode* A, const UEdGraphNode* B) const override;
 	virtual FLinearColor GetPinTypeColor(const FEdGraphPinType& PinType) const override;
+	virtual FConnectionDrawingPolicy* CreateConnectionDrawingPolicy(int32 InBackLayerID, int32 InFrontLayerID, float InZoomFactor, const FSlateRect& InClippingRect, FSlateWindowElementList& InDrawElements, UEdGraph* InGraphObj) const override;
 
 	//virtual int32 GetNodeSelectionCount(const UEdGraph* Graph) const override;
 	//virtual bool IsCacheVisualizationOutOfDate(int32 InVisualizationCacheID) const override;
@@ -31,11 +34,18 @@ class UMovieGraphSchema : public UEdGraphSchema
 	virtual void BreakSinglePinLink(UEdGraphPin* SourcePin, UEdGraphPin* TargetPin) const override;
 	//~ End EdGraphSchema Interface
 	static void InitMoviePipelineNodeClasses();
-	
+
+public:
+	// Allowed "PinCategory" values for use on EdGraphPin
+	static const FName PC_Branch;
+	static const FName PC_Float;
+	static const FName PC_Integer;
+	static const FName PC_Boolean;
+	static const FName PC_String;
+	static const FName PC_IntPoint;
 	
 private:
 	static TArray<UClass*> MoviePipelineNodeClasses;
-	
 };
 
 /** Base class for schema actions in the graph. */
diff --git a/Engine/Plugins/MovieScene/MovieRenderPipeline/Source/MovieRenderPipelineEditor/Private/MovieRenderPipelineEditorModule.cpp b/Engine/Plugins/MovieScene/MovieRenderPipeline/Source/MovieRenderPipelineEditor/Private/MovieRenderPipelineEditorModule.cpp
--- a/Engine/Plugins/MovieScene/MovieRenderPipeline/Source/MovieRenderPipelineEditor/Private/MovieRenderPipelineEditorModule.cpp
+++ b/Engine/Plugins/MovieScene/MovieRenderPipeline/Source/MovieRenderPipelineEditor/Private/MovieRenderPipelineEditorModule.cpp
@@ -1,6 +1,8 @@
 // Copyright Epic Games, Inc. All Rights Reserved.
 
 #include "MovieRenderPipelineEditorModule.h"
+
+#include "EdGraphUtilities.h"
 #include "MovieRenderPipelineCoreModule.h"
 #include "MoviePipelineExecutor.h"
 #include "WorkspaceMenuStructureModule.h"
@@ -32,6 +34,7 @@
 #include "MoviePipelineConsoleVariableSetting.h"
 #include "PropertyEditorModule.h"
 #include "Customizations/ConsoleVariableCustomization.h"
+#include "Graph/MovieGraphPinFactory.h"
 
 #define LOCTEXT_NAMESPACE "FMovieRenderPipelineEditorModule"
 
@@ -213,6 +216,9 @@ void FMovieRenderPipelineEditorModule::StartupModule()
 	RegisterSettings();
 	RegisterTypeCustomizations();
 	RegisterMovieRenderer();
+
+	GraphPanelPinFactory = MakeShared<FMovieGraphPanelPinFactory>();
+	FEdGraphUtilities::RegisterVisualPinFactory(GraphPanelPinFactory);
 }
 
 void FMovieRenderPipelineEditorModule::ShutdownModule()
@@ -221,6 +227,12 @@ void FMovieRenderPipelineEditorModule::ShutdownModule()
 	UnregisterTypeCustomizations();
 	UnregisterSettings();
 	FMoviePipelineCommands::Unregister();
+	
+	if (GraphPanelPinFactory.IsValid())
+	{
+		FEdGraphUtilities::UnregisterVisualPinFactory(GraphPanelPinFactory);
+		GraphPanelPinFactory.Reset();
+	}
 }
 
 #undef LOCTEXT_NAMESPACE
diff --git a/Engine/Plugins/MovieScene/MovieRenderPipeline/Source/MovieRenderPipelineEditor/Private/MovieRenderPipelineEditorModule.h b/Engine/Plugins/MovieScene/MovieRenderPipeline/Source/MovieRenderPipelineEditor/Private/MovieRenderPipelineEditorModule.h
--- a/Engine/Plugins/MovieScene/MovieRenderPipeline/Source/MovieRenderPipelineEditor/Private/MovieRenderPipelineEditorModule.h
+++ b/Engine/Plugins/MovieScene/MovieRenderPipeline/Source/MovieRenderPipelineEditor/Private/MovieRenderPipelineEditorModule.h
@@ -6,6 +6,8 @@
 #include "IMovieRendererInterface.h"
 #include "Delegates/IDelegateInstance.h"
 
+class FMovieGraphPanelPinFactory;
+
 class FMovieRenderPipelineEditorModule : public IMovieRenderPipelineEditorModule
 {
 public:
@@ -24,4 +26,7 @@ private:
 	void UnregisterTypeCustomizations();
 
 	FDelegateHandle MovieRendererDelegate;
+
+	/** Pin factory for the render graph. */
+	TSharedPtr<FMovieGraphPanelPinFactory> GraphPanelPinFactory;
 };
