MoviePipeline: Added more stub nodes to the graph, as well as icons/colors for all graph nodes which are currently available.

Also moved nodes around, some to their own dedicated files, and some to MovieGraphPlaceholderNodes.h, to clean up MovieGraphConfig.h.

diff --git a/Engine/Plugins/MovieScene/MovieRenderPipeline/Source/MovieRenderPipelineCore/Private/Graph/MovieGraphConfig.cpp b/Engine/Plugins/MovieScene/MovieRenderPipeline/Source/MovieRenderPipelineCore/Private/Graph/MovieGraphConfig.cpp
--- a/Engine/Plugins/MovieScene/MovieRenderPipeline/Source/MovieRenderPipelineCore/Private/Graph/MovieGraphConfig.cpp
+++ b/Engine/Plugins/MovieScene/MovieRenderPipeline/Source/MovieRenderPipelineCore/Private/Graph/MovieGraphConfig.cpp
@@ -4,6 +4,9 @@
 
 #include "Algo/Transform.h"
 #include "Graph/MovieGraphEdge.h"
+#include "Graph/Nodes/MovieGraphInputNode.h"
+#include "Graph/Nodes/MovieGraphOutputNode.h"
+#include "Graph/Nodes/MovieGraphVariableNode.h"
 #include "MovieGraphUtils.h"
 #include "MovieRenderPipelineCoreModule.h"
 
diff --git a/Engine/Plugins/MovieScene/MovieRenderPipeline/Source/MovieRenderPipelineCore/Private/Graph/MovieGraphNode.cpp b/Engine/Plugins/MovieScene/MovieRenderPipeline/Source/MovieRenderPipelineCore/Private/Graph/MovieGraphNode.cpp
--- a/Engine/Plugins/MovieScene/MovieRenderPipeline/Source/MovieRenderPipelineCore/Private/Graph/MovieGraphNode.cpp
+++ b/Engine/Plugins/MovieScene/MovieRenderPipeline/Source/MovieRenderPipelineCore/Private/Graph/MovieGraphNode.cpp
@@ -1,8 +1,8 @@
 // Copyright Epic Games, Inc. All Rights Reserved.
 
 #include "Graph/MovieGraphNode.h"
-#include "Graph/MovieGraphConfig.h"
 
+#include "Graph/MovieGraphConfig.h"
 
 UMovieGraphNode::UMovieGraphNode()
 {
@@ -186,221 +186,22 @@ UMovieGraphPin* UMovieGraphNode::GetOutputPin(const FName& Label) const
 	return nullptr;
 }
 
-void UMovieGraphNode::PostLoad()
-{
-	Super::PostLoad();
-
-	RegisterDelegates();
-}
-
-UMovieGraphOutputNode::UMovieGraphOutputNode()
-{
 #if WITH_EDITOR
-	if (!HasAnyFlags(RF_ClassDefaultObject))
-	{
-		if (UMovieGraphConfig* Graph = GetGraph())
-		{
-			// Register delegates for new outputs when they're added to the graph
-			Graph->OnGraphOutputAddedDelegate.AddUObject(this, &UMovieGraphOutputNode::RegisterDelegates);
-		}
-	}
-#endif
-}
-
-TArray<FMovieGraphPinProperties> UMovieGraphOutputNode::GetInputPinProperties() const
+FLinearColor UMovieGraphNode::GetNodeTitleColor() const
 {
-	TArray<FMovieGraphPinProperties> Properties;
-
-	if (const UMovieGraphConfig* ParentGraph = GetGraph())
-	{
-		for (const UMovieGraphOutput* Output : ParentGraph->GetOutputs())
-        {
-        	Properties.Add(FMovieGraphPinProperties(FName(Output->Name), EMovieGraphMemberType::Branch, false));
-        }
-	}
-	
-	return Properties;
+	return FLinearColor::Black;
 }
 
-void UMovieGraphOutputNode::RegisterDelegates() const
+FSlateIcon UMovieGraphNode::GetIconAndTint(FLinearColor& OutColor) const
 {
-	Super::RegisterDelegates();
-	
-	if (const UMovieGraphConfig* Graph = GetGraph())
-	{
-		for (UMovieGraphOutput* OutputMember : Graph->GetOutputs())
-		{
-			RegisterDelegates(OutputMember);
-		}
-	}
-}
-
-void UMovieGraphOutputNode::RegisterDelegates(UMovieGraphOutput* Output) const
-{
-#if WITH_EDITOR
-	if (Output)
-	{
-		Output->OnMovieGraphOutputChangedDelegate.AddUObject(this, &UMovieGraphOutputNode::UpdateExistingPins);
-	}
-#endif
-}
-
-void UMovieGraphOutputNode::UpdateExistingPins(UMovieGraphMember* ChangedOutput) const
-{
-	if (const UMovieGraphConfig* Graph = GetGraph())
-	{
-		const TArray<UMovieGraphOutput*> OutputMembers = Graph->GetOutputs();
-		if (OutputMembers.Num() == InputPins.Num())
-		{
-			for (int32 Index = 0; Index < OutputMembers.Num(); ++Index)
-			{
-				InputPins[Index]->Properties.Label = FName(OutputMembers[Index]->Name);
-			}
-		}
-
-		OnNodeChangedDelegate.Broadcast(this);
-	}
-}
-
-UMovieGraphInputNode::UMovieGraphInputNode()
-{
-#if WITH_EDITOR
-	if (!HasAnyFlags(RF_ClassDefaultObject))
-	{
-		if (UMovieGraphConfig* Graph = GetGraph())
-		{
-			// Register delegates for new inputs when they're added to the graph
-			Graph->OnGraphInputAddedDelegate.AddUObject(this, &UMovieGraphInputNode::RegisterDelegates);
-		}
-	}
-#endif
-}
-
-TArray<FMovieGraphPinProperties> UMovieGraphInputNode::GetOutputPinProperties() const
-{
-	TArray<FMovieGraphPinProperties> Properties;
-
-	if (const UMovieGraphConfig* ParentGraph = GetGraph())
-	{
-		for (const UMovieGraphInput* Input : ParentGraph->GetInputs())
-		{
-			Properties.Add(FMovieGraphPinProperties(FName(Input->Name), EMovieGraphMemberType::Branch, false));
-		}
-	}
-	
-	return Properties;
-}
-
-void UMovieGraphInputNode::RegisterDelegates() const
-{
-	Super::RegisterDelegates();
-	
-	if (const UMovieGraphConfig* Graph = GetGraph())
-	{
-		for (UMovieGraphInput* InputMember : Graph->GetInputs())
-		{
-			RegisterDelegates(InputMember);
-		}
-	}
-}
-
-void UMovieGraphInputNode::RegisterDelegates(UMovieGraphInput* Input) const
-{
-#if WITH_EDITOR
-	if (Input)
-	{
-		Input->OnMovieGraphInputChangedDelegate.AddUObject(this, &UMovieGraphInputNode::UpdateExistingPins);
-	}
-#endif
-}
-
-void UMovieGraphInputNode::UpdateExistingPins(UMovieGraphMember* ChangedInput) const
-{
-	if (const UMovieGraphConfig* Graph = GetGraph())
-	{
-		const TArray<UMovieGraphInput*> InputMembers = Graph->GetInputs();
-		if (InputMembers.Num() == OutputPins.Num())
-		{
-			for (int32 Index = 0; Index < InputMembers.Num(); ++Index)
-			{
-				OutputPins[Index]->Properties.Label = FName(InputMembers[Index]->Name);
-			}
-		}
-		
-		OnNodeChangedDelegate.Broadcast(this);
-	}
-}
-
-UMovieGraphVariableNode::UMovieGraphVariableNode()
-{
-	if (!HasAnyFlags(RF_ClassDefaultObject))
-	{
-		RegisterDelegates();
-	}
-}
-
-TArray<FMovieGraphPinProperties> UMovieGraphVariableNode::GetOutputPinProperties() const
-{
-	TArray<FMovieGraphPinProperties> Properties;
-	
-	if (GraphVariable)
-	{
-		Properties.Add(FMovieGraphPinProperties(FName(GraphVariable->Name), GraphVariable->Type, false));
-	}
-	else
-	{
-		Properties.Add(FMovieGraphPinProperties(TEXT("Unknown"), EMovieGraphMemberType::Float, false));
-	}
-	
-	return Properties;
-}
-
-void UMovieGraphVariableNode::SetVariable(UMovieGraphVariable* InVariable)
-{
-	if (InVariable)
-	{
-		GraphVariable = InVariable;
-
-		// Update the output pin to reflect the new variable, and update the pin whenever the variable changes
-		// (eg, when the variable is renamed)
-		UpdateOutputPin(GraphVariable);
-
-		RegisterDelegates();
-	}
-}
-
-#if WITH_EDITOR
-FText UMovieGraphVariableNode::GetMenuDescription() const
-{
-	return GraphVariable ? FText::FromString(GraphVariable->Name) : FText();
-}
-	
-FText UMovieGraphVariableNode::GetMenuCategory() const
-{
-	return NSLOCTEXT("MovieGraphNode", "VariableNode_Category", "Variables");
+	OutColor = FLinearColor::White;
+	return FSlateIcon();
 }
 #endif // WITH_EDITOR
 
-void UMovieGraphVariableNode::RegisterDelegates() const
-{
-	Super::RegisterDelegates();
-
-#if WITH_EDITOR
-	if (GraphVariable)
-	{
-		GraphVariable->OnMovieGraphVariableChangedDelegate.AddUObject(this, &UMovieGraphVariableNode::UpdateOutputPin);
-	}
-#endif
-}
-
-void UMovieGraphVariableNode::UpdateOutputPin(UMovieGraphMember* ChangedVariable) const
+void UMovieGraphNode::PostLoad()
 {
-	if (!OutputPins.IsEmpty() && ChangedVariable)
-	{
-		// Update the output pin to reflect the variable data model
-		OutputPins[0]->Properties.Label = FName(ChangedVariable->Name);
-		OutputPins[0]->Properties.Type = ChangedVariable->Type;
-	}
+	Super::PostLoad();
 
-	OnNodeChangedDelegate.Broadcast(this);
+	RegisterDelegates();
 }
\ No newline at end of file
diff --git a/Engine/Plugins/MovieScene/MovieRenderPipeline/Source/MovieRenderPipelineCore/Private/Graph/Nodes/MovieGraphInputNode.cpp b/Engine/Plugins/MovieScene/MovieRenderPipeline/Source/MovieRenderPipelineCore/Private/Graph/Nodes/MovieGraphInputNode.cpp
new file mode 100644
--- /dev/null
+++ b/Engine/Plugins/MovieScene/MovieRenderPipeline/Source/MovieRenderPipelineCore/Private/Graph/Nodes/MovieGraphInputNode.cpp
@@ -0,0 +1,100 @@
+﻿// Copyright Epic Games, Inc. All Rights Reserved.
+
+#include "Graph/Nodes/MovieGraphInputNode.h"
+
+#include "Graph/MovieGraphConfig.h"
+#include "Styling/AppStyle.h"
+
+UMovieGraphInputNode::UMovieGraphInputNode()
+{
+#if WITH_EDITOR
+	if (!HasAnyFlags(RF_ClassDefaultObject))
+	{
+		if (UMovieGraphConfig* Graph = GetGraph())
+		{
+			// Register delegates for new inputs when they're added to the graph
+			Graph->OnGraphInputAddedDelegate.AddUObject(this, &UMovieGraphInputNode::RegisterDelegates);
+		}
+	}
+#endif
+}
+
+TArray<FMovieGraphPinProperties> UMovieGraphInputNode::GetOutputPinProperties() const
+{
+	TArray<FMovieGraphPinProperties> Properties;
+
+	if (const UMovieGraphConfig* ParentGraph = GetGraph())
+	{
+		for (const UMovieGraphInput* Input : ParentGraph->GetInputs())
+		{
+			Properties.Add(FMovieGraphPinProperties(FName(Input->Name), EMovieGraphMemberType::Branch, false));
+		}
+	}
+	
+	return Properties;
+}
+
+#if WITH_EDITOR
+FText UMovieGraphInputNode::GetMenuDescription() const
+{
+	return NSLOCTEXT("MovieGraphNodes", "InputNode_Description", "Input");
+}
+
+FText UMovieGraphInputNode::GetMenuCategory() const
+{
+	return NSLOCTEXT("MovieGraphNodes", "InputNode_Category", "Input/Output");
+}
+
+FLinearColor UMovieGraphInputNode::GetNodeTitleColor() const
+{
+	return FLinearColor::Black;
+}
+
+FSlateIcon UMovieGraphInputNode::GetIconAndTint(FLinearColor& OutColor) const
+{
+	static const FSlateIcon InputIcon = FSlateIcon(FAppStyle::GetAppStyleSetName(), "EditorViewport.VisualizeBufferMode");
+
+	OutColor = FLinearColor::White;
+	return InputIcon;
+}
+#endif // WITH_EDITOR
+
+void UMovieGraphInputNode::RegisterDelegates() const
+{
+	Super::RegisterDelegates();
+	
+	if (const UMovieGraphConfig* Graph = GetGraph())
+	{
+		for (UMovieGraphInput* InputMember : Graph->GetInputs())
+		{
+			RegisterDelegates(InputMember);
+		}
+	}
+}
+
+void UMovieGraphInputNode::RegisterDelegates(UMovieGraphInput* Input) const
+{
+#if WITH_EDITOR
+	if (Input)
+	{
+		Input->OnMovieGraphInputChangedDelegate.AddUObject(this, &UMovieGraphInputNode::UpdateExistingPins);
+	}
+#endif
+}
+
+void UMovieGraphInputNode::UpdateExistingPins(UMovieGraphMember* ChangedInput) const
+{
+	if (const UMovieGraphConfig* Graph = GetGraph())
+	{
+		const TArray<UMovieGraphInput*> InputMembers = Graph->GetInputs();
+		if (InputMembers.Num() == OutputPins.Num())
+		{
+			for (int32 Index = 0; Index < InputMembers.Num(); ++Index)
+			{
+				OutputPins[Index]->Properties.Label = FName(InputMembers[Index]->Name);
+			}
+		}
+		
+		OnNodeChangedDelegate.Broadcast(this);
+	}
+}
\ No newline at end of file
diff --git a/Engine/Plugins/MovieScene/MovieRenderPipeline/Source/MovieRenderPipelineCore/Private/Graph/Nodes/MovieGraphOutputNode.cpp b/Engine/Plugins/MovieScene/MovieRenderPipeline/Source/MovieRenderPipelineCore/Private/Graph/Nodes/MovieGraphOutputNode.cpp
new file mode 100644
--- /dev/null
+++ b/Engine/Plugins/MovieScene/MovieRenderPipeline/Source/MovieRenderPipelineCore/Private/Graph/Nodes/MovieGraphOutputNode.cpp
@@ -0,0 +1,100 @@
+﻿// Copyright Epic Games, Inc. All Rights Reserved.
+
+#include "Graph/Nodes/MovieGraphOutputNode.h"
+
+#include "Graph/MovieGraphConfig.h"
+#include "Styling/AppStyle.h"
+
+UMovieGraphOutputNode::UMovieGraphOutputNode()
+{
+#if WITH_EDITOR
+	if (!HasAnyFlags(RF_ClassDefaultObject))
+	{
+		if (UMovieGraphConfig* Graph = GetGraph())
+		{
+			// Register delegates for new outputs when they're added to the graph
+			Graph->OnGraphOutputAddedDelegate.AddUObject(this, &UMovieGraphOutputNode::RegisterDelegates);
+		}
+	}
+#endif
+}
+
+TArray<FMovieGraphPinProperties> UMovieGraphOutputNode::GetInputPinProperties() const
+{
+	TArray<FMovieGraphPinProperties> Properties;
+
+	if (const UMovieGraphConfig* ParentGraph = GetGraph())
+	{
+		for (const UMovieGraphOutput* Output : ParentGraph->GetOutputs())
+		{
+			Properties.Add(FMovieGraphPinProperties(FName(Output->Name), EMovieGraphMemberType::Branch, false));
+		}
+	}
+	
+	return Properties;
+}
+
+#if WITH_EDITOR
+FText UMovieGraphOutputNode::GetMenuDescription() const
+{
+	return NSLOCTEXT("MovieGraphNodes", "OutputNode_Description", "Output");
+}
+
+FText UMovieGraphOutputNode::GetMenuCategory() const
+{
+	return NSLOCTEXT("MovieGraphNodes", "OutputNode_Category", "Input/Output");
+}
+
+FLinearColor UMovieGraphOutputNode::GetNodeTitleColor() const
+{
+	return FLinearColor::Black;
+}
+
+FSlateIcon UMovieGraphOutputNode::GetIconAndTint(FLinearColor& OutColor) const
+{
+	static const FSlateIcon OutputIcon = FSlateIcon(FAppStyle::GetAppStyleSetName(), "EditorViewport.VisualizeBufferMode");
+
+	OutColor = FLinearColor::White;
+	return OutputIcon;
+}
+#endif // WITH_EDITOR
+
+void UMovieGraphOutputNode::RegisterDelegates() const
+{
+	Super::RegisterDelegates();
+	
+	if (const UMovieGraphConfig* Graph = GetGraph())
+	{
+		for (UMovieGraphOutput* OutputMember : Graph->GetOutputs())
+		{
+			RegisterDelegates(OutputMember);
+		}
+	}
+}
+
+void UMovieGraphOutputNode::RegisterDelegates(UMovieGraphOutput* Output) const
+{
+#if WITH_EDITOR
+	if (Output)
+	{
+		Output->OnMovieGraphOutputChangedDelegate.AddUObject(this, &UMovieGraphOutputNode::UpdateExistingPins);
+	}
+#endif
+}
+
+void UMovieGraphOutputNode::UpdateExistingPins(UMovieGraphMember* ChangedOutput) const
+{
+	if (const UMovieGraphConfig* Graph = GetGraph())
+	{
+		const TArray<UMovieGraphOutput*> OutputMembers = Graph->GetOutputs();
+		if (OutputMembers.Num() == InputPins.Num())
+		{
+			for (int32 Index = 0; Index < OutputMembers.Num(); ++Index)
+			{
+				InputPins[Index]->Properties.Label = FName(OutputMembers[Index]->Name);
+			}
+		}
+
+		OnNodeChangedDelegate.Broadcast(this);
+	}
+}
\ No newline at end of file
diff --git a/Engine/Plugins/MovieScene/MovieRenderPipeline/Source/MovieRenderPipelineCore/Private/Graph/Nodes/MovieGraphRenderLayerNode.cpp b/Engine/Plugins/MovieScene/MovieRenderPipeline/Source/MovieRenderPipelineCore/Private/Graph/Nodes/MovieGraphRenderLayerNode.cpp
new file mode 100644
--- /dev/null
+++ b/Engine/Plugins/MovieScene/MovieRenderPipeline/Source/MovieRenderPipelineCore/Private/Graph/Nodes/MovieGraphRenderLayerNode.cpp
@@ -0,0 +1,66 @@
+﻿// Copyright Epic Games, Inc. All Rights Reserved.
+
+#include "Graph/Nodes/MovieGraphRenderLayerNode.h"
+
+#include "Graph/MovieGraphConfig.h"
+#include "Styling/AppStyle.h"
+
+UMovieGraphRenderLayerNode::UMovieGraphRenderLayerNode()
+{
+	LayerName = TEXT("beauty");
+}
+
+TArray<FMovieGraphPinProperties> UMovieGraphRenderLayerNode::GetInputPinProperties() const
+{
+	TArray<FMovieGraphPinProperties> Properties;
+	Properties.Add(FMovieGraphPinProperties(NAME_None, EMovieGraphMemberType::Branch, false));
+	return Properties;
+}
+
+TArray<FMovieGraphPinProperties> UMovieGraphRenderLayerNode::GetOutputPinProperties() const
+{
+	TArray<FMovieGraphPinProperties> Properties;
+	Properties.Add(FMovieGraphPinProperties(NAME_None, EMovieGraphMemberType::Branch, false));
+	return Properties;
+}
+
+TArray<FPropertyBagPropertyDesc> UMovieGraphRenderLayerNode::GetDynamicPropertyDescriptions() const
+{
+	TArray<FPropertyBagPropertyDesc> Properties;
+	FPropertyBagPropertyDesc FloatEditConProperty = FPropertyBagPropertyDesc("bOverride_TestPropName", EPropertyBagPropertyType::Bool);
+	FPropertyBagPropertyDesc FloatProperty = FPropertyBagPropertyDesc("TestPropName", EPropertyBagPropertyType::Float);
+#if WITH_EDITOR
+	FloatEditConProperty.MetaData.Add(FPropertyBagPropertyDescMetaData("InlineEditConditionToggle", "true"));
+	FloatProperty.MetaData.Add(FPropertyBagPropertyDescMetaData("EditCondition", "bOverride_TestPropName"));
+#endif
+	
+	Properties.Add(FloatEditConProperty);
+	Properties.Add(FloatProperty);
+	return Properties;
+}
+
+#if WITH_EDITOR
+FText UMovieGraphRenderLayerNode::GetMenuDescription() const
+{
+	return NSLOCTEXT("MovieGraphNodes", "RenderLayerGraphNode_Description", "Render Layer");
+}
+
+FText UMovieGraphRenderLayerNode::GetMenuCategory() const
+{
+	return NSLOCTEXT("MovieGraphNodes", "RenderLayerGraphNode_Category", "Rendering");
+}
+
+FLinearColor UMovieGraphRenderLayerNode::GetNodeTitleColor() const
+{
+	static const FLinearColor RenderLayerNodeColor = FLinearColor(0.192f, 0.258f, 0.615f);
+	return RenderLayerNodeColor;
+}
+
+FSlateIcon UMovieGraphRenderLayerNode::GetIconAndTint(FLinearColor& OutColor) const
+{
+	static const FSlateIcon RenderLayerIcon = FSlateIcon(FAppStyle::GetAppStyleSetName(), "LevelEditor.Tabs.DataLayers");
+
+	OutColor = FLinearColor::White;
+	return RenderLayerIcon;
+}
+#endif // WITH_EDITOR
\ No newline at end of file
diff --git a/Engine/Plugins/MovieScene/MovieRenderPipeline/Source/MovieRenderPipelineCore/Private/Graph/Nodes/MovieGraphVariableNode.cpp b/Engine/Plugins/MovieScene/MovieRenderPipeline/Source/MovieRenderPipelineCore/Private/Graph/Nodes/MovieGraphVariableNode.cpp
new file mode 100644
--- /dev/null
+++ b/Engine/Plugins/MovieScene/MovieRenderPipeline/Source/MovieRenderPipelineCore/Private/Graph/Nodes/MovieGraphVariableNode.cpp
@@ -0,0 +1,94 @@
+﻿// Copyright Epic Games, Inc. All Rights Reserved.
+
+#include "Graph/Nodes/MovieGraphVariableNode.h"
+
+#include "Graph/MovieGraphConfig.h"
+#include "Styling/AppStyle.h"
+
+UMovieGraphVariableNode::UMovieGraphVariableNode()
+{
+	if (!HasAnyFlags(RF_ClassDefaultObject))
+	{
+		RegisterDelegates();
+	}
+}
+
+TArray<FMovieGraphPinProperties> UMovieGraphVariableNode::GetOutputPinProperties() const
+{
+	TArray<FMovieGraphPinProperties> Properties;
+	
+	if (GraphVariable)
+	{
+		Properties.Add(FMovieGraphPinProperties(FName(GraphVariable->Name), GraphVariable->Type, false));
+	}
+	else
+	{
+		Properties.Add(FMovieGraphPinProperties(TEXT("Unknown"), EMovieGraphMemberType::Float, false));
+	}
+	
+	return Properties;
+}
+
+void UMovieGraphVariableNode::SetVariable(UMovieGraphVariable* InVariable)
+{
+	if (InVariable)
+	{
+		GraphVariable = InVariable;
+
+		// Update the output pin to reflect the new variable, and update the pin whenever the variable changes
+		// (eg, when the variable is renamed)
+		UpdateOutputPin(GraphVariable);
+
+		RegisterDelegates();
+	}
+}
+
+#if WITH_EDITOR
+FText UMovieGraphVariableNode::GetMenuDescription() const
+{
+	return GraphVariable ? FText::FromString(GraphVariable->Name) : FText();
+}
+	
+FText UMovieGraphVariableNode::GetMenuCategory() const
+{
+	return NSLOCTEXT("MovieGraphNode", "VariableNode_Category", "Variables");
+}
+
+FLinearColor UMovieGraphVariableNode::GetNodeTitleColor() const
+{
+	static const FLinearColor VariableNodeColor = FLinearColor(0.188f, 0.309f, 0.286f);
+	return VariableNodeColor;
+}
+
+FSlateIcon UMovieGraphVariableNode::GetIconAndTint(FLinearColor& OutColor) const
+{
+	static const FSlateIcon VariableIcon = FSlateIcon(FAppStyle::GetAppStyleSetName(), "Kismet.AllClasses.VariableIcon");
+
+	OutColor = FLinearColor::White;
+	return VariableIcon;
+}
+#endif // WITH_EDITOR
+
+void UMovieGraphVariableNode::RegisterDelegates() const
+{
+	Super::RegisterDelegates();
+
+#if WITH_EDITOR
+	if (GraphVariable)
+	{
+		GraphVariable->OnMovieGraphVariableChangedDelegate.AddUObject(this, &UMovieGraphVariableNode::UpdateOutputPin);
+	}
+#endif
+}
+
+void UMovieGraphVariableNode::UpdateOutputPin(UMovieGraphMember* ChangedVariable) const
+{
+	if (!OutputPins.IsEmpty() && ChangedVariable)
+	{
+		// Update the output pin to reflect the variable data model
+		OutputPins[0]->Properties.Label = FName(ChangedVariable->Name);
+		OutputPins[0]->Properties.Type = ChangedVariable->Type;
+	}
+
+	OnNodeChangedDelegate.Broadcast(this);
+}
\ No newline at end of file
diff --git a/Engine/Plugins/MovieScene/MovieRenderPipeline/Source/MovieRenderPipelineCore/Public/Graph/MovieGraphNode.h b/Engine/Plugins/MovieScene/MovieRenderPipeline/Source/MovieRenderPipelineCore/Public/Graph/MovieGraphNode.h
--- a/Engine/Plugins/MovieScene/MovieRenderPipeline/Source/MovieRenderPipelineCore/Public/Graph/MovieGraphNode.h
+++ b/Engine/Plugins/MovieScene/MovieRenderPipeline/Source/MovieRenderPipelineCore/Public/Graph/MovieGraphNode.h
@@ -1,12 +1,18 @@
 // Copyright Epic Games, Inc. All Rights Reserved.
 
 #pragma once
+
 #include "CoreMinimal.h"
 #include "MovieGraphPin.h"
 #include "InstancedStruct.h"
 #include "PropertyBag.h"
-#include "MovieGraphNode.generated.h"
 
+#if WITH_EDITOR
+#include "Textures/SlateIcon.h"
+#include "Math/Color.h"
+#endif
+
+#include "MovieGraphNode.generated.h"
 
 // Forward Declares
 class UMovieGraphInput;
@@ -77,6 +83,12 @@ public:
 
 	void SetNodePosX(const int32 InNodePosX) { NodePosX = InNodePosX; }
 	void SetNodePosY(const int32 InNodePosY) { NodePosY = InNodePosY; }
+
+	/** Gets the node's title color, as visible in the graph. */
+	virtual FLinearColor GetNodeTitleColor() const;
+
+	/** Gets the node's icon and icon tint, as visible in the graph. */
+	virtual FSlateIcon GetIconAndTint(FLinearColor& OutColor) const;
 #endif
 
 	//~ Begin UObject Interface
@@ -129,145 +141,4 @@ protected:
 	/** A GUID which uniquely identifies this node. */
 	UPROPERTY()
 	FGuid Guid;
-};
-
-// Dummy test nodes
-UCLASS()
-class UMoviePipelineCollectionNode : public UMovieGraphNode
-{
-	GENERATED_BODY()
-public:
-	UMoviePipelineCollectionNode()
-	{
-	}
-
-	virtual TArray<FMovieGraphPinProperties> GetInputPinProperties() const override
-	{
-		TArray<FMovieGraphPinProperties> Properties;
-		Properties.Add(FMovieGraphPinProperties(NAME_None, EMovieGraphMemberType::Branch, false));
-		return Properties;
-	}
-
-	virtual TArray<FMovieGraphPinProperties> GetOutputPinProperties() const override
-	{
-		TArray<FMovieGraphPinProperties> Properties;
-		Properties.Add(FMovieGraphPinProperties(NAME_None, EMovieGraphMemberType::Branch, false));
-		return Properties;
-	}
-
-	virtual TArray<FPropertyBagPropertyDesc> GetDynamicPropertyDescriptions() const
-	{
-		TArray<FPropertyBagPropertyDesc> Properties;
-		FPropertyBagPropertyDesc FloatEditConProperty = FPropertyBagPropertyDesc("bOverride_TestPropName", EPropertyBagPropertyType::Bool);
-		FPropertyBagPropertyDesc FloatProperty = FPropertyBagPropertyDesc("TestPropName", EPropertyBagPropertyType::Float);
-#if WITH_EDITOR
-		FloatEditConProperty.MetaData.Add(FPropertyBagPropertyDescMetaData("InlineEditConditionToggle", "true"));
-		FloatProperty.MetaData.Add(FPropertyBagPropertyDescMetaData("EditCondition", "bOverride_TestPropName"));
-#endif
-		
-		Properties.Add(FloatEditConProperty);
-		Properties.Add(FloatProperty);
-		return Properties;
-	}
-#if WITH_EDITOR
-	virtual FText GetMenuDescription() const override
-	{
-		return NSLOCTEXT("debug", "collection nodename", "Component Collection");
-	}
-	
-	virtual FText GetMenuCategory() const override
-	{
-		return NSLOCTEXT("debug", "collection cat", "Rendering");
-	}
-#endif
-};
-
-/** A graph node which displays all output members available in the graph. */
-UCLASS()
-class MOVIERENDERPIPELINECORE_API UMovieGraphOutputNode : public UMovieGraphNode
-{
-	GENERATED_BODY()
-
-public:
-	UMovieGraphOutputNode();
-	
-	virtual TArray<FMovieGraphPinProperties> GetInputPinProperties() const override;
-
-#if WITH_EDITOR
-	virtual FText GetMenuDescription() const override { return NSLOCTEXT("MovieGraphNodes", "OutputNode_Description", "Output"); }
-	virtual FText GetMenuCategory() const override { return NSLOCTEXT("MovieGraphNodes", "OutputNode_Category", "Input/Output"); }
-#endif
-
-private:
-	virtual void RegisterDelegates() const override;
-
-	/** Register delegates for the provided output member. */
-	void RegisterDelegates(UMovieGraphOutput* Output) const;
-
-	/** Update data (name, etc) on all existing input pins on this node to reflect the output members on the graph. */
-	void UpdateExistingPins(UMovieGraphMember* ChangedVariable) const;
-};
-
-/** A graph node which displays all input members available in the graph. */
-UCLASS()
-class MOVIERENDERPIPELINECORE_API UMovieGraphInputNode : public UMovieGraphNode
-{
-	GENERATED_BODY()
-
-public:
-	UMovieGraphInputNode();
-	
-	virtual TArray<FMovieGraphPinProperties> GetOutputPinProperties() const override;
-
-#if WITH_EDITOR
-	virtual FText GetMenuDescription() const override { return NSLOCTEXT("MovieGraphNodes", "InputNode_Description", "Input"); }
-	virtual FText GetMenuCategory() const override { return NSLOCTEXT("MovieGraphNodes", "InputNode_Category", "Input/Output"); }
-#endif
-
-private:
-	virtual void RegisterDelegates() const override;
-
-	/** Register delegates for the provided input member. */
-	void RegisterDelegates(UMovieGraphInput* Input) const;
-
-	/** Update data (name, etc) on all existing output pins on this node to reflect the input members on the graph. */
-	void UpdateExistingPins(UMovieGraphMember* ChangedVariable) const;
-};
-
-/** A node which gets the value of a variable which has been defined on the graph. */
-UCLASS()
-class MOVIERENDERPIPELINECORE_API UMovieGraphVariableNode : public UMovieGraphNode
-{
-	GENERATED_BODY()
-
-public:
-	UMovieGraphVariableNode();
-
-	virtual TArray<FMovieGraphPinProperties> GetOutputPinProperties() const override;
-
-	/** Gets the variable that this node represents. */
-	UMovieGraphVariable* GetVariable() const { return GraphVariable; }
-
-	/** Sets the variable that this node represents. */
-	void SetVariable(UMovieGraphVariable* InVariable);
-
-#if WITH_EDITOR
-	virtual FText GetMenuDescription() const override;
-	virtual FText GetMenuCategory() const override;
-#endif
-
-private:
-	virtual void RegisterDelegates() const override;
-	
-	/** Updates the output pin on the node to match the provided variable. */
-	void UpdateOutputPin(UMovieGraphMember* ChangedVariable) const;
-
-private:
-	/** The underlying graph variable this node represents. */
-	UPROPERTY()
-	TObjectPtr<UMovieGraphVariable> GraphVariable = nullptr;
-
-	/** The properties for the output pin on this node. */
-	UPROPERTY(Transient)
-	FMovieGraphPinProperties OutputPin;
 };
\ No newline at end of file
diff --git a/Engine/Plugins/MovieScene/MovieRenderPipeline/Source/MovieRenderPipelineCore/Public/Graph/MovieGraphPlaceholderNodes.cpp b/Engine/Plugins/MovieScene/MovieRenderPipeline/Source/MovieRenderPipelineCore/Public/Graph/MovieGraphPlaceholderNodes.cpp
new file mode 100644
--- /dev/null
+++ b/Engine/Plugins/MovieScene/MovieRenderPipeline/Source/MovieRenderPipelineCore/Public/Graph/MovieGraphPlaceholderNodes.cpp
@@ -0,0 +1,433 @@
+﻿// Copyright Epic Games, Inc. All Rights Reserved.
+
+#include "MovieGraphPlaceholderNodes.h"
+
+#include "MovieGraphConfig.h"
+#include "Styling/AppStyle.h"
+
+#define LOCTEXT_NAMESPACE "MovieGraphNode"
+
+static const FText NodeCategory_Utility = LOCTEXT("NodeCategory_Utility", "Utility");
+static const FText NodeCategory_Globals = LOCTEXT("NodeCategory_Globals", "Globals");
+static const FText NodeCategory_Renderers = LOCTEXT("NodeCategory_Renderers", "Renderers");
+static const FText NodeCategory_OutputType = LOCTEXT("NodeCategory_OutputType", "Output Type");
+static const FText NodeCategory_Conditionals = LOCTEXT("NodeCategory_Conditionals", "Conditionals");
+static const FText NodeCategory_Settings = LOCTEXT("NodeCategory_Settings", "Settings");
+
+TArray<FMovieGraphPinProperties> UMoviePipelineCollectionNode::GetInputPinProperties() const
+{
+	TArray<FMovieGraphPinProperties> Properties;
+	Properties.Add(FMovieGraphPinProperties(NAME_None, EMovieGraphMemberType::Branch, false));
+	return Properties;
+}
+
+TArray<FMovieGraphPinProperties> UMoviePipelineCollectionNode::GetOutputPinProperties() const
+{
+	TArray<FMovieGraphPinProperties> Properties;
+	Properties.Add(FMovieGraphPinProperties(NAME_None, EMovieGraphMemberType::Branch, false));
+	return Properties;
+}
+
+#if WITH_EDITOR
+FText UMoviePipelineCollectionNode::GetMenuDescription() const
+{
+	static const FText CollectionNodeName = LOCTEXT("NodeName_Collection", "Collection");
+	return CollectionNodeName;
+}
+
+FText UMoviePipelineCollectionNode::GetMenuCategory() const
+{
+	return NodeCategory_Utility;
+}
+
+FLinearColor UMoviePipelineCollectionNode::GetNodeTitleColor() const
+{
+	static const FLinearColor CollectionNodeColor = FLinearColor(0.047f, 0.501f, 0.654f);
+	return CollectionNodeColor;
+}
+
+FSlateIcon UMoviePipelineCollectionNode::GetIconAndTint(FLinearColor& OutColor) const
+{
+	static const FSlateIcon CollectionIcon = FSlateIcon(FAppStyle::GetAppStyleSetName(), "SceneOutliner.NewFolderIcon");
+
+	OutColor = FLinearColor::White;
+	return CollectionIcon;
+}
+#endif // WITH_EDITOR
+
+TArray<FMovieGraphPinProperties> UMovieGraphModifierNode::GetInputPinProperties() const
+{
+	TArray<FMovieGraphPinProperties> Properties;
+	Properties.Add(FMovieGraphPinProperties(NAME_None, EMovieGraphMemberType::Branch, false));
+	return Properties;
+}
+
+TArray<FMovieGraphPinProperties> UMovieGraphModifierNode::GetOutputPinProperties() const
+{
+	TArray<FMovieGraphPinProperties> Properties;
+	Properties.Add(FMovieGraphPinProperties(NAME_None, EMovieGraphMemberType::Branch, false));
+	return Properties;
+}
+
+#if WITH_EDITOR
+FText UMovieGraphModifierNode::GetMenuDescription() const
+{
+	static const FText ModifierNodeName = LOCTEXT("NodeName_Modifier", "Modifier");
+	return ModifierNodeName;
+}
+
+FText UMovieGraphModifierNode::GetMenuCategory() const
+{
+	return NodeCategory_Utility;
+}
+
+FLinearColor UMovieGraphModifierNode::GetNodeTitleColor() const
+{
+	static const FLinearColor ModifierNodeColor = FLinearColor(0.6f, 0.113f, 0.113f);
+	return ModifierNodeColor;
+}
+
+FSlateIcon UMovieGraphModifierNode::GetIconAndTint(FLinearColor& OutColor) const
+{
+	static const FSlateIcon ModifierIcon = FSlateIcon(FAppStyle::GetAppStyleSetName(), "ContentBrowser.ReferenceViewer");
+
+	OutColor = FLinearColor::White;
+	return ModifierIcon;
+}
+#endif // WITH_EDITOR
+
+TArray<FMovieGraphPinProperties> UMovieGraphGlobalGameOverridesNode::GetInputPinProperties() const
+{
+	TArray<FMovieGraphPinProperties> Properties;
+	Properties.Add(FMovieGraphPinProperties(NAME_None, EMovieGraphMemberType::Branch, false));
+	return Properties;
+}
+
+TArray<FMovieGraphPinProperties> UMovieGraphGlobalGameOverridesNode::GetOutputPinProperties() const
+{
+	TArray<FMovieGraphPinProperties> Properties;
+	Properties.Add(FMovieGraphPinProperties(NAME_None, EMovieGraphMemberType::Branch, false));
+	return Properties;
+}
+
+#if WITH_EDITOR
+FText UMovieGraphGlobalGameOverridesNode::GetMenuDescription() const
+{
+	static const FText GlobalGameOverridesNodeName = LOCTEXT("NodeName_GlobalGameOverrides", "Global Game Overrides");
+	return GlobalGameOverridesNodeName;
+}
+
+FText UMovieGraphGlobalGameOverridesNode::GetMenuCategory() const
+{
+	return NodeCategory_Globals;
+}
+
+FLinearColor UMovieGraphGlobalGameOverridesNode::GetNodeTitleColor() const
+{
+	static const FLinearColor GlobalGameOverridesNodeColor = FLinearColor(0.549f, 0.f, 0.250f);
+	return GlobalGameOverridesNodeColor;
+}
+
+FSlateIcon UMovieGraphGlobalGameOverridesNode::GetIconAndTint(FLinearColor& OutColor) const
+{
+	static const FSlateIcon GlobalGameOverridesIcon = FSlateIcon(FAppStyle::GetAppStyleSetName(), "Launcher.TabIcon");
+
+	OutColor = FLinearColor::White;
+	return GlobalGameOverridesIcon;
+}
+#endif // WITH_EDITOR
+
+TArray<FMovieGraphPinProperties> UMovieGraphDeferredRendererNode::GetInputPinProperties() const
+{
+	TArray<FMovieGraphPinProperties> Properties;
+	Properties.Add(FMovieGraphPinProperties(NAME_None, EMovieGraphMemberType::Branch, false));
+	return Properties;
+}
+
+TArray<FMovieGraphPinProperties> UMovieGraphDeferredRendererNode::GetOutputPinProperties() const
+{
+	TArray<FMovieGraphPinProperties> Properties;
+	Properties.Add(FMovieGraphPinProperties(NAME_None, EMovieGraphMemberType::Branch, false));
+	return Properties;
+}
+
+#if WITH_EDITOR
+FText UMovieGraphDeferredRendererNode::GetMenuDescription() const
+{
+	static const FText DeferredRendererNodeName = LOCTEXT("NodeName_DeferredRenderer", "Deferred Renderer");
+	return DeferredRendererNodeName;
+}
+
+FText UMovieGraphDeferredRendererNode::GetMenuCategory() const
+{
+	return NodeCategory_Renderers;
+}
+
+FLinearColor UMovieGraphDeferredRendererNode::GetNodeTitleColor() const
+{
+	static const FLinearColor DeferredRendererNodeColor = FLinearColor(0.572f, 0.274f, 1.f);
+	return DeferredRendererNodeColor;
+}
+
+FSlateIcon UMovieGraphDeferredRendererNode::GetIconAndTint(FLinearColor& OutColor) const
+{
+	static const FSlateIcon DeferredRendererIcon = FSlateIcon(FAppStyle::GetAppStyleSetName(), "SequenceRecorder.TabIcon");
+
+	OutColor = FLinearColor::White;
+	return DeferredRendererIcon;
+}
+#endif // WITH_EDITOR
+
+TArray<FMovieGraphPinProperties> UMovieGraphPathTracedRendererNode::GetInputPinProperties() const
+{
+	TArray<FMovieGraphPinProperties> Properties;
+	Properties.Add(FMovieGraphPinProperties(NAME_None, EMovieGraphMemberType::Branch, false));
+	return Properties;
+}
+
+TArray<FMovieGraphPinProperties> UMovieGraphPathTracedRendererNode::GetOutputPinProperties() const
+{
+	TArray<FMovieGraphPinProperties> Properties;
+	Properties.Add(FMovieGraphPinProperties(NAME_None, EMovieGraphMemberType::Branch, false));
+	return Properties;
+}
+
+#if WITH_EDITOR
+FText UMovieGraphPathTracedRendererNode::GetMenuDescription() const
+{
+	static const FText PathTracedRendererNodeName = LOCTEXT("NodeName_PathTracedRenderer", "Path Traced Renderer");
+	return PathTracedRendererNodeName;
+}
+
+FText UMovieGraphPathTracedRendererNode::GetMenuCategory() const
+{
+	return NodeCategory_Renderers;
+}
+
+FLinearColor UMovieGraphPathTracedRendererNode::GetNodeTitleColor() const
+{
+	static const FLinearColor DeferredRendererNodeColor = FLinearColor(0.572f, 0.274f, 1.f);
+	return DeferredRendererNodeColor;
+}
+
+FSlateIcon UMovieGraphPathTracedRendererNode::GetIconAndTint(FLinearColor& OutColor) const
+{
+	static const FSlateIcon DeferredRendererIcon = FSlateIcon(FAppStyle::GetAppStyleSetName(), "SequenceRecorder.TabIcon");
+
+	OutColor = FLinearColor::White;
+	return DeferredRendererIcon;
+}
+#endif // WITH_EDITOR
+
+TArray<FMovieGraphPinProperties> UMovieGraphEXRSequenceNode::GetInputPinProperties() const
+{
+	TArray<FMovieGraphPinProperties> Properties;
+	Properties.Add(FMovieGraphPinProperties(NAME_None, EMovieGraphMemberType::Branch, false));
+	return Properties;
+}
+
+TArray<FMovieGraphPinProperties> UMovieGraphEXRSequenceNode::GetOutputPinProperties() const
+{
+	TArray<FMovieGraphPinProperties> Properties;
+	Properties.Add(FMovieGraphPinProperties(NAME_None, EMovieGraphMemberType::Branch, false));
+	return Properties;
+}
+
+#if WITH_EDITOR
+FText UMovieGraphEXRSequenceNode::GetMenuDescription() const
+{
+	static const FText EXRSequenceNodeName = LOCTEXT("NodeName_EXRSequence", "EXR Sequence");
+	return EXRSequenceNodeName;
+}
+
+FText UMovieGraphEXRSequenceNode::GetMenuCategory() const
+{
+	return NodeCategory_OutputType;
+}
+
+FLinearColor UMovieGraphEXRSequenceNode::GetNodeTitleColor() const
+{
+	static const FLinearColor ImageSequenceNodeColor = FLinearColor(0.047f, 0.654f, 0.537f);
+	return ImageSequenceNodeColor;
+}
+
+FSlateIcon UMovieGraphEXRSequenceNode::GetIconAndTint(FLinearColor& OutColor) const
+{
+	static const FSlateIcon ImageSequenceIcon = FSlateIcon(FAppStyle::GetAppStyleSetName(), "ClassIcon.Texture2D");
+
+	OutColor = FLinearColor::White;
+	return ImageSequenceIcon;
+}
+#endif // WITH_EDITOR
+
+TArray<FMovieGraphPinProperties> UMovieGraphJPGSequenceNode::GetInputPinProperties() const
+{
+	TArray<FMovieGraphPinProperties> Properties;
+	Properties.Add(FMovieGraphPinProperties(NAME_None, EMovieGraphMemberType::Branch, false));
+	return Properties;
+}
+
+TArray<FMovieGraphPinProperties> UMovieGraphJPGSequenceNode::GetOutputPinProperties() const
+{
+	TArray<FMovieGraphPinProperties> Properties;
+	Properties.Add(FMovieGraphPinProperties(NAME_None, EMovieGraphMemberType::Branch, false));
+	return Properties;
+}
+
+#if WITH_EDITOR
+FText UMovieGraphJPGSequenceNode::GetMenuDescription() const
+{
+	static const FText JPGSequenceNodeName = LOCTEXT("NodeName_JPGSequence", "JPG Sequence");
+	return JPGSequenceNodeName;
+}
+
+FText UMovieGraphJPGSequenceNode::GetMenuCategory() const
+{
+	return NodeCategory_OutputType;
+}
+
+FLinearColor UMovieGraphJPGSequenceNode::GetNodeTitleColor() const
+{
+	static const FLinearColor ImageSequenceNodeColor = FLinearColor(0.047f, 0.654f, 0.537f);
+	return ImageSequenceNodeColor;
+}
+
+FSlateIcon UMovieGraphJPGSequenceNode::GetIconAndTint(FLinearColor& OutColor) const
+{
+	static const FSlateIcon ImageSequenceIcon = FSlateIcon(FAppStyle::GetAppStyleSetName(), "ClassIcon.Texture2D");
+
+	OutColor = FLinearColor::White;
+	return ImageSequenceIcon;
+}
+#endif
+
+TArray<FMovieGraphPinProperties> UMovieGraphBranchNode::GetInputPinProperties() const
+{
+	static const FName TrueBranch("True");
+	static const FName FalseBranch("False");
+	static const FName Condition("Condition");
+	
+	TArray<FMovieGraphPinProperties> Properties;
+	Properties.Add(FMovieGraphPinProperties(TrueBranch, EMovieGraphMemberType::Branch, false));
+	Properties.Add(FMovieGraphPinProperties(FalseBranch, EMovieGraphMemberType::Branch, false));
+	Properties.Add(FMovieGraphPinProperties(Condition, EMovieGraphMemberType::Bool, false));
+	return Properties;
+}
+
+TArray<FMovieGraphPinProperties> UMovieGraphBranchNode::GetOutputPinProperties() const
+{
+	TArray<FMovieGraphPinProperties> Properties;
+	Properties.Add(FMovieGraphPinProperties(NAME_None, EMovieGraphMemberType::Branch, false));
+	return Properties;
+}
+
+#if WITH_EDITOR
+FText UMovieGraphBranchNode::GetMenuDescription() const
+{
+	static const FText BranchNodeName = LOCTEXT("NodeName_Branch", "Branch");
+	return BranchNodeName;
+}
+
+FText UMovieGraphBranchNode::GetMenuCategory() const
+{
+	return NodeCategory_Conditionals;
+}
+
+FLinearColor UMovieGraphBranchNode::GetNodeTitleColor() const
+{
+	static const FLinearColor BranchNodeColor = FLinearColor(0.266f, 0.266f, 0.266f);
+	return BranchNodeColor;
+}
+
+FSlateIcon UMovieGraphBranchNode::GetIconAndTint(FLinearColor& OutColor) const
+{
+	static const FSlateIcon BranchIcon = FSlateIcon(FAppStyle::GetAppStyleSetName(), "Icons.Merge");
+
+	OutColor = FLinearColor::White;
+	return BranchIcon;
+}
+#endif // WITH_EDITOR
+
+TArray<FMovieGraphPinProperties> UMovieGraphOutputSettingsNode::GetInputPinProperties() const
+{
+	TArray<FMovieGraphPinProperties> Properties;
+	Properties.Add(FMovieGraphPinProperties(NAME_None, EMovieGraphMemberType::Branch, false));
+	return Properties;
+}
+
+TArray<FMovieGraphPinProperties> UMovieGraphOutputSettingsNode::GetOutputPinProperties() const
+{
+	TArray<FMovieGraphPinProperties> Properties;
+	Properties.Add(FMovieGraphPinProperties(NAME_None, EMovieGraphMemberType::Branch, false));
+	return Properties;
+}
+
+#if WITH_EDITOR
+FText UMovieGraphOutputSettingsNode::GetMenuDescription() const
+{
+	static const FText OutputSettingsNodeName = LOCTEXT("NodeName_OutputSettings", "Output Settings");
+	return OutputSettingsNodeName;
+}
+
+FText UMovieGraphOutputSettingsNode::GetMenuCategory() const
+{
+	return NodeCategory_Settings;
+}
+
+FLinearColor UMovieGraphOutputSettingsNode::GetNodeTitleColor() const
+{
+	static const FLinearColor OutputSettingsColor = FLinearColor(0.854f, 0.509f, 0.039f);
+	return OutputSettingsColor;
+}
+
+FSlateIcon UMovieGraphOutputSettingsNode::GetIconAndTint(FLinearColor& OutColor) const
+{
+	static const FSlateIcon SettingsIcon = FSlateIcon(FAppStyle::GetAppStyleSetName(), "Icons.Settings");
+
+	OutColor = FLinearColor::White;
+	return SettingsIcon;
+}
+#endif // WITH_EDITOR
+
+TArray<FMovieGraphPinProperties> UMovieGraphAntiAliasingNode::GetInputPinProperties() const
+{
+	TArray<FMovieGraphPinProperties> Properties;
+	Properties.Add(FMovieGraphPinProperties(NAME_None, EMovieGraphMemberType::Branch, false));
+	return Properties;
+}
+
+TArray<FMovieGraphPinProperties> UMovieGraphAntiAliasingNode::GetOutputPinProperties() const
+{
+	TArray<FMovieGraphPinProperties> Properties;
+	Properties.Add(FMovieGraphPinProperties(NAME_None, EMovieGraphMemberType::Branch, false));
+	return Properties;
+}
+
+#if WITH_EDITOR
+FText UMovieGraphAntiAliasingNode::GetMenuDescription() const
+{
+	static const FText AntiAliasingNodeName = LOCTEXT("NodeName_AntiAliasing", "Anti-Aliasing");
+	return AntiAliasingNodeName;
+}
+
+FText UMovieGraphAntiAliasingNode::GetMenuCategory() const
+{
+	return NodeCategory_Settings;
+}
+
+FLinearColor UMovieGraphAntiAliasingNode::GetNodeTitleColor() const
+{
+	static const FLinearColor AntiAliasingColor = FLinearColor(0.043f, 0.219f, 0.356f);
+	return AntiAliasingColor;
+}
+
+FSlateIcon UMovieGraphAntiAliasingNode::GetIconAndTint(FLinearColor& OutColor) const
+{
+	static const FSlateIcon SettingsIcon = FSlateIcon(FAppStyle::GetAppStyleSetName(), "Icons.Settings");
+
+	OutColor = FLinearColor::White;
+	return SettingsIcon;
+}
+#endif // WITH_EDITOR
+
+#undef LOCTEXT_NAMESPACE
\ No newline at end of file
diff --git a/Engine/Plugins/MovieScene/MovieRenderPipeline/Source/MovieRenderPipelineCore/Public/Graph/MovieGraphPlaceholderNodes.h b/Engine/Plugins/MovieScene/MovieRenderPipeline/Source/MovieRenderPipelineCore/Public/Graph/MovieGraphPlaceholderNodes.h
new file mode 100644
--- /dev/null
+++ b/Engine/Plugins/MovieScene/MovieRenderPipeline/Source/MovieRenderPipelineCore/Public/Graph/MovieGraphPlaceholderNodes.h
@@ -0,0 +1,222 @@
+﻿// Copyright Epic Games, Inc. All Rights Reserved.
+
+#pragma once
+
+#include "MovieGraphNode.h"
+
+#include "MovieGraphPlaceholderNodes.generated.h"
+
+// TODO: This should probably be called UMovieGraphCollectionNode instead
+/** A graph node which creates collections that are available for use in other nodes. */
+UCLASS()
+class MOVIERENDERPIPELINECORE_API UMoviePipelineCollectionNode : public UMovieGraphNode
+{
+	GENERATED_BODY()
+public:
+	UMoviePipelineCollectionNode() = default;
+
+	virtual TArray<FMovieGraphPinProperties> GetInputPinProperties() const override;
+	virtual TArray<FMovieGraphPinProperties> GetOutputPinProperties() const override;
+
+	virtual TArray<FPropertyBagPropertyDesc> GetDynamicPropertyDescriptions() const override
+	{
+		TArray<FPropertyBagPropertyDesc> Properties;
+		FPropertyBagPropertyDesc FloatEditConProperty = FPropertyBagPropertyDesc("bOverride_TestPropName", EPropertyBagPropertyType::Bool);
+		FPropertyBagPropertyDesc FloatProperty = FPropertyBagPropertyDesc("TestPropName", EPropertyBagPropertyType::Float);
+#if WITH_EDITOR
+		FloatEditConProperty.MetaData.Add(FPropertyBagPropertyDescMetaData("InlineEditConditionToggle", "true"));
+		FloatProperty.MetaData.Add(FPropertyBagPropertyDescMetaData("EditCondition", "bOverride_TestPropName"));
+#endif
+		
+		Properties.Add(FloatEditConProperty);
+		Properties.Add(FloatProperty);
+		return Properties;
+	}
+
+#if WITH_EDITOR
+	virtual FText GetMenuDescription() const override;
+	virtual FText GetMenuCategory() const override;
+	virtual FLinearColor GetNodeTitleColor() const override;
+	virtual FSlateIcon GetIconAndTint(FLinearColor& OutColor) const override;
+#endif
+};
+
+/** A node which modifies part of the world being rendered. */
+UCLASS()
+class MOVIERENDERPIPELINECORE_API UMovieGraphModifierNode : public UMovieGraphNode
+{
+	GENERATED_BODY()
+
+public:
+	UMovieGraphModifierNode() = default;
+
+	virtual TArray<FMovieGraphPinProperties> GetInputPinProperties() const override;
+	virtual TArray<FMovieGraphPinProperties> GetOutputPinProperties() const override;
+
+#if WITH_EDITOR
+	virtual FText GetMenuDescription() const override;
+	virtual FText GetMenuCategory() const override;
+	virtual FLinearColor GetNodeTitleColor() const override;
+	virtual FSlateIcon GetIconAndTint(FLinearColor& OutColor) const override;
+#endif
+};
+
+/** A node which configures the global game overrides. */
+UCLASS()
+class MOVIERENDERPIPELINECORE_API UMovieGraphGlobalGameOverridesNode : public UMovieGraphNode
+{
+	GENERATED_BODY()
+
+public:
+	UMovieGraphGlobalGameOverridesNode() = default;
+
+	virtual TArray<FMovieGraphPinProperties> GetInputPinProperties() const override;
+	virtual TArray<FMovieGraphPinProperties> GetOutputPinProperties() const override;
+
+#if WITH_EDITOR
+	virtual FText GetMenuDescription() const override;
+	virtual FText GetMenuCategory() const override;
+	virtual FLinearColor GetNodeTitleColor() const override;
+	virtual FSlateIcon GetIconAndTint(FLinearColor& OutColor) const override;
+#endif
+};
+
+/** A node which represents a deferred renderer. */
+UCLASS()
+class MOVIERENDERPIPELINECORE_API UMovieGraphDeferredRendererNode : public UMovieGraphNode
+{
+	GENERATED_BODY()
+
+public:
+	UMovieGraphDeferredRendererNode() = default;
+
+	virtual TArray<FMovieGraphPinProperties> GetInputPinProperties() const override;
+	virtual TArray<FMovieGraphPinProperties> GetOutputPinProperties() const override;
+
+#if WITH_EDITOR
+	virtual FText GetMenuDescription() const override;
+	virtual FText GetMenuCategory() const override;
+	virtual FLinearColor GetNodeTitleColor() const override;
+	virtual FSlateIcon GetIconAndTint(FLinearColor& OutColor) const override;
+#endif
+};
+
+/** A node which represents a path traced renderer. */
+UCLASS()
+class MOVIERENDERPIPELINECORE_API UMovieGraphPathTracedRendererNode : public UMovieGraphNode
+{
+	GENERATED_BODY()
+
+public:
+	UMovieGraphPathTracedRendererNode() = default;
+
+	virtual TArray<FMovieGraphPinProperties> GetInputPinProperties() const override;
+	virtual TArray<FMovieGraphPinProperties> GetOutputPinProperties() const override;
+
+#if WITH_EDITOR
+	virtual FText GetMenuDescription() const override;
+	virtual FText GetMenuCategory() const override;
+	virtual FLinearColor GetNodeTitleColor() const override;
+	virtual FSlateIcon GetIconAndTint(FLinearColor& OutColor) const override;
+#endif
+};
+
+/** A node which generates an EXR image sequence. */
+UCLASS()
+class MOVIERENDERPIPELINECORE_API UMovieGraphEXRSequenceNode : public UMovieGraphNode
+{
+	GENERATED_BODY()
+
+public:
+	UMovieGraphEXRSequenceNode() = default;
+
+	virtual TArray<FMovieGraphPinProperties> GetInputPinProperties() const override;
+	virtual TArray<FMovieGraphPinProperties> GetOutputPinProperties() const override;
+
+#if WITH_EDITOR
+	virtual FText GetMenuDescription() const override;
+	virtual FText GetMenuCategory() const override;
+	virtual FLinearColor GetNodeTitleColor() const override;
+	virtual FSlateIcon GetIconAndTint(FLinearColor& OutColor) const override;
+#endif
+};
+
+/** A node which generates a JPG image sequence. */
+UCLASS()
+class MOVIERENDERPIPELINECORE_API UMovieGraphJPGSequenceNode : public UMovieGraphNode
+{
+	GENERATED_BODY()
+
+public:
+	UMovieGraphJPGSequenceNode() = default;
+
+	virtual TArray<FMovieGraphPinProperties> GetInputPinProperties() const override;
+	virtual TArray<FMovieGraphPinProperties> GetOutputPinProperties() const override;
+
+#if WITH_EDITOR
+	virtual FText GetMenuDescription() const override;
+	virtual FText GetMenuCategory() const override;
+	virtual FLinearColor GetNodeTitleColor() const override;
+	virtual FSlateIcon GetIconAndTint(FLinearColor& OutColor) const override;
+#endif
+};
+
+/** A node which creates a branching condition. */
+UCLASS()
+class MOVIERENDERPIPELINECORE_API UMovieGraphBranchNode : public UMovieGraphNode
+{
+	GENERATED_BODY()
+
+public:
+	UMovieGraphBranchNode() = default;
+
+	virtual TArray<FMovieGraphPinProperties> GetInputPinProperties() const override;
+	virtual TArray<FMovieGraphPinProperties> GetOutputPinProperties() const override;
+
+#if WITH_EDITOR
+	virtual FText GetMenuDescription() const override;
+	virtual FText GetMenuCategory() const override;
+	virtual FLinearColor GetNodeTitleColor() const override;
+	virtual FSlateIcon GetIconAndTint(FLinearColor& OutColor) const override;
+#endif
+};
+
+/** A node which configures output settings. */
+UCLASS()
+class MOVIERENDERPIPELINECORE_API UMovieGraphOutputSettingsNode : public UMovieGraphNode
+{
+	GENERATED_BODY()
+
+public:
+	UMovieGraphOutputSettingsNode() = default;
+
+	virtual TArray<FMovieGraphPinProperties> GetInputPinProperties() const override;
+	virtual TArray<FMovieGraphPinProperties> GetOutputPinProperties() const override;
+
+#if WITH_EDITOR
+	virtual FText GetMenuDescription() const override;
+	virtual FText GetMenuCategory() const override;
+	virtual FLinearColor GetNodeTitleColor() const override;
+	virtual FSlateIcon GetIconAndTint(FLinearColor& OutColor) const override;
+#endif
+};
+
+/** A node which configures anti-aliasing settings. */
+UCLASS()
+class MOVIERENDERPIPELINECORE_API UMovieGraphAntiAliasingNode : public UMovieGraphNode
+{
+	GENERATED_BODY()
+
+public:
+	UMovieGraphAntiAliasingNode() = default;
+
+	virtual TArray<FMovieGraphPinProperties> GetInputPinProperties() const override;
+	virtual TArray<FMovieGraphPinProperties> GetOutputPinProperties() const override;
+
+#if WITH_EDITOR
+	virtual FText GetMenuDescription() const override;
+	virtual FText GetMenuCategory() const override;
+	virtual FLinearColor GetNodeTitleColor() const override;
+	virtual FSlateIcon GetIconAndTint(FLinearColor& OutColor) const override;
+#endif
+};
\ No newline at end of file
diff --git a/Engine/Plugins/MovieScene/MovieRenderPipeline/Source/MovieRenderPipelineCore/Public/Graph/Nodes/MovieGraphInputNode.h b/Engine/Plugins/MovieScene/MovieRenderPipeline/Source/MovieRenderPipelineCore/Public/Graph/Nodes/MovieGraphInputNode.h
new file mode 100644
--- /dev/null
+++ b/Engine/Plugins/MovieScene/MovieRenderPipeline/Source/MovieRenderPipelineCore/Public/Graph/Nodes/MovieGraphInputNode.h
@@ -0,0 +1,35 @@
+﻿// Copyright Epic Games, Inc. All Rights Reserved.
+
+#pragma once
+
+#include "Graph/MovieGraphNode.h"
+
+#include "MovieGraphInputNode.generated.h"
+
+/** A graph node which displays all input members available in the graph. */
+UCLASS()
+class MOVIERENDERPIPELINECORE_API UMovieGraphInputNode : public UMovieGraphNode
+{
+	GENERATED_BODY()
+
+public:
+	UMovieGraphInputNode();
+	
+	virtual TArray<FMovieGraphPinProperties> GetOutputPinProperties() const override;
+
+#if WITH_EDITOR
+	virtual FText GetMenuDescription() const override;
+	virtual FText GetMenuCategory() const override;
+	virtual FLinearColor GetNodeTitleColor() const override;
+	virtual FSlateIcon GetIconAndTint(FLinearColor& OutColor) const override;
+#endif
+
+private:
+	virtual void RegisterDelegates() const override;
+
+	/** Register delegates for the provided input member. */
+	void RegisterDelegates(UMovieGraphInput* Input) const;
+
+	/** Update data (name, etc) on all existing output pins on this node to reflect the input members on the graph. */
+	void UpdateExistingPins(UMovieGraphMember* ChangedVariable) const;
+};
\ No newline at end of file
diff --git a/Engine/Plugins/MovieScene/MovieRenderPipeline/Source/MovieRenderPipelineCore/Public/Graph/Nodes/MovieGraphOutputNode.h b/Engine/Plugins/MovieScene/MovieRenderPipeline/Source/MovieRenderPipelineCore/Public/Graph/Nodes/MovieGraphOutputNode.h
new file mode 100644
--- /dev/null
+++ b/Engine/Plugins/MovieScene/MovieRenderPipeline/Source/MovieRenderPipelineCore/Public/Graph/Nodes/MovieGraphOutputNode.h
@@ -0,0 +1,35 @@
+﻿// Copyright Epic Games, Inc. All Rights Reserved.
+
+#pragma once
+
+#include "Graph/MovieGraphNode.h"
+
+#include "MovieGraphOutputNode.generated.h"
+
+/** A graph node which displays all output members available in the graph. */
+UCLASS()
+class MOVIERENDERPIPELINECORE_API UMovieGraphOutputNode : public UMovieGraphNode
+{
+	GENERATED_BODY()
+
+public:
+	UMovieGraphOutputNode();
+	
+	virtual TArray<FMovieGraphPinProperties> GetInputPinProperties() const override;
+
+#if WITH_EDITOR
+	virtual FText GetMenuDescription() const override;
+	virtual FText GetMenuCategory() const override;
+	virtual FLinearColor GetNodeTitleColor() const override;
+	virtual FSlateIcon GetIconAndTint(FLinearColor& OutColor) const override;
+#endif
+
+private:
+	virtual void RegisterDelegates() const override;
+
+	/** Register delegates for the provided output member. */
+	void RegisterDelegates(UMovieGraphOutput* Output) const;
+
+	/** Update data (name, etc) on all existing input pins on this node to reflect the output members on the graph. */
+	void UpdateExistingPins(UMovieGraphMember* ChangedVariable) const;
+};
\ No newline at end of file
diff --git a/Engine/Plugins/MovieScene/MovieRenderPipeline/Source/MovieRenderPipelineCore/Public/Graph/Nodes/MovieGraphRenderLayerNode.h b/Engine/Plugins/MovieScene/MovieRenderPipeline/Source/MovieRenderPipelineCore/Public/Graph/Nodes/MovieGraphRenderLayerNode.h
--- a/Engine/Plugins/MovieScene/MovieRenderPipeline/Source/MovieRenderPipelineCore/Public/Graph/Nodes/MovieGraphRenderLayerNode.h
+++ b/Engine/Plugins/MovieScene/MovieRenderPipeline/Source/MovieRenderPipelineCore/Public/Graph/Nodes/MovieGraphRenderLayerNode.h
@@ -1,60 +1,28 @@
 // Copyright Epic Games, Inc. All Rights Reserved.
 
 #pragma once
+
 #include "Graph/MovieGraphNode.h"
 #include "MovieGraphRenderLayerNode.generated.h"
 
-
 UCLASS()
-class UMovieGraphRenderLayerNode : public UMovieGraphNode
+class MOVIERENDERPIPELINECORE_API UMovieGraphRenderLayerNode : public UMovieGraphNode
 {
 	GENERATED_BODY()
 public:
-	UMovieGraphRenderLayerNode()
-	{
-		LayerName = TEXT("beauty");
-	}
+	UMovieGraphRenderLayerNode();
 
 	FString GetRenderLayerName() const { return LayerName; }
 
-	virtual TArray<FMovieGraphPinProperties> GetInputPinProperties() const override
-	{
-		TArray<FMovieGraphPinProperties> Properties;
-		Properties.Add(FMovieGraphPinProperties(NAME_None, EMovieGraphMemberType::Branch, false));
-		return Properties;
-	}
-
-	virtual TArray<FMovieGraphPinProperties> GetOutputPinProperties() const override
-	{
-		TArray<FMovieGraphPinProperties> Properties;
-		Properties.Add(FMovieGraphPinProperties(NAME_None, EMovieGraphMemberType::Branch, false));
-		return Properties;
-	}
+	virtual TArray<FMovieGraphPinProperties> GetInputPinProperties() const override;
+	virtual TArray<FMovieGraphPinProperties> GetOutputPinProperties() const override;
+	virtual TArray<FPropertyBagPropertyDesc> GetDynamicPropertyDescriptions() const override;
 
-	virtual TArray<FPropertyBagPropertyDesc> GetDynamicPropertyDescriptions() const
-	{
-		TArray<FPropertyBagPropertyDesc> Properties;
-		FPropertyBagPropertyDesc FloatEditConProperty = FPropertyBagPropertyDesc("bOverride_TestPropName", EPropertyBagPropertyType::Bool);
-		FPropertyBagPropertyDesc FloatProperty = FPropertyBagPropertyDesc("TestPropName", EPropertyBagPropertyType::Float);
-#if WITH_EDITOR
-		FloatEditConProperty.MetaData.Add(FPropertyBagPropertyDescMetaData("InlineEditConditionToggle", "true"));
-		FloatProperty.MetaData.Add(FPropertyBagPropertyDescMetaData("EditCondition", "bOverride_TestPropName"));
-#endif
-		
-		Properties.Add(FloatEditConProperty);
-		Properties.Add(FloatProperty);
-		return Properties;
-	}
 #if WITH_EDITOR
-	virtual FText GetMenuDescription() const override
-	{
-		return NSLOCTEXT("MovieGraphNodes", "RenderLayerGraphNode_Description", "Render Layer");
-	}
-	
-	virtual FText GetMenuCategory() const override
-	{
-		return NSLOCTEXT("MovieGraphNodes", "RenderLayerGraphNode_Category", "Rendering");
-	}
+	virtual FText GetMenuDescription() const override;
+	virtual FText GetMenuCategory() const override;
+	virtual FLinearColor GetNodeTitleColor() const override;
+	virtual FSlateIcon GetIconAndTint(FLinearColor& OutColor) const override;
 #endif
 
 	UPROPERTY(BlueprintReadWrite, EditAnywhere, Category = "Settings")
diff --git a/Engine/Plugins/MovieScene/MovieRenderPipeline/Source/MovieRenderPipelineCore/Public/Graph/Nodes/MovieGraphVariableNode.h b/Engine/Plugins/MovieScene/MovieRenderPipeline/Source/MovieRenderPipelineCore/Public/Graph/Nodes/MovieGraphVariableNode.h
new file mode 100644
--- /dev/null
+++ b/Engine/Plugins/MovieScene/MovieRenderPipeline/Source/MovieRenderPipelineCore/Public/Graph/Nodes/MovieGraphVariableNode.h
@@ -0,0 +1,47 @@
+﻿// Copyright Epic Games, Inc. All Rights Reserved.
+
+#pragma once
+
+#include "Graph/MovieGraphNode.h"
+
+#include "MovieGraphVariableNode.generated.h"
+
+/** A node which gets the value of a variable which has been defined on the graph. */
+UCLASS()
+class MOVIERENDERPIPELINECORE_API UMovieGraphVariableNode : public UMovieGraphNode
+{
+	GENERATED_BODY()
+
+public:
+	UMovieGraphVariableNode();
+
+	virtual TArray<FMovieGraphPinProperties> GetOutputPinProperties() const override;
+
+	/** Gets the variable that this node represents. */
+	UMovieGraphVariable* GetVariable() const { return GraphVariable; }
+
+	/** Sets the variable that this node represents. */
+	void SetVariable(UMovieGraphVariable* InVariable);
+
+#if WITH_EDITOR
+	virtual FText GetMenuDescription() const override;
+	virtual FText GetMenuCategory() const override;
+	virtual FLinearColor GetNodeTitleColor() const override;
+	virtual FSlateIcon GetIconAndTint(FLinearColor& OutColor) const override;
+#endif
+
+private:
+	virtual void RegisterDelegates() const override;
+	
+	/** Updates the output pin on the node to match the provided variable. */
+	void UpdateOutputPin(UMovieGraphMember* ChangedVariable) const;
+
+private:
+	/** The underlying graph variable this node represents. */
+	UPROPERTY()
+	TObjectPtr<UMovieGraphVariable> GraphVariable = nullptr;
+
+	/** The properties for the output pin on this node. */
+	UPROPERTY(Transient)
+	FMovieGraphPinProperties OutputPin;
+};
\ No newline at end of file
diff --git a/Engine/Plugins/MovieScene/MovieRenderPipeline/Source/MovieRenderPipelineEditor/Private/Graph/MovieEdGraph.cpp b/Engine/Plugins/MovieScene/MovieRenderPipeline/Source/MovieRenderPipelineEditor/Private/Graph/MovieEdGraph.cpp
--- a/Engine/Plugins/MovieScene/MovieRenderPipeline/Source/MovieRenderPipelineEditor/Private/Graph/MovieEdGraph.cpp
+++ b/Engine/Plugins/MovieScene/MovieRenderPipeline/Source/MovieRenderPipelineEditor/Private/Graph/MovieEdGraph.cpp
@@ -4,6 +4,7 @@
 #include "Graph/MovieGraphConfig.h"
 #include "Graph/MovieGraphEdge.h"
 #include "Graph/MovieGraphPin.h"
+#include "Graph/Nodes/MovieGraphVariableNode.h"
 #include "Graph/MovieGraphNode.h"
 #include "MovieEdGraphOutputNode.h"
 #include "MovieEdGraphInputNode.h"
diff --git a/Engine/Plugins/MovieScene/MovieRenderPipeline/Source/MovieRenderPipelineEditor/Private/Graph/MovieEdGraphNode.cpp b/Engine/Plugins/MovieScene/MovieRenderPipeline/Source/MovieRenderPipelineEditor/Private/Graph/MovieEdGraphNode.cpp
--- a/Engine/Plugins/MovieScene/MovieRenderPipeline/Source/MovieRenderPipelineEditor/Private/Graph/MovieEdGraphNode.cpp
+++ b/Engine/Plugins/MovieScene/MovieRenderPipeline/Source/MovieRenderPipelineEditor/Private/Graph/MovieEdGraphNode.cpp
@@ -267,6 +267,33 @@ void UMoviePipelineEdGraphNodeBase::AutowireNewNode(UEdGraphPin* FromPin)
 	}
 }
 
+FLinearColor UMoviePipelineEdGraphNodeBase::GetNodeTitleColor() const
+{
+	if (RuntimeNode)
+	{
+		return RuntimeNode->GetNodeTitleColor();
+	}
+
+	return FLinearColor::Black;
+}
+
+FSlateIcon UMoviePipelineEdGraphNodeBase::GetIconAndTint(FLinearColor& OutColor) const
+{
+	if (RuntimeNode)
+	{
+		return RuntimeNode->GetIconAndTint(OutColor);
+	}
+
+	OutColor = FLinearColor::White;
+	return FSlateIcon();
+}
+
+bool UMoviePipelineEdGraphNodeBase::ShowPaletteIconOnNode() const
+{
+	// Reveals the icon set by GetIconAndTint() in the top-left corner of the node
+	return true;
+}
+
 void UMoviePipelineEdGraphNodeBase::OnRuntimeNodeChanged(const UMovieGraphNode* InChangedNode)
 {
 	if (InChangedNode == GetRuntimeNode())
diff --git a/Engine/Plugins/MovieScene/MovieRenderPipeline/Source/MovieRenderPipelineEditor/Private/Graph/MovieEdGraphNode.h b/Engine/Plugins/MovieScene/MovieRenderPipeline/Source/MovieRenderPipelineEditor/Private/Graph/MovieEdGraphNode.h
--- a/Engine/Plugins/MovieScene/MovieRenderPipeline/Source/MovieRenderPipelineEditor/Private/Graph/MovieEdGraphNode.h
+++ b/Engine/Plugins/MovieScene/MovieRenderPipeline/Source/MovieRenderPipelineEditor/Private/Graph/MovieEdGraphNode.h
@@ -23,6 +23,9 @@ public:
 	//~ Begin UEdGraphNode Interface
 	virtual void ReconstructNode() override;
 	virtual void AutowireNewNode(UEdGraphPin* FromPin) override;
+	virtual FLinearColor GetNodeTitleColor() const override;
+	virtual FSlateIcon GetIconAndTint(FLinearColor& OutColor) const override;
+	virtual bool ShowPaletteIconOnNode() const override;
 	//~ End UEdGraphNode Interface
 	
 	//~ Begin UObject interface
diff --git a/Engine/Plugins/MovieScene/MovieRenderPipeline/Source/MovieRenderPipelineEditor/Private/Graph/MovieEdGraphVariableNode.cpp b/Engine/Plugins/MovieScene/MovieRenderPipeline/Source/MovieRenderPipelineEditor/Private/Graph/MovieEdGraphVariableNode.cpp
--- a/Engine/Plugins/MovieScene/MovieRenderPipeline/Source/MovieRenderPipelineEditor/Private/Graph/MovieEdGraphVariableNode.cpp
+++ b/Engine/Plugins/MovieScene/MovieRenderPipeline/Source/MovieRenderPipelineEditor/Private/Graph/MovieEdGraphVariableNode.cpp
@@ -4,6 +4,7 @@
 
 #include "EdGraph/EdGraphPin.h"
 #include "Graph/MovieGraphConfig.h"
+#include "Graph/Nodes/MovieGraphVariableNode.h"
 
 #define LOCTEXT_NAMESPACE "MoviePipelineGraph"
 
diff --git a/Engine/Plugins/MovieScene/MovieRenderPipeline/Source/MovieRenderPipelineEditor/Private/Graph/MovieGraphSchema.cpp b/Engine/Plugins/MovieScene/MovieRenderPipeline/Source/MovieRenderPipelineEditor/Private/Graph/MovieGraphSchema.cpp
--- a/Engine/Plugins/MovieScene/MovieRenderPipeline/Source/MovieRenderPipelineEditor/Private/Graph/MovieGraphSchema.cpp
+++ b/Engine/Plugins/MovieScene/MovieRenderPipeline/Source/MovieRenderPipelineEditor/Private/Graph/MovieGraphSchema.cpp
@@ -2,6 +2,9 @@
 
 #include "MovieGraphSchema.h"
 #include "Graph/MovieGraphConfig.h"
+#include "Graph/Nodes/MovieGraphInputNode.h"
+#include "Graph/Nodes/MovieGraphOutputNode.h"
+#include "Graph/Nodes/MovieGraphVariableNode.h"
 #include "Graph/MovieGraphNode.h"
 #include "Graph/MovieEdGraph.h"
 #include "Graph/MovieEdGraphConnectionPolicy.h"
diff --git a/Engine/Plugins/MovieScene/MovieRenderPipeline/Source/MovieRenderPipelineEditor/Private/Widgets/Graph/SMovieGraphConfigPanel.cpp b/Engine/Plugins/MovieScene/MovieRenderPipeline/Source/MovieRenderPipelineEditor/Private/Widgets/Graph/SMovieGraphConfigPanel.cpp
--- a/Engine/Plugins/MovieScene/MovieRenderPipeline/Source/MovieRenderPipelineEditor/Private/Widgets/Graph/SMovieGraphConfigPanel.cpp
+++ b/Engine/Plugins/MovieScene/MovieRenderPipeline/Source/MovieRenderPipelineEditor/Private/Widgets/Graph/SMovieGraphConfigPanel.cpp
@@ -51,6 +51,7 @@
 #include "Misc/FileHelper.h"
 
 #include "Graph/MovieGraphConfig.h"
+#include "Graph/Nodes/MovieGraphVariableNode.h"
 #include "Graph/MovieEdGraph.h"
 #include "Graph/MovieEdGraphNode.h"
 #include "Graph/MovieGraphSchema.h"
