MoviePipeline: Make graph members capable of being non-deletable (eg, for use on the Globals pin). Only deletable Inputs/Outputs are shown in the Members panel.

Also made the Delete*Member() methods private in UMovieGraphConfig; all public requests should go through DeleteMember().

diff --git a/Engine/Plugins/MovieScene/MovieRenderPipeline/Source/MovieRenderPipelineCore/Private/Graph/MovieGraphConfig.cpp b/Engine/Plugins/MovieScene/MovieRenderPipeline/Source/MovieRenderPipelineCore/Private/Graph/MovieGraphConfig.cpp
--- a/Engine/Plugins/MovieScene/MovieRenderPipeline/Source/MovieRenderPipelineCore/Private/Graph/MovieGraphConfig.cpp
+++ b/Engine/Plugins/MovieScene/MovieRenderPipeline/Source/MovieRenderPipelineCore/Private/Graph/MovieGraphConfig.cpp
@@ -12,6 +12,8 @@
 
 #define LOCTEXT_NAMESPACE "MovieGraphConfig"
 
+static const FName MovieGraphGlobalsMemberName("Globals");
+
 #if WITH_EDITOR
 void UMovieGraphVariable::PostEditChangeProperty(FPropertyChangedEvent& PropertyChangedEvent)
 {
@@ -20,6 +22,12 @@ void UMovieGraphVariable::PostEditChangeProperty(FPropertyChangedEvent& Property
 	OnMovieGraphVariableChangedDelegate.Broadcast(this);
 }
 
+bool UMovieGraphInput::IsDeletable() const
+{
+	// The input is deletable as long as it's not the Globals input
+	return Name != MovieGraphGlobalsMemberName;
+}
+
 void UMovieGraphInput::PostEditChangeProperty(FPropertyChangedEvent& PropertyChangedEvent)
 {
 	Super::PostEditChangeProperty(PropertyChangedEvent);
@@ -27,6 +35,12 @@ void UMovieGraphInput::PostEditChangeProperty(FPropertyChangedEvent& PropertyCha
 	OnMovieGraphInputChangedDelegate.Broadcast(this);
 }
 
+bool UMovieGraphOutput::IsDeletable() const
+{
+	// The output is deletable as long as it's not the Globals output
+	return Name != MovieGraphGlobalsMemberName;
+}
+
 void UMovieGraphOutput::PostEditChangeProperty(FPropertyChangedEvent& PropertyChangedEvent)
 {
 	Super::PostEditChangeProperty(PropertyChangedEvent);
@@ -47,6 +61,16 @@ UMovieGraphConfig::UMovieGraphConfig()
 	}
 }
 
+void UMovieGraphConfig::PostInitProperties()
+{
+	Super::PostInitProperties();
+
+	if (!HasAnyFlags(RF_ClassDefaultObject))
+	{
+		AddDefaultMembers();
+	}
+}
+
 void UMovieGraphConfig::TraverseGraphRecursive(UMovieGraphNode* InNode, TSubclassOf<UMovieGraphNode> InClassType, const FMovieGraphTraversalContext& InContext, TArray<UMovieGraphNode*>& OutNodes) const
 {
 	check(InNode);
@@ -91,6 +115,37 @@ void UMovieGraphConfig::TraverseGraphRecursive(UMovieGraphNode* InNode, TSubclas
 
 }
 
+void UMovieGraphConfig::AddDefaultMembers()
+{
+	const bool InputGlobalsExists = Inputs.ContainsByPredicate([](const UMovieGraphMember* Member)
+	{
+		return Member && (Member->Name == MovieGraphGlobalsMemberName);
+	});
+
+	const bool OutputGlobalsExists = Outputs.ContainsByPredicate([](const UMovieGraphMember* Member)
+	{
+		return Member && (Member->Name == MovieGraphGlobalsMemberName);
+	});
+
+	// Ensure there is a Globals input member
+	if (!InputGlobalsExists)
+	{
+		UMovieGraphInput* NewInput = AddInput();
+		NewInput->Name = MovieGraphGlobalsMemberName.ToString();
+
+		InputNode->UpdatePins();
+	}
+
+	// Ensure there is a Globals output member
+	if (!OutputGlobalsExists)
+	{
+		UMovieGraphOutput* NewOutput = AddOutput();
+		NewOutput->Name = MovieGraphGlobalsMemberName.ToString();
+
+		OutputNode->UpdatePins();
+	}
+}
+
 TArray<UMovieGraphNode*> UMovieGraphConfig::TraverseGraph(TSubclassOf<UMovieGraphNode> InClassType, const FMovieGraphTraversalContext& InContext) const
 {
 	TArray<UMovieGraphNode*> OutNodes;
@@ -384,29 +439,44 @@ TArray<UMovieGraphOutput*> UMovieGraphConfig::GetOutputs() const
 	return Outputs;
 }
 
-void UMovieGraphConfig::DeleteMember(UObject* MemberToDelete)
+bool UMovieGraphConfig::DeleteMember(UMovieGraphMember* MemberToDelete)
 {
 	if (!MemberToDelete)
 	{
-		return;
+		return false;
+	}
+
+	if (!MemberToDelete->IsDeletable())
+	{
+		UE_LOG(LogMovieRenderPipeline, Error, TEXT("DeleteMember: The member '%s' cannot be deleted because it is flagged as non-deletable."), *MemberToDelete->Name);
+		return false;
 	}
 
 	if (UMovieGraphVariable* GraphVariableToDelete = Cast<UMovieGraphVariable>(MemberToDelete))
 	{
-		DeleteVariableMember(GraphVariableToDelete);
+		return DeleteVariableMember(GraphVariableToDelete);
 	}
-	else if (UMovieGraphInput* GraphInputToDelete = Cast<UMovieGraphInput>(MemberToDelete))
+
+	if (UMovieGraphInput* GraphInputToDelete = Cast<UMovieGraphInput>(MemberToDelete))
 	{
-		DeleteInputMember(GraphInputToDelete);
+		return DeleteInputMember(GraphInputToDelete);
 	}
-	else if (UMovieGraphOutput* GraphOutputToDelete = Cast<UMovieGraphOutput>(MemberToDelete))
+
+	if (UMovieGraphOutput* GraphOutputToDelete = Cast<UMovieGraphOutput>(MemberToDelete))
 	{
-		DeleteOutputMember(GraphOutputToDelete);
+		return DeleteOutputMember(GraphOutputToDelete);
 	}
+
+	return false;
 }
 
-void UMovieGraphConfig::DeleteVariableMember(UMovieGraphVariable* VariableMemberToDelete)
+bool UMovieGraphConfig::DeleteVariableMember(UMovieGraphVariable* VariableMemberToDelete)
 {
+	if (!VariableMemberToDelete)
+	{
+		return false;
+	}
+	
 	// Find all accessor nodes using this graph variable
 	TArray<TObjectPtr<UMovieGraphNode>> NodesToRemove =
 		AllNodes.FilterByPredicate([VariableMemberToDelete](const TObjectPtr<UMovieGraphNode>& GraphNode)
@@ -439,6 +509,8 @@ void UMovieGraphConfig::DeleteVariableMember(UMovieGraphVariable* VariableMember
 #if WITH_EDITOR
 	OnGraphVariablesChangedDelegate.Broadcast();
 #endif
+
+	return true;
 }
 
 TArray<FMovieGraphBranch> UMovieGraphConfig::GetOutputBranches() const
@@ -456,7 +528,7 @@ TArray<FMovieGraphBranch> UMovieGraphConfig::GetOutputBranches() const
 	return Branches;
 }
 
-void UMovieGraphConfig::DeleteInputMember(UMovieGraphInput* InputMemberToDelete)
+bool UMovieGraphConfig::DeleteInputMember(UMovieGraphInput* InputMemberToDelete)
 {
 	if (InputMemberToDelete)
 	{
@@ -465,10 +537,14 @@ void UMovieGraphConfig::DeleteInputMember(UMovieGraphInput* InputMemberToDelete)
 
 		// This calls OnNodeChangedDelegate to update the graph
 		InputNode->UpdatePins();
+
+		return true;
 	}
+
+	return false;
 }
 
-void UMovieGraphConfig::DeleteOutputMember(UMovieGraphOutput* OutputMemberToDelete)
+bool UMovieGraphConfig::DeleteOutputMember(UMovieGraphOutput* OutputMemberToDelete)
 {
 	if (OutputMemberToDelete)
 	{
@@ -477,7 +553,11 @@ void UMovieGraphConfig::DeleteOutputMember(UMovieGraphOutput* OutputMemberToDele
 
 		// This calls OnNodeChangedDelegate to update the graph
 		OutputNode->UpdatePins();
+
+		return true;
 	}
+
+	return false;
 }
 
 #undef LOCTEXT_NAMESPACE
diff --git a/Engine/Plugins/MovieScene/MovieRenderPipeline/Source/MovieRenderPipelineCore/Public/Graph/MovieGraphConfig.h b/Engine/Plugins/MovieScene/MovieRenderPipeline/Source/MovieRenderPipelineCore/Public/Graph/MovieGraphConfig.h
--- a/Engine/Plugins/MovieScene/MovieRenderPipeline/Source/MovieRenderPipelineCore/Public/Graph/MovieGraphConfig.h
+++ b/Engine/Plugins/MovieScene/MovieRenderPipeline/Source/MovieRenderPipelineCore/Public/Graph/MovieGraphConfig.h
@@ -43,7 +43,7 @@ struct FMovieGraphTraversalContext
 #endif
 
 UCLASS(Abstract)
-class UMovieGraphMember : public UObject
+class MOVIERENDERPIPELINECORE_API UMovieGraphMember : public UObject
 {
 	GENERATED_BODY()
 
@@ -56,6 +56,9 @@ public:
 	/** Sets the GUID that uniquely identifies this member. */
 	void SetGuid(const FGuid& InGuid) { Guid = InGuid; }
 
+	/** Determines if this member can be deleted. */
+	virtual bool IsDeletable() const { return true; }
+
 public:
 	/** The type of data associated with this member. */
 	UPROPERTY(EditAnywhere, Category = "General")
@@ -87,7 +90,7 @@ private:
  * the job level.
  */
 UCLASS(BlueprintType)
-class UMovieGraphVariable : public UMovieGraphMember
+class MOVIERENDERPIPELINECORE_API UMovieGraphVariable : public UMovieGraphMember
 {
 	GENERATED_BODY()
 
@@ -108,13 +111,15 @@ public:
  * An input exposed on the graph that will be available for nodes to connect to.
  */
 UCLASS(BlueprintType)
-class UMovieGraphInput : public UMovieGraphMember
+class MOVIERENDERPIPELINECORE_API UMovieGraphInput : public UMovieGraphMember
 {
 	GENERATED_BODY()
 
 public:
 	UMovieGraphInput() = default;
 
+	virtual bool IsDeletable() const override;
+
 public:
 #if WITH_EDITOR
 	FOnMovieGraphInputChanged OnMovieGraphInputChangedDelegate;
@@ -129,13 +134,15 @@ public:
  * An output exposed on the graph that will be available for nodes to connect to.
  */
 UCLASS(BlueprintType)
-class UMovieGraphOutput : public UMovieGraphMember
+class MOVIERENDERPIPELINECORE_API UMovieGraphOutput : public UMovieGraphMember
 {
 	GENERATED_BODY()
 
 public:
 	UMovieGraphOutput() = default;
 
+	virtual bool IsDeletable() const override;
+
 public:
 #if WITH_EDITOR
 	FOnMovieGraphOutputChanged OnMovieGraphOutputChangedDelegate;
@@ -167,6 +174,10 @@ class MOVIERENDERPIPELINECORE_API UMovieGraphConfig : public UObject
 public:
 	UMovieGraphConfig();
 
+	//~ UObject interface
+	virtual void PostInitProperties() override;
+	//~ End UObject interface
+
 	bool AddLabeledEdge(UMovieGraphNode* FromNode, const FName& FromPinLabel, UMovieGraphNode* ToNode, const FName& ToPinLabel);
 	bool RemoveEdge(UMovieGraphNode* FromNode, const FName& FromPinName, UMovieGraphNode* ToNode, const FName& ToPinName);
 	bool RemoveAllInboundEdges(UMovieGraphNode* InNode);
@@ -204,10 +215,7 @@ public:
 	TArray<UMovieGraphOutput*> GetOutputs() const;
 
 	/** Remove the specified member (input, output, variable) from the graph. */
-	void DeleteMember(UObject* MemberToDelete);
-
-	/** Remove the specified variable member from the graph. */
-	void DeleteVariableMember(UMovieGraphVariable* VariableMemberToDelete);
+	bool DeleteMember(UMovieGraphMember* MemberToDelete);
 
 	/** Returns only the names of the root branches in the Output Node, with no depth information. */
 	TArray<FMovieGraphBranch> GetOutputBranches() const;
@@ -238,12 +246,6 @@ public:
 
 	TArray<UMovieGraphNode*> TraverseGraph(TSubclassOf<UMovieGraphNode> InClassType, const FMovieGraphTraversalContext& InContext) const;
 
-	/** Remove the specified input member from the graph. */
-	void DeleteInputMember(UMovieGraphInput* InputMemberToDelete);
-
-	/** Remove the specified output member from the graph. */
-	void DeleteOutputMember(UMovieGraphOutput* OutputMemberToDelete);
-
 protected:
 	void TraverseGraphRecursive(UMovieGraphNode* InNode, TSubclassOf<UMovieGraphNode> InClassType, const FMovieGraphTraversalContext& InContext, TArray<UMovieGraphNode*>& OutNodes) const;
 
@@ -289,10 +291,22 @@ public:
 	}
 
 private:
+	/** Remove the specified variable member from the graph. */
+	bool DeleteVariableMember(UMovieGraphVariable* VariableMemberToDelete);
+	
+	/** Remove the specified input member from the graph. */
+	bool DeleteInputMember(UMovieGraphInput* InputMemberToDelete);
+
+	/** Remove the specified output member from the graph. */
+	bool DeleteOutputMember(UMovieGraphOutput* OutputMemberToDelete);
+	
 	/** Add a new member of type T to MemberArray, with a unique name that includes BaseName in it. */
 	template<typename T>
 	T* AddMember(TArray<TObjectPtr<T>>& MemberArray, const FText& BaseName);
 
+	/** Adds members to the graph that should always be available. */
+	void AddDefaultMembers();
+
 private:
 	/** All variables which have been defined on the graph. */
 	UPROPERTY()
diff --git a/Engine/Plugins/MovieScene/MovieRenderPipeline/Source/MovieRenderPipelineEditor/Private/Graph/SMovieGraphMembersTabContent.cpp b/Engine/Plugins/MovieScene/MovieRenderPipeline/Source/MovieRenderPipelineEditor/Private/Graph/SMovieGraphMembersTabContent.cpp
--- a/Engine/Plugins/MovieScene/MovieRenderPipeline/Source/MovieRenderPipelineEditor/Private/Graph/SMovieGraphMembersTabContent.cpp
+++ b/Engine/Plugins/MovieScene/MovieRenderPipeline/Source/MovieRenderPipelineEditor/Private/Graph/SMovieGraphMembersTabContent.cpp
@@ -61,7 +61,8 @@ void SMovieGraphMembersTabContent::DeleteSelectedMembers() const
 	{
 		if (const FMovieGraphSchemaAction* MovieGraphAction = static_cast<FMovieGraphSchemaAction*>(SelectedAction.Get()))
 		{
-			CurrentGraph->DeleteMember(MovieGraphAction->ActionTarget.Get());
+			UMovieGraphMember* GraphMember = Cast<UMovieGraphMember>(MovieGraphAction->ActionTarget.Get());
+			CurrentGraph->DeleteMember(GraphMember);
 		}
 	}
 
@@ -103,7 +104,7 @@ void SMovieGraphMembersTabContent::CollectAllActions(FGraphActionListBuilderBase
 
 	for (UMovieGraphInput* Input : CurrentGraph->GetInputs())
 	{
-		if (Input)
+		if (Input && Input->IsDeletable())
 		{
 			AddToActionMenu(Input, EActionSection::Inputs);
             
@@ -114,7 +115,7 @@ void SMovieGraphMembersTabContent::CollectAllActions(FGraphActionListBuilderBase
 
 	for (UMovieGraphOutput* Output : CurrentGraph->GetOutputs())
 	{
-		if (Output)
+		if (Output && Output->IsDeletable())
 		{
 			AddToActionMenu(Output, EActionSection::Outputs);
 
