MoviePipeline: Allow for creating pinnable comment bubbles on nodes in the graph. These comments are placed on individual nodes, and are not graph-wide comments (backdrops).

diff --git a/Engine/Plugins/MovieScene/MovieRenderPipeline/Source/MovieRenderPipelineCore/Public/Graph/MovieGraphNode.h b/Engine/Plugins/MovieScene/MovieRenderPipeline/Source/MovieRenderPipelineCore/Public/Graph/MovieGraphNode.h
--- a/Engine/Plugins/MovieScene/MovieRenderPipeline/Source/MovieRenderPipelineCore/Public/Graph/MovieGraphNode.h
+++ b/Engine/Plugins/MovieScene/MovieRenderPipeline/Source/MovieRenderPipelineCore/Public/Graph/MovieGraphNode.h
@@ -89,6 +89,15 @@ public:
 
 	/** Gets the node's icon and icon tint, as visible in the graph. */
 	virtual FSlateIcon GetIconAndTint(FLinearColor& OutColor) const;
+
+	FString GetNodeComment() const { return NodeComment; }
+	void SetNodeComment(const FString& InNodeComment) { NodeComment = InNodeComment; }
+
+	bool IsCommentBubblePinned() const { return bIsCommentBubblePinned; }
+	void SetIsCommentBubblePinned(const uint8 bIsPinned) { bIsCommentBubblePinned = bIsPinned; }
+
+	bool IsCommentBubbleVisible() const { return bIsCommentBubbleVisible; }
+	void SetIsCommentBubbleVisible(uint8 bIsVisible) { bIsCommentBubbleVisible = bIsVisible; }
 #endif
 
 	//~ Begin UObject Interface
@@ -136,6 +145,15 @@ protected:
 
 	UPROPERTY()
 	int32 NodePosY = 0;
+
+	UPROPERTY()
+	FString NodeComment;
+
+	UPROPERTY()
+	uint8 bIsCommentBubblePinned : 1;
+
+	UPROPERTY()
+	uint8 bIsCommentBubbleVisible : 1;
 #endif
 
 	/** A GUID which uniquely identifies this node. */
diff --git a/Engine/Plugins/MovieScene/MovieRenderPipeline/Source/MovieRenderPipelineEditor/Private/Graph/MovieEdGraphNode.cpp b/Engine/Plugins/MovieScene/MovieRenderPipeline/Source/MovieRenderPipelineEditor/Private/Graph/MovieEdGraphNode.cpp
--- a/Engine/Plugins/MovieScene/MovieRenderPipeline/Source/MovieRenderPipelineEditor/Private/Graph/MovieEdGraphNode.cpp
+++ b/Engine/Plugins/MovieScene/MovieRenderPipeline/Source/MovieRenderPipelineEditor/Private/Graph/MovieEdGraphNode.cpp
@@ -23,9 +23,10 @@ void UMoviePipelineEdGraphNodeBase::Construct(UMovieGraphNode* InRuntimeNode)
 	
 	NodePosX = InRuntimeNode->GetNodePosX();
 	NodePosY = InRuntimeNode->GetNodePosY();
-	// NodeComment = InRuntimeNode->NodeComment;
-	// bCommentBubblePinned = InRuntimeNode->bCommentBubblePinned;
-	// bCommentBubbleVisible = InRuntimeNode->bCommentBubbleVisible;
+	
+	NodeComment = InRuntimeNode->GetNodeComment();
+	bCommentBubblePinned = InRuntimeNode->IsCommentBubblePinned();
+	bCommentBubbleVisible = InRuntimeNode->IsCommentBubbleVisible();
 }
 
 void UMoviePipelineEdGraphNodeBase::PostTransacted(const FTransactionObjectEvent& TransactionEvent)
@@ -37,6 +38,11 @@ void UMoviePipelineEdGraphNodeBase::PostTransacted(const FTransactionObjectEvent
 	{
 		UpdatePosition();
 	}
+
+	if (ChangedProperties.Contains(GET_MEMBER_NAME_CHECKED(UEdGraphNode, bCommentBubblePinned)))
+	{
+		UpdateCommentBubblePinned();
+	}
 }
 
 FEdGraphPinType UMoviePipelineEdGraphNodeBase::GetPinType(const UMovieGraphPin* InPin)
@@ -75,7 +81,7 @@ FEdGraphPinType UMoviePipelineEdGraphNodeBase::GetPinType(const UMovieGraphPin*
 	return EdPinType;
 }
 
-void UMoviePipelineEdGraphNodeBase::UpdatePosition()
+void UMoviePipelineEdGraphNodeBase::UpdatePosition() const
 {
 	if (RuntimeNode)
 	{
@@ -85,6 +91,15 @@ void UMoviePipelineEdGraphNodeBase::UpdatePosition()
 	}
 }
 
+void UMoviePipelineEdGraphNodeBase::UpdateCommentBubblePinned() const
+{
+	if (RuntimeNode)
+	{
+		RuntimeNode->Modify();
+		RuntimeNode->SetIsCommentBubblePinned(bCommentBubblePinned);
+	}
+}
+
 void UMoviePipelineEdGraphNode::AllocateDefaultPins()
 {
 	if(RuntimeNode)
@@ -294,6 +309,26 @@ bool UMoviePipelineEdGraphNodeBase::ShowPaletteIconOnNode() const
 	return true;
 }
 
+void UMoviePipelineEdGraphNodeBase::OnUpdateCommentText(const FString& NewComment)
+{
+	Super::OnUpdateCommentText(NewComment);
+
+	if (RuntimeNode && (RuntimeNode->GetNodeComment() != NewComment))
+	{
+		RuntimeNode->SetNodeComment(NewComment);
+	}
+}
+
+void UMoviePipelineEdGraphNodeBase::OnCommentBubbleToggled(bool bInCommentBubbleVisible)
+{
+	Super::OnCommentBubbleToggled(bInCommentBubbleVisible);
+
+	if (RuntimeNode && (RuntimeNode->IsCommentBubbleVisible() != bInCommentBubbleVisible))
+	{
+		RuntimeNode->SetIsCommentBubbleVisible(bInCommentBubbleVisible);
+	}
+}
+
 void UMoviePipelineEdGraphNodeBase::OnRuntimeNodeChanged(const UMovieGraphNode* InChangedNode)
 {
 	if (InChangedNode == GetRuntimeNode())
diff --git a/Engine/Plugins/MovieScene/MovieRenderPipeline/Source/MovieRenderPipelineEditor/Private/Graph/MovieEdGraphNode.h b/Engine/Plugins/MovieScene/MovieRenderPipeline/Source/MovieRenderPipelineEditor/Private/Graph/MovieEdGraphNode.h
--- a/Engine/Plugins/MovieScene/MovieRenderPipeline/Source/MovieRenderPipelineEditor/Private/Graph/MovieEdGraphNode.h
+++ b/Engine/Plugins/MovieScene/MovieRenderPipeline/Source/MovieRenderPipelineEditor/Private/Graph/MovieEdGraphNode.h
@@ -26,6 +26,8 @@ public:
 	virtual FLinearColor GetNodeTitleColor() const override;
 	virtual FSlateIcon GetIconAndTint(FLinearColor& OutColor) const override;
 	virtual bool ShowPaletteIconOnNode() const override;
+	virtual void OnUpdateCommentText(const FString& NewComment) override;
+	virtual void OnCommentBubbleToggled(bool bInCommentBubbleVisible) override;
 	//~ End UEdGraphNode Interface
 	
 	//~ Begin UObject interface
@@ -42,7 +44,14 @@ protected:
 	/** Recreate the pins on this node, discarding all existing pins. */
 	void ReconstructPins();
 
-	void UpdatePosition();
+	/** Update the position of the underlying runtime node to match the editor node. */
+	void UpdatePosition() const;
+
+	/**
+	 * Update the comment bubble pin state on the underlying runtime node to match the editor node. Note that
+	 * comment bubble visibility state and comment text are handled by UEdGraphNode overrides.
+	 */
+	void UpdateCommentBubblePinned() const;
 
 protected:
 	/** The runtime node that this editor node represents. */
