MoviePipeline: Fix build issue.

diff --git a/Engine/Plugins/MovieScene/MovieRenderPipeline/Source/MovieRenderPipelineCore/Private/Graph/MovieGraphConfig.cpp b/Engine/Plugins/MovieScene/MovieRenderPipeline/Source/MovieRenderPipelineCore/Private/Graph/MovieGraphConfig.cpp
--- a/Engine/Plugins/MovieScene/MovieRenderPipeline/Source/MovieRenderPipelineCore/Private/Graph/MovieGraphConfig.cpp
+++ b/Engine/Plugins/MovieScene/MovieRenderPipeline/Source/MovieRenderPipelineCore/Private/Graph/MovieGraphConfig.cpp
@@ -21,6 +21,7 @@ void UMovieGraphVariable::PostEditChangeProperty(FPropertyChangedEvent& Property
 
 	OnMovieGraphVariableChangedDelegate.Broadcast(this);
 }
+#endif // WITH_EDITOR
 
 bool UMovieGraphInput::IsDeletable() const
 {
@@ -28,12 +29,14 @@ bool UMovieGraphInput::IsDeletable() const
 	return Name != MovieGraphGlobalsMemberName;
 }
 
+#if WITH_EDITOR
 void UMovieGraphInput::PostEditChangeProperty(FPropertyChangedEvent& PropertyChangedEvent)
 {
 	Super::PostEditChangeProperty(PropertyChangedEvent);
 	
 	OnMovieGraphInputChangedDelegate.Broadcast(this);
 }
+#endif // WITH_EDITOR
 
 bool UMovieGraphOutput::IsDeletable() const
 {
@@ -41,13 +44,14 @@ bool UMovieGraphOutput::IsDeletable() const
 	return Name != MovieGraphGlobalsMemberName;
 }
 
+#if WITH_EDITOR
 void UMovieGraphOutput::PostEditChangeProperty(FPropertyChangedEvent& PropertyChangedEvent)
 {
 	Super::PostEditChangeProperty(PropertyChangedEvent);
 
 	OnMovieGraphOutputChangedDelegate.Broadcast(this);
 }
-#endif
+#endif // WITH_EDITOR
 
 UMovieGraphConfig::UMovieGraphConfig()
 {
