MoviePipeline: Set up the bare bones of the GPU->CPU Readback in preparation for creating an output merger and then writing it to disk. This sorts out some of the time tracking stuff that Movie Render Queue needs for rendering, and creates surface pools for GPU->CPU readback.

diff --git a/Engine/Plugins/MovieScene/MovieRenderPipeline/Source/MovieRenderPipelineCore/Private/Graph/MovieGraphDefaultRenderer.cpp b/Engine/Plugins/MovieScene/MovieRenderPipeline/Source/MovieRenderPipelineCore/Private/Graph/MovieGraphDefaultRenderer.cpp
--- a/Engine/Plugins/MovieScene/MovieRenderPipeline/Source/MovieRenderPipelineCore/Private/Graph/MovieGraphDefaultRenderer.cpp
+++ b/Engine/Plugins/MovieScene/MovieRenderPipeline/Source/MovieRenderPipelineCore/Private/Graph/MovieGraphDefaultRenderer.cpp
@@ -9,6 +9,7 @@
 #include "RenderingThread.h"
 #include "Engine/TextureRenderTarget2D.h"
 #include "UObject/Package.h"
+#include "MoviePipelineSurfaceReader.h"
 
 void UMovieGraphDefaultRenderer::SetupRenderingPipelineForShot(UMoviePipelineExecutorShot* InShot)
 {
@@ -91,8 +92,22 @@ void UMovieGraphDefaultRenderer::SetupRenderingPipelineForShot(UMoviePipelineExe
 
 void UMovieGraphDefaultRenderer::TeardownRenderingPipelineForShot(UMoviePipelineExecutorShot* InShot)
 {
+	// Ensure the GPU has actually rendered all of the frames
 	FlushRenderingCommands();
 
+	// Make sure all of the data has actually been copied back from the GPU and accumulation tasks started.
+	for (TPair<UE::MovieGraph::DefaultRenderer::FRenderTargetInitParams, FMoviePipelineSurfaceQueuePtr> SurfaceQueueIt : PooledSurfaceQueues)
+	{
+		if (SurfaceQueueIt.Value.IsValid())
+		{
+			SurfaceQueueIt.Value->Shutdown();
+		}
+	}
+	
+	// Stall until the task graph has completed any pending accumulations.
+	FTaskGraphInterface::Get().WaitUntilTasksComplete(OutstandingTasks, ENamedThreads::GameThread);
+	OutstandingTasks.Reset();
+
 	for (TObjectPtr< UMovieGraphRenderPassNode> RenderPass : RenderPassesInUse)
 	{
 		RenderPass->Teardown();
@@ -102,6 +117,8 @@ void UMovieGraphDefaultRenderer::TeardownRenderingPipelineForShot(UMoviePipeline
 
 	// ToDo: This could probably be preserved across shots to avoid allocations
 	PooledViewRenderTargets.Reset();
+	PooledAccumulators.Reset();
+	PooledSurfaceQueues.Reset();
 }
 
 
@@ -111,27 +128,54 @@ void UMovieGraphDefaultRenderer::AddReferencedObjects(UObject* InThis, FReferenc
 	UMovieGraphDefaultRenderer* This = CastChecked<UMovieGraphDefaultRenderer>(InThis);
 
 	// Can't be a const& due to AddStableReference API
-	for (TPair<FRenderTargetInitParams, TObjectPtr<UTextureRenderTarget2D>>& KVP : This->PooledViewRenderTargets)
+	for (TPair<UE::MovieGraph::DefaultRenderer::FRenderTargetInitParams, TObjectPtr<UTextureRenderTarget2D>>& KVP : This->PooledViewRenderTargets)
 	{
 		Collector.AddStableReference(&KVP.Value);
 	}
 }
 
-
 void UMovieGraphDefaultRenderer::Render(const FMovieGraphTimeStepData& InTimeStepData)
 {
+	// Housekeeping: Clean up any tasks that were completed since the last frame. This lets us have a 
+	// better idea of how much work we're concurrently doing. 
+	{
+		// Don't iterate through the array if someone is actively modifying it.
+		FScopeLock ScopeLock(&OutstandingTasksMutex);
+		for (int32 Index = 0; Index < OutstandingTasks.Num(); Index++)
+		{
+			const FGraphEventRef& Task = OutstandingTasks[Index];
+			if (Task->IsComplete())
+			{
+				Task->Release();
+				OutstandingTasks.RemoveAtSwap(Index);
+
+				// We swapped the end array element into the current spot,
+				// so we need to check this element again, otherwise we skip
+				// checking some things.
+				Index--;
+			}
+		}
+	}
+
 	for (TObjectPtr<UMovieGraphRenderPassNode> RenderPass : RenderPassesInUse)
 	{
-		RenderPass->Render();
+		RenderPass->Render(InTimeStepData);
 	}
 }
 
+void UMovieGraphDefaultRenderer::AddOutstandingRenderTask_AnyThread(FGraphEventRef InTask)
+{
+	// We might be looping through the array to remove previously completed tasks,
+	// so don't modify until that is completed.
+	FScopeLock ScopeLock(&OutstandingTasksMutex);
+	OutstandingTasks.Add(InTask);
+}
 
-UTextureRenderTarget2D* UMovieGraphDefaultRenderer::GetOrCreateViewRenderTarget(const FRenderTargetInitParams& InInitParams)
+UTextureRenderTarget2D* UMovieGraphDefaultRenderer::GetOrCreateViewRenderTarget(const UE::MovieGraph::DefaultRenderer::FRenderTargetInitParams& InInitParams)
 {
-	if (const TObjectPtr<UTextureRenderTarget2D>* ExistViewRenderTarget = PooledViewRenderTargets.Find(InInitParams))
+	if (const TObjectPtr<UTextureRenderTarget2D>* ExistingViewRenderTarget = PooledViewRenderTargets.Find(InInitParams))
 	{
-		return *ExistViewRenderTarget;
+		return *ExistingViewRenderTarget;
 	}
 
 	const TObjectPtr<UTextureRenderTarget2D> NewViewRenderTarget = CreateViewRenderTarget(InInitParams);
@@ -140,7 +184,7 @@ UTextureRenderTarget2D* UMovieGraphDefaultRenderer::GetOrCreateViewRenderTarget(
 	return NewViewRenderTarget.Get();
 }
 
-TObjectPtr<UTextureRenderTarget2D> UMovieGraphDefaultRenderer::CreateViewRenderTarget(const FRenderTargetInitParams& InInitParams) const
+TObjectPtr<UTextureRenderTarget2D> UMovieGraphDefaultRenderer::CreateViewRenderTarget(const UE::MovieGraph::DefaultRenderer::FRenderTargetInitParams& InInitParams) const
 {
 	TObjectPtr<UTextureRenderTarget2D> NewTarget = NewObject<UTextureRenderTarget2D>(GetTransientPackage());
 	NewTarget->ClearColor = FLinearColor(0.0f, 0.0f, 0.0f, 0.0f);
@@ -151,3 +195,66 @@ TObjectPtr<UTextureRenderTarget2D> UMovieGraphDefaultRenderer::CreateViewRenderT
 
 	return NewTarget;
 }
+
+FMoviePipelineSurfaceQueuePtr UMovieGraphDefaultRenderer::GetOrCreateSurfaceQueue(const UE::MovieGraph::DefaultRenderer::FRenderTargetInitParams& InInitParams)
+{
+	if (const FMoviePipelineSurfaceQueuePtr* ExistingSurfaceQueue = PooledSurfaceQueues.Find(InInitParams))
+	{
+		return *ExistingSurfaceQueue;
+	}
+
+	const FMoviePipelineSurfaceQueuePtr NewSurfaceQueue = CreateSurfaceQueue(InInitParams);
+	PooledSurfaceQueues.Emplace(InInitParams, NewSurfaceQueue);
+
+	return NewSurfaceQueue;
+}
+
+FMoviePipelineSurfaceQueuePtr UMovieGraphDefaultRenderer::CreateSurfaceQueue(const UE::MovieGraph::DefaultRenderer::FRenderTargetInitParams& InInitParams) const
+{
+	// ToDo: Refactor these to be dynamically sized, but also allow putting a cap on them. We need at least enough to submit everything for one render
+	FMoviePipelineSurfaceQueuePtr SurfaceQueue = MakeShared<FMoviePipelineSurfaceQueue, ESPMode::ThreadSafe>(InInitParams.Size, InInitParams.PixelFormat, 6, true);
+	
+	UE_LOG(LogMovieRenderPipeline, Log, TEXT("Allocated a Surface Queue sized: (%d, %d)"), InInitParams.Size.X, InInitParams.Size.Y);
+	return SurfaceQueue;
+}
+
+namespace UE::MovieGraph::DefaultRenderer
+{
+FSurfaceAccumulatorPool::FInstancePtr FSurfaceAccumulatorPool::BlockAndGetAccumulator_GameThread(int32 InFrameNumber, const FMovieGraphRenderDataIdentifier& InPassIdentifier)
+{
+	FScopeLock ScopeLock(&CriticalSection);
+
+	int32 AvailableIndex = INDEX_NONE;
+	while (AvailableIndex == INDEX_NONE)
+	{
+		for (int32 Index = 0; Index < Accumulators.Num(); Index++)
+		{
+			if (InFrameNumber == Accumulators[Index]->ActiveFrameNumber && InPassIdentifier == Accumulators[Index]->ActivePassIdentifier)
+			{
+				AvailableIndex = Index;
+				break;
+			}
+		}
+
+		if (AvailableIndex == INDEX_NONE)
+		{
+			// If we don't have an accumulator already working on it let's look for a free one.
+			for (int32 Index = 0; Index < Accumulators.Num(); Index++)
+			{
+				if (!Accumulators[Index]->IsActive())
+				{
+					// Found a free one, tie it to this output frame.
+					Accumulators[Index]->ActiveFrameNumber = InFrameNumber;
+					Accumulators[Index]->ActivePassIdentifier = InPassIdentifier;
+					Accumulators[Index]->bIsActive = true;
+					Accumulators[Index]->TaskPrereq = nullptr;
+					AvailableIndex = Index;
+					break;
+				}
+			}
+		}
+	}
+
+	return Accumulators[AvailableIndex];
+}
+}
\ No newline at end of file
diff --git a/Engine/Plugins/MovieScene/MovieRenderPipeline/Source/MovieRenderPipelineCore/Private/Graph/MovieGraphLinearTimeStep.cpp b/Engine/Plugins/MovieScene/MovieRenderPipeline/Source/MovieRenderPipelineCore/Private/Graph/MovieGraphLinearTimeStep.cpp
--- a/Engine/Plugins/MovieScene/MovieRenderPipeline/Source/MovieRenderPipelineCore/Private/Graph/MovieGraphLinearTimeStep.cpp
+++ b/Engine/Plugins/MovieScene/MovieRenderPipeline/Source/MovieRenderPipelineCore/Private/Graph/MovieGraphLinearTimeStep.cpp
@@ -1,8 +1,58 @@
 // Copyright Epic Games, Inc. All Rights Reserved.
 
 #include "Graph/MovieGraphLinearTimeStep.h"
+#include "Graph/MovieGraphPipeline.h"
+#include "MoviePipelineQueue.h"
+#include "MovieRenderPipelineCoreModule.h"
+#include "LevelSequence.h"
+#include "MovieScene.h"
 
-void UMovieGraphLinearTimeStep::CalculateTimeStep(FMovieGraphTimeStepData& OutTimeData)
+void UMovieGraphLinearTimeStep::TickProducingFrames()
 {
+	int32 CurrentShotIndex = GetOwningGraph()->GetCurrentShotIndex();
+	const TArray<TObjectPtr<UMoviePipelineExecutorShot>>& ActiveShotList = GetOwningGraph()->GetActiveShotList();
+	UMoviePipelineExecutorShot* CurrentCameraCut = ActiveShotList[CurrentShotIndex];
 
-}
\ No newline at end of file
+	if (CurrentCameraCut->ShotInfo.State == EMovieRenderShotState::Uninitialized)
+	{
+		UE_LOG(LogMovieRenderPipeline, Log, TEXT("MovieGraph Initializing Camera Cut [%d/%d] in [%s] %s."),
+			CurrentShotIndex + 1, ActiveShotList.Num(), *CurrentCameraCut->OuterName, *CurrentCameraCut->InnerName);
+
+		GetOwningGraph()->SetupShot(CurrentCameraCut);
+		// InitializeShot(CurrentCameraCut);
+		GetOwningGraph()->GetDataCachingInstance()->InitializeShot(CurrentCameraCut);
+
+		// We can safely fall through to the below states as they're OK to process the same frame we set up.
+		UE_LOG(LogMovieRenderPipeline, Log, TEXT("MovieGraph Finished initializing Camera Cut [%d/%d] in [%s] %s."),
+			CurrentShotIndex + 1, ActiveShotList.Num(), *CurrentCameraCut->OuterName, *CurrentCameraCut->InnerName);
+
+		// Temp...
+		CurrentCameraCut->ShotInfo.State = EMovieRenderShotState::Rendering;
+	}
+
+	
+
+	if (CurrentCameraCut->ShotInfo.State == EMovieRenderShotState::Rendering)
+	{
+		ULevelSequence* CurrentSequence = Cast<ULevelSequence>(GetOwningGraph()->GetCurrentJob()->Sequence.TryLoad());
+		FFrameRate TickResolution = CurrentSequence->GetMovieScene()->GetTickResolution();
+		FFrameRate FrameRate = CurrentSequence->GetMovieScene()->GetDisplayRate();
+
+		FFrameTime TicksPerOutputFrame = FFrameRate::TransformTime(FFrameTime(FFrameNumber(1)), FrameRate, TickResolution);
+		CurrentCameraCut->ShotInfo.CurrentTickInRoot += TicksPerOutputFrame.GetFrame();
+
+		// Convert the CurrentTickInRoot back to a frame number.
+		FFrameTime CurrentTickInRootInOutputFrames = FFrameRate::TransformTime(CurrentCameraCut->ShotInfo.CurrentTickInRoot, TickResolution, FrameRate);
+		CurrentTimeStepData.OutputFrameNumber = CurrentTickInRootInOutputFrames.FloorToFrame().Value;
+		CurrentTimeStepData.FrameDeltaTime = FrameRate.AsInterval();
+		CurrentTimeStepData.WorldSeconds = 0.f;
+		CurrentTimeStepData.MotionBlurFraction = FrameRate.AsInterval();
+
+		if (CurrentCameraCut->ShotInfo.CurrentTickInRoot >= CurrentCameraCut->ShotInfo.TotalOutputRangeRoot.GetUpperBoundValue())
+		{
+			CurrentCameraCut->ShotInfo.State = EMovieRenderShotState::Finished;
+			GetOwningGraph()->TeardownShot(CurrentCameraCut);
+			return;
+		}
+	}
+}
diff --git a/Engine/Plugins/MovieScene/MovieRenderPipeline/Source/MovieRenderPipelineCore/Private/Graph/MovieGraphPipeline.cpp b/Engine/Plugins/MovieScene/MovieRenderPipeline/Source/MovieRenderPipelineCore/Private/Graph/MovieGraphPipeline.cpp
--- a/Engine/Plugins/MovieScene/MovieRenderPipeline/Source/MovieRenderPipelineCore/Private/Graph/MovieGraphPipeline.cpp
+++ b/Engine/Plugins/MovieScene/MovieRenderPipeline/Source/MovieRenderPipelineCore/Private/Graph/MovieGraphPipeline.cpp
@@ -69,7 +69,8 @@ void UMovieGraphPipeline::Initialize(UMoviePipelineExecutorJob* InJob, const FMo
 	
 	CurrentJob = InJob;
 	CurrentShotIndex = 0;
-	
+	GraphInitializationTime = FDateTime::UtcNow();
+
 	// Now that we've created our various systems, we will start using them. First thing we do is cache data about
 	// the world, job, player viewport, etc, before we make any modifications. These will be restored at the end
 	// of the render.
@@ -198,59 +199,7 @@ void UMovieGraphPipeline::TickProducingFrames()
 	//}
 
 
-	FMovieGraphTimeStepData TimeStepData;
-
-	GraphTimeStepInstance->ResetPerTickData();
-
-	UMoviePipelineExecutorShot* CurrentCameraCut = ActiveShotList[CurrentShotIndex];
-	if (CurrentCameraCut->ShotInfo.State == EMovieRenderShotState::Uninitialized)
-	{
-		UE_LOG(LogMovieRenderPipeline, Log, TEXT("MovieGraph Initializing Camera Cut [%d/%d] in [%s] %s."),
-			CurrentShotIndex + 1, ActiveShotList.Num(), *CurrentCameraCut->OuterName, *CurrentCameraCut->InnerName);
-
-		SetupShot(CurrentCameraCut);
-		GraphTimeStepInstance->InitializeShot(CurrentCameraCut, /*Out*/ TimeStepData);
-		GraphDataCachingInstance->InitializeShot(CurrentCameraCut, /*In*/ TimeStepData);
-
-		// We can safely fall through to the below states as they're OK to process the same frame we set up.
-		UE_LOG(LogMovieRenderPipeline, Log, TEXT("MovieGraph Finished initializing Camera Cut [%d/%d] in [%s] %s."),
-			CurrentShotIndex + 1, ActiveShotList.Num(), *CurrentCameraCut->OuterName, *CurrentCameraCut->InnerName);
-
-		// Temp...
-		CurrentCameraCut->ShotInfo.State = EMovieRenderShotState::Rendering;
-	}
-
-	if (CurrentCameraCut->ShotInfo.State == EMovieRenderShotState::Rendering)
-	{
-		ULevelSequence* CurrentSequence = Cast<ULevelSequence>(GetCurrentJob()->Sequence.TryLoad());
-		FFrameRate TickResolution = CurrentSequence->GetMovieScene()->GetTickResolution();
-		FFrameRate FrameRate = CurrentSequence->GetMovieScene()->GetDisplayRate();
-		
-		FFrameTime TicksPerOutputFrame = FFrameRate::TransformTime(FFrameTime(FFrameNumber(1)), FrameRate, TickResolution);
-		CurrentCameraCut->ShotInfo.CurrentTickInRoot += TicksPerOutputFrame.GetFrame();
-
-		if (CurrentCameraCut->ShotInfo.CurrentTickInRoot >= CurrentCameraCut->ShotInfo.TotalOutputRangeRoot.GetUpperBoundValue())
-		{
-			CurrentCameraCut->ShotInfo.State = EMovieRenderShotState::Finished;
-			TeardownShot(CurrentCameraCut);
-			return;
-		}
-	}
-
-	// const MoviePipeline::FFrameConstantMetrics FrameMetrics = CalculateShotFrameMetrics(CurrentCameraCut);
-	// CachedOutputState.ResetPerFrameData();
-
-
-	// Public API:
-	// Which Output Frame are we on (overall)
-
-	// Output State needs to know:
-	// ShotIndex
-	// does output state need sto be an interface? stocastic sampling makes
-	// temporal sample count/temporal index not have any meaning. or maybe
-	// we just re-use it and have some fields be dummies? This is used for
-	// the UI (for progress tracking), and filenames, and debug file output names.
-
+	GraphTimeStepInstance->TickProducingFrames();
 }
 
 
@@ -419,7 +368,7 @@ void UMovieGraphPipeline::OnEngineTickEndFrame()
 void UMovieGraphPipeline::RenderFrame()
 {
 	check(GraphRendererInstance);
-	FMovieGraphTimeStepData TimeStepData;
+	const FMovieGraphTimeStepData& TimeStepData = GraphTimeStepInstance->GetCalculatedTimeData();
 
 	GraphRendererInstance->Render(TimeStepData);
 }
@@ -644,6 +593,7 @@ void UMovieGraphPipeline::TransitionToState(const EMovieRenderPipelineState InNe
 			//}
 
 			OnMoviePipelineFinishedImpl();
+			UE_LOG(LogMovieRenderPipeline, Log, TEXT("Movie Graph Render completed. Duration: %s"), *(FDateTime::UtcNow() - GraphInitializationTime).ToString());
 		}
 		break;
 	}
diff --git a/Engine/Plugins/MovieScene/MovieRenderPipeline/Source/MovieRenderPipelineCore/Private/Graph/Nodes/MovieGraphDeferredRenderPassNode.cpp b/Engine/Plugins/MovieScene/MovieRenderPipeline/Source/MovieRenderPipelineCore/Private/Graph/Nodes/MovieGraphDeferredRenderPassNode.cpp
--- a/Engine/Plugins/MovieScene/MovieRenderPipeline/Source/MovieRenderPipelineCore/Private/Graph/Nodes/MovieGraphDeferredRenderPassNode.cpp
+++ b/Engine/Plugins/MovieScene/MovieRenderPipeline/Source/MovieRenderPipelineCore/Private/Graph/Nodes/MovieGraphDeferredRenderPassNode.cpp
@@ -13,6 +13,8 @@
 #include "SceneViewExtensionContext.h"
 #include "OpenColorIODisplayExtension.h"
 #include "TextureResource.h"
+#include "MovieRenderOverlappedImage.h"
+#include "MoviePipelineSurfaceReader.h"
 
 void UMovieGraphDeferredRenderPassNode::SetupImpl(const FMovieGraphRenderPassSetupData& InSetupData)
 {
@@ -22,7 +24,7 @@ void UMovieGraphDeferredRenderPassNode::SetupImpl(const FMovieGraphRenderPassSet
 	for (const FMovieGraphRenderPassLayerData& LayerData : InSetupData.Layers)
 	{
 		TUniquePtr<FMovieGraphDeferredRenderPass> RendererInstance = MakeUnique<FMovieGraphDeferredRenderPass>();
-		RendererInstance->Setup(InSetupData.Renderer, LayerData);
+		RendererInstance->Setup(InSetupData.Renderer, this, LayerData);
 		CurrentInstances.Add(MoveTemp(RendererInstance));
 	}
 }
@@ -40,11 +42,11 @@ void UMovieGraphDeferredRenderPassNode::TeardownImpl()
 }
 
 
-void UMovieGraphDeferredRenderPassNode::RenderImpl()
+void UMovieGraphDeferredRenderPassNode::RenderImpl(const FMovieGraphTimeStepData& InTimeData)
 {
 	for (TUniquePtr<FMovieGraphDeferredRenderPass>& Instance : CurrentInstances)
 	{
-		Instance->Render();
+		Instance->Render(InTimeData);
 	}
 }
 
@@ -59,12 +61,12 @@ void UMovieGraphDeferredRenderPassNode::AddReferencedObjects(UObject* InThis, FR
 	}
 }
 
-void UMovieGraphDeferredRenderPassNode::FMovieGraphDeferredRenderPass::Setup(TWeakObjectPtr<class UMovieGraphDefaultRenderer> InRenderer, const FMovieGraphRenderPassLayerData& InLayer)
+void UMovieGraphDeferredRenderPassNode::FMovieGraphDeferredRenderPass::Setup(TWeakObjectPtr<UMovieGraphDefaultRenderer> InRenderer, TWeakObjectPtr<UMovieGraphDeferredRenderPassNode> InRenderPassNode, const FMovieGraphRenderPassLayerData& InLayer)
 {
 	LayerData = InLayer;
 	Renderer = InRenderer;
+	RenderPassNode = InRenderPassNode;
 
-	const FIntPoint OutputCameraResolution = FIntPoint(1920, 1080);// Renderer->GetEffectiveOutputResolution(LayerData);
 
 	// Figure out how big each sub-region (tile) is.
 	// const int32 TileSize = Graph->FindSetting(LayerData.LayerName, "highres.tileSize");
@@ -102,15 +104,15 @@ void UMovieGraphDeferredRenderPassNode::FMovieGraphDeferredRenderPass::AddRefere
 	}
 }
 
-void UMovieGraphDeferredRenderPassNode::FMovieGraphDeferredRenderPass::Render()
+void UMovieGraphDeferredRenderPassNode::FMovieGraphDeferredRenderPass::Render(const FMovieGraphTimeStepData& InTimeData)
 {
-
+	// This is the size we actually render at.
+	UE::MovieGraph::DefaultRenderer::FRenderTargetInitParams RenderTargetInitParams;
+	RenderTargetInitParams.Size = FIntPoint(1920, 1080);
+	
 	// OCIO: Since this is a manually created Render target we don't need Gamma to be applied.
 	// We use this render target to render to via a display extension that utilizes Display Gamma
 	// which has a default value of 2.2 (DefaultDisplayGamma), therefore we need to set Gamma on this render target to 2.2 to cancel out any unwanted effects.
-	
-	UMovieGraphDefaultRenderer::FRenderTargetInitParams RenderTargetInitParams;
-	RenderTargetInitParams.Size = FIntPoint(1920, 1080);
 	RenderTargetInitParams.TargetGamma = FOpenColorIODisplayExtension::DefaultDisplayGamma;
 	RenderTargetInitParams.PixelFormat = EPixelFormat::PF_FloatRGBA;
 
@@ -134,10 +136,7 @@ void UMovieGraphDeferredRenderPassNode::FMovieGraphDeferredRenderPass::Render()
 	}
 
 	// ToDo:
-	InitData.TimeData.FrameDeltaTime = 1 / 24.f;
-	InitData.TimeData.TimeDilation = 1.0f;
-	InitData.TimeData.WorldSeconds = 0.f;
-	InitData.TimeData.MotionBlurFraction = 1 / 24.f;
+	InitData.TimeData = InTimeData;
 	InitData.SceneCaptureSource = ESceneCaptureSource::SCS_FinalToneCurveHDR;
 	InitData.bWorldIsPaused = false;
 	InitData.GlobalScreenPercentageFraction = FLegacyScreenPercentageDriver::GetCVarResolutionFraction();
@@ -166,8 +165,96 @@ void UMovieGraphDeferredRenderPassNode::FMovieGraphDeferredRenderPass::Render()
 	// Submit the renderer to be rendered
 	GetRendererModule().BeginRenderingViewFamily(&Canvas, ViewFamily.ToSharedPtr().Get());
 
+	// If this was just to contribute to the history buffer, no need to go any further.
+	//if (InSampleState.bDiscardResult)
+	//{
+	//	return;
+	//}
+
+	FMovieGraphRenderDataIdentifier RenderDataIdentifier;
+	RenderDataIdentifier.RenderLayerName = LayerData.LayerName;
+	RenderDataIdentifier.RendererName = RenderPassNode->GetRendererName();
+	RenderDataIdentifier.SubResourceName = TEXT("beauty");
+	RenderDataIdentifier.CameraName = TEXT("Unnamed_Camera");
+	if (InitData.ViewActor)
+	{
+		RenderDataIdentifier.CameraName = InitData.ViewActor->GetName();
+	}
+
+	UE::MovieGraph::FMovieGraphSampleState SampleState;
+	SampleState.RenderDataIdentifier = RenderDataIdentifier;
+	SampleState.Time = InTimeData;
+
 	// Readback + Accumulate.
-	// PostRendererSubmission(InOutSampleState, PassIdentifierForCurrentCamera, GetOutputFileSortingOrder(), Canvas);
+	PostRendererSubmission(SampleState, RenderTargetInitParams, Canvas);
+}
+
+void UMovieGraphDeferredRenderPassNode::FMovieGraphDeferredRenderPass::PostRendererSubmission(
+	const UE::MovieGraph::FMovieGraphSampleState& InSampleState,
+	const UE::MovieGraph::DefaultRenderer::FRenderTargetInitParams& InRenderTargetInitParams, FCanvas& InCanvas)
+{
+	// ToDo: Draw Letterboxing
+	
+
+	// We have a pool of accumulators - we multi-thread the accumulation on the task graph, and for each frame,
+	// the task has the previous samples as pre-reqs to keep the accumulation in order. However, each accumulator
+	// can only work on one frame at a time, so we create a pool of them to work concurrently. This needs a limit
+	// as large accumulations (16k) can take a lot of system RAM.
+	FMoviePipelineAccumulatorPoolPtr SampleAccumulatorPool = Renderer->GetOrCreateAccumulatorPool<FImageOverlappedAccumulator>();
+	UE::MovieGraph::DefaultRenderer::FSurfaceAccumulatorPool::FInstancePtr AccumulatorInstance = nullptr;
+	{
+		// SCOPE_CYCLE_COUNTER(STAT_MoviePipeline_WaitForAvailableAccumulator);
+		AccumulatorInstance = SampleAccumulatorPool->BlockAndGetAccumulator_GameThread(InSampleState.Time.OutputFrameNumber, InSampleState.RenderDataIdentifier);
+	}
+
+	FMoviePipelineSurfaceQueuePtr LocalSurfaceQueue = Renderer->GetOrCreateSurfaceQueue(InRenderTargetInitParams);
+	LocalSurfaceQueue->BlockUntilAnyAvailable();
+
+	UE::MovieGraph::FMovieGraphRenderDataAccumulationArgs AccumulationArgs;
+	{
+		//AccumulationArgs.OutputMerger = Renderer->GetOutputBuilder();
+		AccumulationArgs.ImageAccumulator = StaticCastSharedPtr<FImageOverlappedAccumulator>(AccumulatorInstance->Accumulator);
+		AccumulationArgs.bIsFirstSample = true; // FramePayload->IsFirstTile() && FramePayload->IsFirstTemporalSample()
+		AccumulationArgs.bIsLastSample = true; // FramePayload->IsLastTile() && FramePayload->IsLastTemporalSample()
+		AccumulationArgs.TaskPrerequisite = AccumulatorInstance->TaskPrereq;
+	}
+
+	auto OnSurfaceReadbackFinished = [this, InSampleState, AccumulationArgs, AccumulatorInstance](TUniquePtr<FImagePixelData>&& InPixelData)
+	{
+		UE::MovieGraph::DefaultRenderer::FMovieGraphAccumulationTask Task;
+		// There may be other accumulations for this accumulator which need to be processed first
+		Task.LastCompletionEvent = AccumulationArgs.TaskPrerequisite;
+		// TUniquePtr<FImagePixelData> Foo = MakeUnique<TImagePixelData<FFloat16Color>>(InPixelData->GetSize(), MoveTemp(InPixelData->MoveImageDataToNew))
+
+		FGraphEventRef Event = Task.Execute([PixelData = MoveTemp(InPixelData), AccumulationArgs, AccumulatorInstance]() mutable
+			{
+				// Enqueue a encode for this frame onto our worker thread.
+				// UE::MovieGraph::DefaultRenderer::AccumulateSample_TaskThread(MoveTemp(PixelData), AccumulationArgs);
+
+				// We have to defer clearing the accumulator until after sample accumulation has finished
+				if (AccumulationArgs.bIsLastSample)
+				{
+					// Final sample has now been executed, break the pre-req chain and free the accumulator for reuse.
+					AccumulatorInstance->bIsActive = false;
+					AccumulatorInstance->TaskPrereq = nullptr;
+				}
+			});
+		AccumulatorInstance->TaskPrereq = Event;
+
+		this->Renderer->AddOutstandingRenderTask_AnyThread(Event);
+	};
+
+	FRenderTarget* RenderTarget = InCanvas.GetRenderTarget();
+
+	ENQUEUE_RENDER_COMMAND(CanvasRenderTargetResolveCommand)(
+		[LocalSurfaceQueue, OnSurfaceReadbackFinished, RenderTarget](FRHICommandListImmediate& RHICmdList) mutable
+		{
+			// The legacy surface reader takes the payload just so it can shuffle it into our callback, but we can just include the data
+			// directly in the callback, so this is just a dummy payload.
+			TSharedRef<FImagePixelDataPayload, ESPMode::ThreadSafe> FramePayload = MakeShared<FImagePixelDataPayload, ESPMode::ThreadSafe>();
+			LocalSurfaceQueue->OnRenderTargetReady_RenderThread(RenderTarget->GetRenderTargetTexture(), FramePayload, MoveTemp(OnSurfaceReadbackFinished));
+		});
+		
 }
 
 TSharedRef<FSceneViewFamilyContext> UMovieGraphDeferredRenderPassNode::FMovieGraphDeferredRenderPass::AllocateSceneViewFamilyContext(const FViewFamilyContextInitData& InInitData)
diff --git a/Engine/Plugins/MovieScene/MovieRenderPipeline/Source/MovieRenderPipelineCore/Public/Graph/MovieGraphDataTypes.h b/Engine/Plugins/MovieScene/MovieRenderPipeline/Source/MovieRenderPipelineCore/Public/Graph/MovieGraphDataTypes.h
--- a/Engine/Plugins/MovieScene/MovieRenderPipeline/Source/MovieRenderPipelineCore/Public/Graph/MovieGraphDataTypes.h
+++ b/Engine/Plugins/MovieScene/MovieRenderPipeline/Source/MovieRenderPipelineCore/Public/Graph/MovieGraphDataTypes.h
@@ -42,6 +42,21 @@ USTRUCT(BlueprintType)
 struct MOVIERENDERPIPELINECORE_API FMovieGraphTimeStepData
 {
 	GENERATED_BODY()
+
+	UPROPERTY(BlueprintReadWrite, EditAnywhere, Category = "Movie Graph")
+	int32 OutputFrameNumber;
+
+	UPROPERTY(BlueprintReadWrite, EditAnywhere, Category = "Movie Graph")
+	float FrameDeltaTime;
+	
+	UPROPERTY(BlueprintReadWrite, EditAnywhere, Category = "Movie Graph")
+	float WorldTimeDilation;
+	
+	UPROPERTY(BlueprintReadWrite, EditAnywhere, Category = "Movie Graph")
+	float WorldSeconds;
+	
+	UPROPERTY(BlueprintReadWrite, EditAnywhere, Category = "Movie Graph")
+	float MotionBlurFraction;
 };
 
 UCLASS(BlueprintType, Abstract)
@@ -49,12 +64,8 @@ class MOVIERENDERPIPELINECORE_API UMovieGraphTimeStepBase : public UObject
 {
 	GENERATED_BODY()
 public:
-	virtual void CalculateTimeStep(FMovieGraphTimeStepData& OutTimeData)
-	{}
-	virtual void ResetPerTickData() {}
-	virtual void InitializeShot(UMoviePipelineExecutorShot* InShot, FMovieGraphTimeStepData& OutTimeStepData) {}
-
-protected:
+	virtual void TickProducingFrames() {}
+	virtual FMovieGraphTimeStepData GetCalculatedTimeData() const { return FMovieGraphTimeStepData(); }
 	UMovieGraphPipeline* GetOwningGraph() const;
 };
 
@@ -66,9 +77,6 @@ public:
 	virtual void Render(const FMovieGraphTimeStepData& InTimeData) {}
 	virtual void SetupRenderingPipelineForShot(UMoviePipelineExecutorShot* InShot) {}
 	virtual void TeardownRenderingPipelineForShot(UMoviePipelineExecutorShot* InShot) {}
-
-
-protected:
 	UMovieGraphPipeline* GetOwningGraph() const;
 };
 
@@ -77,13 +85,7 @@ class MOVIERENDERPIPELINECORE_API UMovieGraphTimeRangeBuilderBase : public UObje
 {
 	GENERATED_BODY()
 public:
-	virtual void BuildTimeRanges()
-	{
-
-	}
-
-
-protected:
+	virtual void BuildTimeRanges() { }
 	UMovieGraphPipeline* GetOwningGraph() const;
 };
 
@@ -95,9 +97,7 @@ public:
 	virtual void CacheDataPreJob(const FMovieGraphInitConfig& InInitConfig) {}
 	virtual void RestoreCachedDataPostJob() {}
 	virtual void UpdateShotList() {}
-	virtual void InitializeShot(UMoviePipelineExecutorShot* InShot, const FMovieGraphTimeStepData& InTimeStepData) {}
-
-protected:
+	virtual void InitializeShot(UMoviePipelineExecutorShot* InShot) {}
 	UMovieGraphPipeline* GetOwningGraph() const;
 };
 
@@ -110,4 +110,86 @@ struct FMovieGraphRenderPassSetupData
 {
 	TWeakObjectPtr<class UMovieGraphDefaultRenderer> Renderer;
 	TArray<FMovieGraphRenderPassLayerData> Layers;
-};
\ No newline at end of file
+};
+
+USTRUCT(BlueprintType)
+struct FMovieGraphRenderDataIdentifier
+{
+	GENERATED_BODY()
+
+	FMovieGraphRenderDataIdentifier()
+	{}
+
+	FMovieGraphRenderDataIdentifier(const FString& InRenderLayerName, const FString& InRendererName,
+		const FString& InSubRenderResourceName, const FString& InCameraName)
+		: RenderLayerName(InRenderLayerName)
+		, RendererName(InRendererName)
+		, SubResourceName(InSubRenderResourceName)
+		, CameraName(InCameraName)
+	{
+	}
+
+	bool operator == (const FMovieGraphRenderDataIdentifier& InRHS) const
+	{
+		return RenderLayerName == InRHS.RenderLayerName && 
+			RendererName == InRHS.RendererName &&
+			SubResourceName == InRHS.SubResourceName &&
+			CameraName == InRHS.CameraName;
+	}
+
+	bool operator != (const FMovieGraphRenderDataIdentifier& InRHS) const
+	{
+		return !(*this == InRHS);
+	}
+
+	friend uint32 GetTypeHash(FMovieGraphRenderDataIdentifier OutputState)
+	{
+		return HashCombineFast(GetTypeHash(OutputState.RenderLayerName),
+			HashCombineFast(GetTypeHash(OutputState.RendererName),
+				HashCombineFast(GetTypeHash(OutputState.SubResourceName),
+					GetTypeHash(OutputState.CameraName))));
+	}
+
+public:
+	// The user provided name for the whole Render Layer ("character", "background", etc.)
+	UPROPERTY(EditAnywhere, BlueprintReadWrite, Category = "Movie Pipeline")
+	FString RenderLayerName;
+
+	// Which renderer was used to produce this image ("panoramic" "deferred" "path tracer", etc.)
+	UPROPERTY(EditAnywhere, BlueprintReadWrite, Category = "Movie Pipeline")
+	FString RendererName;
+
+	// A sub-resource name for the renderer (ie: "beauty", "object id", "depth", etc.)
+	UPROPERTY(EditAnywhere, BlueprintReadWrite, Category = "Movie Pipeline")
+	FString SubResourceName;
+
+	// The name of the camera being used for this render.
+	UPROPERTY(EditAnywhere, BlueprintReadWrite, Category = "Movie Pipeline")
+	FString CameraName;
+};
+
+
+namespace UE
+{
+namespace MovieGraph
+{
+	struct FMovieGraphSampleState
+	{
+		// We need to capture the state of MRQ at the time a sample was rendered so that we can
+		// resolve filenames, etc. This is partially:
+		// Output Frame Data (which output frame, temporal sample, spatial sample)
+		// Traversal Context (layer name, current shot, etc.)
+
+		// The name of the render resource this state was captured for
+		FMovieGraphRenderDataIdentifier RenderDataIdentifier;
+
+		// The backbuffer resolution of this render resource (ie: the resolution the sample was rendered at)
+		FIntPoint BackbufferResolution;
+
+		// The time data (output frame, delta times, etc.)
+		FMovieGraphTimeStepData Time;
+
+	};
+
+}
+}
diff --git a/Engine/Plugins/MovieScene/MovieRenderPipeline/Source/MovieRenderPipelineCore/Public/Graph/MovieGraphDefaultRenderer.h b/Engine/Plugins/MovieScene/MovieRenderPipeline/Source/MovieRenderPipelineCore/Public/Graph/MovieGraphDefaultRenderer.h
--- a/Engine/Plugins/MovieScene/MovieRenderPipeline/Source/MovieRenderPipelineCore/Public/Graph/MovieGraphDefaultRenderer.h
+++ b/Engine/Plugins/MovieScene/MovieRenderPipeline/Source/MovieRenderPipelineCore/Public/Graph/MovieGraphDefaultRenderer.h
@@ -3,21 +3,23 @@
 
 #include "Graph/MovieGraphDataTypes.h"
 #include "PixelFormat.h"
+#include "MovieRenderPipelineCoreModule.h"
+#include "Async/TaskGraphInterfaces.h"
 #include "MovieGraphDefaultRenderer.generated.h"
 
 // Forward Declares
 class UMovieGraphRenderPassNode;
 class UTextureRenderTarget2D;
+struct FMoviePipelineSurfaceQueue;
+namespace MoviePipeline { struct IMoviePipelineOverlappedAccumulator; }
+namespace UE::MovieGraph::DefaultRenderer { struct FSurfaceAccumulatorPool; }
 
-/**
-* This class 
-*/
-UCLASS(BlueprintType)
-class MOVIERENDERPIPELINECORE_API UMovieGraphDefaultRenderer : public UMovieGraphRendererBase
-{
-	GENERATED_BODY()
+typedef TSharedPtr<FMoviePipelineSurfaceQueue, ESPMode::ThreadSafe> FMoviePipelineSurfaceQueuePtr;
+typedef MoviePipeline::IMoviePipelineOverlappedAccumulator MoviePipelineOverlappedAccumulatorInterface;
+typedef TSharedPtr<UE::MovieGraph::DefaultRenderer::FSurfaceAccumulatorPool, ESPMode::ThreadSafe> FMoviePipelineAccumulatorPoolPtr;
 
-public:
+namespace UE::MovieGraph::DefaultRenderer
+{
 	struct FRenderTargetInitParams
 	{
 		FRenderTargetInitParams()
@@ -47,20 +49,143 @@ public:
 		}
 	};
 
+	struct FSurfaceAccumulatorPool : public TSharedFromThis<FSurfaceAccumulatorPool>
+	{
+		struct FInstance
+		{
+			FInstance(TSharedPtr<MoviePipelineOverlappedAccumulatorInterface, ESPMode::ThreadSafe> InAccumulator)
+			{
+				Accumulator = InAccumulator;
+				ActiveFrameNumber = INDEX_NONE;
+				bIsActive = false;
+			}
+
+
+			bool IsActive() const { return bIsActive; }
+			void SetIsActive(const bool bInIsActive) { bIsActive = bInIsActive; }
+
+			TSharedPtr<MoviePipelineOverlappedAccumulatorInterface, ESPMode::ThreadSafe> Accumulator;
+			int32 ActiveFrameNumber;
+			FMovieGraphRenderDataIdentifier ActivePassIdentifier;
+			FThreadSafeBool bIsActive;
+			FGraphEventRef TaskPrereq;
+		};
+		typedef TSharedPtr<FInstance, ESPMode::ThreadSafe> FInstancePtr;
+
+		TArray<TSharedPtr<FInstance, ESPMode::ThreadSafe>> Accumulators;
+		FCriticalSection CriticalSection;
+
+
+		FInstancePtr BlockAndGetAccumulator_GameThread(int32 InFrameNumber, const FMovieGraphRenderDataIdentifier& InPassIdentifier);
+	};
+
+	class FMovieGraphAccumulationTask
+	{
+	public:
+		FGraphEventRef LastCompletionEvent;
+
+	public:
+		FGraphEventRef Execute(TUniqueFunction<void()> InFunctor)
+		{
+			if (LastCompletionEvent)
+			{
+				LastCompletionEvent = FFunctionGraphTask::CreateAndDispatchWhenReady(MoveTemp(InFunctor), GetStatId(), LastCompletionEvent);
+			}
+			else
+			{
+				LastCompletionEvent = FFunctionGraphTask::CreateAndDispatchWhenReady(MoveTemp(InFunctor), GetStatId());
+			}
+			return LastCompletionEvent;
+		}
+
+		FORCEINLINE TStatId GetStatId() const
+		{
+			RETURN_QUICK_DECLARE_CYCLE_STAT(FMovieGraphAccumulationTask, STATGROUP_ThreadPoolAsyncTasks);
+		}
+	};
+}
+
+
+
+/**
+* This class is the default implementation for the Movie Graph Pipeline renderer. This
+* is split off into a separate class to minimize the complexity of the core UMovieGraphPipeline,
+* and to provide a better way to store render-specific data during runtime. It is responsible
+* for taking all of the render passes and rendering them, and then moving their rendered
+* data back to the main UMoviePipeline OutputMerger once finished.
+* 
+* It is unlikely you will want to implement your own renderer. If you need to create new render
+* passes, inherit from UMovieGraphRenderPassNode and add it to your configuration, at which point
+* MRQ will call function on the CDO of it that allow you to set up your own render data.
+*/
+UCLASS(BlueprintType)
+class MOVIERENDERPIPELINECORE_API UMovieGraphDefaultRenderer : public UMovieGraphRendererBase
+{
+	GENERATED_BODY()
+
 public:
+	// UMovieGraphRendererBase Interface
 	virtual void Render(const FMovieGraphTimeStepData& InTimeData) override;
 	virtual void SetupRenderingPipelineForShot(UMoviePipelineExecutorShot* InShot) override;
 	virtual void TeardownRenderingPipelineForShot(UMoviePipelineExecutorShot* InShot) override;
+	// ~UMovieGraphRendererBase Interface
 
+	// UObject Interface
 	static void AddReferencedObjects(UObject* InThis, FReferenceCollector& Collector);
-	UTextureRenderTarget2D* GetOrCreateViewRenderTarget(const FRenderTargetInitParams& InInitParams);
+	// ~UObject Interface
+
+	void AddOutstandingRenderTask_AnyThread(FGraphEventRef InTask);
+	
+	UTextureRenderTarget2D* GetOrCreateViewRenderTarget(const UE::MovieGraph::DefaultRenderer::FRenderTargetInitParams& InInitParams);
+	FMoviePipelineSurfaceQueuePtr GetOrCreateSurfaceQueue(const UE::MovieGraph::DefaultRenderer::FRenderTargetInitParams& InInitParams);
+	
+	template <typename AccumulatorType>
+	FMoviePipelineAccumulatorPoolPtr GetOrCreateAccumulatorPool()
+	{
+		if (const FMoviePipelineAccumulatorPoolPtr* ExistingAccumulatorPool = PooledAccumulators.Find(AccumulatorType::GetName()))
+		{
+			return *ExistingAccumulatorPool;
+		}
+
+		// ToDo: This will need to be dynamically sized based on the number of concurrently produced render datas.
+		FMoviePipelineAccumulatorPoolPtr NewAccumulatorPool = MakeShared<UE::MovieGraph::DefaultRenderer::FSurfaceAccumulatorPool, ESPMode::ThreadSafe>();
+		for (int32 Index = 0; Index < 3; Index++)
+		{
+			TSharedPtr<AccumulatorType> NewAccumulatorInstance = MakeShared<AccumulatorType>();
+			NewAccumulatorPool->Accumulators.Add(MakeShared<UE::MovieGraph::DefaultRenderer::FSurfaceAccumulatorPool::FInstance>(NewAccumulatorInstance));
+		}
+		UE_LOG(LogMovieRenderPipeline, Log, TEXT("Allocated a Accumulator Pool for AccumulatorType: %s"), *AccumulatorType::GetName().ToString());
+		
+		PooledAccumulators.Emplace(AccumulatorType::GetName(), NewAccumulatorPool);
+		
+		return NewAccumulatorPool;
+	}
 
 protected:
-	TObjectPtr<UTextureRenderTarget2D> CreateViewRenderTarget(const FRenderTargetInitParams& InInitParams) const;
+	TObjectPtr<UTextureRenderTarget2D> CreateViewRenderTarget(const UE::MovieGraph::DefaultRenderer::FRenderTargetInitParams& InInitParams) const;
+	FMoviePipelineSurfaceQueuePtr CreateSurfaceQueue(const UE::MovieGraph::DefaultRenderer::FRenderTargetInitParams& InInitParams) const;
 
 protected:
 	UPROPERTY(Transient)
 	TArray<TObjectPtr<UMovieGraphRenderPassNode>> RenderPassesInUse;
 
-	TMap<FRenderTargetInitParams, TObjectPtr<UTextureRenderTarget2D>> PooledViewRenderTargets;
+	// Render Target specify a transient backbuffer to draw the image to. We reuse these
+	// because we queue a copy onto a FMoviePipelineSurface so the only reason it is needed
+	// is for display in the UI.
+	TMap<UE::MovieGraph::DefaultRenderer::FRenderTargetInitParams, TObjectPtr<UTextureRenderTarget2D>> PooledViewRenderTargets;
+
+	// This is a list of Surfaces that we can copy our render target to immediately after
+	// it is drawn to. This is read back to the CPU a frame later (to avoid blocking the
+	// GPU), and it's contents are then copied into a FImagePixelData.
+	TMap<UE::MovieGraph::DefaultRenderer::FRenderTargetInitParams, FMoviePipelineSurfaceQueuePtr> PooledSurfaceQueues;
+
+	// Once a sample is copied back from the GPU, we need to potentially accumulate it over
+	// multiple frames, and to apply high-res tiling. To do this we have a pool of accumulators,
+	// where each unique render resource gets one per frame.
+	TMap<FName, TSharedPtr<UE::MovieGraph::DefaultRenderer::FSurfaceAccumulatorPool>> PooledAccumulators;
+
+	// Accessed by the Render Thread when starting up a new task. Makes sure we don't add tasks to the array while we're removing finished ones.
+	FCriticalSection OutstandingTasksMutex;
+	// Array of outstanding accumulation/blending tasks that are currently being worked on.
+	FGraphEventArray OutstandingTasks;
 };
\ No newline at end of file
diff --git a/Engine/Plugins/MovieScene/MovieRenderPipeline/Source/MovieRenderPipelineCore/Public/Graph/MovieGraphLinearTimeStep.h b/Engine/Plugins/MovieScene/MovieRenderPipeline/Source/MovieRenderPipelineCore/Public/Graph/MovieGraphLinearTimeStep.h
--- a/Engine/Plugins/MovieScene/MovieRenderPipeline/Source/MovieRenderPipelineCore/Public/Graph/MovieGraphLinearTimeStep.h
+++ b/Engine/Plugins/MovieScene/MovieRenderPipeline/Source/MovieRenderPipelineCore/Public/Graph/MovieGraphLinearTimeStep.h
@@ -12,5 +12,9 @@ class MOVIERENDERPIPELINECORE_API UMovieGraphLinearTimeStep : public UMovieGraph
 {
 	GENERATED_BODY()
 public:
-	virtual void CalculateTimeStep(FMovieGraphTimeStepData& OutTimeData) override;
+	virtual void TickProducingFrames() override;
+	virtual FMovieGraphTimeStepData GetCalculatedTimeData() const override { return CurrentTimeStepData; }
+
+private:
+	FMovieGraphTimeStepData CurrentTimeStepData;
 };
\ No newline at end of file
diff --git a/Engine/Plugins/MovieScene/MovieRenderPipeline/Source/MovieRenderPipelineCore/Public/Graph/MovieGraphPipeline.h b/Engine/Plugins/MovieScene/MovieRenderPipeline/Source/MovieRenderPipelineCore/Public/Graph/MovieGraphPipeline.h
--- a/Engine/Plugins/MovieScene/MovieRenderPipeline/Source/MovieRenderPipelineCore/Public/Graph/MovieGraphPipeline.h
+++ b/Engine/Plugins/MovieScene/MovieRenderPipeline/Source/MovieRenderPipelineCore/Public/Graph/MovieGraphPipeline.h
@@ -30,6 +30,23 @@ public:
 public:
 	UMovieGraphConfig* GetRootGraphForShot(UMoviePipelineExecutorShot* InShot) const;
 	FMovieGraphTraversalContext GetTraversalContextForShot(UMoviePipelineExecutorShot* InShot) const;
+
+	// Get the Active Shot list, which is the full shot list generated from the external data source, with disabled shots removed.
+	const TArray<TObjectPtr<UMoviePipelineExecutorShot>>& GetActiveShotList() const { return ActiveShotList; }
+	// Which index of the Active Shot List are we currently on
+	int32 GetCurrentShotIndex() const { return CurrentShotIndex; }
+	// Called by the TimeStepInstance when it's time to set up for another shot. Don't call this unless you know what you're doing.
+	void SetupShot(UMoviePipelineExecutorShot* InShot);
+	// Called by the TimeStepInstance when it's time to tear down the current shot. Don't call this unless you know what you're doing.
+	void TeardownShot(UMoviePipelineExecutorShot* InShot);
+	// Used occasionally to cross-reference other components. Don't call this unless you know what you're doing.
+	UMovieGraphTimeStepBase* GetTimeStepInstance() const { return GraphTimeStepInstance; }
+	// Used occasionally to cross-reference other components. Don't call this unless you know what you're doing.
+	UMovieGraphRendererBase* GetRendererInstance() const { return GraphRendererInstance; }
+	// Used occasionally to cross-reference other components. Don't call this unless you know what you're doing.
+	UMovieGraphTimeRangeBuilderBase* GetTimeRangeBuilderInstance() const { return GraphTimeRangeBuilderInstance; }
+	// Used occasionally to cross-reference other components. Don't call this unless you know what you're doing.
+	UMovieGraphDataCachingBase* GetDataCachingInstance() const { return GraphDataCachingInstance; }
 protected:
 	UFUNCTION(BlueprintCallable, Category = "Movie Render Pipeline")
 	virtual void OnMoviePipelineFinishedImpl();
@@ -39,8 +56,7 @@ protected:
 	virtual void OnEngineTickEndFrame();
 	virtual void RenderFrame();
 	virtual void BuildShotListFromDataSource();
-	virtual void SetupShot(UMoviePipelineExecutorShot* InShot);
-	virtual void TeardownShot(UMoviePipelineExecutorShot* InShot);
+
 
 	virtual void TickProducingFrames();
 	virtual void TickPostFinalizeExport(const bool bInForceFinish);
@@ -54,6 +70,7 @@ protected:
 	virtual EMovieRenderPipelineState GetPipelineStateImpl() const override { return PipelineState; }
 	// ~UMoviePipelineBase Interface
 
+
 protected:
 	UPROPERTY(Transient)
 	TObjectPtr<UMovieGraphTimeStepBase> GraphTimeStepInstance;
@@ -87,4 +104,7 @@ protected:
 
 	/** Which step of the rendering process is the graph currently in. */
 	EMovieRenderPipelineState PipelineState;
+
+	/** What time (in UTC) was Initialization called? Used internally for tracking total job duration. */
+	FDateTime GraphInitializationTime;
 };
\ No newline at end of file
diff --git a/Engine/Plugins/MovieScene/MovieRenderPipeline/Source/MovieRenderPipelineCore/Public/Graph/Nodes/MovieGraphDeferredRenderPassNode.h b/Engine/Plugins/MovieScene/MovieRenderPipeline/Source/MovieRenderPipelineCore/Public/Graph/Nodes/MovieGraphDeferredRenderPassNode.h
--- a/Engine/Plugins/MovieScene/MovieRenderPipeline/Source/MovieRenderPipelineCore/Public/Graph/Nodes/MovieGraphDeferredRenderPassNode.h
+++ b/Engine/Plugins/MovieScene/MovieRenderPipeline/Source/MovieRenderPipelineCore/Public/Graph/Nodes/MovieGraphDeferredRenderPassNode.h
@@ -7,12 +7,15 @@
 #include "Camera/CameraTypes.h"
 #include "MovieRenderPipelineDataTypes.h"
 #include "Engine/EngineTypes.h"
+#include "Async/TaskGraphFwd.h"
 #include "MovieGraphDeferredRenderPassNode.generated.h"
 
 // Forward Declares
 class UMovieGraphDefaultRenderer;
+namespace UE::MovieGraph::DefaultRenderer { struct FRenderTargetInitParams; }
 struct FMovieGraphRenderPassSetupData;
 struct FMovieGraphRenderPassLayerData;
+struct FImageOverlappedAccumulator;
 
 // For FViewFamilyContextInitData
 class FRenderTarget;
@@ -20,7 +23,26 @@ class UWorld;
 class FSceneView;
 class AActor;
 class FSceneViewFamilyContext;
+class FCanvas;
 
+namespace UE::MovieGraph
+{
+	struct FMovieGraphRenderDataAccumulationArgs
+	{
+	public:
+		TWeakPtr<FImageOverlappedAccumulator, ESPMode::ThreadSafe> ImageAccumulator;
+		// TWeakPtr<IMoviePipelineOutputMerger, ESPMode::ThreadSafe> OutputMerger;
+
+		// If it's the first sample then we will reset the accumulator to a clean slate before accumulating into it.
+		bool bIsFirstSample;
+		// If it's the last sample, then we will trigger moving the data to the output merger. It can be both the first and last sample at the same time.
+		bool bIsLastSample;
+		// Does this accumulation event depend on a previous accumulation event before it can be executed?
+		FGraphEventRef TaskPrerequisite;
+	};
+
+	// void MOVIERENDERPIPELINERENDERPASSES_API AccumulateSample_TaskThread(TUniquePtr<FImagePixelData>&& InPixelData, const MoviePipeline::FImageSampleAccumulationArgs& InParams);
+}
 
 UCLASS()
 class UMovieGraphDeferredRenderPassNode : public UMovieGraphRenderPassNode
@@ -38,7 +60,7 @@ protected:
 	virtual FString GetRendererNameImpl() const override { return TEXT("Deferred"); }
 	virtual void SetupImpl(const FMovieGraphRenderPassSetupData& InSetupData) override;
 	virtual void TeardownImpl() override;
-	virtual void RenderImpl() override;
+	virtual void RenderImpl(const FMovieGraphTimeStepData& InTimeData) override;
 	// ~UMovieGraphRenderPassNode Interface
 
 	static void AddReferencedObjects(UObject* InThis, FReferenceCollector& Collector);
@@ -63,7 +85,7 @@ protected:
 
 		class FRenderTarget* RenderTarget;
 		class UWorld* World;
-		FMoviePipelineFrameOutputState::FTimeData TimeData;
+		FMovieGraphTimeStepData TimeData;
 		ESceneCaptureSource SceneCaptureSource;
 		bool bWorldIsPaused;
 		float GlobalScreenPercentageFraction;
@@ -85,16 +107,16 @@ protected:
 	struct FMovieGraphDeferredRenderPass
 	{
 	public:
-		void Setup(TWeakObjectPtr<UMovieGraphDefaultRenderer> InRenderer, const FMovieGraphRenderPassLayerData& InLayer);
+		void Setup(TWeakObjectPtr<UMovieGraphDefaultRenderer> InRenderer, TWeakObjectPtr<UMovieGraphDeferredRenderPassNode> InRenderPassNode, const FMovieGraphRenderPassLayerData& InLayer);
 		void Teardown();
-		void Render();
+		void Render(const FMovieGraphTimeStepData& InTimeData);
 		void AddReferencedObjects(FReferenceCollector& Collector);
 
 	protected:
 		TSharedRef<FSceneViewFamilyContext> AllocateSceneViewFamilyContext(const FViewFamilyContextInitData& InInitData);
 		FSceneView* AllocateSceneView(TSharedPtr<FSceneViewFamilyContext> InViewFamilyContext, FViewFamilyContextInitData& InInitData) const;
 		void ApplyMoviePipelineOverridesToViewFamily(TSharedRef<FSceneViewFamilyContext> InOutFamily, const FViewFamilyContextInitData& InInitData);
-
+		void PostRendererSubmission(const UE::MovieGraph::FMovieGraphSampleState& InSampleState, const UE::MovieGraph::DefaultRenderer::FRenderTargetInitParams& InRenderTargetInitParams, FCanvas& InCanvas);
 
 
 	protected:
@@ -104,6 +126,7 @@ protected:
 		FSceneViewStateReference SceneViewState;
 
 		TWeakObjectPtr<class UMovieGraphDefaultRenderer> Renderer;
+		TWeakObjectPtr<class UMovieGraphDeferredRenderPassNode> RenderPassNode;
 	};
 
 	TArray<TUniquePtr<FMovieGraphDeferredRenderPass>> CurrentInstances;
diff --git a/Engine/Plugins/MovieScene/MovieRenderPipeline/Source/MovieRenderPipelineCore/Public/Graph/Nodes/MovieGraphRenderPassNode.h b/Engine/Plugins/MovieScene/MovieRenderPipeline/Source/MovieRenderPipelineCore/Public/Graph/Nodes/MovieGraphRenderPassNode.h
--- a/Engine/Plugins/MovieScene/MovieRenderPipeline/Source/MovieRenderPipelineCore/Public/Graph/Nodes/MovieGraphRenderPassNode.h
+++ b/Engine/Plugins/MovieScene/MovieRenderPipeline/Source/MovieRenderPipelineCore/Public/Graph/Nodes/MovieGraphRenderPassNode.h
@@ -6,6 +6,7 @@
 
 // Forward Declare
 struct FMovieGraphRenderPassSetupData;
+struct FMovieGraphTimeStepData;
 
 UCLASS(Abstract)
 class UMovieGraphRenderPassNode : public UMovieGraphNode
@@ -20,7 +21,7 @@ public:
 	FString GetRendererName() const { return GetRendererNameImpl(); }
 	void Setup(const FMovieGraphRenderPassSetupData& InSetupData) { SetupImpl(InSetupData); }
 	void Teardown() { TeardownImpl(); }
-	void Render() { RenderImpl(); }
+	void Render(const FMovieGraphTimeStepData& InTimeData) { RenderImpl(InTimeData); }
 	// ~UMovieGraphRenderPassNode interface
 
 	virtual TArray<FMovieGraphPinProperties> GetInputPinProperties() const override
@@ -48,5 +49,5 @@ protected:
 	virtual FString GetRendererNameImpl() const { return TEXT("UnnamedRenderPass"); }
 	virtual void SetupImpl(const FMovieGraphRenderPassSetupData& InSetupData) {}
 	virtual void TeardownImpl() {}
-	virtual void RenderImpl() {}
+	virtual void RenderImpl(const FMovieGraphTimeStepData& InTimeData) {}
 };
\ No newline at end of file
diff --git a/Engine/Plugins/MovieScene/MovieRenderPipeline/Source/MovieRenderPipelineCore/Public/MovieRenderOverlappedImage.h b/Engine/Plugins/MovieScene/MovieRenderPipeline/Source/MovieRenderPipelineCore/Public/MovieRenderOverlappedImage.h
--- a/Engine/Plugins/MovieScene/MovieRenderPipeline/Source/MovieRenderPipelineCore/Public/MovieRenderOverlappedImage.h
+++ b/Engine/Plugins/MovieScene/MovieRenderPipeline/Source/MovieRenderPipelineCore/Public/MovieRenderOverlappedImage.h
@@ -86,6 +86,12 @@ public:
 	{
 	}
 
+	static FName GetName()
+	{
+		static FName Name = FName("ImageOverlappedAccumulator");
+		return Name;
+	}
+
 	/**
 	 * Allocates memory.
 	 *
