CIS: Fixed an error with FMovieGraphTimeStepData not having a default constructor.

diff --git a/Engine/Plugins/MovieScene/MovieRenderPipeline/Source/MovieRenderPipelineCore/Public/Graph/MovieGraphDataTypes.h b/Engine/Plugins/MovieScene/MovieRenderPipeline/Source/MovieRenderPipelineCore/Public/Graph/MovieGraphDataTypes.h
--- a/Engine/Plugins/MovieScene/MovieRenderPipeline/Source/MovieRenderPipelineCore/Public/Graph/MovieGraphDataTypes.h
+++ b/Engine/Plugins/MovieScene/MovieRenderPipeline/Source/MovieRenderPipelineCore/Public/Graph/MovieGraphDataTypes.h
@@ -43,6 +43,14 @@ struct MOVIERENDERPIPELINECORE_API FMovieGraphTimeStepData
 {
 	GENERATED_BODY()
 
+	FMovieGraphTimeStepData()
+		: OutputFrameNumber(0)
+		, FrameDeltaTime(0.f)
+		, WorldTimeDilation(0.f)
+		, WorldSeconds(0.f)
+		, MotionBlurFraction(0.f)
+	{}
+
 	UPROPERTY(BlueprintReadWrite, EditAnywhere, Category = "Movie Graph")
 	int32 OutputFrameNumber;
 
