MoviePipeline: Add more descriptive titles (subtitles) to some graph nodes. Also added a standin "Select" node.

diff --git a/Engine/Plugins/MovieScene/MovieRenderPipeline/Source/MovieRenderPipelineCore/Private/Graph/Nodes/MovieGraphInputNode.cpp b/Engine/Plugins/MovieScene/MovieRenderPipeline/Source/MovieRenderPipelineCore/Private/Graph/Nodes/MovieGraphInputNode.cpp
--- a/Engine/Plugins/MovieScene/MovieRenderPipeline/Source/MovieRenderPipelineCore/Private/Graph/Nodes/MovieGraphInputNode.cpp
+++ b/Engine/Plugins/MovieScene/MovieRenderPipeline/Source/MovieRenderPipelineCore/Private/Graph/Nodes/MovieGraphInputNode.cpp
@@ -35,7 +35,7 @@ TArray<FMovieGraphPinProperties> UMovieGraphInputNode::GetOutputPinProperties()
 }
 
 #if WITH_EDITOR
-FText UMovieGraphInputNode::GetMenuDescription() const
+FText UMovieGraphInputNode::GetNodeTitle(const bool bGetDescriptive) const
 {
 	return NSLOCTEXT("MovieGraphNodes", "InputNode_Description", "Input");
 }
diff --git a/Engine/Plugins/MovieScene/MovieRenderPipeline/Source/MovieRenderPipelineCore/Private/Graph/Nodes/MovieGraphOutputNode.cpp b/Engine/Plugins/MovieScene/MovieRenderPipeline/Source/MovieRenderPipelineCore/Private/Graph/Nodes/MovieGraphOutputNode.cpp
--- a/Engine/Plugins/MovieScene/MovieRenderPipeline/Source/MovieRenderPipelineCore/Private/Graph/Nodes/MovieGraphOutputNode.cpp
+++ b/Engine/Plugins/MovieScene/MovieRenderPipeline/Source/MovieRenderPipelineCore/Private/Graph/Nodes/MovieGraphOutputNode.cpp
@@ -35,7 +35,7 @@ TArray<FMovieGraphPinProperties> UMovieGraphOutputNode::GetInputPinProperties()
 }
 
 #if WITH_EDITOR
-FText UMovieGraphOutputNode::GetMenuDescription() const
+FText UMovieGraphOutputNode::GetNodeTitle(const bool bGetDescriptive) const
 {
 	return NSLOCTEXT("MovieGraphNodes", "OutputNode_Description", "Output");
 }
diff --git a/Engine/Plugins/MovieScene/MovieRenderPipeline/Source/MovieRenderPipelineCore/Private/Graph/Nodes/MovieGraphRenderLayerNode.cpp b/Engine/Plugins/MovieScene/MovieRenderPipeline/Source/MovieRenderPipelineCore/Private/Graph/Nodes/MovieGraphRenderLayerNode.cpp
--- a/Engine/Plugins/MovieScene/MovieRenderPipeline/Source/MovieRenderPipelineCore/Private/Graph/Nodes/MovieGraphRenderLayerNode.cpp
+++ b/Engine/Plugins/MovieScene/MovieRenderPipeline/Source/MovieRenderPipelineCore/Private/Graph/Nodes/MovieGraphRenderLayerNode.cpp
@@ -40,9 +40,17 @@ TArray<FPropertyBagPropertyDesc> UMovieGraphRenderLayerNode::GetDynamicPropertyD
 }
 
 #if WITH_EDITOR
-FText UMovieGraphRenderLayerNode::GetMenuDescription() const
+FText UMovieGraphRenderLayerNode::GetNodeTitle(const bool bGetDescriptive) const
 {
-	return NSLOCTEXT("MovieGraphNodes", "RenderLayerGraphNode_Description", "Render Layer");
+	static const FText RenderLayerNodeName = NSLOCTEXT("MovieGraphNodes", "NodeName_RenderLayer", "Render Layer");
+	static const FText RenderLayerNodeDescription = NSLOCTEXT("MovieGraphNodes", "NodeDescription_RenderLayer", "Render Layer\n{0}");
+
+	if (bGetDescriptive && !LayerName.IsEmpty())
+	{
+		return FText::Format(RenderLayerNodeDescription, FText::FromString(LayerName));
+	}
+	
+	return RenderLayerNodeName;
 }
 
 FText UMovieGraphRenderLayerNode::GetMenuCategory() const
@@ -63,4 +71,14 @@ FSlateIcon UMovieGraphRenderLayerNode::GetIconAndTint(FLinearColor& OutColor) co
 	OutColor = FLinearColor::White;
 	return RenderLayerIcon;
 }
+
+void UMovieGraphRenderLayerNode::PostEditChangeProperty(FPropertyChangedEvent& PropertyChangedEvent)
+{
+	Super::PostEditChangeProperty(PropertyChangedEvent);
+	
+	if (PropertyChangedEvent.GetPropertyName() == GET_MEMBER_NAME_CHECKED(UMovieGraphRenderLayerNode, LayerName))
+	{
+		OnNodeChangedDelegate.Broadcast(this);
+	}
+}
 #endif // WITH_EDITOR
\ No newline at end of file
diff --git a/Engine/Plugins/MovieScene/MovieRenderPipeline/Source/MovieRenderPipelineCore/Private/Graph/Nodes/MovieGraphVariableNode.cpp b/Engine/Plugins/MovieScene/MovieRenderPipeline/Source/MovieRenderPipelineCore/Private/Graph/Nodes/MovieGraphVariableNode.cpp
--- a/Engine/Plugins/MovieScene/MovieRenderPipeline/Source/MovieRenderPipelineCore/Private/Graph/Nodes/MovieGraphVariableNode.cpp
+++ b/Engine/Plugins/MovieScene/MovieRenderPipeline/Source/MovieRenderPipelineCore/Private/Graph/Nodes/MovieGraphVariableNode.cpp
@@ -44,7 +44,7 @@ void UMovieGraphVariableNode::SetVariable(UMovieGraphVariable* InVariable)
 }
 
 #if WITH_EDITOR
-FText UMovieGraphVariableNode::GetMenuDescription() const
+FText UMovieGraphVariableNode::GetNodeTitle(const bool bGetDescriptive) const
 {
 	return GraphVariable ? FText::FromString(GraphVariable->Name) : FText();
 }
diff --git a/Engine/Plugins/MovieScene/MovieRenderPipeline/Source/MovieRenderPipelineCore/Public/Graph/MovieGraphNode.h b/Engine/Plugins/MovieScene/MovieRenderPipeline/Source/MovieRenderPipelineCore/Public/Graph/MovieGraphNode.h
--- a/Engine/Plugins/MovieScene/MovieRenderPipeline/Source/MovieRenderPipelineCore/Public/Graph/MovieGraphNode.h
+++ b/Engine/Plugins/MovieScene/MovieRenderPipeline/Source/MovieRenderPipelineCore/Public/Graph/MovieGraphNode.h
@@ -116,7 +116,13 @@ public:
 #endif
 
 #if WITH_EDITOR
-	virtual FText GetMenuDescription() const PURE_VIRTUAL(UMovieGraphNode::GetMenuDescription, return FText(););
+	/**
+	 * Gets the node's title. Optionally gets a more descriptive, multi-line title for the node if bGetDescriptive is
+	 * set to true.
+	 */
+	virtual FText GetNodeTitle(const bool bGetDescriptive = false) const PURE_VIRTUAL(UMovieGraphNode::GetNodeTitle, return FText(););
+
+	/** Gets the category that the node belongs under. */
 	virtual FText GetMenuCategory() const PURE_VIRTUAL(UMovieGraphNode::GetMenuCategory, return FText(); );
 #endif
 
diff --git a/Engine/Plugins/MovieScene/MovieRenderPipeline/Source/MovieRenderPipelineCore/Public/Graph/MovieGraphPlaceholderNodes.cpp b/Engine/Plugins/MovieScene/MovieRenderPipeline/Source/MovieRenderPipelineCore/Public/Graph/MovieGraphPlaceholderNodes.cpp
--- a/Engine/Plugins/MovieScene/MovieRenderPipeline/Source/MovieRenderPipelineCore/Public/Graph/MovieGraphPlaceholderNodes.cpp
+++ b/Engine/Plugins/MovieScene/MovieRenderPipeline/Source/MovieRenderPipelineCore/Public/Graph/MovieGraphPlaceholderNodes.cpp
@@ -29,9 +29,16 @@ TArray<FMovieGraphPinProperties> UMoviePipelineCollectionNode::GetOutputPinPrope
 }
 
 #if WITH_EDITOR
-FText UMoviePipelineCollectionNode::GetMenuDescription() const
+FText UMoviePipelineCollectionNode::GetNodeTitle(const bool bGetDescriptive) const
 {
 	static const FText CollectionNodeName = LOCTEXT("NodeName_Collection", "Collection");
+	static const FText CollectionNodeDescription = LOCTEXT("NodeDescription_Collection", "Collection\n{0}");
+
+	if (bGetDescriptive && !CollectionName.IsEmpty())
+	{
+		return FText::Format(CollectionNodeDescription, FText::FromString(CollectionName));
+	}
+
 	return CollectionNodeName;
 }
 
@@ -53,6 +60,16 @@ FSlateIcon UMoviePipelineCollectionNode::GetIconAndTint(FLinearColor& OutColor)
 	OutColor = FLinearColor::White;
 	return CollectionIcon;
 }
+
+void UMoviePipelineCollectionNode::PostEditChangeProperty(FPropertyChangedEvent& PropertyChangedEvent)
+{
+	Super::PostEditChangeProperty(PropertyChangedEvent);
+
+	if (PropertyChangedEvent.GetPropertyName() == GET_MEMBER_NAME_CHECKED(UMoviePipelineCollectionNode, CollectionName))
+	{
+		OnNodeChangedDelegate.Broadcast(this);
+	}
+}
 #endif // WITH_EDITOR
 
 TArray<FMovieGraphPinProperties> UMovieGraphModifierNode::GetInputPinProperties() const
@@ -70,9 +87,16 @@ TArray<FMovieGraphPinProperties> UMovieGraphModifierNode::GetOutputPinProperties
 }
 
 #if WITH_EDITOR
-FText UMovieGraphModifierNode::GetMenuDescription() const
+FText UMovieGraphModifierNode::GetNodeTitle(const bool bGetDescriptive) const
 {
 	static const FText ModifierNodeName = LOCTEXT("NodeName_Modifier", "Modifier");
+	static const FText ModifierNodeDescription = LOCTEXT("NodeDescription_Modifier", "Modifier\n{0} - {1}");
+
+	if (bGetDescriptive && !ModifiedCollectionName.IsEmpty() && !ModifierName.IsEmpty())
+	{
+		return FText::Format(ModifierNodeDescription, FText::FromString(ModifiedCollectionName), FText::FromString(ModifierName));
+	}
+
 	return ModifierNodeName;
 }
 
@@ -94,6 +118,17 @@ FSlateIcon UMovieGraphModifierNode::GetIconAndTint(FLinearColor& OutColor) const
 	OutColor = FLinearColor::White;
 	return ModifierIcon;
 }
+
+void UMovieGraphModifierNode::PostEditChangeProperty(FPropertyChangedEvent& PropertyChangedEvent)
+{
+	Super::PostEditChangeProperty(PropertyChangedEvent);
+
+	if ((PropertyChangedEvent.GetPropertyName() == GET_MEMBER_NAME_CHECKED(UMovieGraphModifierNode, ModifierName)) ||
+		(PropertyChangedEvent.GetPropertyName() == GET_MEMBER_NAME_CHECKED(UMovieGraphModifierNode, ModifiedCollectionName)))
+	{
+		OnNodeChangedDelegate.Broadcast(this);
+	}
+}
 #endif // WITH_EDITOR
 
 TArray<FMovieGraphPinProperties> UMovieGraphGlobalGameOverridesNode::GetInputPinProperties() const
@@ -111,7 +146,7 @@ TArray<FMovieGraphPinProperties> UMovieGraphGlobalGameOverridesNode::GetOutputPi
 }
 
 #if WITH_EDITOR
-FText UMovieGraphGlobalGameOverridesNode::GetMenuDescription() const
+FText UMovieGraphGlobalGameOverridesNode::GetNodeTitle(const bool bGetDescriptive) const
 {
 	static const FText GlobalGameOverridesNodeName = LOCTEXT("NodeName_GlobalGameOverrides", "Global Game Overrides");
 	return GlobalGameOverridesNodeName;
@@ -152,7 +187,7 @@ TArray<FMovieGraphPinProperties> UMovieGraphDeferredRendererNode::GetOutputPinPr
 }
 
 #if WITH_EDITOR
-FText UMovieGraphDeferredRendererNode::GetMenuDescription() const
+FText UMovieGraphDeferredRendererNode::GetNodeTitle(const bool bGetDescriptive) const
 {
 	static const FText DeferredRendererNodeName = LOCTEXT("NodeName_DeferredRenderer", "Deferred Renderer");
 	return DeferredRendererNodeName;
@@ -193,7 +228,7 @@ TArray<FMovieGraphPinProperties> UMovieGraphPathTracedRendererNode::GetOutputPin
 }
 
 #if WITH_EDITOR
-FText UMovieGraphPathTracedRendererNode::GetMenuDescription() const
+FText UMovieGraphPathTracedRendererNode::GetNodeTitle(const bool bGetDescriptive) const
 {
 	static const FText PathTracedRendererNodeName = LOCTEXT("NodeName_PathTracedRenderer", "Path Traced Renderer");
 	return PathTracedRendererNodeName;
@@ -234,7 +269,7 @@ TArray<FMovieGraphPinProperties> UMovieGraphEXRSequenceNode::GetOutputPinPropert
 }
 
 #if WITH_EDITOR
-FText UMovieGraphEXRSequenceNode::GetMenuDescription() const
+FText UMovieGraphEXRSequenceNode::GetNodeTitle(const bool bGetDescriptive) const
 {
 	static const FText EXRSequenceNodeName = LOCTEXT("NodeName_EXRSequence", "EXR Sequence");
 	return EXRSequenceNodeName;
@@ -275,7 +310,7 @@ TArray<FMovieGraphPinProperties> UMovieGraphJPGSequenceNode::GetOutputPinPropert
 }
 
 #if WITH_EDITOR
-FText UMovieGraphJPGSequenceNode::GetMenuDescription() const
+FText UMovieGraphJPGSequenceNode::GetNodeTitle(const bool bGetDescriptive) const
 {
 	static const FText JPGSequenceNodeName = LOCTEXT("NodeName_JPGSequence", "JPG Sequence");
 	return JPGSequenceNodeName;
@@ -322,7 +357,7 @@ TArray<FMovieGraphPinProperties> UMovieGraphBranchNode::GetOutputPinProperties()
 }
 
 #if WITH_EDITOR
-FText UMovieGraphBranchNode::GetMenuDescription() const
+FText UMovieGraphBranchNode::GetNodeTitle(const bool bGetDescriptive) const
 {
 	static const FText BranchNodeName = LOCTEXT("NodeName_Branch", "Branch");
 	return BranchNodeName;
@@ -348,6 +383,77 @@ FSlateIcon UMovieGraphBranchNode::GetIconAndTint(FLinearColor& OutColor) const
 }
 #endif // WITH_EDITOR
 
+TArray<FMovieGraphPinProperties> UMovieGraphSelectNode::GetInputPinProperties() const
+{
+	TArray<FMovieGraphPinProperties> Properties;
+
+	// Generate branch pins for each option
+	for (const FString& SelectOption : SelectOptions)
+	{
+		Properties.Add(FMovieGraphPinProperties(FName(SelectOption), EMovieGraphMemberType::Branch, false));
+	}
+
+	Properties.Add(FMovieGraphPinProperties(NAME_None, EMovieGraphMemberType::String, false));
+	
+	return Properties;
+}
+
+TArray<FMovieGraphPinProperties> UMovieGraphSelectNode::GetOutputPinProperties() const
+{
+	TArray<FMovieGraphPinProperties> Properties;
+	Properties.Add(FMovieGraphPinProperties(NAME_None, EMovieGraphMemberType::Branch, false));
+	return Properties;
+}
+
+#if WITH_EDITOR
+FText UMovieGraphSelectNode::GetNodeTitle(const bool bGetDescriptive) const
+{
+	static const FText SelectNodeName = LOCTEXT("NodeName_Select", "Select");
+	static const FText SelectNodeDescription = LOCTEXT("NodeDescription_Select", "Select\n{0}");
+
+	if (bGetDescriptive && !Description.IsEmpty())
+	{
+		return FText::Format(SelectNodeDescription, FText::FromString(Description));
+	}
+
+	return SelectNodeName;
+}
+
+FText UMovieGraphSelectNode::GetMenuCategory() const
+{
+	return NodeCategory_Conditionals;
+}
+
+FLinearColor UMovieGraphSelectNode::GetNodeTitleColor() const
+{
+	static const FLinearColor SelectNodeColor = FLinearColor(0.266f, 0.266f, 0.266f);
+	return SelectNodeColor;
+}
+
+FSlateIcon UMovieGraphSelectNode::GetIconAndTint(FLinearColor& OutColor) const
+{
+	static const FSlateIcon SelectIcon = FSlateIcon(FAppStyle::GetAppStyleSetName(), "Icons.Merge");
+
+	OutColor = FLinearColor::White;
+	return SelectIcon;
+}
+
+void UMovieGraphSelectNode::PostEditChangeProperty(FPropertyChangedEvent& PropertyChangedEvent)
+{
+	Super::PostEditChangeProperty(PropertyChangedEvent);
+
+	if (PropertyChangedEvent.GetPropertyName() == GET_MEMBER_NAME_CHECKED(UMovieGraphSelectNode, SelectOptions))
+	{
+		UpdatePins();
+	}
+	
+	if (PropertyChangedEvent.GetPropertyName() == GET_MEMBER_NAME_CHECKED(UMovieGraphSelectNode, Description))
+	{
+		OnNodeChangedDelegate.Broadcast(this);
+	}
+}
+#endif // WITH_EDITOR
+
 TArray<FMovieGraphPinProperties> UMovieGraphOutputSettingsNode::GetInputPinProperties() const
 {
 	TArray<FMovieGraphPinProperties> Properties;
@@ -363,7 +469,7 @@ TArray<FMovieGraphPinProperties> UMovieGraphOutputSettingsNode::GetOutputPinProp
 }
 
 #if WITH_EDITOR
-FText UMovieGraphOutputSettingsNode::GetMenuDescription() const
+FText UMovieGraphOutputSettingsNode::GetNodeTitle(const bool bGetDescriptive) const
 {
 	static const FText OutputSettingsNodeName = LOCTEXT("NodeName_OutputSettings", "Output Settings");
 	return OutputSettingsNodeName;
@@ -404,7 +510,7 @@ TArray<FMovieGraphPinProperties> UMovieGraphAntiAliasingNode::GetOutputPinProper
 }
 
 #if WITH_EDITOR
-FText UMovieGraphAntiAliasingNode::GetMenuDescription() const
+FText UMovieGraphAntiAliasingNode::GetNodeTitle(const bool bGetDescriptive) const
 {
 	static const FText AntiAliasingNodeName = LOCTEXT("NodeName_AntiAliasing", "Anti-Aliasing");
 	return AntiAliasingNodeName;
diff --git a/Engine/Plugins/MovieScene/MovieRenderPipeline/Source/MovieRenderPipelineCore/Public/Graph/MovieGraphPlaceholderNodes.h b/Engine/Plugins/MovieScene/MovieRenderPipeline/Source/MovieRenderPipelineCore/Public/Graph/MovieGraphPlaceholderNodes.h
--- a/Engine/Plugins/MovieScene/MovieRenderPipeline/Source/MovieRenderPipelineCore/Public/Graph/MovieGraphPlaceholderNodes.h
+++ b/Engine/Plugins/MovieScene/MovieRenderPipeline/Source/MovieRenderPipelineCore/Public/Graph/MovieGraphPlaceholderNodes.h
@@ -12,6 +12,7 @@ UCLASS()
 class MOVIERENDERPIPELINECORE_API UMoviePipelineCollectionNode : public UMovieGraphNode
 {
 	GENERATED_BODY()
+
 public:
 	UMoviePipelineCollectionNode() = default;
 
@@ -34,11 +35,19 @@ public:
 	}
 
 #if WITH_EDITOR
-	virtual FText GetMenuDescription() const override;
+	virtual FText GetNodeTitle(const bool bGetDescriptive = false) const override;
 	virtual FText GetMenuCategory() const override;
 	virtual FLinearColor GetNodeTitleColor() const override;
 	virtual FSlateIcon GetIconAndTint(FLinearColor& OutColor) const override;
+
+	//~ Begin UObject interface
+	virtual void PostEditChangeProperty(FPropertyChangedEvent& PropertyChangedEvent) override;
+	//~ End UObject interface
 #endif
+
+	/** The name of this collection, which is used to reference this collection in the graph. */
+	UPROPERTY(EditAnywhere, Category = "General")
+	FString CollectionName;
 };
 
 /** A node which modifies part of the world being rendered. */
@@ -54,11 +63,24 @@ public:
 	virtual TArray<FMovieGraphPinProperties> GetOutputPinProperties() const override;
 
 #if WITH_EDITOR
-	virtual FText GetMenuDescription() const override;
+	virtual FText GetNodeTitle(const bool bGetDescriptive = false) const override;
 	virtual FText GetMenuCategory() const override;
 	virtual FLinearColor GetNodeTitleColor() const override;
 	virtual FSlateIcon GetIconAndTint(FLinearColor& OutColor) const override;
+
+	//~ Begin UObject interface
+	virtual void PostEditChangeProperty(FPropertyChangedEvent& PropertyChangedEvent) override;
+	//~ End UObject interface
 #endif
+
+	// TODO: Temporary - The modifier name will come from the specific specialized modifier node type (eg, MaterialModifier) in the future
+	/** The name of this modifier. */
+	UPROPERTY(EditAnywhere, Category = "General")
+	FString ModifierName;
+
+	/** The name of the collection being modified. */
+	UPROPERTY(EditAnywhere, Category = "General")
+	FString ModifiedCollectionName;
 };
 
 /** A node which configures the global game overrides. */
@@ -74,7 +96,7 @@ public:
 	virtual TArray<FMovieGraphPinProperties> GetOutputPinProperties() const override;
 
 #if WITH_EDITOR
-	virtual FText GetMenuDescription() const override;
+	virtual FText GetNodeTitle(const bool bGetDescriptive = false) const override;
 	virtual FText GetMenuCategory() const override;
 	virtual FLinearColor GetNodeTitleColor() const override;
 	virtual FSlateIcon GetIconAndTint(FLinearColor& OutColor) const override;
@@ -94,7 +116,7 @@ public:
 	virtual TArray<FMovieGraphPinProperties> GetOutputPinProperties() const override;
 
 #if WITH_EDITOR
-	virtual FText GetMenuDescription() const override;
+	virtual FText GetNodeTitle(const bool bGetDescriptive = false) const override;
 	virtual FText GetMenuCategory() const override;
 	virtual FLinearColor GetNodeTitleColor() const override;
 	virtual FSlateIcon GetIconAndTint(FLinearColor& OutColor) const override;
@@ -114,7 +136,7 @@ public:
 	virtual TArray<FMovieGraphPinProperties> GetOutputPinProperties() const override;
 
 #if WITH_EDITOR
-	virtual FText GetMenuDescription() const override;
+	virtual FText GetNodeTitle(const bool bGetDescriptive = false) const override;
 	virtual FText GetMenuCategory() const override;
 	virtual FLinearColor GetNodeTitleColor() const override;
 	virtual FSlateIcon GetIconAndTint(FLinearColor& OutColor) const override;
@@ -134,7 +156,7 @@ public:
 	virtual TArray<FMovieGraphPinProperties> GetOutputPinProperties() const override;
 
 #if WITH_EDITOR
-	virtual FText GetMenuDescription() const override;
+	virtual FText GetNodeTitle(const bool bGetDescriptive = false) const override;
 	virtual FText GetMenuCategory() const override;
 	virtual FLinearColor GetNodeTitleColor() const override;
 	virtual FSlateIcon GetIconAndTint(FLinearColor& OutColor) const override;
@@ -154,7 +176,7 @@ public:
 	virtual TArray<FMovieGraphPinProperties> GetOutputPinProperties() const override;
 
 #if WITH_EDITOR
-	virtual FText GetMenuDescription() const override;
+	virtual FText GetNodeTitle(const bool bGetDescriptive = false) const override;
 	virtual FText GetMenuCategory() const override;
 	virtual FLinearColor GetNodeTitleColor() const override;
 	virtual FSlateIcon GetIconAndTint(FLinearColor& OutColor) const override;
@@ -174,11 +196,48 @@ public:
 	virtual TArray<FMovieGraphPinProperties> GetOutputPinProperties() const override;
 
 #if WITH_EDITOR
-	virtual FText GetMenuDescription() const override;
+	virtual FText GetNodeTitle(const bool bGetDescriptive = false) const override;
+	virtual FText GetMenuCategory() const override;
+	virtual FLinearColor GetNodeTitleColor() const override;
+	virtual FSlateIcon GetIconAndTint(FLinearColor& OutColor) const override;
+#endif
+};
+
+// TODO: Currently this node is restricted to just accepting string-based options.
+/** A node which creates a condition that selects from a set of input branches. */
+UCLASS()
+class MOVIERENDERPIPELINECORE_API UMovieGraphSelectNode : public UMovieGraphNode
+{
+	GENERATED_BODY()
+
+public:
+	UMovieGraphSelectNode() = default;
+
+	virtual TArray<FMovieGraphPinProperties> GetInputPinProperties() const override;
+	virtual TArray<FMovieGraphPinProperties> GetOutputPinProperties() const override;
+
+#if WITH_EDITOR
+	virtual FText GetNodeTitle(const bool bGetDescriptive = false) const override;
 	virtual FText GetMenuCategory() const override;
 	virtual FLinearColor GetNodeTitleColor() const override;
 	virtual FSlateIcon GetIconAndTint(FLinearColor& OutColor) const override;
+
+	//~ Begin UObject interface
+	virtual void PostEditChangeProperty(FPropertyChangedEvent& PropertyChangedEvent) override;
+	//~ End UObject interface
 #endif
+
+	/** The values of options which can be selected. */
+	UPROPERTY(EditAnywhere, Category = "General")
+	TArray<FString> SelectOptions;
+
+	/** The value of the option which has been selected. */
+	UPROPERTY(EditAnywhere, Category = "General")
+	FString SelectedOption;
+
+	/** A description of what this select is doing. */
+	UPROPERTY(EditAnywhere, Category = "General")
+	FString Description;
 };
 
 /** A node which configures output settings. */
@@ -194,7 +253,7 @@ public:
 	virtual TArray<FMovieGraphPinProperties> GetOutputPinProperties() const override;
 
 #if WITH_EDITOR
-	virtual FText GetMenuDescription() const override;
+	virtual FText GetNodeTitle(const bool bGetDescriptive = false) const override;
 	virtual FText GetMenuCategory() const override;
 	virtual FLinearColor GetNodeTitleColor() const override;
 	virtual FSlateIcon GetIconAndTint(FLinearColor& OutColor) const override;
@@ -214,7 +273,7 @@ public:
 	virtual TArray<FMovieGraphPinProperties> GetOutputPinProperties() const override;
 
 #if WITH_EDITOR
-	virtual FText GetMenuDescription() const override;
+	virtual FText GetNodeTitle(const bool bGetDescriptive = false) const override;
 	virtual FText GetMenuCategory() const override;
 	virtual FLinearColor GetNodeTitleColor() const override;
 	virtual FSlateIcon GetIconAndTint(FLinearColor& OutColor) const override;
diff --git a/Engine/Plugins/MovieScene/MovieRenderPipeline/Source/MovieRenderPipelineCore/Public/Graph/Nodes/MovieGraphDeferredRenderPassNode.h b/Engine/Plugins/MovieScene/MovieRenderPipeline/Source/MovieRenderPipelineCore/Public/Graph/Nodes/MovieGraphDeferredRenderPassNode.h
--- a/Engine/Plugins/MovieScene/MovieRenderPipeline/Source/MovieRenderPipelineCore/Public/Graph/Nodes/MovieGraphDeferredRenderPassNode.h
+++ b/Engine/Plugins/MovieScene/MovieRenderPipeline/Source/MovieRenderPipelineCore/Public/Graph/Nodes/MovieGraphDeferredRenderPassNode.h
@@ -50,7 +50,7 @@ class UMovieGraphDeferredRenderPassNode : public UMovieGraphRenderPassNode
 	GENERATED_BODY()
 public:
 #if WITH_EDITOR
-	virtual FText GetMenuDescription() const override
+	virtual FText GetNodeTitle(const bool bGetDescriptive = false) const override
 	{
 		return NSLOCTEXT("MovieGraphNodes", "DeferredRenderPassGraphNode_Description", "Deferred Renderer");
 	}
diff --git a/Engine/Plugins/MovieScene/MovieRenderPipeline/Source/MovieRenderPipelineCore/Public/Graph/Nodes/MovieGraphInputNode.h b/Engine/Plugins/MovieScene/MovieRenderPipeline/Source/MovieRenderPipelineCore/Public/Graph/Nodes/MovieGraphInputNode.h
--- a/Engine/Plugins/MovieScene/MovieRenderPipeline/Source/MovieRenderPipelineCore/Public/Graph/Nodes/MovieGraphInputNode.h
+++ b/Engine/Plugins/MovieScene/MovieRenderPipeline/Source/MovieRenderPipelineCore/Public/Graph/Nodes/MovieGraphInputNode.h
@@ -18,7 +18,7 @@ public:
 	virtual TArray<FMovieGraphPinProperties> GetOutputPinProperties() const override;
 
 #if WITH_EDITOR
-	virtual FText GetMenuDescription() const override;
+	virtual FText GetNodeTitle(const bool bGetDescriptive = false) const override;
 	virtual FText GetMenuCategory() const override;
 	virtual FLinearColor GetNodeTitleColor() const override;
 	virtual FSlateIcon GetIconAndTint(FLinearColor& OutColor) const override;
diff --git a/Engine/Plugins/MovieScene/MovieRenderPipeline/Source/MovieRenderPipelineCore/Public/Graph/Nodes/MovieGraphOutputNode.h b/Engine/Plugins/MovieScene/MovieRenderPipeline/Source/MovieRenderPipelineCore/Public/Graph/Nodes/MovieGraphOutputNode.h
--- a/Engine/Plugins/MovieScene/MovieRenderPipeline/Source/MovieRenderPipelineCore/Public/Graph/Nodes/MovieGraphOutputNode.h
+++ b/Engine/Plugins/MovieScene/MovieRenderPipeline/Source/MovieRenderPipelineCore/Public/Graph/Nodes/MovieGraphOutputNode.h
@@ -18,7 +18,7 @@ public:
 	virtual TArray<FMovieGraphPinProperties> GetInputPinProperties() const override;
 
 #if WITH_EDITOR
-	virtual FText GetMenuDescription() const override;
+	virtual FText GetNodeTitle(const bool bGetDescriptive = false) const override;
 	virtual FText GetMenuCategory() const override;
 	virtual FLinearColor GetNodeTitleColor() const override;
 	virtual FSlateIcon GetIconAndTint(FLinearColor& OutColor) const override;
diff --git a/Engine/Plugins/MovieScene/MovieRenderPipeline/Source/MovieRenderPipelineCore/Public/Graph/Nodes/MovieGraphRenderLayerNode.h b/Engine/Plugins/MovieScene/MovieRenderPipeline/Source/MovieRenderPipelineCore/Public/Graph/Nodes/MovieGraphRenderLayerNode.h
--- a/Engine/Plugins/MovieScene/MovieRenderPipeline/Source/MovieRenderPipelineCore/Public/Graph/Nodes/MovieGraphRenderLayerNode.h
+++ b/Engine/Plugins/MovieScene/MovieRenderPipeline/Source/MovieRenderPipelineCore/Public/Graph/Nodes/MovieGraphRenderLayerNode.h
@@ -19,10 +19,14 @@ public:
 	virtual TArray<FPropertyBagPropertyDesc> GetDynamicPropertyDescriptions() const override;
 
 #if WITH_EDITOR
-	virtual FText GetMenuDescription() const override;
+	virtual FText GetNodeTitle(const bool bGetDescriptive = false) const override;
 	virtual FText GetMenuCategory() const override;
 	virtual FLinearColor GetNodeTitleColor() const override;
 	virtual FSlateIcon GetIconAndTint(FLinearColor& OutColor) const override;
+
+	//~ Begin UObject interface
+	virtual void PostEditChangeProperty(FPropertyChangedEvent& PropertyChangedEvent) override;
+	//~ End UObject interface
 #endif
 
 	UPROPERTY(BlueprintReadWrite, EditAnywhere, Category = "Settings")
diff --git a/Engine/Plugins/MovieScene/MovieRenderPipeline/Source/MovieRenderPipelineCore/Public/Graph/Nodes/MovieGraphVariableNode.h b/Engine/Plugins/MovieScene/MovieRenderPipeline/Source/MovieRenderPipelineCore/Public/Graph/Nodes/MovieGraphVariableNode.h
--- a/Engine/Plugins/MovieScene/MovieRenderPipeline/Source/MovieRenderPipelineCore/Public/Graph/Nodes/MovieGraphVariableNode.h
+++ b/Engine/Plugins/MovieScene/MovieRenderPipeline/Source/MovieRenderPipelineCore/Public/Graph/Nodes/MovieGraphVariableNode.h
@@ -24,7 +24,7 @@ public:
 	void SetVariable(UMovieGraphVariable* InVariable);
 
 #if WITH_EDITOR
-	virtual FText GetMenuDescription() const override;
+	virtual FText GetNodeTitle(const bool bGetDescriptive = false) const override;
 	virtual FText GetMenuCategory() const override;
 	virtual FLinearColor GetNodeTitleColor() const override;
 	virtual FSlateIcon GetIconAndTint(FLinearColor& OutColor) const override;
diff --git a/Engine/Plugins/MovieScene/MovieRenderPipeline/Source/MovieRenderPipelineEditor/Private/Graph/MovieEdGraphNode.cpp b/Engine/Plugins/MovieScene/MovieRenderPipeline/Source/MovieRenderPipelineEditor/Private/Graph/MovieEdGraphNode.cpp
--- a/Engine/Plugins/MovieScene/MovieRenderPipeline/Source/MovieRenderPipelineEditor/Private/Graph/MovieEdGraphNode.cpp
+++ b/Engine/Plugins/MovieScene/MovieRenderPipeline/Source/MovieRenderPipelineEditor/Private/Graph/MovieEdGraphNode.cpp
@@ -120,10 +120,11 @@ FText UMoviePipelineEdGraphNode::GetNodeTitle(ENodeTitleType::Type TitleType) co
 {
 	if (RuntimeNode)
 	{
-		return RuntimeNode->GetMenuDescription();
+		const bool bGetDescriptive = true;
+		return RuntimeNode->GetNodeTitle(bGetDescriptive);
 	}
 
-	return LOCTEXT("GraphTestTitle", "TestTitle");
+	return LOCTEXT("UnknownNodeTitle", "Unknown");
 }
 
 FText UMoviePipelineEdGraphNode::GetTooltipText() const
diff --git a/Engine/Plugins/MovieScene/MovieRenderPipeline/Source/MovieRenderPipelineEditor/Private/Graph/MovieGraphSchema.cpp b/Engine/Plugins/MovieScene/MovieRenderPipeline/Source/MovieRenderPipelineEditor/Private/Graph/MovieGraphSchema.cpp
--- a/Engine/Plugins/MovieScene/MovieRenderPipeline/Source/MovieRenderPipelineEditor/Private/Graph/MovieGraphSchema.cpp
+++ b/Engine/Plugins/MovieScene/MovieRenderPipeline/Source/MovieRenderPipelineEditor/Private/Graph/MovieGraphSchema.cpp
@@ -111,7 +111,7 @@ void UMovieGraphSchema::GetGraphContextActions(FGraphContextMenuBuilder& Context
 		// to filter out incompatible nodes.
 		// if (!ContextMenuBuilder.FromPin || ContextMenuBuilder.FromPin->Direction == EGPD_Input)
 		{
-			const FText Name = PipelineNode->GetMenuDescription();
+			const FText Name = PipelineNode->GetNodeTitle();
 			const FText Category = PipelineNode->GetMenuCategory();
 			const FText Tooltip = LOCTEXT("CreateNode_Tooltip", "Create a node of this type.");
 			
