MoviePipeline: Added the ability to add comments to the graph. These are comments which go around multiple nodes, not pinnable comments on individual nodes.

diff --git a/Engine/Plugins/MovieScene/MovieRenderPipeline/Source/MovieRenderPipelineCore/Private/Graph/MovieGraphConfig.cpp b/Engine/Plugins/MovieScene/MovieRenderPipeline/Source/MovieRenderPipelineCore/Private/Graph/MovieGraphConfig.cpp
--- a/Engine/Plugins/MovieScene/MovieRenderPipeline/Source/MovieRenderPipelineCore/Private/Graph/MovieGraphConfig.cpp
+++ b/Engine/Plugins/MovieScene/MovieRenderPipeline/Source/MovieRenderPipelineCore/Private/Graph/MovieGraphConfig.cpp
@@ -532,6 +532,18 @@ TArray<FMovieGraphBranch> UMovieGraphConfig::GetOutputBranches() const
 	return Branches;
 }
 
+#if WITH_EDITOR
+void UMovieGraphConfig::SetEditorOnlyNodes(const TArray<TObjectPtr<const UObject>>& InNodes)
+{
+	EditorOnlyNodes.Empty();
+
+	for (const TObjectPtr<const UObject>& Node : InNodes)
+	{
+		EditorOnlyNodes.Add(DuplicateObject(Node.Get(), this));
+	}
+}
+#endif	// WITH_EDITOR
+
 bool UMovieGraphConfig::DeleteInputMember(UMovieGraphInput* InputMemberToDelete)
 {
 	if (InputMemberToDelete)
diff --git a/Engine/Plugins/MovieScene/MovieRenderPipeline/Source/MovieRenderPipelineCore/Public/Graph/MovieGraphConfig.h b/Engine/Plugins/MovieScene/MovieRenderPipeline/Source/MovieRenderPipelineCore/Public/Graph/MovieGraphConfig.h
--- a/Engine/Plugins/MovieScene/MovieRenderPipeline/Source/MovieRenderPipelineCore/Public/Graph/MovieGraphConfig.h
+++ b/Engine/Plugins/MovieScene/MovieRenderPipeline/Source/MovieRenderPipelineCore/Public/Graph/MovieGraphConfig.h
@@ -220,6 +220,14 @@ public:
 	/** Returns only the names of the root branches in the Output Node, with no depth information. */
 	TArray<FMovieGraphBranch> GetOutputBranches() const;
 
+#if WITH_EDITOR
+	/** Gets the editor-only nodes in this graph. Editor-only nodes do not have an equivalent runtime node. */
+	const TArray<TObjectPtr<UObject>>& GetEditorOnlyNodes() const { return EditorOnlyNodes; }
+
+	/** Sets the editor-only nodes in this graph. */
+	void SetEditorOnlyNodes(const TArray<TObjectPtr<const UObject>>& InNodes);
+#endif
+
 	template<typename NodeType>
 	NodeType* IterateGraphForClass(const FMovieGraphTraversalContext& InContext) const
 	{
@@ -319,4 +327,10 @@ private:
 	/** All outputs which have been defined on the graph. */
 	UPROPERTY()
 	TArray<TObjectPtr<UMovieGraphOutput>> Outputs;
+
+#if WITH_EDITORONLY_DATA
+	/** Nodes which are only useful in the editor (like comments) and have no runtime equivalent */
+	UPROPERTY()
+	TArray<TObjectPtr<UObject>> EditorOnlyNodes;
+#endif
 };
\ No newline at end of file
diff --git a/Engine/Plugins/MovieScene/MovieRenderPipeline/Source/MovieRenderPipelineEditor/Private/Graph/MovieEdGraph.cpp b/Engine/Plugins/MovieScene/MovieRenderPipeline/Source/MovieRenderPipelineEditor/Private/Graph/MovieEdGraph.cpp
--- a/Engine/Plugins/MovieScene/MovieRenderPipeline/Source/MovieRenderPipelineEditor/Private/Graph/MovieEdGraph.cpp
+++ b/Engine/Plugins/MovieScene/MovieRenderPipeline/Source/MovieRenderPipelineEditor/Private/Graph/MovieEdGraph.cpp
@@ -54,7 +54,7 @@ void UMoviePipelineEdGraph::InitFromRuntimeGraph(UMovieGraphConfig* InGraph)
 		}
 	}
 
-	// Now that we've added an Editor Graph representation for every node in the graph, link
+	// Now that we've added an Editor Graph representation for every runtime node in the graph, link
 	// the editor nodes together to match the Runtime Layout.
 	for (const TPair< UMovieGraphNode*, UMoviePipelineEdGraphNodeBase*> Pair : NodeLookup)
 	{
@@ -63,6 +63,18 @@ void UMoviePipelineEdGraph::InitFromRuntimeGraph(UMovieGraphConfig* InGraph)
 		CreateLinks(Pair.Value, bCreateInboundLinks, bCreateOutboundLinks, NodeLookup);
 	}
 
+	// Restore editor-only nodes, which have no runtime node equivalent
+	for (const TObjectPtr<UObject> EditorOnlyNodeObject : InGraph->GetEditorOnlyNodes())
+	{
+		if (const UEdGraphNode* EdGraphNode = Cast<UEdGraphNode>(EditorOnlyNodeObject))
+		{
+			UEdGraphNode* NewEdGraphNode = DuplicateObject(EdGraphNode, /*Outer=*/this);
+			const bool bIsUserAction = false;
+			const bool bSelectNewNode = false;
+			AddNode(NewEdGraphNode, bIsUserAction, bSelectNewNode);
+		}
+	}
+
 	InGraph->OnGraphChangedDelegate.AddUObject(this, &UMoviePipelineEdGraph::OnGraphConfigChanged);
 	InGraph->OnGraphNodesDeletedDelegate.AddUObject(this, &UMoviePipelineEdGraph::OnGraphNodesDeleted);
 
diff --git a/Engine/Plugins/MovieScene/MovieRenderPipeline/Source/MovieRenderPipelineEditor/Private/Graph/MovieGraphAssetToolkit.cpp b/Engine/Plugins/MovieScene/MovieRenderPipeline/Source/MovieRenderPipelineEditor/Private/Graph/MovieGraphAssetToolkit.cpp
--- a/Engine/Plugins/MovieScene/MovieRenderPipeline/Source/MovieRenderPipelineEditor/Private/Graph/MovieGraphAssetToolkit.cpp
+++ b/Engine/Plugins/MovieScene/MovieRenderPipeline/Source/MovieRenderPipelineEditor/Private/Graph/MovieGraphAssetToolkit.cpp
@@ -4,6 +4,7 @@
 
 #include "Framework/Commands/GenericCommands.h"
 #include "Graph/MovieGraphConfig.h"
+#include "MovieEdGraphNode.h"
 #include "MovieGraphSchema.h"
 #include "PropertyEditorModule.h"
 #include "SMovieGraphMembersTabContent.h"
@@ -233,6 +234,25 @@ bool FMovieGraphAssetToolkit::CanDeleteSelectedMembers()
 	return false;
 }
 
+void FMovieGraphAssetToolkit::PersistEditorOnlyNodes() const
+{
+	if (InitialGraph && InitialGraph->PipelineEdGraph)
+	{
+		TArray<TObjectPtr<const UObject>> EditorOnlyNodes;
+		for (const TObjectPtr<UEdGraphNode>& GraphEdNode : InitialGraph->PipelineEdGraph->Nodes)
+		{
+			// Treat any non-MRQ nodes as editor-only nodes
+			check(GraphEdNode);
+			if (!GraphEdNode->IsA<UMoviePipelineEdGraphNodeBase>())
+			{
+				EditorOnlyNodes.Add(GraphEdNode);
+			}
+		}
+
+		InitialGraph->SetEditorOnlyNodes(EditorOnlyNodes);
+	}
+}
+
 FName FMovieGraphAssetToolkit::GetToolkitFName() const
 {
 	return FName("MovieGraphEditor");
@@ -255,6 +275,9 @@ FLinearColor FMovieGraphAssetToolkit::GetWorldCentricTabColorScale() const
 
 void FMovieGraphAssetToolkit::SaveAsset_Execute()
 {
+	// Editor-only nodes are copied to the underlying runtime graph on save/close
+	PersistEditorOnlyNodes();
+	
 	// Perform the default save process
 	// TODO: This will fail silently on a transient graph and won't trigger a Save As
 	FAssetEditorToolkit::SaveAsset_Execute();
@@ -262,4 +285,12 @@ void FMovieGraphAssetToolkit::SaveAsset_Execute()
 	// TODO: Any custom save logic here
 }
 
+void FMovieGraphAssetToolkit::OnClose()
+{
+	// Editor-only nodes are copied to the underlying runtime graph on save/close
+	PersistEditorOnlyNodes();
+	
+	FAssetEditorToolkit::OnClose();
+}
+
 #undef LOCTEXT_NAMESPACE
diff --git a/Engine/Plugins/MovieScene/MovieRenderPipeline/Source/MovieRenderPipelineEditor/Private/Graph/MovieGraphAssetToolkit.h b/Engine/Plugins/MovieScene/MovieRenderPipeline/Source/MovieRenderPipelineEditor/Private/Graph/MovieGraphAssetToolkit.h
--- a/Engine/Plugins/MovieScene/MovieRenderPipeline/Source/MovieRenderPipelineEditor/Private/Graph/MovieGraphAssetToolkit.h
+++ b/Engine/Plugins/MovieScene/MovieRenderPipeline/Source/MovieRenderPipelineEditor/Private/Graph/MovieGraphAssetToolkit.h
@@ -18,6 +18,7 @@ public:
 	virtual FText GetBaseToolkitName() const override;
 	virtual FString GetWorldCentricTabPrefix() const override;
 	virtual void SaveAsset_Execute() override;
+	virtual void OnClose() override;
 	//~ End FAssetEditorToolkit interface
 
 	//~ Begin IToolkit interface
@@ -40,6 +41,9 @@ private:
 	/** Determines if the selected member(s) can be deleted. */
 	bool CanDeleteSelectedMembers();
 
+	/** Copies editor-only nodes to the underlying runtime graph so they can be restored later. */
+	void PersistEditorOnlyNodes() const;
+
 private:
 	/** The details panel for the selected object(s) in the graph */
 	TSharedPtr<IDetailsView> SelectedGraphObjectsDetailsWidget;
diff --git a/Engine/Plugins/MovieScene/MovieRenderPipeline/Source/MovieRenderPipelineEditor/Private/Graph/MovieGraphSchema.cpp b/Engine/Plugins/MovieScene/MovieRenderPipeline/Source/MovieRenderPipelineEditor/Private/Graph/MovieGraphSchema.cpp
--- a/Engine/Plugins/MovieScene/MovieRenderPipeline/Source/MovieRenderPipelineEditor/Private/Graph/MovieGraphSchema.cpp
+++ b/Engine/Plugins/MovieScene/MovieRenderPipeline/Source/MovieRenderPipelineEditor/Private/Graph/MovieGraphSchema.cpp
@@ -346,4 +346,25 @@ UEdGraphNode* FMovieGraphSchemaAction_NewVariableNode::PerformAction(UEdGraph* P
 	return GraphNode;
 }
 
+UEdGraphNode* FMovieGraphSchemaAction_NewComment::PerformAction(UEdGraph* ParentGraph, UEdGraphPin* FromPin, const FVector2D Location, bool bSelectNewNode)
+{
+	UEdGraphNode_Comment* CommentTemplate = NewObject<UEdGraphNode_Comment>();
+
+	const TSharedPtr<SGraphEditor> GraphEditorPtr = SGraphEditor::FindGraphEditorForGraph(ParentGraph);
+
+	FVector2D SpawnLocation = Location;
+	FSlateRect Bounds;
+	if (GraphEditorPtr && GraphEditorPtr->GetBoundsForSelectedNodes(Bounds, 50.0f))
+	{
+		CommentTemplate->SetBounds(Bounds);
+		SpawnLocation.X = CommentTemplate->NodePosX;
+		SpawnLocation.Y = CommentTemplate->NodePosY;
+	}
+	
+	const FScopedTransaction Transaction(LOCTEXT("GraphEditor_NewCommentNode", "Add New Comment Node"));
+	ParentGraph->Modify();
+	
+	return FEdGraphSchemaAction_NewNode::SpawnNodeFromTemplate<UEdGraphNode_Comment>(ParentGraph, CommentTemplate, SpawnLocation, bSelectNewNode);;
+}
+
 #undef LOCTEXT_NAMESPACE // "MoviePipelineGraphSchema"
diff --git a/Engine/Plugins/MovieScene/MovieRenderPipeline/Source/MovieRenderPipelineEditor/Private/Graph/MovieGraphSchema.h b/Engine/Plugins/MovieScene/MovieRenderPipeline/Source/MovieRenderPipelineEditor/Private/Graph/MovieGraphSchema.h
--- a/Engine/Plugins/MovieScene/MovieRenderPipeline/Source/MovieRenderPipelineEditor/Private/Graph/MovieGraphSchema.h
+++ b/Engine/Plugins/MovieScene/MovieRenderPipeline/Source/MovieRenderPipelineEditor/Private/Graph/MovieGraphSchema.h
@@ -126,4 +126,26 @@ private:
 	/** GUID of the runtime variable this action relates to. */
 	UPROPERTY()
 	FGuid VariableGuid;
+};
+
+/** Schema action for creating a new comment node in the graph. */
+USTRUCT()
+struct FMovieGraphSchemaAction_NewComment : public FMovieGraphSchemaAction
+{
+	GENERATED_BODY()
+
+	static FName StaticGetTypeId()
+	{
+		static FName Type("FMovieGraphSchemaAction_NewComment");
+		return Type;
+	}
+
+	FMovieGraphSchemaAction_NewComment()
+		: FMovieGraphSchemaAction()
+	{}
+
+	//~ Begin FEdGraphSchemaAction Interface
+	virtual FName GetTypeId() const override { return StaticGetTypeId(); }
+	virtual UEdGraphNode* PerformAction(UEdGraph* ParentGraph, UEdGraphPin* FromPin, const FVector2D Location, bool bSelectNewNode = true) override;
+	//~ End FEdGraphSchemaAction Interface
 };
\ No newline at end of file
diff --git a/Engine/Plugins/MovieScene/MovieRenderPipeline/Source/MovieRenderPipelineEditor/Private/Widgets/Graph/SMovieGraphConfigPanel.cpp b/Engine/Plugins/MovieScene/MovieRenderPipeline/Source/MovieRenderPipelineEditor/Private/Widgets/Graph/SMovieGraphConfigPanel.cpp
--- a/Engine/Plugins/MovieScene/MovieRenderPipeline/Source/MovieRenderPipelineEditor/Private/Widgets/Graph/SMovieGraphConfigPanel.cpp
+++ b/Engine/Plugins/MovieScene/MovieRenderPipeline/Source/MovieRenderPipelineEditor/Private/Widgets/Graph/SMovieGraphConfigPanel.cpp
@@ -57,6 +57,7 @@
 #include "Graph/MovieGraphSchema.h"
 #include "Kismet2/BlueprintEditorUtils.h"
 #include "GraphEditor.h"
+#include "GraphEditorActions.h"
 
 #define LOCTEXT_NAMESPACE "SMoviePipelineGraphPanel"
 
@@ -121,6 +122,10 @@ void SMoviePipelineGraphPanel::MakeEditorCommands()
 {
 	GraphEditorCommands = MakeShareable(new FUICommandList);
 	{
+		GraphEditorCommands->MapAction(FGraphEditorCommands::Get().CreateComment,
+			FExecuteAction::CreateSP(this, &SMoviePipelineGraphPanel::OnCreateComment)
+		);
+		
 		// Editing commands
 		GraphEditorCommands->MapAction(FGenericCommands::Get().SelectAll,
 			FExecuteAction::CreateSP(this, &SMoviePipelineGraphPanel::SelectAllNodes),
@@ -134,6 +139,22 @@ void SMoviePipelineGraphPanel::MakeEditorCommands()
 	}
 }
 
+void SMoviePipelineGraphPanel::OnCreateComment() const
+{
+	if (CurrentGraph)
+	{
+		FMovieGraphSchemaAction_NewComment CommentAction;
+		
+		const TSharedPtr<SGraphEditor> GraphEditorPtr = SGraphEditor::FindGraphEditorForGraph(CurrentGraph->PipelineEdGraph);
+		FVector2D Location;
+		if (GraphEditorPtr)
+		{
+			Location = GraphEditorPtr->GetPasteLocation();
+		}
+		
+		CommentAction.PerformAction(CurrentGraph->PipelineEdGraph, nullptr, Location);
+	}
+}
 
 void SMoviePipelineGraphPanel::SelectAllNodes()
 {
diff --git a/Engine/Plugins/MovieScene/MovieRenderPipeline/Source/MovieRenderPipelineEditor/Private/Widgets/Graph/SMovieGraphConfigPanel.h b/Engine/Plugins/MovieScene/MovieRenderPipeline/Source/MovieRenderPipelineEditor/Private/Widgets/Graph/SMovieGraphConfigPanel.h
--- a/Engine/Plugins/MovieScene/MovieRenderPipeline/Source/MovieRenderPipelineEditor/Private/Widgets/Graph/SMovieGraphConfigPanel.h
+++ b/Engine/Plugins/MovieScene/MovieRenderPipeline/Source/MovieRenderPipelineEditor/Private/Widgets/Graph/SMovieGraphConfigPanel.h
@@ -74,6 +74,9 @@ private:
 	//UMoviePipelineConfigBase* AllocateTransientPreset();
 	void MakeEditorCommands();
 
+	/** Creates a new, empty comment in the graph around the currently selected nodes. */
+	void OnCreateComment() const;
+	
 	void SelectAllNodes();
 	bool CanSelectAllNodes() const;
 	bool CanDeleteSelectedNodes() const;
