MoviePipeline: Set up the Output Merger to collect all of the data created each frame before passing it onto the actual output step, implement a basic temporary jpeg writer, and track files written to disk so they can be provided in a post-render structure.

diff --git a/Engine/Plugins/MovieScene/MovieRenderPipeline/Source/MovieRenderPipelineCore/Private/Graph/MovieGraphDefaultRenderer.cpp b/Engine/Plugins/MovieScene/MovieRenderPipeline/Source/MovieRenderPipelineCore/Private/Graph/MovieGraphDefaultRenderer.cpp
--- a/Engine/Plugins/MovieScene/MovieRenderPipeline/Source/MovieRenderPipelineCore/Private/Graph/MovieGraphDefaultRenderer.cpp
+++ b/Engine/Plugins/MovieScene/MovieRenderPipeline/Source/MovieRenderPipelineCore/Private/Graph/MovieGraphDefaultRenderer.cpp
@@ -29,7 +29,7 @@ void UMovieGraphDefaultRenderer::SetupRenderingPipelineForShot(UMoviePipelineExe
 	for (const FMovieGraphBranch& Branch : GraphBranches)
 	{
 		// We follow each branch looking for Render Layer nodes to figure out what render layer this should be.
-		FMovieGraphTraversalContext TraversalContext = GetOwningGraph()->GetTraversalContextForShot(InShot);
+		FMovieGraphTraversalContext TraversalContext = GetOwningGraph()->GetCurrentTraversalContext();
 		TraversalContext.RootBranch = Branch;
 
 		FString RenderLayerName = TEXT("Unnamed_Layer");
@@ -157,9 +157,36 @@ void UMovieGraphDefaultRenderer::Render(const FMovieGraphTimeStepData& InTimeSte
 		}
 	}
 
+	if (InTimeStepData.bIsFirstTemporalSampleForFrame)
+	{
+		// If this is the first sample for this output frame, then we need to 
+		// talk to all of our render passes and ask them for what data they will
+		// produce, and set the Output Merger up with that knowledge.
+		UE::MovieGraph::FMovieGraphOutputMergerFrame& NewOutputFrame = GetOwningGraph()->GetOutputMerger()->AllocateNewOutputFrame_GameThread(InTimeStepData.OutputFrameNumber);
+
+		// Get the Traversal Context (not specific to any render pass) at the first sample. This is so
+		// we can easily fetch things that are shared between all render layers later.
+		NewOutputFrame.TraversalContext = GetOwningGraph()->GetCurrentTraversalContext();
+
+		for (const TObjectPtr<UMovieGraphRenderPassNode>& RenderPass : RenderPassesInUse)
+		{
+			RenderPass->GatherOutputPasses(NewOutputFrame.ExpectedRenderPasses);
+		}
+
+		// Register the frame with our render statistics as being worked on
+
+		UE::MovieGraph::FRenderTimeStatistics* TimeStats = GetRenderTimeStatistics(NewOutputFrame.TraversalContext.Time.OutputFrameNumber);
+		if (ensure(TimeStats))
+		{
+			TimeStats->StartTime = FDateTime::UtcNow();
+		}
+	}
+
 	for (TObjectPtr<UMovieGraphRenderPassNode> RenderPass : RenderPassesInUse)
 	{
-		RenderPass->Render(InTimeStepData);
+		// Pass in a copy of the traversal context so the renderer can decide what to do with it.
+		UE::MovieGraph::FMovieGraphOutputMergerFrame& OutputFrame = GetOwningGraph()->GetOutputMerger()->GetOutputFrame_GameThread(InTimeStepData.OutputFrameNumber);
+		RenderPass->Render(OutputFrame.TraversalContext, InTimeStepData);
 	}
 }
 
@@ -218,6 +245,11 @@ FMoviePipelineSurfaceQueuePtr UMovieGraphDefaultRenderer::CreateSurfaceQueue(con
 	return SurfaceQueue;
 }
 
+UE::MovieGraph::FRenderTimeStatistics* UMovieGraphDefaultRenderer::GetRenderTimeStatistics(const int32 InFrameNumber)
+{
+	return &RenderTimeStatistics.FindOrAdd(InFrameNumber);
+}
+
 namespace UE::MovieGraph::DefaultRenderer
 {
 FSurfaceAccumulatorPool::FInstancePtr FSurfaceAccumulatorPool::BlockAndGetAccumulator_GameThread(int32 InFrameNumber, const FMovieGraphRenderDataIdentifier& InPassIdentifier)
diff --git a/Engine/Plugins/MovieScene/MovieRenderPipeline/Source/MovieRenderPipelineCore/Private/Graph/MovieGraphLinearTimeStep.cpp b/Engine/Plugins/MovieScene/MovieRenderPipeline/Source/MovieRenderPipelineCore/Private/Graph/MovieGraphLinearTimeStep.cpp
--- a/Engine/Plugins/MovieScene/MovieRenderPipeline/Source/MovieRenderPipelineCore/Private/Graph/MovieGraphLinearTimeStep.cpp
+++ b/Engine/Plugins/MovieScene/MovieRenderPipeline/Source/MovieRenderPipelineCore/Private/Graph/MovieGraphLinearTimeStep.cpp
@@ -47,6 +47,11 @@ void UMovieGraphLinearTimeStep::TickProducingFrames()
 		CurrentTimeStepData.FrameDeltaTime = FrameRate.AsInterval();
 		CurrentTimeStepData.WorldSeconds = 0.f;
 		CurrentTimeStepData.MotionBlurFraction = FrameRate.AsInterval();
+		
+		// ToDo: This will change once we have temporal sub-sample support
+		CurrentTimeStepData.bIsFirstTemporalSampleForFrame = true;
+		CurrentTimeStepData.bIsLastTemporalSampleForFrame = true;
+
 
 		if (CurrentCameraCut->ShotInfo.CurrentTickInRoot >= CurrentCameraCut->ShotInfo.TotalOutputRangeRoot.GetUpperBoundValue())
 		{
diff --git a/Engine/Plugins/MovieScene/MovieRenderPipeline/Source/MovieRenderPipelineCore/Private/Graph/MovieGraphOutputMerger.cpp b/Engine/Plugins/MovieScene/MovieRenderPipeline/Source/MovieRenderPipelineCore/Private/Graph/MovieGraphOutputMerger.cpp
new file mode 100644
--- /dev/null
+++ b/Engine/Plugins/MovieScene/MovieRenderPipeline/Source/MovieRenderPipelineCore/Private/Graph/MovieGraphOutputMerger.cpp
@@ -0,0 +1,115 @@
+// Copyright Epic Games, Inc. All Rights Reserved.
+
+#include "MovieGraphOutputMerger.h"
+#include "Graph/MovieGraphPipeline.h"
+#include "Graph/MovieGraphRenderDataIdentifier.h"
+#include "Async/Async.h"
+
+namespace UE::MovieGraph
+{
+	FMovieGraphOutputMerger::FMovieGraphOutputMerger(UMovieGraphPipeline* InOwningMoviePipeline)
+		: WeakMoviePipeline(MakeWeakObjectPtr(InOwningMoviePipeline))
+	{
+	}
+
+	void FMovieGraphOutputMerger::AbandonOutstandingWork()
+	{
+		FScopeLock ScopeLock(&PendingDataMutex);
+		FinishedFrames.Empty();
+		PendingData.Empty();
+	}
+
+	FMovieGraphOutputMergerFrame& FMovieGraphOutputMerger::AllocateNewOutputFrame_GameThread(const int32 InFrameNumber) 
+	{
+		FScopeLock ScopeLock(&PendingDataMutex);
+
+		// Ensure this frame hasn't already been entered somehow.
+		check(!PendingData.Find(InFrameNumber));
+
+		FMovieGraphOutputMergerFrame& NewFrame = PendingData.Add(InFrameNumber);
+		return NewFrame;
+	}
+
+	FMovieGraphOutputMergerFrame& FMovieGraphOutputMerger::GetOutputFrame_GameThread(int32 InFrameNumber)
+	{
+		check(PendingData.Find(InFrameNumber));
+		return *PendingData.Find(InFrameNumber);
+	}
+
+	
+	void FMovieGraphOutputMerger::OnSingleSampleDataAvailable_AnyThread(TUniquePtr<FImagePixelData>&& InData)
+	{
+		// This is to support outputting individual samples (skipping accumulation) for debug reasons,
+		// or because you want to post-process them yourself. We just forward this directly on for output to disk.
+
+		TWeakObjectPtr<UMovieGraphPipeline> LocalWeakPipeline = WeakMoviePipeline;
+
+		AsyncTask(ENamedThreads::GameThread, [LocalData = MoveTemp(InData), LocalWeakPipeline]() mutable
+		{
+			if (ensureAlwaysMsgf(LocalWeakPipeline.IsValid(), TEXT("A memory lifespan issue has left an output builder alive without an owning Movie Pipeline.")))
+			{
+				// LocalWeakPipeline->Debug_OnSampleRendered(MoveTemp(LocalData));
+				// ToDo: Read the finished frames from MRQ
+				// also implement the above function.
+			}
+		}
+		);
+	}
+	
+	void FMovieGraphOutputMerger::OnCompleteRenderPassDataAvailable_AnyThread(TUniquePtr<FImagePixelData>&& InData)
+	{	
+		// Lock the ActiveData when we're updating what data has been gathered.
+		FScopeLock ScopeLock(&PendingDataMutex);
+		
+		// Fetch our payload from the data. If this check fails then you didn't attach the payload to the image.
+		FMovieGraphSampleState* Payload = InData->GetPayload<FMovieGraphSampleState>();
+		check(Payload);
+		
+		const int32 FrameNumber = Payload->TraversalContext.Time.OutputFrameNumber;
+
+		// See if we can find the frame this data is for. This should always be valid, if it's not
+		// valid it means they either forgot to declare they were going to produce it, or this is
+		// coming in after the system already thinks it's finished that frame.
+		FMovieGraphOutputMergerFrame* OutputFrame = PendingData.Find(FrameNumber);
+		
+		// Make sure we expected this frame number
+		if (!ensureAlwaysMsgf(OutputFrame, TEXT("Received data for unknown frame. Frame was either already processed or not queued yet!")))
+		{
+			return;
+		}
+		
+		// Make sure this render pass identifier was expected as well.
+		FMovieGraphRenderDataIdentifier NewLayerId = Payload->TraversalContext.RenderDataIdentifier;
+		if (!ensureAlwaysMsgf(OutputFrame->ExpectedRenderPasses.Contains(NewLayerId), TEXT("Received data for unexpected render pass: %s"), *LexToString(NewLayerId)))
+		{
+			return;
+		}
+		
+		// Put the new data inside this output frame.
+		OutputFrame->ImageOutputData.FindOrAdd(NewLayerId) = MoveTemp(InData);
+		
+		// ToDo:
+		// Merge the metadata from each output state. Metadata is part of the output state but gets forked when
+		// we submit different render passes, so we need to merge it again. Doesn't handle conflicts.
+		//for (const TPair<FString, FString>& KVP : Payload->SampleState.OutputState.FileMetadata)
+		//{
+		//	OutputFrame->FrameOutputState.FileMetadata.Add(KVP.Key, KVP.Value);
+		//}
+		
+		// Check to see if this was the last piece of data needed for this frame.
+		int32 TotalPasses = OutputFrame->ExpectedRenderPasses.Num();
+		int32 FinishedPasses = OutputFrame->ImageOutputData.Num();
+		
+		if (FinishedPasses == TotalPasses)
+		{
+			// ToDo: We potentialy want to sort the output using a stable-sort for EXRs so they can choose which is the primary rgba layer.
+			
+			// Move this frame into our FinishedFrames array so the Game Thread can read it at its leisure
+			FMovieGraphOutputMergerFrame FinalFrame;
+			ensureMsgf(PendingData.RemoveAndCopyValue(FrameNumber, FinalFrame), TEXT("Could not find frame in pending data, output will be skipped!"));
+			
+			// TQueue is thread safe so it's okay to just push the data into it.
+			FinishedFrames.Enqueue(MoveTemp(FinalFrame));
+		}
+	}
+}		
diff --git a/Engine/Plugins/MovieScene/MovieRenderPipeline/Source/MovieRenderPipelineCore/Private/Graph/MovieGraphOutputMerger.h b/Engine/Plugins/MovieScene/MovieRenderPipeline/Source/MovieRenderPipelineCore/Private/Graph/MovieGraphOutputMerger.h
new file mode 100644
--- /dev/null
+++ b/Engine/Plugins/MovieScene/MovieRenderPipeline/Source/MovieRenderPipelineCore/Private/Graph/MovieGraphOutputMerger.h
@@ -0,0 +1,50 @@
+// Copyright Epic Games, Inc. All Rights Reserved.
+
+#pragma once
+
+#include "Graph/MovieGraphDataTypes.h"
+#include "Containers/Queue.h"
+
+class UMovieGraphPipeline;
+
+namespace UE::MovieGraph
+{
+	/**
+	* When MRQ produces a single output frame, it's built out of many individual parts that are spread
+	* out over time. Different render passes within a single output frame may take longer to produce
+	* due to splitting out work onto the Task Graph, so this class is responsible for handling finished
+	* data for each render pass coming in at an arbitrary time, and once all of the pieces are there
+	* only then does this release it onto the actual pipeline saying the frame is finished.
+	*/
+	class FMovieGraphOutputMerger : public UE::MovieGraph::IMovieGraphOutputMerger
+	{
+	public:
+		FMovieGraphOutputMerger(UMovieGraphPipeline* InOwningMoviePipeline);
+
+		// IMovieGraphOutputMerger Interface
+		virtual FMovieGraphOutputMergerFrame& AllocateNewOutputFrame_GameThread(const int32 InFrameNumber) override;
+		virtual FMovieGraphOutputMergerFrame& GetOutputFrame_GameThread(int32 InFrameNumber) override;
+		virtual void OnCompleteRenderPassDataAvailable_AnyThread(TUniquePtr<FImagePixelData>&& InData) override;
+		virtual void OnSingleSampleDataAvailable_AnyThread(TUniquePtr<FImagePixelData>&& InData) override;
+		virtual void AbandonOutstandingWork() override;
+		virtual int32 GetNumOutstandingFrames() const override { return PendingData.Num(); }
+		virtual TQueue<FMovieGraphOutputMergerFrame>& GetFinishedFrames() override { return FinishedFrames; }
+		// ~IMovieGraphOutputMerger Interface
+
+	public:
+		/** 
+		* A thread safe FIFO queue of finished frames that have had all of their render passes complete. 
+		* All frames in the queue are processed on the next Game Thread tick and moved onto the output containers.
+		*/
+		TQueue<UE::MovieGraph::FMovieGraphOutputMergerFrame> FinishedFrames;
+	private:
+		/** The Movie Pipeline that owns us. */
+		TWeakObjectPtr<UMovieGraphPipeline> WeakMoviePipeline;
+
+		/** Data that is expected but not fully available yet. */
+		TMap<int32, FMovieGraphOutputMergerFrame> PendingData;
+
+		/** Mutex that protects adding/updating/removing from PendingData */
+		FCriticalSection PendingDataMutex;
+	};
+}
\ No newline at end of file
diff --git a/Engine/Plugins/MovieScene/MovieRenderPipeline/Source/MovieRenderPipelineCore/Private/Graph/MovieGraphPipeline.cpp b/Engine/Plugins/MovieScene/MovieRenderPipeline/Source/MovieRenderPipelineCore/Private/Graph/MovieGraphPipeline.cpp
--- a/Engine/Plugins/MovieScene/MovieRenderPipeline/Source/MovieRenderPipelineCore/Private/Graph/MovieGraphPipeline.cpp
+++ b/Engine/Plugins/MovieScene/MovieRenderPipeline/Source/MovieRenderPipelineCore/Private/Graph/MovieGraphPipeline.cpp
@@ -1,12 +1,16 @@
 // Copyright Epic Games, Inc. All Rights Reserved.
 
 #include "Graph/MovieGraphPipeline.h"
+#include "Graph/MovieGraphDataTypes.h"
+#include "Graph/MovieGraphOutputMerger.h"
+#include "Graph/Nodes/MovieGraphFileOutputNode.h"
 #include "MovieRenderPipelineCoreModule.h"
 #include "Misc/CoreDelegates.h"
 #include "MoviePipelineQueue.h"
-#include "Graph/MovieGraphDataTypes.h"
 #include "MovieScene.h"
 #include "RenderingThread.h"
+#include "ImageWriteQueue.h"
+#include "Modules/ModuleManager.h"
 
 // Temp
 #include "LevelSequence.h"
@@ -16,6 +20,9 @@ UMovieGraphPipeline::UMovieGraphPipeline()
 	, bIsTransitioningState(false)
 	, PipelineState(EMovieRenderPipelineState::Uninitialized)
 {
+	OutputMerger = MakeShared<UE::MovieGraph::FMovieGraphOutputMerger>(this);
+
+	Debug_ImageWriteQueue = &FModuleManager::Get().LoadModuleChecked<IImageWriteQueueModule>("ImageWriteQueue").GetWriteQueue();
 }
 
 void UMovieGraphPipeline::Initialize(UMoviePipelineExecutorJob* InJob, const FMovieGraphInitConfig& InitConfig)
@@ -170,6 +177,15 @@ void UMovieGraphPipeline::TickProducingFrames()
 {
 	check(GraphTimeStepInstance);
 
+	// Move any output frames that have been finished from the Output Merger
+	// into the actual outputs. This will generate new futures (for actual 
+	// disk writes) which we keep track of below.
+	ProcessOutstandingFinishedFrames();
+	
+	// Process any files that have finished writing to disk and push them into our list of 
+	// files made by this shot.
+	ProcessOutstandingFutures();
+
 	// The callback for this function does not get registered until Initialization has been called, which sets
 	// the state to Render. If it's not, we have a initialization order/flow issue!
 	//check(PipelineState == EMovieRenderPipelineState::ProducingFrames);
@@ -644,6 +660,104 @@ void UMovieGraphPipeline::OnMoviePipelineFinishedImpl()
 //	}
 //}
 
+void UMovieGraphPipeline::ProcessOutstandingFinishedFrames()
+{
+	while (!OutputMerger->GetFinishedFrames().IsEmpty())
+	{
+		UE::MovieGraph::FMovieGraphOutputMergerFrame OutputFrame;
+		OutputMerger->GetFinishedFrames().Dequeue(OutputFrame);
+
+		UE::MovieGraph::FRenderTimeStatistics* TimeStats = GetRendererInstance()->GetRenderTimeStatistics(OutputFrame.TraversalContext.Time.OutputFrameNumber);
+		if (ensure(TimeStats))
+		{
+			TimeStats->EndTime = FDateTime::UtcNow();
+		}
+
+		int32 FrameNumber = OutputFrame.TraversalContext.Time.OutputFrameNumber;
+		FString DurationTimeStr = (TimeStats->EndTime - TimeStats->StartTime).ToString();
+		// UE_LOG(LogTemp, Log, TEXT("Frame: %d Duration: %s"), FrameNumber, *DurationTimeStr);
+
+		// We need to get the list of output nodes for this frame's shot. We pass all of the data at once to each relevant one
+		// and individual output nodes can decide if they should skip writing that layer's data, or combine them into one file
+		// or multiple files, etc. To get the right evaluation context though we need to use common data from one node for now.
+		TSet<UMovieGraphFileOutputNode*> ContainerCDOs;
+
+		TArray<UMovieGraphFileOutputNode*> AllNodeInstances = UMovieGraphConfig::IterateGraphForClassAll<UMovieGraphFileOutputNode>(OutputFrame.TraversalContext);
+		for (UMovieGraphFileOutputNode* Node : AllNodeInstances)
+		{
+			ContainerCDOs.Add(Node->GetClass()->GetDefaultObject<UMovieGraphFileOutputNode>());
+		}
+
+		for (UMovieGraphFileOutputNode* CDOInstance : ContainerCDOs)
+		{
+			CDOInstance->OnReceiveImageData(this, &OutputFrame);
+		}
+
+	}
+}
+
+void UMovieGraphPipeline::AddOutputFuture(TFuture<bool>&& InOutputFuture, const UE::MovieGraph::FMovieGraphOutputFutureData& InData)
+{
+	OutstandingOutputFutures.Add(
+		UE::MovieGraph::FMovieGraphOutputFuture(MoveTemp(InOutputFuture), InData)
+	);
+}
+
+void UMovieGraphPipeline::ProcessOutstandingFutures()
+{
+	// Check if any frames failed to output
+	TArray<int32> CompletedOutputFutures;
+	for (int32 Index = 0; Index < OutstandingOutputFutures.Num(); ++Index)
+	{
+		// Output futures are pushed in order to the OutputFutures array. However they are
+		// completed asyncronously, so we don't process any futures after a not-yet-ready one
+		// otherwise we push into the GeneratedShotOutputData array out of order.
+		const UE::MovieGraph::FMovieGraphOutputFuture& OutputFuture = OutstandingOutputFutures[Index];
+		if (!OutputFuture.Get<0>().IsReady())
+		{
+			break;
+		}
+
+		CompletedOutputFutures.Add(Index);
+
+		const UE::MovieGraph::FMovieGraphOutputFutureData& FutureData = OutputFuture.Get<1>();
+
+		// The future was completed, time to add it to our shot output data.
+		FMovieGraphRenderOutputData* ShotOutputData = nullptr;
+		for (int32 OutputDataIndex = 0; OutputDataIndex < GeneratedOutputData.Num(); OutputDataIndex++)
+		{
+			if (FutureData.Shot == GeneratedOutputData[OutputDataIndex].Shot)
+			{
+				ShotOutputData = &GeneratedOutputData[OutputDataIndex];
+			}
+		}
+
+		if (!ShotOutputData)
+		{
+			GeneratedOutputData.Add(FMovieGraphRenderOutputData());
+			ShotOutputData = &GeneratedOutputData.Last();
+			ShotOutputData->Shot = FutureData.Shot;
+		}
+
+		// Add the filepath to the renderpass data.
+		ShotOutputData->RenderPassData.FindOrAdd(FutureData.DataIdentifier).FilePaths.Add(FutureData.FilePath);
+
+		// Sometime futures can be completed, but will be set to an error state (such as we couldn't write to the specified disk path.)
+		if (!OutputFuture.Get<0>().Get())
+		{
+			UE_LOG(LogMovieRenderPipeline, Error, TEXT("Error exporting frame, canceling movie export."));
+			RequestShutdown(true);
+			break;
+		}
+	}
+
+	// Remove any output futures that have been completed now.
+	for (int32 Index = CompletedOutputFutures.Num() - 1; Index >= 0; --Index)
+	{
+		OutstandingOutputFutures.RemoveAt(CompletedOutputFutures[Index]);
+	}
+}
+
 UMovieGraphConfig* UMovieGraphPipeline::GetRootGraphForShot(UMoviePipelineExecutorShot* InShot) const
 {
 	if (GetCurrentJob())
@@ -678,9 +792,14 @@ UMovieGraphConfig* UMovieGraphPipeline::GetRootGraphForShot(UMoviePipelineExecut
 	return nullptr;
 }
 
-FMovieGraphTraversalContext UMovieGraphPipeline::GetTraversalContextForShot(UMoviePipelineExecutorShot* InShot) const
+FMovieGraphTraversalContext UMovieGraphPipeline::GetCurrentTraversalContext() const
 {
-	// ToDo: Fill out shot index/current shot/current frame/job/etc.
-	return FMovieGraphTraversalContext();
-}
-
+	FMovieGraphTraversalContext CurrentContext;
+	CurrentContext.ShotIndex = GetCurrentShotIndex();
+	CurrentContext.ShotCount = GetActiveShotList().Num();
+	CurrentContext.Job = GetCurrentJob();
+	CurrentContext.RootGraph = GetRootGraphForShot(GetActiveShotList()[GetCurrentShotIndex()]);
+	CurrentContext.Time = GetTimeStepInstance()->GetCalculatedTimeData();
+
+	return CurrentContext;
+}
\ No newline at end of file
diff --git a/Engine/Plugins/MovieScene/MovieRenderPipeline/Source/MovieRenderPipelineCore/Private/Graph/Nodes/MovieGraphDeferredRenderPassNode.cpp b/Engine/Plugins/MovieScene/MovieRenderPipeline/Source/MovieRenderPipelineCore/Private/Graph/Nodes/MovieGraphDeferredRenderPassNode.cpp
--- a/Engine/Plugins/MovieScene/MovieRenderPipeline/Source/MovieRenderPipelineCore/Private/Graph/Nodes/MovieGraphDeferredRenderPassNode.cpp
+++ b/Engine/Plugins/MovieScene/MovieRenderPipeline/Source/MovieRenderPipelineCore/Private/Graph/Nodes/MovieGraphDeferredRenderPassNode.cpp
@@ -3,7 +3,10 @@
 #include "Graph/Nodes/MovieGraphDeferredRenderPassNode.h"
 #include "Graph/MovieGraphDataTypes.h"
 #include "Graph/MovieGraphDefaultRenderer.h"
+#include "Graph/MovieGraphPipeline.h"
 #include "MovieRenderPipelineCoreModule.h"
+#include "MovieRenderOverlappedImage.h"
+#include "MoviePipelineSurfaceReader.h"
 
 #include "EngineModule.h"
 #include "SceneManagement.h"
@@ -42,14 +45,21 @@ void UMovieGraphDeferredRenderPassNode::TeardownImpl()
 }
 
 
-void UMovieGraphDeferredRenderPassNode::RenderImpl(const FMovieGraphTimeStepData& InTimeData)
+void UMovieGraphDeferredRenderPassNode::RenderImpl(FMovieGraphTraversalContext InFrameTraversalContext, const FMovieGraphTimeStepData& InTimeData)
 {
 	for (TUniquePtr<FMovieGraphDeferredRenderPass>& Instance : CurrentInstances)
 	{
-		Instance->Render(InTimeData);
+		Instance->Render(InFrameTraversalContext, InTimeData);
 	}
 }
 
+void UMovieGraphDeferredRenderPassNode::GatherOutputPassesImpl(TArray<FMovieGraphRenderDataIdentifier>& OutExpectedPasses) const
+{
+	for (const TUniquePtr<FMovieGraphDeferredRenderPass>& Instance : CurrentInstances)
+	{
+		Instance->GatherOutputPassesImpl(OutExpectedPasses);
+	}
+}
 
 void UMovieGraphDeferredRenderPassNode::AddReferencedObjects(UObject* InThis, FReferenceCollector& Collector)
 {
@@ -85,6 +95,18 @@ void UMovieGraphDeferredRenderPassNode::FMovieGraphDeferredRenderPass::Setup(TWe
 
 }
 
+void UMovieGraphDeferredRenderPassNode::FMovieGraphDeferredRenderPass::GatherOutputPassesImpl(TArray<FMovieGraphRenderDataIdentifier>& OutExpectedPasses) const
+{
+	// ToDo: Store this when the render pass is created, so we don't create it in two places at once.
+	// This is a problem right now due to the camera name not being pre-resolved.
+	FMovieGraphRenderDataIdentifier RenderDataIdentifier;
+	RenderDataIdentifier.RenderLayerName = LayerData.LayerName;
+	RenderDataIdentifier.RendererName = RenderPassNode->GetRendererName();
+	RenderDataIdentifier.SubResourceName = TEXT("beauty");
+	RenderDataIdentifier.CameraName = TEXT("Unnamed_Camera");
+
+	OutExpectedPasses.Add(RenderDataIdentifier);
+}
 void UMovieGraphDeferredRenderPassNode::FMovieGraphDeferredRenderPass::Teardown()
 {
 	FSceneViewStateInterface* Ref = SceneViewState.GetReference();
@@ -104,7 +126,7 @@ void UMovieGraphDeferredRenderPassNode::FMovieGraphDeferredRenderPass::AddRefere
 	}
 }
 
-void UMovieGraphDeferredRenderPassNode::FMovieGraphDeferredRenderPass::Render(const FMovieGraphTimeStepData& InTimeData)
+void UMovieGraphDeferredRenderPassNode::FMovieGraphDeferredRenderPass::Render(FMovieGraphTraversalContext InFrameTraversalContext, const FMovieGraphTimeStepData& InTimeData)
 {
 	// This is the size we actually render at.
 	UE::MovieGraph::DefaultRenderer::FRenderTargetInitParams RenderTargetInitParams;
@@ -176,14 +198,20 @@ void UMovieGraphDeferredRenderPassNode::FMovieGraphDeferredRenderPass::Render(co
 	RenderDataIdentifier.RendererName = RenderPassNode->GetRendererName();
 	RenderDataIdentifier.SubResourceName = TEXT("beauty");
 	RenderDataIdentifier.CameraName = TEXT("Unnamed_Camera");
-	if (InitData.ViewActor)
-	{
-		RenderDataIdentifier.CameraName = InitData.ViewActor->GetName();
-	}
+	//if (InitData.ViewActor)
+	//{
+	//	RenderDataIdentifier.CameraName = InitData.ViewActor->GetName();
+	//}
+
+	// Take our per-frame Traversal Context and update it with context specific to this sample.
+	FMovieGraphTraversalContext UpdatedTraversalContext = InFrameTraversalContext;
+	UpdatedTraversalContext.Time = InTimeData;
+	UpdatedTraversalContext.RenderDataIdentifier = RenderDataIdentifier;
+
 
 	UE::MovieGraph::FMovieGraphSampleState SampleState;
-	SampleState.RenderDataIdentifier = RenderDataIdentifier;
-	SampleState.Time = InTimeData;
+	SampleState.TraversalContext = MoveTemp(UpdatedTraversalContext);
+	SampleState.BackbufferResolution = RenderTargetInitParams.Size;
 
 	// Readback + Accumulate.
 	PostRendererSubmission(SampleState, RenderTargetInitParams, Canvas);
@@ -204,7 +232,7 @@ void UMovieGraphDeferredRenderPassNode::FMovieGraphDeferredRenderPass::PostRende
 	UE::MovieGraph::DefaultRenderer::FSurfaceAccumulatorPool::FInstancePtr AccumulatorInstance = nullptr;
 	{
 		// SCOPE_CYCLE_COUNTER(STAT_MoviePipeline_WaitForAvailableAccumulator);
-		AccumulatorInstance = SampleAccumulatorPool->BlockAndGetAccumulator_GameThread(InSampleState.Time.OutputFrameNumber, InSampleState.RenderDataIdentifier);
+		AccumulatorInstance = SampleAccumulatorPool->BlockAndGetAccumulator_GameThread(InSampleState.TraversalContext.Time.OutputFrameNumber, InSampleState.TraversalContext.RenderDataIdentifier);
 	}
 
 	FMoviePipelineSurfaceQueuePtr LocalSurfaceQueue = Renderer->GetOrCreateSurfaceQueue(InRenderTargetInitParams);
@@ -212,10 +240,10 @@ void UMovieGraphDeferredRenderPassNode::FMovieGraphDeferredRenderPass::PostRende
 
 	UE::MovieGraph::FMovieGraphRenderDataAccumulationArgs AccumulationArgs;
 	{
-		//AccumulationArgs.OutputMerger = Renderer->GetOutputBuilder();
+		AccumulationArgs.OutputMerger = Renderer->GetOwningGraph()->GetOutputMerger();
 		AccumulationArgs.ImageAccumulator = StaticCastSharedPtr<FImageOverlappedAccumulator>(AccumulatorInstance->Accumulator);
-		AccumulationArgs.bIsFirstSample = true; // FramePayload->IsFirstTile() && FramePayload->IsFirstTemporalSample()
-		AccumulationArgs.bIsLastSample = true; // FramePayload->IsLastTile() && FramePayload->IsLastTemporalSample()
+		AccumulationArgs.bIsFirstSample = InSampleState.TraversalContext.Time.bIsFirstTemporalSampleForFrame; // FramePayload->IsFirstTile() && FramePayload->IsFirstTemporalSample()
+		AccumulationArgs.bIsLastSample = InSampleState.TraversalContext.Time.bIsLastTemporalSampleForFrame; // FramePayload->IsLastTile() && FramePayload->IsLastTemporalSample()
 		AccumulationArgs.TaskPrerequisite = AccumulatorInstance->TaskPrereq;
 	}
 
@@ -224,12 +252,11 @@ void UMovieGraphDeferredRenderPassNode::FMovieGraphDeferredRenderPass::PostRende
 		UE::MovieGraph::DefaultRenderer::FMovieGraphAccumulationTask Task;
 		// There may be other accumulations for this accumulator which need to be processed first
 		Task.LastCompletionEvent = AccumulationArgs.TaskPrerequisite;
-		// TUniquePtr<FImagePixelData> Foo = MakeUnique<TImagePixelData<FFloat16Color>>(InPixelData->GetSize(), MoveTemp(InPixelData->MoveImageDataToNew))
 
-		FGraphEventRef Event = Task.Execute([PixelData = MoveTemp(InPixelData), AccumulationArgs, AccumulatorInstance]() mutable
+		FGraphEventRef Event = Task.Execute([PixelData = MoveTemp(InPixelData), InSampleState, AccumulationArgs, AccumulatorInstance]() mutable
 			{
 				// Enqueue a encode for this frame onto our worker thread.
-				// UE::MovieGraph::DefaultRenderer::AccumulateSample_TaskThread(MoveTemp(PixelData), AccumulationArgs);
+				UE::MovieGraph::AccumulateSample_TaskThread(MoveTemp(PixelData), InSampleState, AccumulationArgs);
 
 				// We have to defer clearing the accumulator until after sample accumulation has finished
 				if (AccumulationArgs.bIsLastSample)
@@ -467,3 +494,38 @@ void UMovieGraphDeferredRenderPassNode::FMovieGraphDeferredRenderPass::ApplyMovi
 	return OutViewFamily;*/
 }
 
+namespace UE::MovieGraph
+{
+	void AccumulateSample_TaskThread(TUniquePtr<FImagePixelData>&& InPixelData, const UE::MovieGraph::FMovieGraphSampleState InSampleState, const UE::MovieGraph::FMovieGraphRenderDataAccumulationArgs& InAccumulationParams)
+	{
+		TRACE_CPUPROFILER_EVENT_SCOPE(MoviePipeline_AccumulateSample);
+		
+		TUniquePtr<FImagePixelData> SamplePixelData = MoveTemp(InPixelData);
+
+		// Associate the sample state with the image as payload data, this allows downstream systems to fetch the values without us having to store the data
+		// separately and ensure they stay paired the whole way down.
+		TSharedPtr<UE::MovieGraph::FMovieGraphSampleState> SampleStatePayload = MakeShared<UE::MovieGraph::FMovieGraphSampleState>(InSampleState);
+		SamplePixelData->SetPayload(StaticCastSharedPtr<IImagePixelDataPayload>(SampleStatePayload));
+
+		TSharedPtr<IMovieGraphOutputMerger, ESPMode::ThreadSafe> AccumulatorPin = InAccumulationParams.OutputMerger.Pin();
+		if (!AccumulatorPin.IsValid())
+		{
+			return;
+		}
+
+		const bool bIsWellFormed = SamplePixelData->IsDataWellFormed();
+		check(bIsWellFormed);
+
+		// ToDo:
+		bool bWriteSampleToDisk = false;
+		if (bWriteSampleToDisk)
+		{
+			// Debug Feature: Write the raw sample to disk for debugging purposes. We copy the data here,
+			// as we don't want to disturb the memory flow below.
+			TUniquePtr<FImagePixelData> SampleData = SamplePixelData->CopyImageData();
+			AccumulatorPin->OnSingleSampleDataAvailable_AnyThread(MoveTemp(SampleData));
+		}
+
+		AccumulatorPin->OnCompleteRenderPassDataAvailable_AnyThread(MoveTemp(SamplePixelData));
+	}
+}
\ No newline at end of file
diff --git a/Engine/Plugins/MovieScene/MovieRenderPipeline/Source/MovieRenderPipelineCore/Private/Graph/Nodes/MovieGraphImageSequenceOutputNode.cpp b/Engine/Plugins/MovieScene/MovieRenderPipeline/Source/MovieRenderPipelineCore/Private/Graph/Nodes/MovieGraphImageSequenceOutputNode.cpp
new file mode 100644
--- /dev/null
+++ b/Engine/Plugins/MovieScene/MovieRenderPipeline/Source/MovieRenderPipelineCore/Private/Graph/Nodes/MovieGraphImageSequenceOutputNode.cpp
@@ -0,0 +1,56 @@
+// Copyright Epic Games, Inc. All Rights Reserved.
+
+#include "Graph/Nodes/MovieGraphImageSequenceOutputNode.h"
+#include "Graph/MovieGraphDataTypes.h"
+#include "Graph/MovieGraphPipeline.h"
+#include "Modules/ModuleManager.h"
+#include "ImageWriteQueue.h"
+#include "Misc/Paths.h"
+#include "Async/TaskGraphInterfaces.h"
+
+UMovieGraphImageSequenceOutputNode::UMovieGraphImageSequenceOutputNode()
+{
+	ImageWriteQueue = &FModuleManager::Get().LoadModuleChecked<IImageWriteQueueModule>("ImageWriteQueue").GetWriteQueue();
+}
+
+void UMovieGraphImageSequenceOutputNode::OnReceiveImageDataImpl(UMovieGraphPipeline* InPipeline, UE::MovieGraph::FMovieGraphOutputMergerFrame* InRawFrameData)
+{
+	check(InRawFrameData);
+	// ToDo:
+	// The ImageWriteQueue is set up in a fire-and-forget manner. This means that the data needs to be placed in the WriteQueue
+	// as a TUniquePtr (so it can free the data when its done). Unfortunately we can have multiple output formats at once,
+	// so we can't MoveTemp the data into it, we need to make a copy (though we could optimize for the common case where there is
+	// only one output format).
+	// Copying can be expensive (3ms @ 1080p, 12ms at 4k for a single layer image) so ideally we'd like to do it on the task graph
+	// but this isn't really compatible with the ImageWriteQueue API as we need the future returned by the ImageWriteQueue to happen
+	// in order, so that we push our futures to the main Movie Pipeline in order, otherwise when we encode files to videos they'll
+	// end up with frames out of order. A workaround for this would be to chain all of the send-to-imagewritequeue tasks to each
+	// other with dependencies, but I'm not sure that's going to scale to the potentialy high data volume going wide MRQ will eventually
+	// need.
+
+	// The base ImageSequenceOutputNode doesn't support any multilayer formats, so we write out each render pass separately.
+	for (TPair<FMovieGraphRenderDataIdentifier, TUniquePtr<FImagePixelData>>& RenderData : InRawFrameData->ImageOutputData)
+	{
+		UE::MovieGraph::FMovieGraphSampleState* Payload = RenderData.Value->GetPayload<UE::MovieGraph::FMovieGraphSampleState>();
+
+		// ToDo: Real file paths.
+		const FString ResolvedProjectDir = FPaths::ConvertRelativePathToFull(FPaths::ProjectDir());
+		FString FileName = FString::Printf(TEXT("%s/Saved/MovieRenders/%d.jpg"), *ResolvedProjectDir, Payload->TraversalContext.Time.OutputFrameNumber);
+
+		TUniquePtr<FImageWriteTask> TileImageTask = MakeUnique<FImageWriteTask>();
+		TileImageTask->Format = OutputFormat;
+		TileImageTask->CompressionQuality = 100;
+		TileImageTask->Filename = FileName;
+		TileImageTask->PixelData = RenderData.Value->CopyImageData();
+
+		UE::MovieGraph::FMovieGraphOutputFutureData OutputData;
+		OutputData.Shot = nullptr;
+		OutputData.FilePath = FileName;
+		OutputData.DataIdentifier = RenderData.Key;
+
+		TFuture<bool> Future = ImageWriteQueue->Enqueue(MoveTemp(TileImageTask));
+
+		InPipeline->AddOutputFuture(MoveTemp(Future), OutputData);
+	}
+
+}
\ No newline at end of file
diff --git a/Engine/Plugins/MovieScene/MovieRenderPipeline/Source/MovieRenderPipelineCore/Public/Graph/MovieGraphConfig.h b/Engine/Plugins/MovieScene/MovieRenderPipeline/Source/MovieRenderPipelineCore/Public/Graph/MovieGraphConfig.h
--- a/Engine/Plugins/MovieScene/MovieRenderPipeline/Source/MovieRenderPipelineCore/Public/Graph/MovieGraphConfig.h
+++ b/Engine/Plugins/MovieScene/MovieRenderPipeline/Source/MovieRenderPipelineCore/Public/Graph/MovieGraphConfig.h
@@ -4,38 +4,13 @@
 #include "CoreMinimal.h"
 #include "EdGraph/EdGraph.h"
 #include "Graph/MovieGraphNode.h"
+#include "Graph/MovieGraphTraversalContext.h"
 
 #include "MovieGraphConfig.generated.h"
 
 // Forward Declare
 class UMovieGraphNode;
 
-USTRUCT(BlueprintType)
-struct FMovieGraphBranch
-{
-	GENERATED_BODY()
-
-public:
-	UPROPERTY(BlueprintReadWrite, EditAnywhere, Category = "MovieGraph")
-	FName BranchName;
-};
-
-
-USTRUCT(BlueprintType)
-struct FMovieGraphTraversalContext
-{
-	GENERATED_BODY();
-
-	UPROPERTY(BlueprintReadWrite, EditAnywhere, Category = "MovieGraph")
-	FMovieGraphBranch RootBranch;
-
-	UPROPERTY(BlueprintReadWrite, EditAnywhere, Category = "MovieGraph")
-	FString ShotName;
-
-	UPROPERTY(BlueprintReadWrite, EditAnywhere, Category = "MovieGraph")
-	FString RenderLayerName;
-};
-
 #if WITH_EDITOR
 	DECLARE_MULTICAST_DELEGATE_OneParam(FOnMovieGraphVariableChanged, class UMovieGraphMember*);
 	DECLARE_MULTICAST_DELEGATE_OneParam(FOnMovieGraphInputChanged, class UMovieGraphMember*);
@@ -229,7 +204,7 @@ public:
 #endif
 
 	template<typename NodeType>
-	NodeType* IterateGraphForClass(const FMovieGraphTraversalContext& InContext) const
+	static NodeType* IterateGraphForClass(const FMovieGraphTraversalContext& InContext)
 	{
 		TArray<NodeType*> AllSettings = IterateGraphForClassAll<NodeType>(InContext);
 		if (AllSettings.Num() > 0)
@@ -240,10 +215,15 @@ public:
 	}
 	
 	template<typename NodeType>
-	TArray<NodeType*> IterateGraphForClassAll(const FMovieGraphTraversalContext& InContext) const
+	static TArray<NodeType*> IterateGraphForClassAll(const FMovieGraphTraversalContext& InContext)
 	{
 		TArray<NodeType*> TypedNodes;
-		const TArray<UMovieGraphNode*> FoundNodes = TraverseGraph(NodeType::StaticClass(), InContext);
+		if (!ensureMsgf(InContext.RootGraph, TEXT("You must specify a RootGraph to traverse with")))
+		{
+			return TypedNodes;
+		}
+
+		const TArray<UMovieGraphNode*> FoundNodes = InContext.RootGraph->TraverseGraph(NodeType::StaticClass(), InContext);
 		for (UMovieGraphNode* Node : FoundNodes)
 		{
 			TypedNodes.Add(CastChecked<NodeType>(Node));
diff --git a/Engine/Plugins/MovieScene/MovieRenderPipeline/Source/MovieRenderPipelineCore/Public/Graph/MovieGraphDataTypes.h b/Engine/Plugins/MovieScene/MovieRenderPipeline/Source/MovieRenderPipelineCore/Public/Graph/MovieGraphDataTypes.h
--- a/Engine/Plugins/MovieScene/MovieRenderPipeline/Source/MovieRenderPipelineCore/Public/Graph/MovieGraphDataTypes.h
+++ b/Engine/Plugins/MovieScene/MovieRenderPipeline/Source/MovieRenderPipelineCore/Public/Graph/MovieGraphDataTypes.h
@@ -1,8 +1,12 @@
 // Copyright Epic Games, Inc. All Rights Reserved.
 #pragma once
 
+#include "Graph/MovieGraphTraversalContext.h"
+#include "Graph/MovieGraphRenderDataIdentifier.h"
 #include "UObject/Object.h"
 #include "Templates/SubclassOf.h"
+#include "ImagePixelData.h"
+#include "Containers/Queue.h"
 #include "MovieGraphDataTypes.generated.h"
 
 // Forward Declares
@@ -13,6 +17,12 @@ class UMovieGraphTimeRangeBuilderBase;
 class UMovieGraphDataCachingBase;
 class UMoviePipelineExecutorJob;
 class UMoviePipelineExecutorShot;
+struct FImagePixelData;
+
+namespace UE::MovieGraph
+{
+	struct FRenderTimeStatistics;
+}
 
 USTRUCT(BlueprintType)
 struct MOVIERENDERPIPELINECORE_API FMovieGraphInitConfig
@@ -21,9 +31,17 @@ struct MOVIERENDERPIPELINECORE_API FMovieGraphInitConfig
 	
 	FMovieGraphInitConfig();
 	
+	/** 
+	* Which class should the UMovieGraphPipeline use to handle calculating per frame 
+	* timesteps? Defaults to UMovieGraphLinearTimeStep.
+	*/
 	UPROPERTY(BlueprintReadWrite, EditAnywhere, Category = "Movie Graph")
 	TSubclassOf<UMovieGraphTimeStepBase> TimeStepClass;
 	
+	/**
+	* Which class should the UMovieGraphPipeline use to look for render layers and
+	* request renders from. Defaults to UMovieGraphDefaultRenderer.
+	*/
 	UPROPERTY(BlueprintReadWrite, EditAnywhere, Category = "Movie Graph")
 	TSubclassOf<UMovieGraphRendererBase> RendererClass;
 
@@ -33,40 +51,43 @@ struct MOVIERENDERPIPELINECORE_API FMovieGraphInitConfig
 	UPROPERTY(BlueprintReadWrite, EditAnywhere, Category = "Movie Graph")
 	TSubclassOf<UMovieGraphDataCachingBase> DataCachingClass;
 
+	/**
+	* Should the UMovieGraphPipeline render the full player viewport? Defaults
+	* to false (so no 3d content is rendered) so we can display the UMG widgets
+	* and MRQ rendering always happens in an off-screen render target.
+	*/
 	UPROPERTY(BlueprintReadWrite, EditAnywhere, Category = "Movie Graph")
 	bool bRenderViewport;
 };
 
-
 USTRUCT(BlueprintType)
-struct MOVIERENDERPIPELINECORE_API FMovieGraphTimeStepData
+struct FMovieGraphRenderPassOutputData
 {
 	GENERATED_BODY()
+public:
+	/** A list of file paths on disk (in order) that were generated for this particular render pass. */
+	UPROPERTY(VisibleAnywhere, BlueprintReadOnly, Category = "Movie Pipeline")
+	TArray<FString> FilePaths;
+};
 
-	FMovieGraphTimeStepData()
-		: OutputFrameNumber(0)
-		, FrameDeltaTime(0.f)
-		, WorldTimeDilation(0.f)
-		, WorldSeconds(0.f)
-		, MotionBlurFraction(0.f)
-	{}
-
-	UPROPERTY(BlueprintReadWrite, EditAnywhere, Category = "Movie Graph")
-	int32 OutputFrameNumber;
-
-	UPROPERTY(BlueprintReadWrite, EditAnywhere, Category = "Movie Graph")
-	float FrameDeltaTime;
-	
-	UPROPERTY(BlueprintReadWrite, EditAnywhere, Category = "Movie Graph")
-	float WorldTimeDilation;
-	
-	UPROPERTY(BlueprintReadWrite, EditAnywhere, Category = "Movie Graph")
-	float WorldSeconds;
-	
-	UPROPERTY(BlueprintReadWrite, EditAnywhere, Category = "Movie Graph")
-	float MotionBlurFraction;
+USTRUCT(BlueprintType)
+struct MOVIERENDERPIPELINECORE_API FMovieGraphRenderOutputData
+{
+	GENERATED_BODY()
+public:
+	/** Which shot is this output data for. */
+	UPROPERTY(VisibleAnywhere, BlueprintReadOnly, Category = "Movie Graph")
+	TWeakObjectPtr<UMoviePipelineExecutorShot> Shot;
+
+	/**
+	* A mapping between render passes (such as "beauty") and an array containing the files written for that shot.
+	* Will be multiple files if using image sequences
+	*/
+	UPROPERTY(VisibleAnywhere, BlueprintReadOnly, Category = "Movie Graph")
+	TMap<FMovieGraphRenderDataIdentifier, FMovieGraphRenderPassOutputData> RenderPassData;
 };
 
+
 UCLASS(BlueprintType, Abstract)
 class MOVIERENDERPIPELINECORE_API UMovieGraphTimeStepBase : public UObject
 {
@@ -85,6 +106,8 @@ public:
 	virtual void Render(const FMovieGraphTimeStepData& InTimeData) {}
 	virtual void SetupRenderingPipelineForShot(UMoviePipelineExecutorShot* InShot) {}
 	virtual void TeardownRenderingPipelineForShot(UMoviePipelineExecutorShot* InShot) {}
+	virtual UE::MovieGraph::FRenderTimeStatistics* GetRenderTimeStatistics(const int32 InFrameNumber) { return nullptr; }
+	
 	UMovieGraphPipeline* GetOwningGraph() const;
 };
 
@@ -109,10 +132,12 @@ public:
 	UMovieGraphPipeline* GetOwningGraph() const;
 };
 
+// ToDo: Both of these can probably go into the Default Renderer implementation.
 struct FMovieGraphRenderPassLayerData
 {
 	FString LayerName;
 };
+// ToDo: Both of these can probably go into the Default Renderer implementation.
 
 struct FMovieGraphRenderPassSetupData
 {
@@ -120,84 +145,106 @@ struct FMovieGraphRenderPassSetupData
 	TArray<FMovieGraphRenderPassLayerData> Layers;
 };
 
-USTRUCT(BlueprintType)
-struct FMovieGraphRenderDataIdentifier
+
+namespace UE::MovieGraph
 {
-	GENERATED_BODY()
+	/** Supplemental data to store with a file future so we can retrieve information about what the future was doing once it finishes writing to disk. */
+	struct FMovieGraphOutputFutureData
+	{
+		FMovieGraphOutputFutureData()
+			: Shot(nullptr)
+		{}
 
-	FMovieGraphRenderDataIdentifier()
-	{}
+		/** Which shot is this for? */
+		class UMoviePipelineExecutorShot* Shot;
 
-	FMovieGraphRenderDataIdentifier(const FString& InRenderLayerName, const FString& InRendererName,
-		const FString& InSubRenderResourceName, const FString& InCameraName)
-		: RenderLayerName(InRenderLayerName)
-		, RendererName(InRendererName)
-		, SubResourceName(InSubRenderResourceName)
-		, CameraName(InCameraName)
-	{
-	}
+		/** What was the filepath on disk that was generated by the future? */
+		FString FilePath;
 
-	bool operator == (const FMovieGraphRenderDataIdentifier& InRHS) const
+		/** Which render pass is this for? */
+		FMovieGraphRenderDataIdentifier DataIdentifier;
+	};
+	
+	/**
+	* Statistics tracked for a given frame during rendering, such as when the
+	* render started and when it ended. This can be added as metadata during
+	* file output to include useful information. 
+	*
+	* ToDo: This could potentially be tracked as just regular filemetadata
+	* with a rendered sample but requires strong typing support and metadata
+	*  merging support.
+	*/
+	struct FRenderTimeStatistics
 	{
-		return RenderLayerName == InRHS.RenderLayerName && 
-			RendererName == InRHS.RendererName &&
-			SubResourceName == InRHS.SubResourceName &&
-			CameraName == InRHS.CameraName;
-	}
+		/** Time in UTC that the frame was first started. */
+		FDateTime StartTime;
 
-	bool operator != (const FMovieGraphRenderDataIdentifier& InRHS) const
-	{
-		return !(*this == InRHS);
-	}
+		/** Time in UTC that the frame was fully rendered and merged together. */
+		FDateTime EndTime;
+	};
 
-	friend uint32 GetTypeHash(FMovieGraphRenderDataIdentifier OutputState)
+	struct FMovieGraphSampleState : IImagePixelDataPayload, public TSharedFromThis<FMovieGraphSampleState>
 	{
-		return HashCombineFast(GetTypeHash(OutputState.RenderLayerName),
-			HashCombineFast(GetTypeHash(OutputState.RendererName),
-				HashCombineFast(GetTypeHash(OutputState.SubResourceName),
-					GetTypeHash(OutputState.CameraName))));
-	}
+		/** The traversal context used to read graph values at the time of submission. */
+		FMovieGraphTraversalContext TraversalContext;
 
-public:
-	// The user provided name for the whole Render Layer ("character", "background", etc.)
-	UPROPERTY(EditAnywhere, BlueprintReadWrite, Category = "Movie Pipeline")
-	FString RenderLayerName;
+		/** The backbuffer resolution of this render resource(ie: the resolution the sample was rendered at). */
+		FIntPoint BackbufferResolution;
+	};
 
-	// Which renderer was used to produce this image ("panoramic" "deferred" "path tracer", etc.)
-	UPROPERTY(EditAnywhere, BlueprintReadWrite, Category = "Movie Pipeline")
-	FString RendererName;
+	/**
+	* A collection of all of the data needed to produce the one output frame on disk. This holds
+	* all of the pixel data associated with the output frame at once, so we can pass it to our 
+	* image writing/video encoders in one go, to ensure we can do things like Multilayer EXRs which
+	* embed all of the render passes inside of one file.
+	*/
+	struct FMovieGraphOutputMergerFrame
+	{
+		FMovieGraphOutputMergerFrame() = default;
+		
+
+		/** 
+		* Traversal context for the Movie Pipeline that isn't layer specific. Each image layer will have
+		* its own traversal context within its payload, specific to itself. */
+		FMovieGraphTraversalContext TraversalContext;
+
+		/**
+		* An array of expected data identifiers this frame. This is generated when the frame starts, 
+		* and this frame won't be complete until ImageOutputData contains an image for each expected pass.
+		*/
+		TArray<FMovieGraphRenderDataIdentifier> ExpectedRenderPasses;
+
+		/** Stores the actual pixel data for each render pass. */
+		TMap<FMovieGraphRenderDataIdentifier, TUniquePtr<FImagePixelData>> ImageOutputData;
+	};
 
-	// A sub-resource name for the renderer (ie: "beauty", "object id", "depth", etc.)
-	UPROPERTY(EditAnywhere, BlueprintReadWrite, Category = "Movie Pipeline")
-	FString SubResourceName;
+	/**
+	* Interface for the intermediate threadsafe class that allocates new frames, holds onto the individual
+	* pieces as they come in, and eventually moves it back to the game thread once the data is complete.
+	*/
+	struct IMovieGraphOutputMerger : public TSharedFromThis<IMovieGraphOutputMerger>
+	{
+		virtual ~IMovieGraphOutputMerger() {};
 
-	// The name of the camera being used for this render.
-	UPROPERTY(EditAnywhere, BlueprintReadWrite, Category = "Movie Pipeline")
-	FString CameraName;
-};
+		/** Call once per output frame to generate a struct to hold data. */
+		virtual FMovieGraphOutputMergerFrame& AllocateNewOutputFrame_GameThread(int32 InFrameNumber) = 0;
 
+		/** Getter for the Output Frame. Make sure you call AllocateNewOutputFrame_GameThread first, otherwise this trips a check. */
+		virtual FMovieGraphOutputMergerFrame& GetOutputFrame_GameThread(int32 InFrameNumber) = 0;
 
-namespace UE
-{
-namespace MovieGraph
-{
-	struct FMovieGraphSampleState
-	{
-		// We need to capture the state of MRQ at the time a sample was rendered so that we can
-		// resolve filenames, etc. This is partially:
-		// Output Frame Data (which output frame, temporal sample, spatial sample)
-		// Traversal Context (layer name, current shot, etc.)
+		/** When a final accumulated render pass is available, call this function with the pixel data. */
+		virtual void OnCompleteRenderPassDataAvailable_AnyThread(TUniquePtr<FImagePixelData>&& InData) = 0;
 
-		// The name of the render resource this state was captured for
-		FMovieGraphRenderDataIdentifier RenderDataIdentifier;
+		/** When a individual sample is read back from the GPU, call this function. Debug Feature, immediately writes to disk. */
+		virtual void OnSingleSampleDataAvailable_AnyThread(TUniquePtr<FImagePixelData>&& InData) = 0;
 
-		// The backbuffer resolution of this render resource (ie: the resolution the sample was rendered at)
-		FIntPoint BackbufferResolution;
+		/** Abandon any frames we're waiting on, ie; render stopped in the middle of a many temporal sample count frame. */
+		virtual void AbandonOutstandingWork() = 0;
 
-		// The time data (output frame, delta times, etc.)
-		FMovieGraphTimeStepData Time;
+		/** Number of frames that have been started (via AllocateNewOutputFrame_GameThread) but haven't had all of their data come in yet via OnCompleteRenderPassDataAvailable_AnyThread. */
+		virtual int32 GetNumOutstandingFrames() const = 0;
 
+		/** Get a thread-safe FIFO queue of finished frame data. */
+		virtual TQueue<FMovieGraphOutputMergerFrame>& GetFinishedFrames() = 0;
 	};
-
-}
 }
diff --git a/Engine/Plugins/MovieScene/MovieRenderPipeline/Source/MovieRenderPipelineCore/Public/Graph/MovieGraphDefaultRenderer.h b/Engine/Plugins/MovieScene/MovieRenderPipeline/Source/MovieRenderPipelineCore/Public/Graph/MovieGraphDefaultRenderer.h
--- a/Engine/Plugins/MovieScene/MovieRenderPipeline/Source/MovieRenderPipelineCore/Public/Graph/MovieGraphDefaultRenderer.h
+++ b/Engine/Plugins/MovieScene/MovieRenderPipeline/Source/MovieRenderPipelineCore/Public/Graph/MovieGraphDefaultRenderer.h
@@ -128,6 +128,8 @@ public:
 	virtual void Render(const FMovieGraphTimeStepData& InTimeData) override;
 	virtual void SetupRenderingPipelineForShot(UMoviePipelineExecutorShot* InShot) override;
 	virtual void TeardownRenderingPipelineForShot(UMoviePipelineExecutorShot* InShot) override;
+	virtual UE::MovieGraph::FRenderTimeStatistics* GetRenderTimeStatistics(const int32 InFrameNumber);
+
 	// ~UMovieGraphRendererBase Interface
 
 	// UObject Interface
@@ -166,26 +168,36 @@ protected:
 	FMoviePipelineSurfaceQueuePtr CreateSurfaceQueue(const UE::MovieGraph::DefaultRenderer::FRenderTargetInitParams& InInitParams) const;
 
 protected:
+	/** A pointer to the CDOs of the Render Pass nodes that are valid for the current shot render. */
 	UPROPERTY(Transient)
 	TArray<TObjectPtr<UMovieGraphRenderPassNode>> RenderPassesInUse;
 
-	// Render Target specify a transient backbuffer to draw the image to. We reuse these
-	// because we queue a copy onto a FMoviePipelineSurface so the only reason it is needed
-	// is for display in the UI.
+	/** Keep track of some statistics about render frame data for metadata purposes. */
+	TMap<int32, UE::MovieGraph::FRenderTimeStatistics> RenderTimeStatistics;
+
+	/*
+	* Render Target specify a transient backbuffer to draw the image to.We reuse these
+	* because we queue a copy onto a FMoviePipelineSurface so the only reason it is needed
+	* is for display in the UI.
+	*/
 	TMap<UE::MovieGraph::DefaultRenderer::FRenderTargetInitParams, TObjectPtr<UTextureRenderTarget2D>> PooledViewRenderTargets;
 
-	// This is a list of Surfaces that we can copy our render target to immediately after
-	// it is drawn to. This is read back to the CPU a frame later (to avoid blocking the
-	// GPU), and it's contents are then copied into a FImagePixelData.
+	/*
+	* This is a list of Surfaces that we can copy our render target to immediately after
+	* it is drawn to. This is read back to the CPU a frame later (to avoid blocking the
+	* GPU), and it's contents are then copied into a FImagePixelData.
+	*/
 	TMap<UE::MovieGraph::DefaultRenderer::FRenderTargetInitParams, FMoviePipelineSurfaceQueuePtr> PooledSurfaceQueues;
 
-	// Once a sample is copied back from the GPU, we need to potentially accumulate it over
-	// multiple frames, and to apply high-res tiling. To do this we have a pool of accumulators,
-	// where each unique render resource gets one per frame.
+	/*
+	* Once a sample is copied back from the GPU, we need to potentially accumulate it over
+	* multiple frames, and to apply high-res tiling. To do this we have a pool of accumulators,
+	* where each unique render resource gets one per frame.
+	*/
 	TMap<FName, TSharedPtr<UE::MovieGraph::DefaultRenderer::FSurfaceAccumulatorPool>> PooledAccumulators;
 
-	// Accessed by the Render Thread when starting up a new task. Makes sure we don't add tasks to the array while we're removing finished ones.
+	/** Accessed by the Render Thread when starting up a new task.Makes sure we don't add tasks to the array while we're removing finished ones. */
 	FCriticalSection OutstandingTasksMutex;
-	// Array of outstanding accumulation/blending tasks that are currently being worked on.
+	/** Array of outstanding accumulation / blending tasks that are currently being worked on. */
 	FGraphEventArray OutstandingTasks;
 };
\ No newline at end of file
diff --git a/Engine/Plugins/MovieScene/MovieRenderPipeline/Source/MovieRenderPipelineCore/Public/Graph/MovieGraphPipeline.h b/Engine/Plugins/MovieScene/MovieRenderPipeline/Source/MovieRenderPipelineCore/Public/Graph/MovieGraphPipeline.h
--- a/Engine/Plugins/MovieScene/MovieRenderPipeline/Source/MovieRenderPipelineCore/Public/Graph/MovieGraphPipeline.h
+++ b/Engine/Plugins/MovieScene/MovieRenderPipeline/Source/MovieRenderPipelineCore/Public/Graph/MovieGraphPipeline.h
@@ -11,6 +11,12 @@
 // Forward Declares
 class UMoviePipelineExecutorJob;
 class UMoviePipelineExecutorShot;
+class IImageWriteQueue;
+
+namespace UE::MovieGraph
+{
+	typedef TTuple<TFuture<bool>, UE::MovieGraph::FMovieGraphOutputFutureData> FMovieGraphOutputFuture;
+}
 
 UCLASS(BlueprintType)
 class MOVIERENDERPIPELINECORE_API UMovieGraphPipeline : public UMoviePipelineBase
@@ -29,24 +35,29 @@ public:
 
 public:
 	UMovieGraphConfig* GetRootGraphForShot(UMoviePipelineExecutorShot* InShot) const;
-	FMovieGraphTraversalContext GetTraversalContextForShot(UMoviePipelineExecutorShot* InShot) const;
+	FMovieGraphTraversalContext GetCurrentTraversalContext() const;
 
-	// Get the Active Shot list, which is the full shot list generated from the external data source, with disabled shots removed.
+	/** Get the Active Shot list, which is the full shot list generated from the external data source, with disabled shots removed. */
 	const TArray<TObjectPtr<UMoviePipelineExecutorShot>>& GetActiveShotList() const { return ActiveShotList; }
-	// Which index of the Active Shot List are we currently on
+	/** Which index of the Active Shot List are we currently on */
 	int32 GetCurrentShotIndex() const { return CurrentShotIndex; }
-	// Called by the TimeStepInstance when it's time to set up for another shot. Don't call this unless you know what you're doing.
+	/** Called by the TimeStepInstance when it's time to set up for another shot. Don't call this unless you know what you're doing. */
 	void SetupShot(UMoviePipelineExecutorShot* InShot);
-	// Called by the TimeStepInstance when it's time to tear down the current shot. Don't call this unless you know what you're doing.
+	/** Called by the TimeStepInstance when it's time to tear down the current shot. Don't call this unless you know what you're doing. */
 	void TeardownShot(UMoviePipelineExecutorShot* InShot);
-	// Used occasionally to cross-reference other components. Don't call this unless you know what you're doing.
+	/** Used occasionally to cross-reference other components. Don't call this unless you know what you're doing. */
 	UMovieGraphTimeStepBase* GetTimeStepInstance() const { return GraphTimeStepInstance; }
-	// Used occasionally to cross-reference other components. Don't call this unless you know what you're doing.
+	/** Used occasionally to cross-reference other components. Don't call this unless you know what you're doing. */
 	UMovieGraphRendererBase* GetRendererInstance() const { return GraphRendererInstance; }
-	// Used occasionally to cross-reference other components. Don't call this unless you know what you're doing.
+	/** Used occasionally to cross-reference other components. Don't call this unless you know what you're doing. */
 	UMovieGraphTimeRangeBuilderBase* GetTimeRangeBuilderInstance() const { return GraphTimeRangeBuilderInstance; }
-	// Used occasionally to cross-reference other components. Don't call this unless you know what you're doing.
+	/** Used occasionally to cross-reference other components. Don't call this unless you know what you're doing. */
 	UMovieGraphDataCachingBase* GetDataCachingInstance() const { return GraphDataCachingInstance; }
+	/** Gets the Output Merger for this Movie Pipeline which is responsible for gathering all of the data coming in for a given output frame, before making it available to the MovieGraphPipeline. */
+	TSharedPtr<UE::MovieGraph::IMovieGraphOutputMerger> GetOutputMerger() const { return OutputMerger; }
+	/** Writing images to disk is an async process. When you start writing, declare a future with the filename you will eventually write to, and complete the future once it is on disk. */
+	void AddOutputFuture(TFuture<bool>&& InOutputFuture, const UE::MovieGraph::FMovieGraphOutputFutureData& InData);
+
 protected:
 	UFUNCTION(BlueprintCallable, Category = "Movie Render Pipeline")
 	virtual void OnMoviePipelineFinishedImpl();
@@ -56,7 +67,8 @@ protected:
 	virtual void OnEngineTickEndFrame();
 	virtual void RenderFrame();
 	virtual void BuildShotListFromDataSource();
-
+	virtual void ProcessOutstandingFinishedFrames();
+	virtual void ProcessOutstandingFutures();
 
 	virtual void TickProducingFrames();
 	virtual void TickPostFinalizeExport(const bool bInForceFinish);
@@ -107,4 +119,22 @@ protected:
 
 	/** What time (in UTC) was Initialization called? Used internally for tracking total job duration. */
 	FDateTime GraphInitializationTime;
+
+	/** This gathers all of the produced data for an output frame (which may come in async, many frames later) before passing them onto the Output Containers. */
+	TSharedPtr<UE::MovieGraph::IMovieGraphOutputMerger> OutputMerger;
+
+	/** A debug image sequence writer in the event they want to dump every sample generated on its own. */
+	IImageWriteQueue* Debug_ImageWriteQueue;
+
+	/** 
+	* An array of Output Futures for files that have started writing to disk, but have not finished. 
+	* Each frame we take each future that finishes writing to disk and push it into the GeneratedOutputData array.
+	*/
+	TArray<UE::MovieGraph::FMovieGraphOutputFuture> OutstandingOutputFutures;
+
+	/**
+	* An array of output data, one per shot. This only keeps track of the filenames generated during
+	* rendering (and is only pushed to once the file is actually on disk), and does not store pixel data.
+	*/
+	TArray<FMovieGraphRenderOutputData> GeneratedOutputData;
 };
\ No newline at end of file
diff --git a/Engine/Plugins/MovieScene/MovieRenderPipeline/Source/MovieRenderPipelineCore/Public/Graph/MovieGraphPlaceholderNodes.cpp b/Engine/Plugins/MovieScene/MovieRenderPipeline/Source/MovieRenderPipelineCore/Public/Graph/MovieGraphPlaceholderNodes.cpp
--- a/Engine/Plugins/MovieScene/MovieRenderPipeline/Source/MovieRenderPipelineCore/Public/Graph/MovieGraphPlaceholderNodes.cpp
+++ b/Engine/Plugins/MovieScene/MovieRenderPipeline/Source/MovieRenderPipelineCore/Public/Graph/MovieGraphPlaceholderNodes.cpp
@@ -1,4 +1,4 @@
-﻿// Copyright Epic Games, Inc. All Rights Reserved.
+// Copyright Epic Games, Inc. All Rights Reserved.
 
 #include "MovieGraphPlaceholderNodes.h"
 
@@ -172,47 +172,6 @@ FSlateIcon UMovieGraphGlobalGameOverridesNode::GetIconAndTint(FLinearColor& OutC
 }
 #endif // WITH_EDITOR
 
-TArray<FMovieGraphPinProperties> UMovieGraphDeferredRendererNode::GetInputPinProperties() const
-{
-	TArray<FMovieGraphPinProperties> Properties;
-	Properties.Add(FMovieGraphPinProperties(NAME_None, EMovieGraphMemberType::Branch, false));
-	return Properties;
-}
-
-TArray<FMovieGraphPinProperties> UMovieGraphDeferredRendererNode::GetOutputPinProperties() const
-{
-	TArray<FMovieGraphPinProperties> Properties;
-	Properties.Add(FMovieGraphPinProperties(NAME_None, EMovieGraphMemberType::Branch, false));
-	return Properties;
-}
-
-#if WITH_EDITOR
-FText UMovieGraphDeferredRendererNode::GetNodeTitle(const bool bGetDescriptive) const
-{
-	static const FText DeferredRendererNodeName = LOCTEXT("NodeName_DeferredRenderer", "Deferred Renderer");
-	return DeferredRendererNodeName;
-}
-
-FText UMovieGraphDeferredRendererNode::GetMenuCategory() const
-{
-	return NodeCategory_Renderers;
-}
-
-FLinearColor UMovieGraphDeferredRendererNode::GetNodeTitleColor() const
-{
-	static const FLinearColor DeferredRendererNodeColor = FLinearColor(0.572f, 0.274f, 1.f);
-	return DeferredRendererNodeColor;
-}
-
-FSlateIcon UMovieGraphDeferredRendererNode::GetIconAndTint(FLinearColor& OutColor) const
-{
-	static const FSlateIcon DeferredRendererIcon = FSlateIcon(FAppStyle::GetAppStyleSetName(), "SequenceRecorder.TabIcon");
-
-	OutColor = FLinearColor::White;
-	return DeferredRendererIcon;
-}
-#endif // WITH_EDITOR
-
 TArray<FMovieGraphPinProperties> UMovieGraphPathTracedRendererNode::GetInputPinProperties() const
 {
 	TArray<FMovieGraphPinProperties> Properties;
@@ -295,47 +254,6 @@ FSlateIcon UMovieGraphEXRSequenceNode::GetIconAndTint(FLinearColor& OutColor) co
 }
 #endif // WITH_EDITOR
 
-TArray<FMovieGraphPinProperties> UMovieGraphJPGSequenceNode::GetInputPinProperties() const
-{
-	TArray<FMovieGraphPinProperties> Properties;
-	Properties.Add(FMovieGraphPinProperties(NAME_None, EMovieGraphMemberType::Branch, false));
-	return Properties;
-}
-
-TArray<FMovieGraphPinProperties> UMovieGraphJPGSequenceNode::GetOutputPinProperties() const
-{
-	TArray<FMovieGraphPinProperties> Properties;
-	Properties.Add(FMovieGraphPinProperties(NAME_None, EMovieGraphMemberType::Branch, false));
-	return Properties;
-}
-
-#if WITH_EDITOR
-FText UMovieGraphJPGSequenceNode::GetNodeTitle(const bool bGetDescriptive) const
-{
-	static const FText JPGSequenceNodeName = LOCTEXT("NodeName_JPGSequence", "JPG Sequence");
-	return JPGSequenceNodeName;
-}
-
-FText UMovieGraphJPGSequenceNode::GetMenuCategory() const
-{
-	return NodeCategory_OutputType;
-}
-
-FLinearColor UMovieGraphJPGSequenceNode::GetNodeTitleColor() const
-{
-	static const FLinearColor ImageSequenceNodeColor = FLinearColor(0.047f, 0.654f, 0.537f);
-	return ImageSequenceNodeColor;
-}
-
-FSlateIcon UMovieGraphJPGSequenceNode::GetIconAndTint(FLinearColor& OutColor) const
-{
-	static const FSlateIcon ImageSequenceIcon = FSlateIcon(FAppStyle::GetAppStyleSetName(), "ClassIcon.Texture2D");
-
-	OutColor = FLinearColor::White;
-	return ImageSequenceIcon;
-}
-#endif
-
 TArray<FMovieGraphPinProperties> UMovieGraphBranchNode::GetInputPinProperties() const
 {
 	static const FName TrueBranch("True");
diff --git a/Engine/Plugins/MovieScene/MovieRenderPipeline/Source/MovieRenderPipelineCore/Public/Graph/MovieGraphPlaceholderNodes.h b/Engine/Plugins/MovieScene/MovieRenderPipeline/Source/MovieRenderPipelineCore/Public/Graph/MovieGraphPlaceholderNodes.h
--- a/Engine/Plugins/MovieScene/MovieRenderPipeline/Source/MovieRenderPipelineCore/Public/Graph/MovieGraphPlaceholderNodes.h
+++ b/Engine/Plugins/MovieScene/MovieRenderPipeline/Source/MovieRenderPipelineCore/Public/Graph/MovieGraphPlaceholderNodes.h
@@ -1,4 +1,4 @@
-﻿// Copyright Epic Games, Inc. All Rights Reserved.
+// Copyright Epic Games, Inc. All Rights Reserved.
 
 #pragma once
 
@@ -103,26 +103,6 @@ public:
 #endif
 };
 
-/** A node which represents a deferred renderer. */
-UCLASS()
-class MOVIERENDERPIPELINECORE_API UMovieGraphDeferredRendererNode : public UMovieGraphNode
-{
-	GENERATED_BODY()
-
-public:
-	UMovieGraphDeferredRendererNode() = default;
-
-	virtual TArray<FMovieGraphPinProperties> GetInputPinProperties() const override;
-	virtual TArray<FMovieGraphPinProperties> GetOutputPinProperties() const override;
-
-#if WITH_EDITOR
-	virtual FText GetNodeTitle(const bool bGetDescriptive = false) const override;
-	virtual FText GetMenuCategory() const override;
-	virtual FLinearColor GetNodeTitleColor() const override;
-	virtual FSlateIcon GetIconAndTint(FLinearColor& OutColor) const override;
-#endif
-};
-
 /** A node which represents a path traced renderer. */
 UCLASS()
 class MOVIERENDERPIPELINECORE_API UMovieGraphPathTracedRendererNode : public UMovieGraphNode
@@ -163,26 +143,6 @@ public:
 #endif
 };
 
-/** A node which generates a JPG image sequence. */
-UCLASS()
-class MOVIERENDERPIPELINECORE_API UMovieGraphJPGSequenceNode : public UMovieGraphNode
-{
-	GENERATED_BODY()
-
-public:
-	UMovieGraphJPGSequenceNode() = default;
-
-	virtual TArray<FMovieGraphPinProperties> GetInputPinProperties() const override;
-	virtual TArray<FMovieGraphPinProperties> GetOutputPinProperties() const override;
-
-#if WITH_EDITOR
-	virtual FText GetNodeTitle(const bool bGetDescriptive = false) const override;
-	virtual FText GetMenuCategory() const override;
-	virtual FLinearColor GetNodeTitleColor() const override;
-	virtual FSlateIcon GetIconAndTint(FLinearColor& OutColor) const override;
-#endif
-};
-
 /** A node which creates a branching condition. */
 UCLASS()
 class MOVIERENDERPIPELINECORE_API UMovieGraphBranchNode : public UMovieGraphNode
diff --git a/Engine/Plugins/MovieScene/MovieRenderPipeline/Source/MovieRenderPipelineCore/Public/Graph/MovieGraphRenderDataIdentifier.h b/Engine/Plugins/MovieScene/MovieRenderPipeline/Source/MovieRenderPipelineCore/Public/Graph/MovieGraphRenderDataIdentifier.h
new file mode 100644
--- /dev/null
+++ b/Engine/Plugins/MovieScene/MovieRenderPipeline/Source/MovieRenderPipelineCore/Public/Graph/MovieGraphRenderDataIdentifier.h
@@ -0,0 +1,71 @@
+// Copyright Epic Games, Inc. All Rights Reserved.
+#pragma once
+
+#include "CoreMinimal.h"
+#include "MovieGraphRenderDataIdentifier.generated.h"
+
+/**
+* This data structure can be used to identify what render data a set of pixels represents
+* by knowing what the render layer name is, what renderer produced it, if it's a sub-resource,
+* and what camera it is for. Can be used as the key in a TMap.
+*/
+USTRUCT(BlueprintType)
+struct FMovieGraphRenderDataIdentifier
+{
+	GENERATED_BODY()
+
+	FMovieGraphRenderDataIdentifier()
+	{}
+
+	FMovieGraphRenderDataIdentifier(const FString& InRenderLayerName, const FString& InRendererName,
+		const FString& InSubRenderResourceName, const FString& InCameraName)
+		: RenderLayerName(InRenderLayerName)
+		, RendererName(InRendererName)
+		, SubResourceName(InSubRenderResourceName)
+		, CameraName(InCameraName)
+	{
+	}
+
+	bool operator == (const FMovieGraphRenderDataIdentifier& InRHS) const
+	{
+		return RenderLayerName == InRHS.RenderLayerName && 
+			RendererName == InRHS.RendererName &&
+			SubResourceName == InRHS.SubResourceName &&
+			CameraName == InRHS.CameraName;
+	}
+
+	bool operator != (const FMovieGraphRenderDataIdentifier& InRHS) const
+	{
+		return !(*this == InRHS);
+	}
+
+	friend uint32 GetTypeHash(FMovieGraphRenderDataIdentifier InIdentifier)
+	{
+		return HashCombineFast(GetTypeHash(InIdentifier.RenderLayerName),
+			HashCombineFast(GetTypeHash(InIdentifier.RendererName),
+				HashCombineFast(GetTypeHash(InIdentifier.SubResourceName),
+					GetTypeHash(InIdentifier.CameraName))));
+	}
+
+	friend FString LexToString(const FMovieGraphRenderDataIdentifier InIdentifier)
+	{
+		return FString::Printf(TEXT("RenderLayer: %s Renderer:%s SubResource: %s Camera: %s"), *InIdentifier.RenderLayerName, *InIdentifier.RendererName, *InIdentifier.SubResourceName, *InIdentifier.CameraName);
+	}
+
+public:
+	/** The user provided name for the whole Render Layer ("character", "background", etc.) */
+	UPROPERTY(EditAnywhere, BlueprintReadWrite, Category = "Movie Pipeline")
+	FString RenderLayerName;
+
+	/** Which renderer was used to produce this image ("panoramic" "deferred" "path tracer", etc.) */
+	UPROPERTY(EditAnywhere, BlueprintReadWrite, Category = "Movie Pipeline")
+	FString RendererName;
+
+	/** A sub-resource name for the renderer (ie: "beauty", "object id", "depth", etc.) */
+	UPROPERTY(EditAnywhere, BlueprintReadWrite, Category = "Movie Pipeline")
+	FString SubResourceName;
+
+	/** The name of the camera being used for this render. */
+	UPROPERTY(EditAnywhere, BlueprintReadWrite, Category = "Movie Pipeline")
+	FString CameraName;
+};
\ No newline at end of file
diff --git a/Engine/Plugins/MovieScene/MovieRenderPipeline/Source/MovieRenderPipelineCore/Public/Graph/MovieGraphTimeStepData.h b/Engine/Plugins/MovieScene/MovieRenderPipeline/Source/MovieRenderPipelineCore/Public/Graph/MovieGraphTimeStepData.h
new file mode 100644
--- /dev/null
+++ b/Engine/Plugins/MovieScene/MovieRenderPipeline/Source/MovieRenderPipelineCore/Public/Graph/MovieGraphTimeStepData.h
@@ -0,0 +1,46 @@
+// Copyright Epic Games, Inc. All Rights Reserved.
+#pragma once
+
+#include "CoreMinimal.h"
+#include "MovieGraphTimeStepData.generated.h"
+
+/** 
+* This data structure needs to be filled out each frame by the UMovieGraphTimeStepBase,
+* which will eventually be passed to the renderer. It controls per-sample behavior such
+* as the delta time, if this is the first/last sample for an output frame, etc.
+*/
+USTRUCT(BlueprintType)
+struct MOVIERENDERPIPELINECORE_API FMovieGraphTimeStepData
+{
+	GENERATED_BODY()
+
+	UPROPERTY(BlueprintReadWrite, EditAnywhere, Category = "Movie Graph")
+	int32 OutputFrameNumber;
+
+	UPROPERTY(BlueprintReadWrite, EditAnywhere, Category = "Movie Graph")
+	float FrameDeltaTime;
+	
+	UPROPERTY(BlueprintReadWrite, EditAnywhere, Category = "Movie Graph")
+	float WorldTimeDilation;
+	
+	UPROPERTY(BlueprintReadWrite, EditAnywhere, Category = "Movie Graph")
+	float WorldSeconds;
+	
+	UPROPERTY(BlueprintReadWrite, EditAnywhere, Category = "Movie Graph")
+	float MotionBlurFraction;
+
+	/** 
+	* Should be set to true for the first sample of each output frame. Used to determine
+	* if various systems need to reset or gather data for a new frame. Can be true at
+	* the same time as bIsLastTemporalSampleForFrame (ie: 1TS)
+	*/
+	UPROPERTY(BlueprintReadWrite, EditAnywhere, Category = "Movie Graph")
+	bool bIsFirstTemporalSampleForFrame;
+
+	/**
+	* Should be set to true for the last sample of each output frame. Can be true at
+	* the same time as bIsFirstTemporalSampleForFrame (ie: 1TS)
+	*/
+	UPROPERTY(BlueprintReadWrite, EditAnywhere, Category = "Movie Graph")
+	bool bIsLastTemporalSampleForFrame;
+};
diff --git a/Engine/Plugins/MovieScene/MovieRenderPipeline/Source/MovieRenderPipelineCore/Public/Graph/MovieGraphTraversalContext.h b/Engine/Plugins/MovieScene/MovieRenderPipeline/Source/MovieRenderPipelineCore/Public/Graph/MovieGraphTraversalContext.h
new file mode 100644
--- /dev/null
+++ b/Engine/Plugins/MovieScene/MovieRenderPipeline/Source/MovieRenderPipelineCore/Public/Graph/MovieGraphTraversalContext.h
@@ -0,0 +1,56 @@
+// Copyright Epic Games, Inc. All Rights Reserved.
+
+#pragma once
+#include "CoreMinimal.h"
+#include "Graph/MovieGraphRenderDataIdentifier.h"
+#include "Graph/MovieGraphTimeStepData.h"
+#include "MovieGraphTraversalContext.generated.h"
+
+// Forward Declares
+class UMovieGraphConfig;
+class UMoviePipelineExecutorJob;
+
+
+USTRUCT(BlueprintType)
+struct FMovieGraphBranch
+{
+	GENERATED_BODY()
+
+public:
+	UPROPERTY(BlueprintReadWrite, EditAnywhere, Category = "MovieGraph")
+	FName BranchName;
+};
+
+
+USTRUCT(BlueprintType)
+struct FMovieGraphTraversalContext
+{
+	GENERATED_BODY();
+
+	UPROPERTY(BlueprintReadWrite, EditAnywhere, Category = "MovieGraph")
+	FMovieGraphBranch RootBranch;
+
+	/** Which shot (out of ShotCount) is this time step for? */
+	UPROPERTY(BlueprintReadWrite, EditAnywhere, Category = "Movie Graph")
+	int32 ShotIndex;
+
+	/** The total number of shots being rendered for this job. This is from the active shot list, not total in the Level Sequence. */
+	UPROPERTY(BlueprintReadWrite, EditAnywhere, Category = "Movie Graph")
+	int32 ShotCount;
+
+	/** The job in the queue this traversal context is for. Needed to fetch variable values from the job. */
+	UPROPERTY(BlueprintReadWrite, EditAnywhere, Category = "Movie Graph")
+	TObjectPtr<UMoviePipelineExecutorJob> Job;
+
+	/** The root graph to start our traversal from. This could be a shared config for the whole job, or a shot-specific override. */
+	UPROPERTY(BlueprintReadWrite, EditAnywhere, Category = "Movie Graph")
+	TObjectPtr<UMovieGraphConfig> RootGraph;
+
+	/** The name of the render resource this state was captured for. */
+	UPROPERTY(BlueprintReadWrite, EditAnywhere, Category = "MovieGraph")
+	FMovieGraphRenderDataIdentifier RenderDataIdentifier;
+
+	/** The time data(output frame, delta times, etc.) */
+	UPROPERTY(BlueprintReadWrite, EditAnywhere, Category = "MovieGraph")
+	FMovieGraphTimeStepData Time;
+};
\ No newline at end of file
diff --git a/Engine/Plugins/MovieScene/MovieRenderPipeline/Source/MovieRenderPipelineCore/Public/Graph/Nodes/MovieGraphDeferredRenderPassNode.h b/Engine/Plugins/MovieScene/MovieRenderPipeline/Source/MovieRenderPipelineCore/Public/Graph/Nodes/MovieGraphDeferredRenderPassNode.h
--- a/Engine/Plugins/MovieScene/MovieRenderPipeline/Source/MovieRenderPipelineCore/Public/Graph/Nodes/MovieGraphDeferredRenderPassNode.h
+++ b/Engine/Plugins/MovieScene/MovieRenderPipeline/Source/MovieRenderPipelineCore/Public/Graph/Nodes/MovieGraphDeferredRenderPassNode.h
@@ -8,11 +8,16 @@
 #include "MovieRenderPipelineDataTypes.h"
 #include "Engine/EngineTypes.h"
 #include "Async/TaskGraphFwd.h"
+#include "Styling/AppStyle.h"
 #include "MovieGraphDeferredRenderPassNode.generated.h"
 
 // Forward Declares
 class UMovieGraphDefaultRenderer;
-namespace UE::MovieGraph::DefaultRenderer { struct FRenderTargetInitParams; }
+namespace UE::MovieGraph::DefaultRenderer
+{
+	struct FRenderTargetInitParams;
+	struct IMovieGraphOutputMerger;
+}
 struct FMovieGraphRenderPassSetupData;
 struct FMovieGraphRenderPassLayerData;
 struct FImageOverlappedAccumulator;
@@ -31,7 +36,7 @@ namespace UE::MovieGraph
 	{
 	public:
 		TWeakPtr<FImageOverlappedAccumulator, ESPMode::ThreadSafe> ImageAccumulator;
-		// TWeakPtr<IMoviePipelineOutputMerger, ESPMode::ThreadSafe> OutputMerger;
+		TWeakPtr<IMovieGraphOutputMerger, ESPMode::ThreadSafe> OutputMerger;
 
 		// If it's the first sample then we will reset the accumulator to a clean slate before accumulating into it.
 		bool bIsFirstSample;
@@ -41,7 +46,7 @@ namespace UE::MovieGraph
 		FGraphEventRef TaskPrerequisite;
 	};
 
-	// void MOVIERENDERPIPELINERENDERPASSES_API AccumulateSample_TaskThread(TUniquePtr<FImagePixelData>&& InPixelData, const MoviePipeline::FImageSampleAccumulationArgs& InParams);
+	void AccumulateSample_TaskThread(TUniquePtr<FImagePixelData>&& InPixelData, const UE::MovieGraph::FMovieGraphSampleState InSampleState, const UE::MovieGraph::FMovieGraphRenderDataAccumulationArgs& InAccumulationParams);
 }
 
 UCLASS()
@@ -54,13 +59,27 @@ public:
 	{
 		return NSLOCTEXT("MovieGraphNodes", "DeferredRenderPassGraphNode_Description", "Deferred Renderer");
 	}
+	
+	FLinearColor GetNodeTitleColor() const
+	{
+		return FLinearColor(0.572f, 0.274f, 1.f);
+	}
+
+	FSlateIcon GetIconAndTint(FLinearColor& OutColor) const
+	{
+		static const FSlateIcon DeferredRendererIcon = FSlateIcon(FAppStyle::GetAppStyleSetName(), "SequenceRecorder.TabIcon");
+
+		OutColor = FLinearColor::White;
+		return DeferredRendererIcon;
+	}
 #endif
 protected:
 	// UMovieGraphRenderPassNode Interface
 	virtual FString GetRendererNameImpl() const override { return TEXT("Deferred"); }
 	virtual void SetupImpl(const FMovieGraphRenderPassSetupData& InSetupData) override;
 	virtual void TeardownImpl() override;
-	virtual void RenderImpl(const FMovieGraphTimeStepData& InTimeData) override;
+	virtual void RenderImpl(FMovieGraphTraversalContext InFrameTraversalContext, const FMovieGraphTimeStepData& InTimeData) override;
+	virtual void GatherOutputPassesImpl(TArray<FMovieGraphRenderDataIdentifier>& OutExpectedPasses) const override;
 	// ~UMovieGraphRenderPassNode Interface
 
 	static void AddReferencedObjects(UObject* InThis, FReferenceCollector& Collector);
@@ -109,7 +128,8 @@ protected:
 	public:
 		void Setup(TWeakObjectPtr<UMovieGraphDefaultRenderer> InRenderer, TWeakObjectPtr<UMovieGraphDeferredRenderPassNode> InRenderPassNode, const FMovieGraphRenderPassLayerData& InLayer);
 		void Teardown();
-		void Render(const FMovieGraphTimeStepData& InTimeData);
+		void Render(FMovieGraphTraversalContext InFrameTraversalContext, const FMovieGraphTimeStepData& InTimeData);
+		void GatherOutputPassesImpl(TArray<FMovieGraphRenderDataIdentifier>& OutExpectedPasses) const;
 		void AddReferencedObjects(FReferenceCollector& Collector);
 
 	protected:
diff --git a/Engine/Plugins/MovieScene/MovieRenderPipeline/Source/MovieRenderPipelineCore/Public/Graph/Nodes/MovieGraphFileOutputNode.h b/Engine/Plugins/MovieScene/MovieRenderPipeline/Source/MovieRenderPipelineCore/Public/Graph/Nodes/MovieGraphFileOutputNode.h
new file mode 100644
--- /dev/null
+++ b/Engine/Plugins/MovieScene/MovieRenderPipeline/Source/MovieRenderPipelineCore/Public/Graph/Nodes/MovieGraphFileOutputNode.h
@@ -0,0 +1,57 @@
+// Copyright Epic Games, Inc. All Rights Reserved.
+
+#pragma once
+#include "Graph/MovieGraphNode.h"
+#include "Graph/MovieGraphDataTypes.h"
+#include "MovieGraphFileOutputNode.generated.h"
+
+/**
+* The UMovieGraphFileOutputNode node defines a output file format that MRQ may place produced render data into.
+* At runtime, the individual pieces of render data will be generated by the Movie Graph Renderer, and then 
+* all data will be collected into a single datastructure by the FMovieGraphOutputMerger before being passed back
+* into the main Movie Graph Pipeline. Once this happens, the CDOs deriving from UMovieGraphFileOutputNode will be
+* given a pointer to the information for the whole frame, so they can decide what to do with it. There are potentially
+* multiple output formats at once (such as exr and jpeg).
+*/
+UCLASS(Abstract)
+class UMovieGraphFileOutputNode : public UMovieGraphNode
+{
+	GENERATED_BODY()
+public:
+	UMovieGraphFileOutputNode() = default;
+	
+	// UMovieGraphFileOutputNode Interface
+	/** 
+	* This is called by the Game Thread containing the final output data for a given output frame. If you need to hold 
+	* onto this data for longer than the duration of the single function call (such as async/task graph operations) then
+	* you need to increment the reference count on the data in the UMovieGraphPipeline, and decrement when you are done.
+	* Failing to increment will lead to the data being deleted out from underneath you, and failing to decrement will leak 
+	* memory until the end of a render.
+	*/
+	void OnReceiveImageData(UMovieGraphPipeline* InPipeline, UE::MovieGraph::FMovieGraphOutputMergerFrame* InRawFrameData) { OnReceiveImageDataImpl(InPipeline, InRawFrameData); }
+	// ~UMovieGraphFileOutputNode Interface
+
+	virtual TArray<FMovieGraphPinProperties> GetInputPinProperties() const override
+	{
+		TArray<FMovieGraphPinProperties> Properties;
+		Properties.Add(FMovieGraphPinProperties(NAME_None, EMovieGraphMemberType::Branch, false));
+		return Properties;
+	}
+
+	virtual TArray<FMovieGraphPinProperties> GetOutputPinProperties() const override
+	{
+		TArray<FMovieGraphPinProperties> Properties;
+		Properties.Add(FMovieGraphPinProperties(NAME_None, EMovieGraphMemberType::Branch, false));
+		return Properties;
+	}
+
+#if WITH_EDITOR
+	virtual FText GetMenuCategory() const override
+	{
+		return NSLOCTEXT("MovieGraphNodes", "FileOutputGraphNode_Category", "Output Type");
+	}
+#endif
+
+protected:
+	virtual void OnReceiveImageDataImpl(UMovieGraphPipeline* InPipeline, UE::MovieGraph::FMovieGraphOutputMergerFrame* InRawFrameData) {}
+};
\ No newline at end of file
diff --git a/Engine/Plugins/MovieScene/MovieRenderPipeline/Source/MovieRenderPipelineCore/Public/Graph/Nodes/MovieGraphImageSequenceOutputNode.h b/Engine/Plugins/MovieScene/MovieRenderPipeline/Source/MovieRenderPipelineCore/Public/Graph/Nodes/MovieGraphImageSequenceOutputNode.h
new file mode 100644
--- /dev/null
+++ b/Engine/Plugins/MovieScene/MovieRenderPipeline/Source/MovieRenderPipelineCore/Public/Graph/Nodes/MovieGraphImageSequenceOutputNode.h
@@ -0,0 +1,147 @@
+// Copyright Epic Games, Inc. All Rights Reserved.
+
+#pragma once
+#include "Graph/Nodes/MovieGraphFileOutputNode.h"
+#include "IImageWrapper.h"
+#include "Styling/AppStyle.h"
+#include "MovieGraphImageSequenceOutputNode.generated.h"
+
+// Forward Declare
+class IImageWriteQueue;
+
+/**
+* The UMovieGraphFileOutputNode node defines a output file format that MRQ may place produced render data into.
+* At runtime, the individual pieces of render data will be generated by the Movie Graph Renderer, and then 
+* all data will be collected into a single datastructure by the FMovieGraphOutputMerger before being passed back
+* into the main Movie Graph Pipeline. Once this happens, the classes deriving from UMovieGraphFileOutputNode will be
+* given a pointer to the information for the whole frame, so they can decide what to do with it. There are potentially
+* multiple output formats at once (such as exr and jpeg).
+*/
+UCLASS(Abstract)
+class UMovieGraphImageSequenceOutputNode : public UMovieGraphFileOutputNode
+{
+	GENERATED_BODY()
+public:
+	UMovieGraphImageSequenceOutputNode();
+	
+	// UMovieGraphFileOutputNode Interface
+	virtual void OnReceiveImageDataImpl(UMovieGraphPipeline* InPipeline, UE::MovieGraph::FMovieGraphOutputMergerFrame* InRawFrameData) override;
+	// ~UMovieGraphFileOutputNode Interface
+
+protected:
+	/** The output format (as known used by the ImageWriteQueue) to output into. */
+	EImageFormat OutputFormat;
+
+	/** A pointer to the image write queue used for asynchronously writing images */
+	IImageWriteQueue* ImageWriteQueue;
+};
+
+UCLASS()
+class UMovieGraphImageSequenceOutputNode_BMP : public UMovieGraphImageSequenceOutputNode
+{
+	GENERATED_BODY()
+
+public:
+	UMovieGraphImageSequenceOutputNode_BMP()
+	{
+		OutputFormat = EImageFormat::BMP;
+	}
+
+#if WITH_EDITOR
+	virtual FText GetNodeTitle(const bool bGetDescriptive = false) const override 
+	{ 
+		if(bGetDescriptive)
+		{
+			return NSLOCTEXT("MovieGraph", "ImgSequenceBMPSetting_NodeTitleFull", ".bmp Sequence\n[8bit]"); 
+		}
+		return NSLOCTEXT("MovieGraph", "ImgSequenceBMPSetting_NodeTitleShort", ".bmp Sequence");
+	}
+	
+	virtual FLinearColor GetNodeTitleColor() const override
+	{
+		return FLinearColor(0.047f, 0.654f, 0.537f);
+	}
+	
+	virtual FSlateIcon GetIconAndTint(FLinearColor& OutColor) const override
+	{
+		static const FSlateIcon ImageSequenceIcon = FSlateIcon(FAppStyle::GetAppStyleSetName(), "ClassIcon.Texture2D");
+
+		OutColor = FLinearColor::White;
+		return ImageSequenceIcon;
+	}
+#endif
+};
+
+
+UCLASS()
+class UMovieGraphImageSequenceOutputNode_JPG : public UMovieGraphImageSequenceOutputNode
+{
+	GENERATED_BODY()
+
+public:
+	UMovieGraphImageSequenceOutputNode_JPG()
+	{
+		OutputFormat = EImageFormat::JPEG;
+	}
+
+#if WITH_EDITOR
+	virtual FText GetNodeTitle(const bool bGetDescriptive = false) const override 
+	{ 
+		if(bGetDescriptive)
+		{
+			return NSLOCTEXT("MovieGraph", "ImgSequenceJPGSetting_NodeTitleFull", ".jpg Sequence\n[8bit]"); 
+		}
+		return NSLOCTEXT("MovieGraph", "ImgSequenceJPGSetting_NodeTitleShort", ".jpg Sequence");
+	}
+	
+	virtual FLinearColor GetNodeTitleColor() const override
+	{
+		return FLinearColor(0.047f, 0.654f, 0.537f);
+	}
+	
+	virtual FSlateIcon GetIconAndTint(FLinearColor& OutColor) const override
+	{
+		static const FSlateIcon ImageSequenceIcon = FSlateIcon(FAppStyle::GetAppStyleSetName(), "ClassIcon.Texture2D");
+
+		OutColor = FLinearColor::White;
+		return ImageSequenceIcon;
+	}
+#endif
+};
+
+UCLASS()
+class UMovieGraphImageSequenceOutputNode_PNG : public UMovieGraphImageSequenceOutputNode
+{
+	GENERATED_BODY()
+
+public:
+	UMovieGraphImageSequenceOutputNode_PNG()
+	{
+		OutputFormat = EImageFormat::PNG;
+	}
+
+#if WITH_EDITOR
+	virtual FText GetNodeTitle(const bool bGetDescriptive = false) const override 
+	{ 
+		if(bGetDescriptive)
+		{
+			return NSLOCTEXT("MovieGraph", "ImgSequencePNGSetting_NodeTitleFull", ".png Sequence\n[8bit]"); 
+		}
+		return NSLOCTEXT("MovieGraph", "ImgSequencePNGSetting_NodeTitleShort", ".png Sequence");
+	}
+	
+	virtual FLinearColor GetNodeTitleColor() const override
+	{
+		return FLinearColor(0.047f, 0.654f, 0.537f);
+	}
+	
+	virtual FSlateIcon GetIconAndTint(FLinearColor& OutColor) const override
+	{
+		static const FSlateIcon ImageSequenceIcon = FSlateIcon(FAppStyle::GetAppStyleSetName(), "ClassIcon.Texture2D");
+
+		OutColor = FLinearColor::White;
+		return ImageSequenceIcon;
+	}
+#endif
+};
+
diff --git a/Engine/Plugins/MovieScene/MovieRenderPipeline/Source/MovieRenderPipelineCore/Public/Graph/Nodes/MovieGraphRenderPassNode.h b/Engine/Plugins/MovieScene/MovieRenderPipeline/Source/MovieRenderPipelineCore/Public/Graph/Nodes/MovieGraphRenderPassNode.h
--- a/Engine/Plugins/MovieScene/MovieRenderPipeline/Source/MovieRenderPipelineCore/Public/Graph/Nodes/MovieGraphRenderPassNode.h
+++ b/Engine/Plugins/MovieScene/MovieRenderPipeline/Source/MovieRenderPipelineCore/Public/Graph/Nodes/MovieGraphRenderPassNode.h
@@ -2,12 +2,20 @@
 
 #pragma once
 #include "Graph/MovieGraphNode.h"
+#include "Graph/MovieGraphDataTypes.h"
 #include "MovieGraphRenderPassNode.generated.h"
 
 // Forward Declare
 struct FMovieGraphRenderPassSetupData;
 struct FMovieGraphTimeStepData;
 
+/**
+* The UMovieGraphRenderPassNode node defines a render pass that MRQ may produce. This node can be implemented
+* in the graph multiple times, and the exact settings it should use can be created out of a mixture of nodes. Because
+* of this, when rendering, MRQ will figure out how many layers there are that actually use this CDO and will call
+* the function on the CDO once, providing the information about all instances. This will allow the node to create any
+* number of instances (decoupled from the number of times the node is used in the graph).
+*/
 UCLASS(Abstract)
 class UMovieGraphRenderPassNode : public UMovieGraphNode
 {
@@ -17,12 +25,23 @@ public:
 	{
 	}
 
-	// UMovieGraphRenderPassNode interface
+	// UMovieGraphRenderPassNode Interface
+	/** Get the name of this renderer. Deferred, Path Tracer, Panoramic, etc. Called on the CDO. */
 	FString GetRendererName() const { return GetRendererNameImpl(); }
+	/** Called when this should set up for rendering a new shot. Called on the CDO. */
 	void Setup(const FMovieGraphRenderPassSetupData& InSetupData) { SetupImpl(InSetupData); }
+	/** Called when this should do teardown of resources. FlushRenderingCommands() will have already been called by this point. Called on the CDO. */
 	void Teardown() { TeardownImpl(); }
-	void Render(const FMovieGraphTimeStepData& InTimeData) { RenderImpl(InTimeData); }
-	// ~UMovieGraphRenderPassNode interface
+	/** Called each tick (once per temporal sample). Called on the CDO. */
+	void Render(FMovieGraphTraversalContext InFrameTraversalContext, const FMovieGraphTimeStepData& InTimeData) { RenderImpl(InFrameTraversalContext, InTimeData); }
+	/** 
+	* Called each output frame. Should add a series of FMovieGraphRenderDataIdentifiers to the array, and then when producing frames
+	* in Render, the resulting image data should have the matching FMovieGraphRenderDataIdentifiers associated with it. Used by the
+	* Output Merger to ensure all of the render data for a given frame has been generated before passing it on to write to disk.
+	* Called on the CDO.
+	*/
+	void GatherOutputPasses(TArray<FMovieGraphRenderDataIdentifier>& OutExpectedPasses) const { GatherOutputPassesImpl(OutExpectedPasses); }
+	// ~UMovieGraphRenderPassNode Interface
 
 	virtual TArray<FMovieGraphPinProperties> GetInputPinProperties() const override
 	{
@@ -49,5 +68,6 @@ protected:
 	virtual FString GetRendererNameImpl() const { return TEXT("UnnamedRenderPass"); }
 	virtual void SetupImpl(const FMovieGraphRenderPassSetupData& InSetupData) {}
 	virtual void TeardownImpl() {}
-	virtual void RenderImpl(const FMovieGraphTimeStepData& InTimeData) {}
+	virtual void RenderImpl(FMovieGraphTraversalContext InFrameTraversalContext, const FMovieGraphTimeStepData& InTimeData) {}
+	virtual void GatherOutputPassesImpl(TArray<FMovieGraphRenderDataIdentifier>& OutExpectedPasses) const {}
 };
\ No newline at end of file
diff --git a/Engine/Source/Runtime/ImageWriteQueue/Public/ImagePixelData.h b/Engine/Source/Runtime/ImageWriteQueue/Public/ImagePixelData.h
--- a/Engine/Source/Runtime/ImageWriteQueue/Public/ImagePixelData.h
+++ b/Engine/Source/Runtime/ImageWriteQueue/Public/ImagePixelData.h
@@ -195,6 +195,11 @@ struct FImagePixelData
 	template<typename T>
 	const T* GetPayload() const { return static_cast<T*>(Payload.Get()); }
 
+	/**
+	* Sets the payload after construction.
+	*/
+	void SetPayload(FImagePixelPayloadPtr NewPayload) { Payload = NewPayload; }
+
 	bool GetSRGB() const { return bSRGB; }
 	void SetSRGB(bool InSRGB) { bSRGB = InSRGB; }
 
