MoviePipeline: Fixed a duplicate class comment, fixed some missing const &

diff --git a/Engine/Plugins/MovieScene/MovieRenderPipeline/Source/MovieRenderPipelineCore/Private/Graph/Nodes/MovieGraphDeferredRenderPassNode.cpp b/Engine/Plugins/MovieScene/MovieRenderPipeline/Source/MovieRenderPipelineCore/Private/Graph/Nodes/MovieGraphDeferredRenderPassNode.cpp
--- a/Engine/Plugins/MovieScene/MovieRenderPipeline/Source/MovieRenderPipelineCore/Private/Graph/Nodes/MovieGraphDeferredRenderPassNode.cpp
+++ b/Engine/Plugins/MovieScene/MovieRenderPipeline/Source/MovieRenderPipelineCore/Private/Graph/Nodes/MovieGraphDeferredRenderPassNode.cpp
@@ -45,7 +45,7 @@ void UMovieGraphDeferredRenderPassNode::TeardownImpl()
 }
 
 
-void UMovieGraphDeferredRenderPassNode::RenderImpl(FMovieGraphTraversalContext InFrameTraversalContext, const FMovieGraphTimeStepData& InTimeData)
+void UMovieGraphDeferredRenderPassNode::RenderImpl(const FMovieGraphTraversalContext& InFrameTraversalContext, const FMovieGraphTimeStepData& InTimeData)
 {
 	for (TUniquePtr<FMovieGraphDeferredRenderPass>& Instance : CurrentInstances)
 	{
@@ -126,7 +126,7 @@ void UMovieGraphDeferredRenderPassNode::FMovieGraphDeferredRenderPass::AddRefere
 	}
 }
 
-void UMovieGraphDeferredRenderPassNode::FMovieGraphDeferredRenderPass::Render(FMovieGraphTraversalContext InFrameTraversalContext, const FMovieGraphTimeStepData& InTimeData)
+void UMovieGraphDeferredRenderPassNode::FMovieGraphDeferredRenderPass::Render(const FMovieGraphTraversalContext& InFrameTraversalContext, const FMovieGraphTimeStepData& InTimeData)
 {
 	// This is the size we actually render at.
 	UE::MovieGraph::DefaultRenderer::FRenderTargetInitParams RenderTargetInitParams;
diff --git a/Engine/Plugins/MovieScene/MovieRenderPipeline/Source/MovieRenderPipelineCore/Public/Graph/Nodes/MovieGraphDeferredRenderPassNode.h b/Engine/Plugins/MovieScene/MovieRenderPipeline/Source/MovieRenderPipelineCore/Public/Graph/Nodes/MovieGraphDeferredRenderPassNode.h
--- a/Engine/Plugins/MovieScene/MovieRenderPipeline/Source/MovieRenderPipelineCore/Public/Graph/Nodes/MovieGraphDeferredRenderPassNode.h
+++ b/Engine/Plugins/MovieScene/MovieRenderPipeline/Source/MovieRenderPipelineCore/Public/Graph/Nodes/MovieGraphDeferredRenderPassNode.h
@@ -78,7 +78,7 @@ protected:
 	virtual FString GetRendererNameImpl() const override { return TEXT("Deferred"); }
 	virtual void SetupImpl(const FMovieGraphRenderPassSetupData& InSetupData) override;
 	virtual void TeardownImpl() override;
-	virtual void RenderImpl(FMovieGraphTraversalContext InFrameTraversalContext, const FMovieGraphTimeStepData& InTimeData) override;
+	virtual void RenderImpl(const FMovieGraphTraversalContext& InFrameTraversalContext, const FMovieGraphTimeStepData& InTimeData) override;
 	virtual void GatherOutputPassesImpl(TArray<FMovieGraphRenderDataIdentifier>& OutExpectedPasses) const override;
 	// ~UMovieGraphRenderPassNode Interface
 
@@ -128,7 +128,7 @@ protected:
 	public:
 		void Setup(TWeakObjectPtr<UMovieGraphDefaultRenderer> InRenderer, TWeakObjectPtr<UMovieGraphDeferredRenderPassNode> InRenderPassNode, const FMovieGraphRenderPassLayerData& InLayer);
 		void Teardown();
-		void Render(FMovieGraphTraversalContext InFrameTraversalContext, const FMovieGraphTimeStepData& InTimeData);
+		void Render(const FMovieGraphTraversalContext& InFrameTraversalContext, const FMovieGraphTimeStepData& InTimeData);
 		void GatherOutputPassesImpl(TArray<FMovieGraphRenderDataIdentifier>& OutExpectedPasses) const;
 		void AddReferencedObjects(FReferenceCollector& Collector);
 
diff --git a/Engine/Plugins/MovieScene/MovieRenderPipeline/Source/MovieRenderPipelineCore/Public/Graph/Nodes/MovieGraphImageSequenceOutputNode.h b/Engine/Plugins/MovieScene/MovieRenderPipeline/Source/MovieRenderPipelineCore/Public/Graph/Nodes/MovieGraphImageSequenceOutputNode.h
--- a/Engine/Plugins/MovieScene/MovieRenderPipeline/Source/MovieRenderPipelineCore/Public/Graph/Nodes/MovieGraphImageSequenceOutputNode.h
+++ b/Engine/Plugins/MovieScene/MovieRenderPipeline/Source/MovieRenderPipelineCore/Public/Graph/Nodes/MovieGraphImageSequenceOutputNode.h
@@ -10,12 +10,9 @@
 class IImageWriteQueue;
 
 /**
-* The UMovieGraphFileOutputNode node defines a output file format that MRQ may place produced render data into.
-* At runtime, the individual pieces of render data will be generated by the Movie Graph Renderer, and then 
-* all data will be collected into a single datastructure by the FMovieGraphOutputMerger before being passed back
-* into the main Movie Graph Pipeline. Once this happens, the classes deriving from UMovieGraphFileOutputNode will be
-* given a pointer to the information for the whole frame, so they can decide what to do with it. There are potentially
-* multiple output formats at once (such as exr and jpeg).
+* The UMovieGraphImageSequenceOutputNode node is the base class for all image sequence outputs, such as 
+* a series of jpeg, png, bmp, or .exr images. Create an instance of the appropriate class (such as 
+* UMovieGraphImageSequenceOutputNode_JPG) instead of this abstract base class.
 */
 UCLASS(Abstract)
 class UMovieGraphImageSequenceOutputNode : public UMovieGraphFileOutputNode
@@ -36,6 +33,11 @@ protected:
 	IImageWriteQueue* ImageWriteQueue;
 };
 
+/**
+* Save the images generated by the Movie Graph Pipeline as an lossless 8 bit bmp format. This can
+* be useful in rare occasions (bmp files are uncompressed but larger). sRGB is applied.
+* No metadata is supported.
+*/
 UCLASS()
 class UMovieGraphImageSequenceOutputNode_BMP : public UMovieGraphImageSequenceOutputNode
 {
@@ -73,6 +75,11 @@ public:
 };
 
 
+/**
+* Save the images generated by the Movie Graph Pipeline as an 8 bit jpg format. JPEG image files
+* are lossy, but a good balance between compression speed and final filesize. sRGB is applied.
+* No metadata is supported.
+*/
 UCLASS()
 class UMovieGraphImageSequenceOutputNode_JPG : public UMovieGraphImageSequenceOutputNode
 {
@@ -109,6 +116,11 @@ public:
 #endif
 };
 
+/**
+* Save the images generated by the Movie Graph Pipeline as an 8 bit png format. PNG image files
+* are lossless but slow to compress and have a larger final filesize than JPEG. sRGB is applied.
+* No metadata is supported.
+*/
 UCLASS()
 class UMovieGraphImageSequenceOutputNode_PNG : public UMovieGraphImageSequenceOutputNode
 {
diff --git a/Engine/Plugins/MovieScene/MovieRenderPipeline/Source/MovieRenderPipelineCore/Public/Graph/Nodes/MovieGraphRenderPassNode.h b/Engine/Plugins/MovieScene/MovieRenderPipeline/Source/MovieRenderPipelineCore/Public/Graph/Nodes/MovieGraphRenderPassNode.h
--- a/Engine/Plugins/MovieScene/MovieRenderPipeline/Source/MovieRenderPipelineCore/Public/Graph/Nodes/MovieGraphRenderPassNode.h
+++ b/Engine/Plugins/MovieScene/MovieRenderPipeline/Source/MovieRenderPipelineCore/Public/Graph/Nodes/MovieGraphRenderPassNode.h
@@ -33,7 +33,7 @@ public:
 	/** Called when this should do teardown of resources. FlushRenderingCommands() will have already been called by this point. Called on the CDO. */
 	void Teardown() { TeardownImpl(); }
 	/** Called each tick (once per temporal sample). Called on the CDO. */
-	void Render(FMovieGraphTraversalContext InFrameTraversalContext, const FMovieGraphTimeStepData& InTimeData) { RenderImpl(InFrameTraversalContext, InTimeData); }
+	void Render(const FMovieGraphTraversalContext& InFrameTraversalContext, const FMovieGraphTimeStepData& InTimeData) { RenderImpl(InFrameTraversalContext, InTimeData); }
 	/** 
 	* Called each output frame. Should add a series of FMovieGraphRenderDataIdentifiers to the array, and then when producing frames
 	* in Render, the resulting image data should have the matching FMovieGraphRenderDataIdentifiers associated with it. Used by the
@@ -68,6 +68,6 @@ protected:
 	virtual FString GetRendererNameImpl() const { return TEXT("UnnamedRenderPass"); }
 	virtual void SetupImpl(const FMovieGraphRenderPassSetupData& InSetupData) {}
 	virtual void TeardownImpl() {}
-	virtual void RenderImpl(FMovieGraphTraversalContext InFrameTraversalContext, const FMovieGraphTimeStepData& InTimeData) {}
+	virtual void RenderImpl(const FMovieGraphTraversalContext& InFrameTraversalContext, const FMovieGraphTimeStepData& InTimeData) {}
 	virtual void GatherOutputPassesImpl(TArray<FMovieGraphRenderDataIdentifier>& OutExpectedPasses) const {}
 };
\ No newline at end of file
