Movie Render Pipeline: Move EnsureJobHasDefaultSettings to only when there is no active job and a new job is created. This fixes an issue where default settings would get added to a job that already exists when clicking on the Render Movie button from Sequencer

diff --git a/Engine/Plugins/MovieScene/MovieRenderPipeline/Source/MovieRenderPipelineEditor/Private/MovieRenderPipelineEditorModule.cpp b/Engine/Plugins/MovieScene/MovieRenderPipeline/Source/MovieRenderPipelineEditor/Private/MovieRenderPipelineEditorModule.cpp
--- a/Engine/Plugins/MovieScene/MovieRenderPipeline/Source/MovieRenderPipelineEditor/Private/MovieRenderPipelineEditorModule.cpp
+++ b/Engine/Plugins/MovieScene/MovieRenderPipeline/Source/MovieRenderPipelineEditor/Private/MovieRenderPipelineEditorModule.cpp
@@ -140,9 +140,9 @@ class FMovieRenderPipelineRenderer : public IMovieRendererInterface
 			{
 				ActiveJob->SetPresetOrigin(ProjectSettings->LastPresetOrigin.Get());
 			}
-		}
 
-		UMoviePipelineEditorBlueprintLibrary::EnsureJobHasDefaultSettings(ActiveJob);
+			UMoviePipelineEditorBlueprintLibrary::EnsureJobHasDefaultSettings(ActiveJob);
+		}
 
 		TArray<FString> ShotNames;
 		for (UMovieSceneCinematicShotSection* ShotSection : InSections)
