CIS: Fixed some uninitialized variables in MovieGraph code.

diff --git a/Engine/Plugins/MovieScene/MovieRenderPipeline/Source/MovieRenderPipelineCore/Public/Graph/MovieGraphTimeStepData.h b/Engine/Plugins/MovieScene/MovieRenderPipeline/Source/MovieRenderPipelineCore/Public/Graph/MovieGraphTimeStepData.h
--- a/Engine/Plugins/MovieScene/MovieRenderPipeline/Source/MovieRenderPipelineCore/Public/Graph/MovieGraphTimeStepData.h
+++ b/Engine/Plugins/MovieScene/MovieRenderPipeline/Source/MovieRenderPipelineCore/Public/Graph/MovieGraphTimeStepData.h
@@ -14,6 +14,17 @@ struct MOVIERENDERPIPELINECORE_API FMovieGraphTimeStepData
 {
 	GENERATED_BODY()
 
+public:
+	FMovieGraphTimeStepData()
+		: OutputFrameNumber(0)
+		, FrameDeltaTime(0.f)
+		, WorldTimeDilation(0.f)
+		, WorldSeconds(0.f)
+		, MotionBlurFraction(0.f)
+		, bIsFirstTemporalSampleForFrame(false)
+		, bIsLastTemporalSampleForFrame(false)
+	{}
+
 	UPROPERTY(BlueprintReadWrite, EditAnywhere, Category = "Movie Graph")
 	int32 OutputFrameNumber;
 
diff --git a/Engine/Plugins/MovieScene/MovieRenderPipeline/Source/MovieRenderPipelineCore/Public/Graph/MovieGraphTraversalContext.h b/Engine/Plugins/MovieScene/MovieRenderPipeline/Source/MovieRenderPipelineCore/Public/Graph/MovieGraphTraversalContext.h
--- a/Engine/Plugins/MovieScene/MovieRenderPipeline/Source/MovieRenderPipelineCore/Public/Graph/MovieGraphTraversalContext.h
+++ b/Engine/Plugins/MovieScene/MovieRenderPipeline/Source/MovieRenderPipelineCore/Public/Graph/MovieGraphTraversalContext.h
@@ -27,6 +27,12 @@ struct FMovieGraphTraversalContext
 {
 	GENERATED_BODY();
 
+public:
+	FMovieGraphTraversalContext()
+		: ShotIndex(0)
+		, ShotCount(0)
+	{}
+
 	UPROPERTY(BlueprintReadWrite, EditAnywhere, Category = "MovieGraph")
 	FMovieGraphBranch RootBranch;
 
