MoviePipeline:
Combine the DataCaching and TimeRangeBuilder concepts into one, in practice, the source that generates your data is also the thing you want to update every frame as yo walk through time.

diff --git a/Engine/Plugins/MovieScene/MovieRenderPipeline/Source/MovieRenderPipelineCore/Private/Graph/MovieGraphDataTypes.cpp b/Engine/Plugins/MovieScene/MovieRenderPipeline/Source/MovieRenderPipelineCore/Private/Graph/MovieGraphDataTypes.cpp
--- a/Engine/Plugins/MovieScene/MovieRenderPipeline/Source/MovieRenderPipelineCore/Private/Graph/MovieGraphDataTypes.cpp
+++ b/Engine/Plugins/MovieScene/MovieRenderPipeline/Source/MovieRenderPipelineCore/Private/Graph/MovieGraphDataTypes.cpp
@@ -3,16 +3,14 @@
 #include "Graph/MovieGraphDataTypes.h"
 #include "Graph/MovieGraphLinearTimeStep.h"
 #include "Graph/MovieGraphDefaultRenderer.h"
-#include "Graph/MovieGraphDefaultDataCaching.h"
-#include "Graph/MovieGraphSequenceTimeRangeBuilder.h"
+#include "Graph/MovieGraphSequenceDataSource.h"
 #include "Graph/MovieGraphPipeline.h"
 
 FMovieGraphInitConfig::FMovieGraphInitConfig()
 {
 	TimeStepClass = UMovieGraphLinearTimeStep::StaticClass();
 	RendererClass = UMovieGraphDefaultRenderer::StaticClass();
-	TimeRangeBuilderClass = UMovieGraphSequenceTimeRangeBuilder::StaticClass();
-	DataCachingClass = UMovieGraphDefaultDataCaching::StaticClass();
+	DataSourceClass = UMovieGraphSequenceDataSource::StaticClass();
 	bRenderViewport = false;
 }
 
@@ -26,12 +24,7 @@ UMovieGraphPipeline* UMovieGraphRendererBase::GetOwningGraph() const
 	return GetTypedOuter<UMovieGraphPipeline>();
 }
 
-UMovieGraphPipeline* UMovieGraphTimeRangeBuilderBase::GetOwningGraph() const
+UMovieGraphPipeline* UMovieGraphDataSourceBase::GetOwningGraph() const
 {
 	return GetTypedOuter<UMovieGraphPipeline>();
 }
-
-UMovieGraphPipeline* UMovieGraphDataCachingBase::GetOwningGraph() const
-{
-	return GetTypedOuter<UMovieGraphPipeline>();
-}
\ No newline at end of file
diff --git a/Engine/Plugins/MovieScene/MovieRenderPipeline/Source/MovieRenderPipelineCore/Private/Graph/MovieGraphDefaultDataCaching.cpp b/Engine/Plugins/MovieScene/MovieRenderPipeline/Source/MovieRenderPipelineCore/Private/Graph/MovieGraphDefaultDataCaching.cpp
deleted file mode 100644
diff --git a/Engine/Plugins/MovieScene/MovieRenderPipeline/Source/MovieRenderPipelineCore/Private/Graph/MovieGraphLinearTimeStep.cpp b/Engine/Plugins/MovieScene/MovieRenderPipeline/Source/MovieRenderPipelineCore/Private/Graph/MovieGraphLinearTimeStep.cpp
--- a/Engine/Plugins/MovieScene/MovieRenderPipeline/Source/MovieRenderPipelineCore/Private/Graph/MovieGraphLinearTimeStep.cpp
+++ b/Engine/Plugins/MovieScene/MovieRenderPipeline/Source/MovieRenderPipelineCore/Private/Graph/MovieGraphLinearTimeStep.cpp
@@ -6,6 +6,10 @@
 #include "MovieRenderPipelineCoreModule.h"
 #include "LevelSequence.h"
 #include "MovieScene.h"
+#include "Camera/PlayerCameraManager.h"
+#include "Camera/CameraComponent.h"
+#include "Camera/CameraActor.h"
+#include "Interfaces/Interface_PostProcessVolume.h"
 
 void UMovieGraphLinearTimeStep::TickProducingFrames()
 {
@@ -13,6 +17,38 @@ void UMovieGraphLinearTimeStep::TickProducingFrames()
 	const TArray<TObjectPtr<UMoviePipelineExecutorShot>>& ActiveShotList = GetOwningGraph()->GetActiveShotList();
 	UMoviePipelineExecutorShot* CurrentCameraCut = ActiveShotList[CurrentShotIndex];
 
+	// When start up we want to override the engine's Custom Timestep with our own.
+	// This gives us the ability to completely control the engine tick/delta time before the frame
+	// is started so that we don't have to always be thinking of delta times one frame ahead. We need
+	// to do this only once we're ready to set the timestep though, as Initialize can be called as
+	// a result of a OnBeginFrame, meaning that Initialize is called on the frame before TickProducingFrames
+	// so there would be one frame where it used the custom timestep (after initialize) before TPF was called.
+	//if (GEngine->GetCustomTimeStep() != CustomTimeStep)
+	//{
+	//	CachedPrevCustomTimeStep = GEngine->GetCustomTimeStep();
+	//	GEngine->SetCustomTimeStep(CustomTimeStep);
+	//}
+
+	// We cache the frame metrics for the duration of a single output frame, instead of recalculating them every tick.
+	// This is required for stochastic timesteps (which will jump around the actual evaluated time randomly within one
+	// output frame), but it helps resolve some of our issues related to time dilation. Historically when doing it every
+	// tick, we didn't know if we were going to overrun the end range of time represented until we had actually done it,
+	// partway through the frame, at which point we had to backtrack and abandon the work. So now we calculate it all
+	// in advance, check if our new time would overrun the end of the current camera cut, and if so, skip actually starting
+	// that frame. Unfortunately this does cause some complications with actor time dilations - the time dilation track
+	// will be updating each tick within the output sample and trying to change the world time dilation which we don't want.
+	// ToDo: Still tbd how we'll fix that, currently the plan is to allow custom clocksources to disable world time dilation
+	// support, at which point we'll handle the time dilation by affecting the whole engine tick instead of just world tick.
+	//if (IsFirstTemporalSample())
+	{
+		UpdateFrameMetrics();
+
+		// Now that we've calculated the total range of time we're trying to represent, we can check to see
+		// if this would put us beyond our range of time this shot is supposed to represent.
+
+	}
+
+
 	if (CurrentCameraCut->ShotInfo.State == EMovieRenderShotState::Uninitialized)
 	{
 		UE_LOG(LogMovieRenderPipeline, Log, TEXT("MovieGraph Initializing Camera Cut [%d/%d] in [%s] %s."),
@@ -20,7 +56,7 @@ void UMovieGraphLinearTimeStep::TickProducingFrames()
 
 		GetOwningGraph()->SetupShot(CurrentCameraCut);
 		// InitializeShot(CurrentCameraCut);
-		GetOwningGraph()->GetDataCachingInstance()->InitializeShot(CurrentCameraCut);
+		GetOwningGraph()->GetDataSourceInstance()->InitializeShot(CurrentCameraCut);
 
 		// We can safely fall through to the below states as they're OK to process the same frame we set up.
 		UE_LOG(LogMovieRenderPipeline, Log, TEXT("MovieGraph Finished initializing Camera Cut [%d/%d] in [%s] %s."),
@@ -61,3 +97,125 @@ void UMovieGraphLinearTimeStep::TickProducingFrames()
 		}
 	}
 }
+
+void UMovieGraphLinearTimeStep::UpdateFrameMetrics()
+{
+	FFrameData FrameData;
+
+	// ToDo: This needs to come from the graph
+	int32 TemporalSampleCount = 1;
+
+	// We inherit a tick resolution from the level sequence so that we can represent the same
+	// range of time that the level sequence does.
+	FrameData.TickResolution = GetOwningGraph()->GetDataSourceInstance()->GetTickResolution();
+	FrameData.FrameRate = GetOwningGraph()->GetDataSourceInstance()->GetDisplayRate(); // ToDo, needs to come from config (config can override)
+	FrameData.FrameTimePerOutputFrame = FFrameRate::TransformTime(FFrameTime(FFrameNumber(1)), FrameData.FrameRate, FrameData.TickResolution);
+
+	// Manually perform blending of the Post Process Volumes/Camera/Camera Modifiers to match what the renderer will do.
+	// This uses the primary camera specified by the PlayerCameraManager to get the motion blur amount so in the event of
+	// multi-camera rendering, all cameras will end up using the same motion blur amount defined by the primary camera).
+	FrameData.MotionBlurAmount = GetBlendedMotionBlurAmount();
+
+	// Calculate how long of a duration we want to represent where the camera shutter is open.
+	FrameData.FrameTimeWhileShutterOpen = FrameData.FrameTimePerOutputFrame * FrameData.MotionBlurAmount;
+
+	// Now that we know how long the shutter is open, figure out how long each temporal sub-sample gets.
+	FrameData.FrameTimePerTemporalSample = FrameData.FrameTimeWhileShutterOpen / TemporalSampleCount;
+
+	// The amount of time closed + time open should add up to exactly how long a frame is.
+	FrameData.FrameTimeWhileShutterClosed = FrameData.FrameTimePerOutputFrame - FrameData.FrameTimeWhileShutterOpen;
+
+	// Shutter timing is a bias applied to the final evaluation time to let us change what we consider a frame
+	// ie: Do we consider a frame the start of the timespan we captured? Or is the frame the end of the timespan?
+	// We default to Centered so that the center of your evaluated time is what you see in Level Sequences.
+	// ToDo: Shutter Timing Offset
+	//switch (CameraSettings->ShutterTiming)
+	//{
+	//	// Subtract the entire time the shutter is open.
+	//case EMoviePipelineShutterTiming::FrameClose:
+	//	Output.ShutterOffsetTicks = -Output.TicksWhileShutterOpen;
+	//	break;
+	//	// Only subtract half the time the shutter is open.
+	//case EMoviePipelineShutterTiming::FrameCenter:
+	FrameData.ShutterOffsetFrameTime = -FrameData.FrameTimeWhileShutterOpen / 2.0;
+	//	break;
+	//	// No offset needed
+	//case EMoviePipelineShutterTiming::FrameOpen:
+	//	break;
+	//}
+
+	// Then, calculate our motion blur offset. Motion Blur in the engine is always
+	// centered around the object so we offset our time sampling by half of the
+	// motion blur distance so that the distance blurred represents that time.
+	FrameData.MotionBlurCenteringOffsetTime = FrameData.FrameTimePerTemporalSample / 2.0;
+}
+
+float UMovieGraphLinearTimeStep::GetBlendedMotionBlurAmount()
+{
+	// 0.5f is the default engine motion blur in the event no Post Process/Camera overrides it.
+	float FinalMotionBlurAmount = 0.5f;
+
+	APlayerCameraManager* PlayerCameraManager = GetWorld()->GetFirstPlayerController()->PlayerCameraManager;
+	if (PlayerCameraManager)
+	{
+		// Apply any motion blur settings from post process volumes in the world
+		FVector ViewLocation = PlayerCameraManager->GetCameraLocation();
+		for (IInterface_PostProcessVolume* PPVolume : GetWorld()->PostProcessVolumes)
+		{
+			const FPostProcessVolumeProperties VolumeProperties = PPVolume->GetProperties();
+
+			// Skip any volumes which are either disabled or don't modify blur amount
+			if (!VolumeProperties.bIsEnabled || !VolumeProperties.Settings->bOverride_MotionBlurAmount)
+			{
+				continue;
+			}
+
+			float LocalWeight = FMath::Clamp(VolumeProperties.BlendWeight, 0.0f, 1.0f);
+
+			if (!VolumeProperties.bIsUnbound)
+			{
+				float DistanceToPoint = 0.0f;
+				PPVolume->EncompassesPoint(ViewLocation, 0.0f, &DistanceToPoint);
+
+				if (DistanceToPoint >= 0 && DistanceToPoint < VolumeProperties.BlendRadius)
+				{
+					LocalWeight *= FMath::Clamp(1.0f - DistanceToPoint / VolumeProperties.BlendRadius, 0.0f, 1.0f);
+				}
+				else
+				{
+					LocalWeight = 0.0f;
+				}
+			}
+
+			if (LocalWeight > 0.0f)
+			{
+				FinalMotionBlurAmount = FMath::Lerp(FinalMotionBlurAmount, VolumeProperties.Settings->MotionBlurAmount, LocalWeight);
+			}
+		}
+
+		// Now try from the camera, which takes priority over post processing volumes.
+		ACameraActor* CameraActor = Cast<ACameraActor>(PlayerCameraManager->GetViewTarget());
+		if (CameraActor)
+		{
+			UCameraComponent* CameraComponent = CameraActor->GetCameraComponent();
+			if (CameraComponent && CameraComponent->PostProcessSettings.bOverride_MotionBlurAmount)
+			{
+				FinalMotionBlurAmount = FMath::Lerp(FinalMotionBlurAmount, CameraComponent->PostProcessSettings.MotionBlurAmount, CameraComponent->PostProcessBlendWeight);
+			}
+		}
+
+		// Apply any motion blur settings from post processing blends attached to the camera manager
+		TArray<FPostProcessSettings> const* CameraAnimPPSettings;
+		TArray<float> const* CameraAnimPPBlendWeights;
+		PlayerCameraManager->GetCachedPostProcessBlends(CameraAnimPPSettings, CameraAnimPPBlendWeights);
+		for (int32 PPIdx = 0; PPIdx < CameraAnimPPBlendWeights->Num(); ++PPIdx)
+		{
+			if ((*CameraAnimPPSettings)[PPIdx].bOverride_MotionBlurAmount)
+			{
+				FinalMotionBlurAmount = FMath::Lerp(FinalMotionBlurAmount, (*CameraAnimPPSettings)[PPIdx].MotionBlurAmount, (*CameraAnimPPBlendWeights)[PPIdx]);
+			}
+		}
+	}
+
+	return FinalMotionBlurAmount;
+}
\ No newline at end of file
diff --git a/Engine/Plugins/MovieScene/MovieRenderPipeline/Source/MovieRenderPipelineCore/Private/Graph/MovieGraphPipeline.cpp b/Engine/Plugins/MovieScene/MovieRenderPipeline/Source/MovieRenderPipelineCore/Private/Graph/MovieGraphPipeline.cpp
--- a/Engine/Plugins/MovieScene/MovieRenderPipeline/Source/MovieRenderPipelineCore/Private/Graph/MovieGraphPipeline.cpp
+++ b/Engine/Plugins/MovieScene/MovieRenderPipeline/Source/MovieRenderPipelineCore/Private/Graph/MovieGraphPipeline.cpp
@@ -71,8 +71,7 @@ void UMovieGraphPipeline::Initialize(UMoviePipelineExecutorJob* InJob, const FMo
 	// Create instances of our different classes from the InitConfig
 	GraphTimeStepInstance = NewObject<UMovieGraphTimeStepBase>(this, InitConfig.TimeStepClass);
 	GraphRendererInstance = NewObject<UMovieGraphRendererBase>(this, InitConfig.RendererClass);
-	GraphTimeRangeBuilderInstance = NewObject<UMovieGraphTimeRangeBuilderBase>(this, InitConfig.TimeRangeBuilderClass);
-	GraphDataCachingInstance = NewObject<UMovieGraphDataCachingBase>(this, InitConfig.DataCachingClass);
+	GraphDataSourceInstance = NewObject<UMovieGraphDataSourceBase>(this, InitConfig.DataSourceClass);
 	
 	CurrentJob = InJob;
 	CurrentShotIndex = 0;
@@ -81,7 +80,7 @@ void UMovieGraphPipeline::Initialize(UMoviePipelineExecutorJob* InJob, const FMo
 	// Now that we've created our various systems, we will start using them. First thing we do is cache data about
 	// the world, job, player viewport, etc, before we make any modifications. These will be restored at the end
 	// of the render.
-	GraphDataCachingInstance->CacheDataPreJob(InitConfig);
+	GraphDataSourceInstance->CacheDataPreJob(InitConfig);
 
 	// Construct a debug UI and bind it to this instance.
 	// LoadDebugWidget();
@@ -127,7 +126,7 @@ void UMovieGraphPipeline::BuildShotListFromDataSource()
 {
 	// Synchronize our shot list with our target data source. New shots will be added and outdated shots removed.
 	// Shots that are already in the list will be updated but their enable flag will be respected. 
-	GraphDataCachingInstance->UpdateShotList();
+	GraphDataSourceInstance->UpdateShotList();
 
 	for (UMoviePipelineExecutorShot* Shot : GetCurrentJob()->ShotInfo)
 	{
diff --git a/Engine/Plugins/MovieScene/MovieRenderPipeline/Source/MovieRenderPipelineCore/Private/Graph/MovieGraphSequenceDataSource.cpp b/Engine/Plugins/MovieScene/MovieRenderPipeline/Source/MovieRenderPipelineCore/Private/Graph/MovieGraphSequenceDataSource.cpp
new file mode 100644
--- /dev/null
+++ b/Engine/Plugins/MovieScene/MovieRenderPipeline/Source/MovieRenderPipelineCore/Private/Graph/MovieGraphSequenceDataSource.cpp
@@ -0,0 +1,124 @@
+// Copyright Epic Games, Inc. All Rights Reserved.
+
+#include "Graph/MovieGraphSequenceDataSource.h"
+#include "Graph/MovieGraphPipeline.h"
+#include "LevelSequence.h"
+#include "LevelSequenceActor.h"
+#include "LevelSequencePlayer.h"
+#include "MoviePipelineQueue.h"
+#include "CoreGlobals.h"
+#include "EngineUtils.h"
+#include "MovieRenderPipelineCoreModule.h"
+#include "MoviePipelineBlueprintLibrary.h"
+#include "Engine/GameViewportClient.h"
+ 
+void UMovieGraphSequenceDataSource::CacheDataPreJob(const FMovieGraphInitConfig& InInitConfig)
+{
+	// Turn off screen messages as some forms are drawn directly to final render targets
+	// which will polute final frames.
+	GAreScreenMessagesEnabled = false;
+
+	// Turn off player viewport rendering to avoid overhead of an extra render.
+	if (UGameViewportClient* Viewport = GetWorld()->GetGameViewport())
+	{
+		Viewport->bDisableWorldRendering = !InInitConfig.bRenderViewport;
+	}
+
+	ULevelSequence* RootSequence = Cast<ULevelSequence>(GetOwningGraph()->GetCurrentJob()->Sequence.TryLoad());
+	if (RootSequence)
+	{
+		CacheLevelSequenceData(RootSequence);
+	}
+}
+
+void UMovieGraphSequenceDataSource::RestoreCachedDataPostJob()
+{
+	GAreScreenMessagesEnabled = true;
+	if (UGameViewportClient* Viewport = GetWorld()->GetGameViewport())
+	{
+		Viewport->bDisableWorldRendering = false;
+	}
+}
+
+void UMovieGraphSequenceDataSource::UpdateShotList()
+{
+	ULevelSequence* RootSequence = Cast<ULevelSequence>(GetOwningGraph()->GetCurrentJob()->Sequence.TryLoad());
+	if (RootSequence)
+	{
+		bool bShotsChanged = false;
+		UMoviePipelineBlueprintLibrary::UpdateJobShotListFromSequence(RootSequence, GetOwningGraph()->GetCurrentJob(), bShotsChanged);
+	}
+}
+
+void UMovieGraphSequenceDataSource::CacheLevelSequenceData(ULevelSequence* InSequence)
+{
+	// There is a reasonable chance that there exists a Level Sequence Actor in the world already set up to play this sequence.
+	ALevelSequenceActor* ExistingActor = nullptr;
+
+	for (auto It = TActorIterator<ALevelSequenceActor>(GetWorld()); It; ++It)
+	{
+		// Iterate through all of them in the event someone has multiple copies in the world on accident.
+		if (It->GetSequence() == InSequence)
+		{
+			// Found it!
+			ExistingActor = *It;
+
+			// Stop it from playing if it's already playing.
+			if (ExistingActor->GetSequencePlayer())
+			{
+				ExistingActor->GetSequencePlayer()->Stop();
+			}
+		}
+	}
+
+	LevelSequenceActor = ExistingActor;
+	if (!LevelSequenceActor)
+	{
+		// Spawn a new level sequence
+		LevelSequenceActor = GetWorld()->SpawnActor<ALevelSequenceActor>();
+		check(LevelSequenceActor);
+	}
+
+	// Enforce settings.
+	LevelSequenceActor->PlaybackSettings.LoopCount.Value = 0;
+	LevelSequenceActor->PlaybackSettings.bAutoPlay = false;
+	LevelSequenceActor->PlaybackSettings.bPauseAtEnd = true;
+	LevelSequenceActor->PlaybackSettings.bRestoreState = true;
+
+	// Ensure the (possibly new) Level Sequence Actor uses our sequence
+	LevelSequenceActor->SetSequence(InSequence);
+
+	// LevelSequenceActor->GetSequencePlayer()->SetTimeController(CustomSequenceTimeController);
+	LevelSequenceActor->GetSequencePlayer()->Stop();
+
+	LevelSequenceActor->GetSequencePlayer()->OnSequenceUpdated().AddUObject(this, &UMovieGraphSequenceDataSource::OnSequenceEvaluated);
+
+}
+
+void UMovieGraphSequenceDataSource::OnSequenceEvaluated(const UMovieSceneSequencePlayer& Player, FFrameTime CurrentTime, FFrameTime PreviousTime)
+{
+	// This callback exists for logging purposes. DO NOT HINGE LOGIC ON THIS CALLBACK
+	// because this may get called multiple times per frame and may be the result of
+	// a seek operation which is reverted before a frame is even rendered.
+	UE_LOG(LogMovieRenderPipeline, VeryVerbose, TEXT("Sequence Evaluated. CurrentTime: %s PreviousTime: %s"), *LexToString(CurrentTime), *LexToString(PreviousTime));
+}
+
+FFrameRate UMovieGraphSequenceDataSource::GetTickResolution() const
+{
+	if (ensure(LevelSequenceActor && LevelSequenceActor->GetSequence()))
+	{
+		return LevelSequenceActor->GetSequence()->GetMovieScene()->GetTickResolution();
+	}
+
+	return FFrameRate(24000, 1);
+}
+
+FFrameRate UMovieGraphSequenceDataSource::GetDisplayRate() const
+{
+	if (ensure(LevelSequenceActor && LevelSequenceActor->GetSequence()))
+	{
+		return LevelSequenceActor->GetSequence()->GetMovieScene()->GetDisplayRate();
+	}
+
+	return FFrameRate(24, 1);
+}
\ No newline at end of file
diff --git a/Engine/Plugins/MovieScene/MovieRenderPipeline/Source/MovieRenderPipelineCore/Public/Graph/MovieGraphDataTypes.h b/Engine/Plugins/MovieScene/MovieRenderPipeline/Source/MovieRenderPipelineCore/Public/Graph/MovieGraphDataTypes.h
--- a/Engine/Plugins/MovieScene/MovieRenderPipeline/Source/MovieRenderPipelineCore/Public/Graph/MovieGraphDataTypes.h
+++ b/Engine/Plugins/MovieScene/MovieRenderPipeline/Source/MovieRenderPipelineCore/Public/Graph/MovieGraphDataTypes.h
@@ -7,6 +7,7 @@
 #include "Templates/SubclassOf.h"
 #include "ImagePixelData.h"
 #include "Containers/Queue.h"
+#include "Misc/FrameRate.h"
 #include "MovieGraphDataTypes.generated.h"
 
 // Forward Declares
@@ -45,11 +46,13 @@ struct MOVIERENDERPIPELINECORE_API FMovieGraphInitConfig
 	UPROPERTY(BlueprintReadWrite, EditAnywhere, Category = "Movie Graph")
 	TSubclassOf<UMovieGraphRendererBase> RendererClass;
 
+	/**
+	* Which class should the UMovieGraphPipeline use to build time ranges from, and
+	* during evaluation, send callbacks about the time actually evaluated so you
+	* can sync with an external source. Defaults to UMovieGraphSequenceDataSource.
+	*/
 	UPROPERTY(BlueprintReadWrite, EditAnywhere, Category = "Movie Graph")
-	TSubclassOf<UMovieGraphTimeRangeBuilderBase> TimeRangeBuilderClass;
-
-	UPROPERTY(BlueprintReadWrite, EditAnywhere, Category = "Movie Graph")
-	TSubclassOf<UMovieGraphDataCachingBase> DataCachingClass;
+	TSubclassOf<UMovieGraphDataSourceBase> DataSourceClass;
 
 	/**
 	* Should the UMovieGraphPipeline render the full player viewport? Defaults
@@ -111,20 +114,28 @@ public:
 	UMovieGraphPipeline* GetOwningGraph() const;
 };
 
+/**
+* Movie Graph Pipeline is mostly interested in knowing about ranges of time that
+* it should render, and less concerned with the specifics of where that data comes
+* from (ie: a Level Sequence). This lets you synchronize with a different data source
+* to provide the ranges of time to render, and then the UMovieGraphTimeStepBase class
+* figures out how to move around within that time step, before calling some functions
+* to synchronize your external data source to actually match the evaluated time.
+*/
 UCLASS(BlueprintType, Abstract)
-class MOVIERENDERPIPELINECORE_API UMovieGraphTimeRangeBuilderBase : public UObject
+class MOVIERENDERPIPELINECORE_API UMovieGraphDataSourceBase : public UObject
 {
 	GENERATED_BODY()
 public:
+	/** An internal, high resolution framerate that seeks, etc. will be returned in. (ie: 24,000/1) */
+	virtual FFrameRate GetTickResolution() const { return FFrameRate(); }
+	/** A lower res, human readable Frame Rate. We convert to the Tick Resolution internally. (ie: 24/1) */
+	virtual FFrameRate GetDisplayRate() const { return FFrameRate(); }
 	virtual void BuildTimeRanges() { }
-	UMovieGraphPipeline* GetOwningGraph() const;
-};
-
-UCLASS(BlueprintType, Abstract)
-class MOVIERENDERPIPELINECORE_API UMovieGraphDataCachingBase : public UObject
-{
-	GENERATED_BODY()
-public:
+	/** 
+	* Called when the Movie Graph Pipeline starts before anything has happened, allowing you to 
+	* cache your datasource before making any modifications to it as a result of rendering.
+	*/
 	virtual void CacheDataPreJob(const FMovieGraphInitConfig& InInitConfig) {}
 	virtual void RestoreCachedDataPostJob() {}
 	virtual void UpdateShotList() {}
diff --git a/Engine/Plugins/MovieScene/MovieRenderPipeline/Source/MovieRenderPipelineCore/Public/Graph/MovieGraphDefaultDataCaching.h b/Engine/Plugins/MovieScene/MovieRenderPipeline/Source/MovieRenderPipelineCore/Public/Graph/MovieGraphDefaultDataCaching.h
deleted file mode 100644
diff --git a/Engine/Plugins/MovieScene/MovieRenderPipeline/Source/MovieRenderPipelineCore/Public/Graph/MovieGraphLinearTimeStep.h b/Engine/Plugins/MovieScene/MovieRenderPipeline/Source/MovieRenderPipelineCore/Public/Graph/MovieGraphLinearTimeStep.h
--- a/Engine/Plugins/MovieScene/MovieRenderPipeline/Source/MovieRenderPipelineCore/Public/Graph/MovieGraphLinearTimeStep.h
+++ b/Engine/Plugins/MovieScene/MovieRenderPipeline/Source/MovieRenderPipelineCore/Public/Graph/MovieGraphLinearTimeStep.h
@@ -2,10 +2,19 @@
 #pragma once
 
 #include "Graph/MovieGraphDataTypes.h"
+#include "Misc/FrameRate.h"
+#include "Misc/FrameTime.h"
 #include "MovieGraphLinearTimeStep.generated.h"
 
 /**
-* This class 
+* This class is responsible for calculating the time step of each tick of the engine during a
+* Movie Render Queue render using a linear strategy - the number of temporal sub-samples is
+* read from the graph for each output frame, and we then take the time the shutter is open
+* and break it into that many sub-samples. Time always advances forward until we reach the
+* end of the range of time we wish to render. This is useful for deferred rendering (where 
+* we have a small number of temporal sub-samples and no feedback mechanism for measuring
+* noise in the final image) but is less useful for Path Traced images which have a varying
+* amount of noise (and thus want a varying amount of samples) based on their content.
 */
 UCLASS(BlueprintType)
 class MOVIERENDERPIPELINECORE_API UMovieGraphLinearTimeStep : public UMovieGraphTimeStepBase
@@ -15,6 +24,55 @@ public:
 	virtual void TickProducingFrames() override;
 	virtual FMovieGraphTimeStepData GetCalculatedTimeData() const override { return CurrentTimeStepData; }
 
-private:
+protected:
+	virtual void UpdateFrameMetrics();
+	virtual float GetBlendedMotionBlurAmount();
+
+protected:
+	/** This is the output data needed by the rest of MRQ to produce a frame. */
 	FMovieGraphTimeStepData CurrentTimeStepData;
-};
\ No newline at end of file
+
+	struct FFrameData
+	{
+		FFrameData()
+			: MotionBlurAmount(0.f)
+		{}
+		
+		/** The human readable frame rate, adjusted by the config file (24fps, 30fps, etc.) */
+		FFrameRate FrameRate;
+
+		/** The internal resolution data is actually stored at. (24,000, 120,000 etc.) */
+		FFrameRate TickResolution;
+
+		/** The amount of time (in Tick Resolution) that represents one output frame (ie: 1/24). */
+		FFrameTime FrameTimePerOutputFrame;
+
+		/** The 0-1 Motion Blur amount the camera has. */
+		float MotionBlurAmount;
+
+		/** The amount of time (in Tick Resolution) that the shutter is open for. Could be <= FrameTimePerOutputFrame. */
+		FFrameTime FrameTimeWhileShutterOpen;
+
+		/** 
+		* The amount of time (in Tick Resolution) per temporal sample. The actual time slice for a temporal sub-sample
+		* may be different than this (as we have to jump over the time period that the shutter is closed too for one 
+		* temporal sub-sample.
+		*/
+		FFrameTime FrameTimePerTemporalSample;
+		
+		/** The inverse of FrameTimeWhileShutterOpen. ShutterClosed+ShutterOpen should add up to FrameTimePerOutputFrame. */
+		FFrameTime FrameTimeWhileShutterClosed;
+
+		/** A constant offset applied to the final evaluated time to allow users to influence what counts as a frame. (the time before, during, or after a frame) */
+		FFrameTime ShutterOffsetFrameTime;
+
+		/** An offset to the evaluated time to get us into the center of the time period (since motion blur is bi-directional and centered) */
+		FFrameTime MotionBlurCenteringOffsetTime;
+	};
+
+	/**
+	* A set of cached values that are true for the current output frame. They will be recalculated
+	* at the start of the next output frame.
+	*/
+	FFrameData CurrentFrameMetrics;
+};
diff --git a/Engine/Plugins/MovieScene/MovieRenderPipeline/Source/MovieRenderPipelineCore/Public/Graph/MovieGraphPipeline.h b/Engine/Plugins/MovieScene/MovieRenderPipeline/Source/MovieRenderPipelineCore/Public/Graph/MovieGraphPipeline.h
--- a/Engine/Plugins/MovieScene/MovieRenderPipeline/Source/MovieRenderPipelineCore/Public/Graph/MovieGraphPipeline.h
+++ b/Engine/Plugins/MovieScene/MovieRenderPipeline/Source/MovieRenderPipelineCore/Public/Graph/MovieGraphPipeline.h
@@ -50,9 +50,7 @@ public:
 	/** Used occasionally to cross-reference other components. Don't call this unless you know what you're doing. */
 	UMovieGraphRendererBase* GetRendererInstance() const { return GraphRendererInstance; }
 	/** Used occasionally to cross-reference other components. Don't call this unless you know what you're doing. */
-	UMovieGraphTimeRangeBuilderBase* GetTimeRangeBuilderInstance() const { return GraphTimeRangeBuilderInstance; }
-	/** Used occasionally to cross-reference other components. Don't call this unless you know what you're doing. */
-	UMovieGraphDataCachingBase* GetDataCachingInstance() const { return GraphDataCachingInstance; }
+	UMovieGraphDataSourceBase* GetDataSourceInstance() const { return GraphDataSourceInstance; }
 	/** Gets the Output Merger for this Movie Pipeline which is responsible for gathering all of the data coming in for a given output frame, before making it available to the MovieGraphPipeline. */
 	TSharedPtr<UE::MovieGraph::IMovieGraphOutputMerger> GetOutputMerger() const { return OutputMerger; }
 	/** Writing images to disk is an async process. When you start writing, declare a future with the filename you will eventually write to, and complete the future once it is on disk. */
@@ -91,10 +89,7 @@ protected:
 	TObjectPtr<UMovieGraphRendererBase> GraphRendererInstance;
 
 	UPROPERTY(Transient)
-	TObjectPtr<UMovieGraphTimeRangeBuilderBase> GraphTimeRangeBuilderInstance;
-
-	UPROPERTY(Transient)
-	TObjectPtr<UMovieGraphDataCachingBase> GraphDataCachingInstance;
+	TObjectPtr<UMovieGraphDataSourceBase> GraphDataSourceInstance;
 
 protected:
 	UPROPERTY(Transient)
diff --git a/Engine/Plugins/MovieScene/MovieRenderPipeline/Source/MovieRenderPipelineCore/Public/Graph/MovieGraphSequenceDataSource.h b/Engine/Plugins/MovieScene/MovieRenderPipeline/Source/MovieRenderPipelineCore/Public/Graph/MovieGraphSequenceDataSource.h
new file mode 100644
--- /dev/null
+++ b/Engine/Plugins/MovieScene/MovieRenderPipeline/Source/MovieRenderPipelineCore/Public/Graph/MovieGraphSequenceDataSource.h
@@ -0,0 +1,39 @@
+// Copyright Epic Games, Inc. All Rights Reserved.
+#pragma once
+
+#include "Graph/MovieGraphDataTypes.h"
+#include "Misc/FrameTime.h"
+#include "MovieGraphSequenceDataSource.generated.h"
+
+// Forward Declares
+class ULevelSequence;
+class ALevelSequenceActor;
+class UMovieSceneSequencePlayer;
+
+/**
+* The UMovieGraphSequenceDataSource allows using a ULevelSequence as the external datasource for the Movie Graph.
+* It will build the range of shots based on the contents of the level sequence (one shot per camera cut found inside
+* the sequence hierarchy, not allowing overlapping Camera Cut sections), and then it will evaluate the level sequence
+* for the given time when rendering.
+*/
+UCLASS(BlueprintType)
+class MOVIERENDERPIPELINECORE_API UMovieGraphSequenceDataSource : public UMovieGraphDataSourceBase
+{
+	GENERATED_BODY()
+
+public:
+	virtual void CacheDataPreJob(const FMovieGraphInitConfig& InInitConfig) override;
+	virtual void RestoreCachedDataPostJob() override;
+	virtual void UpdateShotList() override;
+	virtual FFrameRate GetTickResolution() const override;
+	virtual FFrameRate GetDisplayRate() const override;
+
+protected:
+	void CacheLevelSequenceData(ULevelSequence* InSequence);
+	void OnSequenceEvaluated(const UMovieSceneSequencePlayer& Player, FFrameTime CurrentTime, FFrameTime PreviousTime);
+
+protected:
+	UPROPERTY(Transient)
+	TObjectPtr<ALevelSequenceActor> LevelSequenceActor;
+
+};
\ No newline at end of file
diff --git a/Engine/Plugins/MovieScene/MovieRenderPipeline/Source/MovieRenderPipelineCore/Public/Graph/MovieGraphSequenceTimeRangeBuilder.h b/Engine/Plugins/MovieScene/MovieRenderPipeline/Source/MovieRenderPipelineCore/Public/Graph/MovieGraphSequenceTimeRangeBuilder.h
deleted file mode 100644
