MoviePipeline: Exposed global variables to the graph, which are available in every graph by default and are not deletable or editable.

diff --git a/Engine/Plugins/MovieScene/MovieRenderPipeline/Source/MovieRenderPipelineCore/Private/Graph/MovieGraphConfig.cpp b/Engine/Plugins/MovieScene/MovieRenderPipeline/Source/MovieRenderPipelineCore/Private/Graph/MovieGraphConfig.cpp
--- a/Engine/Plugins/MovieScene/MovieRenderPipeline/Source/MovieRenderPipelineCore/Private/Graph/MovieGraphConfig.cpp
+++ b/Engine/Plugins/MovieScene/MovieRenderPipeline/Source/MovieRenderPipelineCore/Private/Graph/MovieGraphConfig.cpp
@@ -53,24 +53,37 @@ void UMovieGraphOutput::PostEditChangeProperty(FPropertyChangedEvent& PropertyCh
 }
 #endif // WITH_EDITOR
 
+FName UMovieGraphConfig::GlobalVariable_ShotName = "shot_name";
+FName UMovieGraphConfig::GlobalVariable_SequenceName = "seq_name";
+FName UMovieGraphConfig::GlobalVariable_FrameNumber = "frame_num";
+FName UMovieGraphConfig::GlobalVariable_CameraName = "camera_name";
+FName UMovieGraphConfig::GlobalVariable_RenderLayerName = "render_layer_name";
+
 UMovieGraphConfig::UMovieGraphConfig()
 {
 	InputNode = CreateDefaultSubobject<UMovieGraphInputNode>(TEXT("DefaultInputNode"));
 	OutputNode = CreateDefaultSubobject<UMovieGraphOutputNode>(TEXT("DefaultOutputNode"));
-	
-	if (!HasAnyFlags(RF_ClassDefaultObject))
+
+	// Don't add default members in the ctor if this object is being loaded (ie, it's not a new object). Defer that
+	// until PostLoad(), otherwise the default members may be overwritten when properties are loaded.
+	const bool bIsNewObject = !HasAnyFlags(RF_ClassDefaultObject | RF_NeedLoad | RF_NeedPostLoad);
+	if (bIsNewObject)
 	{
+		AddDefaultMembers();
 		InputNode->UpdatePins();
 		OutputNode->UpdatePins();
 	}
 }
 
-void UMovieGraphConfig::PostInitProperties()
+void UMovieGraphConfig::PostLoad()
 {
-	Super::PostInitProperties();
+	Super::PostLoad();
 
 	if (!HasAnyFlags(RF_ClassDefaultObject))
 	{
+		// TODO: When the graph has stabilized, we can remove this and replace with a system that solely performs
+		// upgrades/deprecations. For now, we assume that each load of the graph should re-initialize all default
+		// members.
 		AddDefaultMembers();
 	}
 }
@@ -119,6 +132,29 @@ void UMovieGraphConfig::TraverseGraphRecursive(UMovieGraphNode* InNode, TSubclas
 
 }
 
+UMovieGraphVariable* UMovieGraphConfig::AddGlobalVariable(const FName& InName)
+{
+	// Don't add duplicate global variables
+	const bool VariableExists = Variables.ContainsByPredicate([&InName](const TObjectPtr<UMovieGraphVariable>& Variable)
+	{
+		return Variable && (Variable->Name == InName);
+	});
+
+	if (VariableExists)
+	{
+		return nullptr;
+	}
+	
+	if (UMovieGraphVariable* NewVariable = AddVariable(InName))
+	{
+		NewVariable->bIsGlobal = true;
+		NewVariable->bIsEditable = false;
+		return NewVariable;
+	}
+
+	return nullptr;
+}
+
 void UMovieGraphConfig::AddDefaultMembers()
 {
 	const bool InputGlobalsExists = Inputs.ContainsByPredicate([](const UMovieGraphMember* Member)
@@ -148,6 +184,15 @@ void UMovieGraphConfig::AddDefaultMembers()
 
 		OutputNode->UpdatePins();
 	}
+
+	// Add all of the global variables that should be available in the graph
+	static const TArray<FName> GlobalVariableNames =
+		{GlobalVariable_ShotName, GlobalVariable_SequenceName, GlobalVariable_FrameNumber,
+		 GlobalVariable_CameraName, GlobalVariable_RenderLayerName};
+	for (const FName& GlobalVariableName : GlobalVariableNames)
+	{
+		AddGlobalVariable(GlobalVariableName);
+	}
 }
 
 TArray<UMovieGraphNode*> UMovieGraphConfig::TraverseGraph(TSubclassOf<UMovieGraphNode> InClassType, const FMovieGraphTraversalContext& InContext) const
@@ -354,31 +399,48 @@ bool UMovieGraphConfig::RemoveNode(UMovieGraphNode* InNode)
 }
 
 template<typename T>
-T* UMovieGraphConfig::AddMember(TArray<TObjectPtr<T>>& MemberArray, const FText& BaseName)
+T* UMovieGraphConfig::AddMember(TArray<TObjectPtr<T>>& InMemberArray, const FName& InBaseName)
 {
 	static_assert(std::is_base_of_v<UMovieGraphMember, T>, "T is not derived from UMovieGraphMember");
 	
 	using namespace UE::MoviePipeline::RenderGraph;
+
+	// TODO: This can be replaced with just CreateDefaultSubobject() when AddDefaultMembers() isn't called from PostLoad()
+	//
+	// This method will be called in two cases: 1) when default members are being added to a new graph when it is being
+	// initially created or loaded via PostLoad(), or 2) a member is being added to the graph by the user. For case 1,
+	// when the constructor is running, RF_NeedInitialization will be set. CreateDefaultSubobject() needs to be called
+	// in this scenario instead of NewObject().
+	const bool bIsNewObject = HasAnyFlags(RF_NeedInitialization);
+	T* NewMember = bIsNewObject
+		? CreateDefaultSubobject<T>(MakeUniqueObjectName(this, T::StaticClass()))
+		: NewObject<T>(this, NAME_None);
 	
-	T* NewMember = NewObject<T>(this, NAME_None, RF_Transactional);
-	MemberArray.Add(NewMember);
+	if (!NewMember)
+	{
+		UE_LOG(LogMovieRenderPipeline, Warning, TEXT("Unable to create new member object in the graph."));
+		return nullptr;
+	}
 
+	InMemberArray.Add(NewMember);
+	NewMember->SetFlags(RF_Transactional);
 	NewMember->Type = EMovieGraphMemberType::Float;
 	NewMember->SetGuid(FGuid::NewGuid());
 
 	// Generate and set a unique name
 	TArray<FString> ExistingMemberNames;
-	Algo::Transform(MemberArray, ExistingMemberNames, [](const T* Member) { return Member->Name; });
-	NewMember->Name = GetUniqueName(ExistingMemberNames, BaseName.ToString());
+	Algo::Transform(InMemberArray, ExistingMemberNames, [](const T* Member) { return Member->Name; });
+	NewMember->Name = GetUniqueName(ExistingMemberNames, InBaseName.ToString());
 
 	return NewMember;
 }
 
-UMovieGraphVariable* UMovieGraphConfig::AddVariable()
+UMovieGraphVariable* UMovieGraphConfig::AddVariable(const FName InCustomBaseName)
 {
 	static const FText VariableBaseName = LOCTEXT("VariableBaseName", "Variable");
 	
-	UMovieGraphVariable* NewVariable = AddMember(Variables, VariableBaseName);
+	UMovieGraphVariable* NewVariable = AddMember(
+		Variables, !InCustomBaseName.IsNone() ? InCustomBaseName : FName(*VariableBaseName.ToString()));
 
 #if WITH_EDITOR
 	OnGraphVariablesChangedDelegate.Broadcast();
@@ -391,7 +453,7 @@ UMovieGraphInput* UMovieGraphConfig::AddInput()
 {
 	static const FText InputBaseName = LOCTEXT("InputBaseName", "Input");
 
-	UMovieGraphInput* NewInput = AddMember(Inputs, InputBaseName);
+	UMovieGraphInput* NewInput = AddMember(Inputs, FName(*InputBaseName.ToString()));
 	InputNode->UpdatePins();
 	
 #if WITH_EDITOR
@@ -405,7 +467,7 @@ UMovieGraphOutput* UMovieGraphConfig::AddOutput()
 {
 	static const FText OutputBaseName = LOCTEXT("OutputBaseName", "Output");
 	
-	UMovieGraphOutput* NewOutput = AddMember(Outputs, OutputBaseName);
+	UMovieGraphOutput* NewOutput = AddMember(Outputs, FName(*OutputBaseName.ToString()));
 	OutputNode->UpdatePins();
 
 #if WITH_EDITOR
@@ -428,9 +490,14 @@ UMovieGraphVariable* UMovieGraphConfig::GetVariableByGuid(const FGuid& InGuid) c
 	return nullptr;
 }
 
-TArray<UMovieGraphVariable*> UMovieGraphConfig::GetVariables() const
+TArray<UMovieGraphVariable*> UMovieGraphConfig::GetVariables(const bool bIncludeGlobal) const
 {
-	return Variables;
+	if (bIncludeGlobal)
+	{
+		return Variables;
+	}
+
+	return Variables.FilterByPredicate([](const UMovieGraphVariable* Var) { return Var && !Var->IsGlobal(); });
 }
 
 TArray<UMovieGraphInput*> UMovieGraphConfig::GetInputs() const
diff --git a/Engine/Plugins/MovieScene/MovieRenderPipeline/Source/MovieRenderPipelineCore/Public/Graph/MovieGraphConfig.h b/Engine/Plugins/MovieScene/MovieRenderPipeline/Source/MovieRenderPipelineCore/Public/Graph/MovieGraphConfig.h
--- a/Engine/Plugins/MovieScene/MovieRenderPipeline/Source/MovieRenderPipelineCore/Public/Graph/MovieGraphConfig.h
+++ b/Engine/Plugins/MovieScene/MovieRenderPipeline/Source/MovieRenderPipelineCore/Public/Graph/MovieGraphConfig.h
@@ -20,6 +20,9 @@ class UMovieGraphNode;
 UCLASS(Abstract)
 class MOVIERENDERPIPELINECORE_API UMovieGraphMember : public UObject
 {
+	// The graph needs to set private flags during construction time
+	friend class UMovieGraphConfig;
+	
 	GENERATED_BODY()
 
 public:
@@ -36,42 +39,59 @@ public:
 
 public:
 	/** The type of data associated with this member. */
-	UPROPERTY(EditAnywhere, Category = "General")
+	UPROPERTY(EditAnywhere, Category = "General", meta=(EditCondition="bIsEditable", HideEditConditionToggle))
 	EMovieGraphMemberType Type = EMovieGraphMemberType::Float;
 
 	// TODO: Need a details customization that validates whether or not the name is valid/unique
 	/** The name of this member, which is user-facing. */
-	UPROPERTY(EditAnywhere, Category = "General")
+	UPROPERTY(EditAnywhere, Category = "General", meta=(EditCondition="bIsEditable", HideEditConditionToggle))
 	FString Name;
 
 	/** The optional description of this member, which is user-facing. */
-	UPROPERTY(EditAnywhere, Category = "General")
+	UPROPERTY(EditAnywhere, Category = "General", meta=(EditCondition="bIsEditable", HideEditConditionToggle))
 	FString Description;
 
 	// TODO: This needs to eventually be capable of storing multiple types. Using TVariant looks promising, but it will
 	// still require extensive details panel customizations
 	/** The default value of this member. */
-	UPROPERTY(EditAnywhere, Category = "General")
+	UPROPERTY(EditAnywhere, Category = "General", meta=(EditCondition="bIsEditable", HideEditConditionToggle))
 	float Default = 0.f;
 
 private:
 	/** A GUID that uniquely identifies this member within its graph. */
 	UPROPERTY()
 	FGuid Guid;
+
+	// Note: This is a bool flag rather than a method (eg, IsEditable()) for now in order to allow it to drive the
+	// EditCondition metadata on properties.
+	/** Whether this member can be edited in the UI. */
+	UPROPERTY()
+	bool bIsEditable = true;
 };
 
 /**
- * A variable that can be created and used inside the graph. These variables can be controlled at
- * the job level.
+ * A variable that can be used inside the graph. Most variables are created by the user, and can have their value
+ * changed at the job level. Global variables, however, are not user-created and their values are provided when the
+ * graph is evaluated. Overriding them at the job level is not possible.
  */
 UCLASS(BlueprintType)
 class MOVIERENDERPIPELINECORE_API UMovieGraphVariable : public UMovieGraphMember
 {
+	// The graph needs to set private flags during construction time
+	friend class UMovieGraphConfig;
+	
 	GENERATED_BODY()
 
 public:
 	UMovieGraphVariable() = default;
 
+	/** Returns true if this variable is a global variable. */
+	bool IsGlobal() const { return bIsGlobal; }
+
+	//~ Begin UMovieGraphMember interface
+	virtual bool IsDeletable() const override { return !bIsGlobal; }
+	//~ End UMovieGraphMember interface
+
 public:
 #if WITH_EDITOR
 	FOnMovieGraphVariableChanged OnMovieGraphVariableChangedDelegate;
@@ -80,6 +100,11 @@ public:
 	virtual void PostEditChangeProperty(FPropertyChangedEvent& PropertyChangedEvent) override;
 	//~ End UObject overrides
 #endif // WITH_EDITOR
+
+private:
+	/** Whether this variable represents a global variable. */
+	UPROPERTY()
+	bool bIsGlobal = false;
 };
 
 /**
@@ -150,7 +175,7 @@ public:
 	UMovieGraphConfig();
 
 	//~ UObject interface
-	virtual void PostInitProperties() override;
+	virtual void PostLoad() override;
 	//~ End UObject interface
 
 	bool AddLabeledEdge(UMovieGraphNode* FromNode, const FName& FromPinLabel, UMovieGraphNode* ToNode, const FName& ToPinLabel);
@@ -168,8 +193,11 @@ public:
 	UMovieGraphNode* GetOutputNode() const { return OutputNode; }
 	const TArray<TObjectPtr<UMovieGraphNode>>& GetNodes() const { return AllNodes; }
 
-	/** Adds a new variable member with default values to the graph. Returns the new variable on success, else nullptr. */
-	UMovieGraphVariable* AddVariable();
+	/**
+	 * Adds a new variable member with default values to the graph. The new variable will have a base name of
+	 * "Variable" unless specified in InCustomBaseName. Returns the new variable on success, else nullptr.
+	 */
+	UMovieGraphVariable* AddVariable(const FName InCustomBaseName = NAME_None);
 
 	/** Adds a new input member to the graph. Returns the new input on success, else nullptr. */
 	UMovieGraphInput* AddInput();
@@ -180,8 +208,11 @@ public:
 	/** Gets the variable in the graph with the specified GUID, else nullptr if one could not be found. */
 	UMovieGraphVariable* GetVariableByGuid(const FGuid& InGuid) const;
 
-	/** Gets all variables that have been defined on the graph. */
-	TArray<UMovieGraphVariable*> GetVariables() const;
+	/**
+	 * Gets all variables that are available to be used in the graph. Global variables can optionally be included if
+	 * bIncludeGlobal is set to true.
+	 */
+	TArray<UMovieGraphVariable*> GetVariables(const bool bIncludeGlobal = false) const;
 
 	/** Gets all inputs that have been defined on the graph. */
 	TArray<UMovieGraphInput*> GetInputs() const;
@@ -238,6 +269,13 @@ protected:
 	void TraverseGraphRecursive(UMovieGraphNode* InNode, TSubclassOf<UMovieGraphNode> InClassType, const FMovieGraphTraversalContext& InContext, TArray<UMovieGraphNode*>& OutNodes) const;
 
 public:
+	// Names of global variables that are provided by the graph
+	static FName GlobalVariable_ShotName;
+	static FName GlobalVariable_SequenceName;
+	static FName GlobalVariable_FrameNumber;
+	static FName GlobalVariable_CameraName;
+	static FName GlobalVariable_RenderLayerName;
+	
 #if WITH_EDITOR
 	FOnMovieGraphChanged OnGraphChangedDelegate;
 	FOnMovieGraphVariablesChanged OnGraphVariablesChangedDelegate;
@@ -290,13 +328,16 @@ private:
 	
 	/** Add a new member of type T to MemberArray, with a unique name that includes BaseName in it. */
 	template<typename T>
-	T* AddMember(TArray<TObjectPtr<T>>& MemberArray, const FText& BaseName);
+	T* AddMember(TArray<TObjectPtr<T>>& InMemberArray, const FName& InBaseName);
+
+	/** Adds a global variable to the graph with the provided name. */
+	UMovieGraphVariable* AddGlobalVariable(const FName& InName);
 
 	/** Adds members to the graph that should always be available. */
 	void AddDefaultMembers();
 
 private:
-	/** All variables which have been defined on the graph. */
+	/** All variables (user and global) which are available for use in the graph. */
 	UPROPERTY()
 	TArray<TObjectPtr<UMovieGraphVariable>> Variables;
 
diff --git a/Engine/Plugins/MovieScene/MovieRenderPipeline/Source/MovieRenderPipelineCore/Public/Graph/Nodes/MovieGraphVariableNode.h b/Engine/Plugins/MovieScene/MovieRenderPipeline/Source/MovieRenderPipelineCore/Public/Graph/Nodes/MovieGraphVariableNode.h
--- a/Engine/Plugins/MovieScene/MovieRenderPipeline/Source/MovieRenderPipelineCore/Public/Graph/Nodes/MovieGraphVariableNode.h
+++ b/Engine/Plugins/MovieScene/MovieRenderPipeline/Source/MovieRenderPipelineCore/Public/Graph/Nodes/MovieGraphVariableNode.h
@@ -2,6 +2,7 @@
 
 #pragma once
 
+#include "Graph/MovieGraphConfig.h"
 #include "Graph/MovieGraphNode.h"
 
 #include "MovieGraphVariableNode.generated.h"
@@ -23,6 +24,9 @@ public:
 	/** Sets the variable that this node represents. */
 	void SetVariable(UMovieGraphVariable* InVariable);
 
+	/** Returns true if this node represents a global variable, else false. */
+	bool IsGlobalVariable() const { return GraphVariable && GraphVariable->IsGlobal(); }
+
 #if WITH_EDITOR
 	virtual FText GetNodeTitle(const bool bGetDescriptive = false) const override;
 	virtual FText GetMenuCategory() const override;
diff --git a/Engine/Plugins/MovieScene/MovieRenderPipeline/Source/MovieRenderPipelineEditor/Private/Graph/MovieEdGraphVariableNode.cpp b/Engine/Plugins/MovieScene/MovieRenderPipeline/Source/MovieRenderPipelineEditor/Private/Graph/MovieEdGraphVariableNode.cpp
--- a/Engine/Plugins/MovieScene/MovieRenderPipeline/Source/MovieRenderPipelineEditor/Private/Graph/MovieEdGraphVariableNode.cpp
+++ b/Engine/Plugins/MovieScene/MovieRenderPipeline/Source/MovieRenderPipelineEditor/Private/Graph/MovieEdGraphVariableNode.cpp
@@ -10,7 +10,12 @@
 
 FText UMoviePipelineEdGraphVariableNode::GetNodeTitle(ENodeTitleType::Type TitleType) const
 {
-	return LOCTEXT("GetVariableNodeTitle", "Get Variable");
+	static const FText VariableNodeTitle = LOCTEXT("GetVariableNodeTitle", "Get Variable");
+	static const FText GlobalVariableNodeTitle = LOCTEXT("GetGlobalVariableNodeTitle", "Get Global Variable");
+
+	const UMovieGraphVariableNode* VariableNode = Cast<UMovieGraphVariableNode>(RuntimeNode);
+
+	return (VariableNode && VariableNode->IsGlobalVariable()) ? GlobalVariableNodeTitle : VariableNodeTitle;
 }
 
 void UMoviePipelineEdGraphVariableNode::AllocateDefaultPins()
diff --git a/Engine/Plugins/MovieScene/MovieRenderPipeline/Source/MovieRenderPipelineEditor/Private/Graph/MovieGraphSchema.cpp b/Engine/Plugins/MovieScene/MovieRenderPipeline/Source/MovieRenderPipelineEditor/Private/Graph/MovieGraphSchema.cpp
--- a/Engine/Plugins/MovieScene/MovieRenderPipeline/Source/MovieRenderPipelineEditor/Private/Graph/MovieGraphSchema.cpp
+++ b/Engine/Plugins/MovieScene/MovieRenderPipeline/Source/MovieRenderPipelineEditor/Private/Graph/MovieGraphSchema.cpp
@@ -123,10 +123,11 @@ void UMovieGraphSchema::GetGraphContextActions(FGraphContextMenuBuilder& Context
 	}
 
 	// Create an accessor node action for each variable the graph has
-	for (const UMovieGraphVariable* Variable : RuntimeGraph->GetVariables())
+	const bool bIncludeGlobal = true;
+	for (const UMovieGraphVariable* Variable : RuntimeGraph->GetVariables(bIncludeGlobal))
 	{
 		const FText Name = FText::FromString(Variable->Name);
-		const FText Category = LOCTEXT("CreateVariable_Category", "Variables");
+		const FText Category = Variable->IsGlobal() ? LOCTEXT("CreateGlobalVariable_Category", "Global Variables") : LOCTEXT("CreateVariable_Category", "Variables");
 		const FText Tooltip = LOCTEXT("CreateVariable_Tooltip", "Create an accessor node for this variable.");
 		
 		TSharedPtr<FMovieGraphSchemaAction> NewAction = MakeShared<FMovieGraphSchemaAction_NewVariableNode>(Category, Name, Variable->GetGuid(), Tooltip);
@@ -342,7 +343,6 @@ UEdGraphNode* FMovieGraphSchemaAction_NewVariableNode::PerformAction(UEdGraph* P
 		GraphNode->AutowireNewNode(FromPin);
 	}
 
-
 	return GraphNode;
 }
 
diff --git a/Engine/Plugins/MovieScene/MovieRenderPipeline/Source/MovieRenderPipelineEditor/Private/Graph/SMovieGraphMembersTabContent.cpp b/Engine/Plugins/MovieScene/MovieRenderPipeline/Source/MovieRenderPipelineEditor/Private/Graph/SMovieGraphMembersTabContent.cpp
--- a/Engine/Plugins/MovieScene/MovieRenderPipeline/Source/MovieRenderPipelineEditor/Private/Graph/SMovieGraphMembersTabContent.cpp
+++ b/Engine/Plugins/MovieScene/MovieRenderPipeline/Source/MovieRenderPipelineEditor/Private/Graph/SMovieGraphMembersTabContent.cpp
@@ -76,12 +76,32 @@ bool SMovieGraphMembersTabContent::CanDeleteSelectedMembers() const
 		return false;
 	}
 
-	// There may need to be some more extensive validation here in the future
+	TArray<TSharedPtr<FEdGraphSchemaAction>> SelectedActions;
+	ActionMenu->GetSelectedActions(SelectedActions);
+	
+	for (const TSharedPtr<FEdGraphSchemaAction>& SelectedAction : SelectedActions)
+	{
+		const FMovieGraphSchemaAction* MovieAction = static_cast<FMovieGraphSchemaAction*>(SelectedAction.Get());
+
+		// Don't allow deletion if the member was explicitly marked as non-deletable
+		if (const UMovieGraphMember* Member = Cast<UMovieGraphMember>(MovieAction->ActionTarget))
+		{
+			if (!Member->IsDeletable())
+			{
+				return false;
+			}
+		}
+	}
+
 	return true;
 }
 
 void SMovieGraphMembersTabContent::CollectAllActions(FGraphActionListBuilderBase& OutAllActions)
 {
+	static const FText UserVariablesCategory = LOCTEXT("UserVariablesCategory", "User Variables");
+	static const FText GlobalVariablesCategory = LOCTEXT("GlobalVariablesCategory", "Global Variables");
+	static const FText EmptyCategory = FText::GetEmpty();
+	
 	if (!CurrentGraph)
 	{
 		return;
@@ -90,23 +110,22 @@ void SMovieGraphMembersTabContent::CollectAllActions(FGraphActionListBuilderBase
 	FGraphActionMenuBuilder ActionMenuBuilder;
 
 	// Creates a new action in the action menu under a specific section w/ the provided action target
-	auto AddToActionMenu = [&ActionMenuBuilder](UMovieGraphMember* ActionTarget, const EActionSection Section) -> void
-	{
-		const FText InputActionDesc = FText::FromString(ActionTarget->Name);
-		const FText InputActionCategory;
-		const FText InputActionTooltip;
-		const FText InputActionKeywords;
-		const int32 InputActionSectionID = static_cast<int32>(Section);
-		const TSharedPtr<FMovieGraphSchemaAction> InputAction(new FMovieGraphSchemaAction(InputActionCategory, InputActionDesc, InputActionTooltip, 0, InputActionKeywords, InputActionSectionID));
-		InputAction->ActionTarget = ActionTarget;
-		ActionMenuBuilder.AddAction(InputAction);
+	auto AddToActionMenu = [&ActionMenuBuilder](UMovieGraphMember* ActionTarget, const EActionSection Section, const FText& Category) -> void
+	{
+		const FText MemberActionDesc = FText::FromString(ActionTarget->Name);
+		const FText MemberActionTooltip;
+		const FText MemberActionKeywords;
+		const int32 MemberActionSectionID = static_cast<int32>(Section);
+		const TSharedPtr<FMovieGraphSchemaAction> MemberAction(new FMovieGraphSchemaAction(Category, MemberActionDesc, MemberActionTooltip, 0, MemberActionKeywords, MemberActionSectionID));
+		MemberAction->ActionTarget = ActionTarget;
+		ActionMenuBuilder.AddAction(MemberAction);
 	};
 
 	for (UMovieGraphInput* Input : CurrentGraph->GetInputs())
 	{
 		if (Input && Input->IsDeletable())
 		{
-			AddToActionMenu(Input, EActionSection::Inputs);
+			AddToActionMenu(Input, EActionSection::Inputs, EmptyCategory);
             
             // Update actions when an input is updated (renamed, etc)
             Input->OnMovieGraphInputChangedDelegate.AddSP(this, &SMovieGraphMembersTabContent::RefreshMemberActions);
@@ -117,23 +136,36 @@ void SMovieGraphMembersTabContent::CollectAllActions(FGraphActionListBuilderBase
 	{
 		if (Output && Output->IsDeletable())
 		{
-			AddToActionMenu(Output, EActionSection::Outputs);
+			AddToActionMenu(Output, EActionSection::Outputs, EmptyCategory);
 
 			// Update actions when an output is updated (renamed, etc)
 			Output->OnMovieGraphOutputChangedDelegate.AddSP(this, &SMovieGraphMembersTabContent::RefreshMemberActions);
 		}
 	}
 
-	for (UMovieGraphVariable* Variable : CurrentGraph->GetVariables())
+	const bool bIncludeGlobal = true;
+	const TArray<UMovieGraphVariable*> AllVariables = CurrentGraph->GetVariables(bIncludeGlobal);
+
+	// Add non-global variables first
+	for (UMovieGraphVariable* Variable : AllVariables)
 	{
-		if (Variable)
+		if (Variable && !Variable->IsGlobal())
 		{
-			AddToActionMenu(Variable, EActionSection::Variables);
+			AddToActionMenu(Variable, EActionSection::Variables, UserVariablesCategory);
 
 			// Update actions when a variable is updated (renamed, etc)
 			Variable->OnMovieGraphVariableChangedDelegate.AddSP(this, &SMovieGraphMembersTabContent::RefreshMemberActions);
 		}
 	}
+
+	// Add global variables after user-declared variables
+	for (UMovieGraphVariable* Variable : AllVariables)
+	{
+		if (Variable && Variable->IsGlobal())
+		{
+			AddToActionMenu(Variable, EActionSection::Variables, GlobalVariablesCategory);
+		}
+	}
 	
 	OutAllActions.Append(ActionMenuBuilder);
 }
