MoviePipeline plugin:
- Made the PIE Executor automatically set the render progression ('StatusProgress') of the job it is rendering.
#p4v-preflight-copy 24889039

diff --git a/Engine/Plugins/MovieScene/MovieRenderPipeline/Source/MovieRenderPipelineEditor/Private/MoviePipelinePIEExecutor.cpp b/Engine/Plugins/MovieScene/MovieRenderPipeline/Source/MovieRenderPipelineEditor/Private/MoviePipelinePIEExecutor.cpp
--- a/Engine/Plugins/MovieScene/MovieRenderPipeline/Source/MovieRenderPipelineEditor/Private/MoviePipelinePIEExecutor.cpp
+++ b/Engine/Plugins/MovieScene/MovieRenderPipeline/Source/MovieRenderPipelineEditor/Private/MoviePipelinePIEExecutor.cpp
@@ -13,6 +13,7 @@
 #include "Editor.h"
 #include "Settings/LevelEditorPlaySettings.h"
 #include "MoviePipelineQueue.h"
+#include "MoviePipelineBlueprintLibrary.h"
 #include "MoviePipelinePIEExecutorSettings.h"
 #include "MoviePipelineEditorBlueprintLibrary.h"
 #include "Misc/MessageDialog.h"
@@ -301,7 +302,12 @@ void UMoviePipelinePIEExecutor::OnTick()
 
 		RemainingInitializationFrames--;
 	}
-
+	
+	if (RemainingInitializationFrames <= 0)
+	{
+		Queue->GetJobs()[CurrentPipelineIndex]->SetStatusProgress(UMoviePipelineBlueprintLibrary::GetCompletionPercentage(Cast<UMoviePipeline>(ActiveMoviePipeline)));
+	}
+	
 	FText WindowTitle = GetWindowTitle();
 	TSharedPtr<SWindow> CustomWindow = WeakCustomWindow.Pin();
 	if (CustomWindow)
