MRQ: Fix Float16 conversion issue with templated function exercised by the ProRes/OCIO render path.

diff --git a/Engine/Plugins/MovieScene/MovieRenderPipeline/Source/MovieRenderPipelineCore/Private/MoviePipelineImageQuantization.cpp b/Engine/Plugins/MovieScene/MovieRenderPipeline/Source/MovieRenderPipelineCore/Private/MoviePipelineImageQuantization.cpp
--- a/Engine/Plugins/MovieScene/MovieRenderPipeline/Source/MovieRenderPipelineCore/Private/MoviePipelineImageQuantization.cpp
+++ b/Engine/Plugins/MovieScene/MovieRenderPipeline/Source/MovieRenderPipelineCore/Private/MoviePipelineImageQuantization.cpp
@@ -21,6 +21,11 @@ DECLARE_CYCLE_STAT(TEXT("STAT_MoviePipeline_ImageQuantization"), STAT_ImageQuant
 template<class TToColorBitDepthType, typename TFromColorBitDepthType, typename TColorChannelNumericType>
 TArray<TToColorBitDepthType> ConvertLinearToLinearBitDepth(TFromColorBitDepthType* InColor, const int32 InCount)
 {
+	if constexpr (std::is_same_v<TToColorBitDepthType, TFromColorBitDepthType>)
+	{
+		return TArray(InColor, InCount);
+	}
+
 	// Convert all of our pixels.
 	TArray<TToColorBitDepthType> OutsRGBData;
 	OutsRGBData.SetNumUninitialized(InCount);
@@ -355,8 +360,7 @@ static TUniquePtr<FImagePixelData> QuantizePixelDataTo16bpp(const FImagePixelDat
 		}
 		else
 		{
-			TArray<FFloat16Color> sRGBEncoded = ConvertLinearToLinearBitDepth<FFloat16Color, FFloat16Color, uint16>((FFloat16Color*)SrcRawDataPtr, RawSize.X * RawSize.Y);
-			QuantizedPixelData = MakeUnique<TImagePixelData<FFloat16Color>>(RawSize, TArray64<FFloat16Color>(MoveTemp(sRGBEncoded)), InPayload);
+			QuantizedPixelData = MakeUnique<TImagePixelData<FFloat16Color>>(RawSize, TArray64<FFloat16Color>((FFloat16Color*)SrcRawDataPtr, RawSize.X * RawSize.Y), InPayload);
 		}
 		break;
 	}
