MoviePipeline: Moved UMoviePipelineImageSequenceOutput_EXR from /Private/ to /Public/ so that you can use it from a Plugin's C++ code.

diff --git a/Engine/Plugins/MovieScene/MovieRenderPipeline/Source/MovieRenderPipelineRenderPasses/Private/MoviePipelineEXROutput.h b/Engine/Plugins/MovieScene/MovieRenderPipeline/Source/MovieRenderPipelineRenderPasses/Private/MoviePipelineEXROutput.h
deleted file mode 100644
diff --git a/Engine/Plugins/MovieScene/MovieRenderPipeline/Source/MovieRenderPipelineRenderPasses/Public/MoviePipelineEXROutput.h b/Engine/Plugins/MovieScene/MovieRenderPipeline/Source/MovieRenderPipelineRenderPasses/Public/MoviePipelineEXROutput.h
new file mode 100644
--- /dev/null
+++ b/Engine/Plugins/MovieScene/MovieRenderPipeline/Source/MovieRenderPipelineRenderPasses/Public/MoviePipelineEXROutput.h
@@ -0,0 +1,164 @@
+// Copyright Epic Games, Inc. All Rights Reserved.
+#pragma once
+
+#include "ImageWriteTask.h"
+#include "ImagePixelData.h"
+#include "MoviePipelineImageSequenceOutput.h"
+#include "Misc/StringFormatArg.h"
+
+#if WITH_UNREALEXR
+THIRD_PARTY_INCLUDES_START
+#include "Imath/ImathBox.h"
+#include "OpenEXR/ImfArray.h"
+#include "OpenEXR/ImfChannelList.h"
+#include "OpenEXR/ImfHeader.h"
+#include "OpenEXR/ImfIO.h"
+#include "OpenEXR/ImfInputFile.h"
+#include "OpenEXR/ImfOutputFile.h"
+#include "OpenEXR/ImfRgbaFile.h"
+#include "OpenEXR/ImfStdIO.h"
+THIRD_PARTY_INCLUDES_END
+#endif // WITH_UNREALEXR
+
+#include "MoviePipelineEXROutput.generated.h"
+
+class UMoviePipelineColorSetting;
+
+UENUM(BlueprintType)
+enum class EEXRCompressionFormat : uint8
+{
+	/** No compression is applied. */
+	None,
+	/** Good compression quality for grainy images. Lossless.*/
+	PIZ,
+	/** Good compression quality for images with low amounts of noise. Lossless. */
+	ZIP,
+	/** Lossy DCT-based compression for RGB channels. Alpha and other channels are uncompressed. More efficient than DWAB for partial buffer access on read in 3rd party tools. */
+	DWAA,
+	/** Similar to DWAA but goes in blocks of 256 scanlines instead of 32. More efficient disk space and faster to decode than DWAA. */
+	DWAB
+};
+
+#if WITH_UNREALEXR
+class FEXRImageWriteTask : public IImageWriteTaskBase
+{
+public:
+
+	/** The filename to write to */
+	FString Filename;
+
+	/** True if this task is allowed to overwrite an existing file, false otherwise. */
+	bool bOverwriteFile;
+	
+	/** Compression method used for the resulting EXR files. */
+	EEXRCompressionFormat Compression;
+
+	/** When using a lossy compression format, what is the base-error (CompressionLevel/100000) */
+	int32 CompressionLevel;
+
+	/** A function to invoke on the game thread when the task has completed */
+	TFunction<void(bool)> OnCompleted;
+	
+	/** Width/Height of the image data. All samples should match this. */
+	int32 Width;
+
+	int32 Height;
+
+	/** A set of key/value pairs to write into the exr file as metadata. */
+	TMap<FString, FStringFormatArg> FileMetadata;
+
+	/** The image data to write. Supports multiple layers of different bitdepths. */
+	TArray<TUniquePtr<FImagePixelData>> Layers;
+
+	/** Optional. A mapping between the FImagePixelData and a name. The standard is that the default layer is nameless (at which point it would be omitted) and other layers are prefixed. */
+	TMap<FImagePixelData*, FString> LayerNames;
+
+	/** Overscan info used to create apropriate dataWindow for EXR output. Goes from 0.0 to 1.0. */
+	float OverscanPercentage;
+
+	/** Color space chromaticity metadata. */
+	TArray<FVector2d> ColorSpaceChromaticities;
+
+	FEXRImageWriteTask()
+		: bOverwriteFile(true)
+		, Compression(EEXRCompressionFormat::PIZ)
+		, CompressionLevel(45)
+		, OverscanPercentage(0.0f)
+	{}
+
+public:
+
+	virtual bool RunTask() override final;
+	virtual void OnAbandoned() override final;
+
+private:
+
+	/**
+	 * Run the task, attempting to write out the raw data using the currently specified parameters
+	 *
+	 * @return true on success, false on any failure
+	 */
+	bool WriteToDisk();
+
+	/**
+	 * Ensures that the desired output filename is writable, deleting an existing file if bOverwriteFile is true
+	 *
+	 * @return True if the file is writable and the task can proceed, false otherwise
+	 */
+	bool EnsureWritableFile();
+
+	/**
+	* Adds arbitrary key/value pair metadata to the header of the file.
+	*/
+	void AddFileMetadata(Imf::Header& InHeader);
+
+	template <Imf::PixelType OutputFormat>
+	int64 CompressRaw(Imf::Header& InHeader, Imf::FrameBuffer& InFrameBuffer, FImagePixelData* InLayer);
+};
+#endif // WITH_UNREALEXR
+
+UCLASS()
+class MOVIERENDERPIPELINERENDERPASSES_API UMoviePipelineImageSequenceOutput_EXR : public UMoviePipelineImageSequenceOutputBase
+{
+	GENERATED_BODY()
+public:
+#if WITH_EDITOR
+	virtual FText GetDisplayText() const override { return NSLOCTEXT("MovieRenderPipeline", "ImgSequenceEXRSettingDisplayName", ".exr Sequence [16bit]"); }
+#endif
+public:
+	UMoviePipelineImageSequenceOutput_EXR()
+	{
+		OutputFormat = EImageFormat::EXR;
+		Compression = EEXRCompressionFormat::PIZ;
+		bMultilayer = true;
+	}
+
+	virtual void OnReceiveImageDataImpl(FMoviePipelineMergerOutputFrame* InMergedOutputFrame) override;
+	virtual bool IsAlphaAllowed() const override { return true; }
+
+public:
+	/**
+	* Which compression method should the resulting EXR file be compressed with
+	*/
+	UPROPERTY(EditAnywhere, BlueprintReadWrite, Category = "EXR")
+	EEXRCompressionFormat Compression;
+
+	/**
+	* Should we write all render passes to the same exr file? Not all software supports multi-layer exr files.
+	*/
+	UPROPERTY(EditAnywhere, BlueprintReadWrite, Category = "EXR")
+	bool bMultilayer;
+
+protected:
+	struct FColorSpaceMetadata
+	{
+		FString SourceName;
+		FString DestinationName;
+		TArray<FVector2d> Chromaticities;
+	};
+
+	/**
+	* Get color space chromaticities, source and destination names from the color settings or working color space.
+	*/
+	static FColorSpaceMetadata GetColorSpaceMetadata(UMoviePipelineColorSetting* InColorSettings);
+};
\ No newline at end of file
