MoviePipeline: Add initial support for Temporal Sub-Sampling in Movie Render Graph version. The Movie Graph calculates temporal sub-sampling differently than the previous version and works by defining TRange<FFrameTime> for each output frame, and then breaking that into smaller ranges.

Evaluation is then just moving between a list of ranges, which simplifies a lot of the original code that had to move long steps vs. short steps, and this should help with future needs (such as stochastic sub-sampling). This does present some interesting challenges for Time Dilation (which will want to change the length of each range, but the amount it changes updates every frame) so we will need a solution for this in the future.

#UE-180934

diff --git a/Engine/Plugins/MovieScene/MovieRenderPipeline/Source/MovieRenderPipelineCore/Private/Graph/MovieGraphDefaultRenderer.cpp b/Engine/Plugins/MovieScene/MovieRenderPipeline/Source/MovieRenderPipelineCore/Private/Graph/MovieGraphDefaultRenderer.cpp
--- a/Engine/Plugins/MovieScene/MovieRenderPipeline/Source/MovieRenderPipelineCore/Private/Graph/MovieGraphDefaultRenderer.cpp
+++ b/Engine/Plugins/MovieScene/MovieRenderPipeline/Source/MovieRenderPipelineCore/Private/Graph/MovieGraphDefaultRenderer.cpp
@@ -92,33 +92,54 @@ void UMovieGraphDefaultRenderer::SetupRenderingPipelineForShot(UMoviePipelineExe
 
 void UMovieGraphDefaultRenderer::TeardownRenderingPipelineForShot(UMoviePipelineExecutorShot* InShot)
 {
+	TRACE_CPUPROFILER_EVENT_SCOPE(MRQ_TeardownRenderingPipeline);
+
 	// Ensure the GPU has actually rendered all of the frames
-	FlushRenderingCommands();
+	{
+		TRACE_CPUPROFILER_EVENT_SCOPE(MRQ_WaitForRenderWorkFinished);
+		FlushRenderingCommands();
+	}
 
-	// Make sure all of the data has actually been copied back from the GPU and accumulation tasks started.
-	for (TPair<UE::MovieGraph::DefaultRenderer::FRenderTargetInitParams, FMoviePipelineSurfaceQueuePtr> SurfaceQueueIt : PooledSurfaceQueues)
 	{
-		if (SurfaceQueueIt.Value.IsValid())
+		TRACE_CPUPROFILER_EVENT_SCOPE(MRQ_WaitForGPUReadbackFinished);
+		// Make sure all of the data has actually been copied back from the GPU and accumulation tasks started.
+		for (const TPair<UE::MovieGraph::DefaultRenderer::FRenderTargetInitParams, FMoviePipelineSurfaceQueuePtr>& SurfaceQueueIt : PooledSurfaceQueues)
 		{
-			SurfaceQueueIt.Value->Shutdown();
+			if (SurfaceQueueIt.Value.IsValid())
+			{
+				SurfaceQueueIt.Value->Shutdown();
+			}
 		}
+
 	}
 	
 	// Stall until the task graph has completed any pending accumulations.
-	FTaskGraphInterface::Get().WaitUntilTasksComplete(OutstandingTasks, ENamedThreads::GameThread);
-	OutstandingTasks.Reset();
-
-	for (TObjectPtr< UMovieGraphRenderPassNode> RenderPass : RenderPassesInUse)
 	{
-		RenderPass->Teardown();
+		TRACE_CPUPROFILER_EVENT_SCOPE(MRQ_WaitForAccumulationTasksFinished);
+
+		UE::Tasks::Wait(OutstandingTasks);
+		OutstandingTasks.Empty();
 	}
+	
 
-	RenderPassesInUse.Reset();
+	{
+		TRACE_CPUPROFILER_EVENT_SCOPE(MRQ_TeardownRenderPasses);
+		
+		for (const TObjectPtr<UMovieGraphRenderPassNode>& RenderPass : RenderPassesInUse)
+		{
+			RenderPass->Teardown();
+		}
+
+		RenderPassesInUse.Reset();
+	}
 
 	// ToDo: This could probably be preserved across shots to avoid allocations
-	PooledViewRenderTargets.Reset();
-	PooledAccumulators.Reset();
-	PooledSurfaceQueues.Reset();
+	{
+		TRACE_CPUPROFILER_EVENT_SCOPE(MRQ_ResetPooledData);
+		PooledViewRenderTargets.Reset();
+		PooledAccumulators.Reset();
+		PooledSurfaceQueues.Reset();
+	}
 }
 
 
@@ -140,15 +161,15 @@ void UMovieGraphDefaultRenderer::Render(const FMovieGraphTimeStepData& InTimeSte
 	// better idea of how much work we're concurrently doing. 
 	{
 		// Don't iterate through the array if someone is actively modifying it.
+		TRACE_CPUPROFILER_EVENT_SCOPE(MRQ_RemoveCompletedOutstandingTasks);
 		FScopeLock ScopeLock(&OutstandingTasksMutex);
 		for (int32 Index = 0; Index < OutstandingTasks.Num(); Index++)
 		{
-			const FGraphEventRef& Task = OutstandingTasks[Index];
-			if (Task->IsComplete())
+			const UE::Tasks::FTask& Task = OutstandingTasks[Index];
+			if (Task.IsCompleted())
 			{
-				Task->Release();
 				OutstandingTasks.RemoveAtSwap(Index);
-
+		
 				// We swapped the end array element into the current spot,
 				// so we need to check this element again, otherwise we skip
 				// checking some things.
@@ -190,12 +211,32 @@ void UMovieGraphDefaultRenderer::Render(const FMovieGraphTimeStepData& InTimeSte
 	}
 }
 
-void UMovieGraphDefaultRenderer::AddOutstandingRenderTask_AnyThread(FGraphEventRef InTask)
+void UMovieGraphDefaultRenderer::AddOutstandingRenderTask_AnyThread(UE::Tasks::FTask InTask)
 {
 	// We might be looping through the array to remove previously completed tasks,
 	// so don't modify until that is completed.
 	FScopeLock ScopeLock(&OutstandingTasksMutex);
-	OutstandingTasks.Add(InTask);
+	OutstandingTasks.Add(MoveTemp(InTask));
+}
+
+UE::MovieGraph::DefaultRenderer::FCameraInfo UMovieGraphDefaultRenderer::GetCameraInfo(const FGuid& InCameraIdentifier) const
+{
+	UE::MovieGraph::DefaultRenderer::FCameraInfo CameraInfo;
+
+	// We only support the primary camera right now
+	APlayerController* LocalPlayerController = GetWorld()->GetFirstPlayerController();
+	if (LocalPlayerController && LocalPlayerController->PlayerCameraManager)
+	{
+		CameraInfo.ViewInfo = LocalPlayerController->PlayerCameraManager->GetCameraCacheView();
+		CameraInfo.ViewActor = LocalPlayerController->GetViewTarget();
+		CameraInfo.CameraName = TEXT("Unnamed_Camera"); // ToDo: This eventually needs to come from Level Sequences
+	}
+	else
+	{
+		UE_LOG(LogMovieRenderPipeline, Error, TEXT("Failed to find Local Player Controller/Camera Manager to get viewpoint!"));
+	}
+	
+	return CameraInfo;
 }
 
 UTextureRenderTarget2D* UMovieGraphDefaultRenderer::GetOrCreateViewRenderTarget(const UE::MovieGraph::DefaultRenderer::FRenderTargetInitParams& InInitParams)
@@ -279,7 +320,7 @@ FSurfaceAccumulatorPool::FInstancePtr FSurfaceAccumulatorPool::BlockAndGetAccumu
 					Accumulators[Index]->ActiveFrameNumber = InFrameNumber;
 					Accumulators[Index]->ActivePassIdentifier = InPassIdentifier;
 					Accumulators[Index]->bIsActive = true;
-					Accumulators[Index]->TaskPrereq = nullptr;
+					Accumulators[Index]->TaskPrereq = UE::Tasks::FTask();
 					AvailableIndex = Index;
 					break;
 				}
diff --git a/Engine/Plugins/MovieScene/MovieRenderPipeline/Source/MovieRenderPipelineCore/Private/Graph/MovieGraphLinearTimeStep.cpp b/Engine/Plugins/MovieScene/MovieRenderPipeline/Source/MovieRenderPipelineCore/Private/Graph/MovieGraphLinearTimeStep.cpp
--- a/Engine/Plugins/MovieScene/MovieRenderPipeline/Source/MovieRenderPipelineCore/Private/Graph/MovieGraphLinearTimeStep.cpp
+++ b/Engine/Plugins/MovieScene/MovieRenderPipeline/Source/MovieRenderPipelineCore/Private/Graph/MovieGraphLinearTimeStep.cpp
@@ -10,6 +10,17 @@
 #include "Camera/CameraComponent.h"
 #include "Camera/CameraActor.h"
 #include "Interfaces/Interface_PostProcessVolume.h"
+#include "GameFramework/WorldSettings.h"
+#include "ProfilingDebugging/MiscTrace.h"
+
+UMovieGraphLinearTimeStep::UMovieGraphLinearTimeStep()
+{
+	CustomTimeStep = CreateDefaultSubobject<UMovieGraphEngineTimeStep>("MovieGraphEngineTimeStep");
+
+	// This set up our internal state to pick up on the first temporal sub-sample in the pattern
+	ResetForEndOfOutputFrame();
+	CurrentTimeStepData.OutputFrameNumber = 0;
+}
 
 void UMovieGraphLinearTimeStep::TickProducingFrames()
 {
@@ -23,11 +34,11 @@ void UMovieGraphLinearTimeStep::TickProducingFrames()
 	// to do this only once we're ready to set the timestep though, as Initialize can be called as
 	// a result of a OnBeginFrame, meaning that Initialize is called on the frame before TickProducingFrames
 	// so there would be one frame where it used the custom timestep (after initialize) before TPF was called.
-	//if (GEngine->GetCustomTimeStep() != CustomTimeStep)
-	//{
-	//	CachedPrevCustomTimeStep = GEngine->GetCustomTimeStep();
-	//	GEngine->SetCustomTimeStep(CustomTimeStep);
-	//}
+	if (GEngine->GetCustomTimeStep() != CustomTimeStep)
+	{
+		PrevCustomTimeStep = GEngine->GetCustomTimeStep();
+		GEngine->SetCustomTimeStep(CustomTimeStep);
+	}
 
 	// We cache the frame metrics for the duration of a single output frame, instead of recalculating them every tick.
 	// This is required for stochastic timesteps (which will jump around the actual evaluated time randomly within one
@@ -37,6 +48,7 @@ void UMovieGraphLinearTimeStep::TickProducingFrames()
 	// in advance, check if our new time would overrun the end of the current camera cut, and if so, skip actually starting
 	// that frame. Unfortunately this does cause some complications with actor time dilations - the time dilation track
 	// will be updating each tick within the output sample and trying to change the world time dilation which we don't want.
+	
 	// ToDo: Still tbd how we'll fix that, currently the plan is to allow custom clocksources to disable world time dilation
 	// support, at which point we'll handle the time dilation by affecting the whole engine tick instead of just world tick.
 	//if (IsFirstTemporalSample())
@@ -47,17 +59,30 @@ void UMovieGraphLinearTimeStep::TickProducingFrames()
 		// if this would put us beyond our range of time this shot is supposed to represent.
 
 	}
-
+	// Handle shot initialization
 
 	if (CurrentCameraCut->ShotInfo.State == EMovieRenderShotState::Uninitialized)
 	{
 		UE_LOG(LogMovieRenderPipeline, Log, TEXT("MovieGraph Initializing Camera Cut [%d/%d] in [%s] %s."),
 			CurrentShotIndex + 1, ActiveShotList.Num(), *CurrentCameraCut->OuterName, *CurrentCameraCut->InnerName);
 
+		// Sets up the render state, etc.
 		GetOwningGraph()->SetupShot(CurrentCameraCut);
-		// InitializeShot(CurrentCameraCut);
+
+		// Seeks the external data source to match
 		GetOwningGraph()->GetDataSourceInstance()->InitializeShot(CurrentCameraCut);
 
+		CurrentFrameData.TemporalSampleIndex = 0;
+
+		FFrameRate TickResolution = GetOwningGraph()->GetDataSourceInstance()->GetTickResolution();
+		FFrameRate FrameRate = GetOwningGraph()->GetDataSourceInstance()->GetDisplayRate(); // ToDo, needs to come from config (config can override)
+		FFrameTime FrameTimePerOutputFrame = FFrameRate::TransformTime(FFrameTime(FFrameNumber(1)), FrameRate, TickResolution);
+
+		// Dummy range so that the first rendering frame knows to start at the right time.
+		FFrameTime UpperBound = CurrentCameraCut->ShotInfo.CurrentTimeInRoot;
+		CurrentFrameData.LastOutputFrameRange = TRange<FFrameTime>(UpperBound - FrameTimePerOutputFrame, UpperBound);
+		CurrentFrameData.LastSampleRange = CurrentFrameData.LastOutputFrameRange;
+
 		// We can safely fall through to the below states as they're OK to process the same frame we set up.
 		UE_LOG(LogMovieRenderPipeline, Log, TEXT("MovieGraph Finished initializing Camera Cut [%d/%d] in [%s] %s."),
 			CurrentShotIndex + 1, ActiveShotList.Num(), *CurrentCameraCut->OuterName, *CurrentCameraCut->InnerName);
@@ -66,44 +91,212 @@ void UMovieGraphLinearTimeStep::TickProducingFrames()
 		CurrentCameraCut->ShotInfo.State = EMovieRenderShotState::Rendering;
 	}
 
-	
+
 
 	if (CurrentCameraCut->ShotInfo.State == EMovieRenderShotState::Rendering)
 	{
-		ULevelSequence* CurrentSequence = Cast<ULevelSequence>(GetOwningGraph()->GetCurrentJob()->Sequence.TryLoad());
-		FFrameRate TickResolution = CurrentSequence->GetMovieScene()->GetTickResolution();
-		FFrameRate FrameRate = CurrentSequence->GetMovieScene()->GetDisplayRate();
-
-		FFrameTime TicksPerOutputFrame = FFrameRate::TransformTime(FFrameTime(FFrameNumber(1)), FrameRate, TickResolution);
-		CurrentCameraCut->ShotInfo.CurrentTickInRoot += TicksPerOutputFrame.GetFrame();
-
-		// Convert the CurrentTickInRoot back to a frame number.
-		FFrameTime CurrentTickInRootInOutputFrames = FFrameRate::TransformTime(CurrentCameraCut->ShotInfo.CurrentTickInRoot, TickResolution, FrameRate);
-		CurrentTimeStepData.OutputFrameNumber = CurrentTickInRootInOutputFrames.FloorToFrame().Value;
-		CurrentTimeStepData.FrameDeltaTime = FrameRate.AsInterval();
-		CurrentTimeStepData.WorldSeconds = 0.f;
-		CurrentTimeStepData.MotionBlurFraction = FrameRate.AsInterval();
+		if (IsFirstTemporalSample())
+		{
+			// The temporal sample count can change every frame due to graph evaluations, so when we're on our first temporal
+			// sub-sample of the new frame, we need to re-fetch the value.
+			UpdateTemporalSampleCount();
+
+			// Re-calculate some timing statistics about the given TS count
+			UpdateFrameMetrics();
+
+			// We create a range of time that we want to represent with this frame in absolute root sequence
+			// space. This is because each frame can have a different temporal sample count, and when we jump
+			// over the time period the shutter is closed, it's no longer an easy fixed number. So instead,
+			// we build a time range we wish to represesnt with this frame, and then split that into sub-regions
+			// and simply move between pre-calculated regions. This lets us handle the fact that the delta time
+			// between output frames is going to be different when moving between different TS counts on frames.
+			
+			// To figure out what range of time we want to represent, we actually need to look at the last frame's
+			// calculated total output range (ignoring all sub-sampling). The new frame is then multiplied by time
+			// dilation, so our entire output range is shortened. We can't just jump between whole frames (due to
+			// time dilation) so we have to take the last range, figure out how long the new frame is, and then
+			// build ontop of that.
+			TRangeBound<FFrameTime> EndOfPreviousFrame = CurrentFrameData.LastOutputFrameRange.GetUpperBound();
+
+			// The upper bound is exclusive, so we initialize a new TRange with just the value to start inclusively
+			// ToDo: Apply time dilation by multiplying the calculated range by slow-mo duration.
+			CurrentFrameData.CurrentOutputFrameRange = TRange<FFrameTime>(EndOfPreviousFrame.GetValue(), EndOfPreviousFrame.GetValue() + CurrentFrameMetrics.FrameTimePerOutputFrame);
+
+			// Now that we've calculated the total range of time we're trying to represent, we can check to see
+			// if this would put us beyond our range of time this shot is supposed to represent.
+			if (CurrentFrameData.CurrentOutputFrameRange.GetUpperBoundValue() > CurrentCameraCut->ShotInfo.TotalOutputRangeRoot.GetUpperBoundValue())
+			{
+				// We're going to spend this frame tearing down the shot (not rendering anything), next frame
+				// we'll re-enter this loop and pick up the start of the next shot.
+				ResetForEndOfOutputFrame();
+
+				// The delta time isn't very relevant here but you must specify at _a_ delta time each frame to ensure
+				// we never have a frame that uses a stale delta time.
+				CustomTimeStep->SetCachedFrameTiming(UMovieGraphEngineTimeStep::FTimeStepCache(CurrentFrameMetrics.FrameRate.AsInterval()));
+
+				CurrentCameraCut->ShotInfo.State = EMovieRenderShotState::Finished;
+				GetOwningGraph()->TeardownShot(CurrentCameraCut);
+				return;
+			}
+
+			// Figure out ranges for how long the shutter is open and closed
+			{
+				TArray<TRange<FFrameTime>> SplitRange = CurrentFrameData.CurrentOutputFrameRange.Split(CurrentFrameData.CurrentOutputFrameRange.GetLowerBoundValue() + CurrentFrameMetrics.FrameTimeWhileShutterOpen);
+				if (ensure(SplitRange.Num() == 2))
+				{
+					CurrentFrameData.RangeShutterOpen = SplitRange[0];
+					CurrentFrameData.RangeShutterClosed = SplitRange[1];
+				}
+			}
+			
+			// Now split the range the shutter is open per temporal sample.
+			{
+				CurrentFrameData.TemporalRanges.Reset();
+				CurrentFrameData.TemporalRanges.Reserve(CurrentFrameData.TemporalSampleCount);
+				FFrameTime SplitStartTime = CurrentFrameData.RangeShutterOpen.GetLowerBoundValue();
+				TRange<FFrameTime> RemainingRange = CurrentFrameData.RangeShutterOpen;
+				for (int32 Index = 0; Index < CurrentFrameData.TemporalSampleCount - 1; Index++)
+				{
+					SplitStartTime += CurrentFrameMetrics.FrameTimePerTemporalSample;
+					TArray<TRange<FFrameTime>> NewRanges = RemainingRange.Split(SplitStartTime);
+					if (ensure(NewRanges.Num() == 2))
+					{
+						CurrentFrameData.TemporalRanges.Add(NewRanges[0]);
+						RemainingRange = NewRanges[1];
+					}
+				}
+
+				// Add the remaining range, it's the leftover from the last split.
+				CurrentFrameData.TemporalRanges.Add(RemainingRange);
+
+				//for (int32 Index = 0; Index < CurrentFrameData.TemporalRanges.Num(); Index++)
+				//{
+				//	UE_LOG(LogTemp, Warning, TEXT("Range: [%s, %s)"),
+				//		*LexToString(CurrentFrameData.TemporalRanges[Index].GetLowerBoundValue()),
+				//		*LexToString(CurrentFrameData.TemporalRanges[Index].GetUpperBoundValue()));
+				//}
+			}
+		}
 		
-		// ToDo: This will change once we have temporal sub-sample support
-		CurrentTimeStepData.bIsFirstTemporalSampleForFrame = true;
-		CurrentTimeStepData.bIsLastTemporalSampleForFrame = true;
+		// The delta time for this frame is the difference between the current range index, and the last range index.
+		const TRange<FFrameTime>& PreviousRange = CurrentFrameData.LastSampleRange;
+		const TRange<FFrameTime>& NextRange = CurrentFrameData.TemporalRanges[CurrentFrameData.TemporalSampleIndex];
+
+		FFrameTime FrameDeltaTime = NextRange.GetLowerBoundValue() - PreviousRange.GetLowerBoundValue();
+
+		// ToDo: Propagate delta time multipliers to cloth
+
+		// Because we know what time range we're supposed to represent, we can just assign the CurrentTimeInRoot absolutely,
+		// instead of accumulating delta times. // ToDo: shutter timing, motion blur offset?
+		CurrentCameraCut->ShotInfo.CurrentTimeInRoot = NextRange.GetLowerBoundValue();
+
+		FFrameTime FinalEvalTime = CurrentCameraCut->ShotInfo.CurrentTimeInRoot + CurrentFrameMetrics.MotionBlurCenteringOffsetTime + CurrentFrameMetrics.ShutterOffsetFrameTime;
+		UE_LOG(LogMovieRenderPipeline, VeryVerbose, TEXT("FinalEvalTime: %s (Tick: %s SOffset: %s MOffset: %s)"),
+			*LexToString(FinalEvalTime),
+			*LexToString(CurrentCameraCut->ShotInfo.CurrentTimeInRoot), 
+			*LexToString(CurrentFrameMetrics.ShutterOffsetFrameTime),
+			*LexToString(CurrentFrameMetrics.MotionBlurCenteringOffsetTime));
+
+
+		// Now we need to fill out some of our current timestep data for the renderer portion to use.
+		double FrameDeltaTimeAsSeconds = CurrentFrameMetrics.TickResolution.AsSeconds(FrameDeltaTime);
+		CurrentTimeStepData.FrameDeltaTime = FrameDeltaTimeAsSeconds;
+		CurrentTimeStepData.WorldSeconds = CurrentTimeStepData.WorldSeconds + FrameDeltaTimeAsSeconds;
+		CurrentTimeStepData.WorldTimeDilation = 1.0; // Temp
+
+		// The combination of shutter angle percentage, non-uniform render frame delta times and dividing by sample
+		// count produce the correct length for motion blur in all cases.
+		CurrentTimeStepData.MotionBlurFraction = CurrentFrameMetrics.MotionBlurAmount / CurrentFrameData.TemporalSampleCount;
+		CurrentTimeStepData.bIsFirstTemporalSampleForFrame = IsFirstTemporalSample();
+		CurrentTimeStepData.bIsLastTemporalSampleForFrame = IsLastTemporalSample();
+		CurrentTimeStepData.bRequiresAccumulator = CurrentFrameData.TemporalSampleCount > 1;
+		CurrentTimeStepData.OutputFrameNumber = CurrentFrameData.OutputFrameNumber;
+
+		UE_LOG(LogTemp, Warning, TEXT("F# %d bFirst: %d bLast: %d bReqAc: %d"),
+			CurrentTimeStepData.OutputFrameNumber, CurrentTimeStepData.bIsFirstTemporalSampleForFrame,
+			CurrentTimeStepData.bIsLastTemporalSampleForFrame, CurrentTimeStepData.bRequiresAccumulator);
+
+		// Set our time step for the next frame. We use the undilated delta time for the Custom Timestep as the engine will
+		// apply the time dilation to the world tick for us, so we don't want to double up time dilation.
+		double UndilatedDeltaTime = CurrentFrameMetrics.TickResolution.AsSeconds(FrameDeltaTime);
+		CustomTimeStep->SetCachedFrameTiming(UMovieGraphEngineTimeStep::FTimeStepCache(UndilatedDeltaTime));
+		GetOwningGraph()->GetDataSourceInstance()->SyncDataSourceTime(FinalEvalTime);
+
+		// ToDo: This should be converted back to an 'effective' frame number (source frame in external data asset)
+		// so you can line up profiling with the acutal content on screen.
+		TRACE_BOOKMARK(TEXT("MoviePipeline - Rendering Frame %d [TS: %d]"), CurrentTimeStepData.OutputFrameNumber, CurrentFrameData.TemporalSampleIndex);
+
+		// Increment various post-frame counters to set them up for the next frame. This is okay
+		// because the rest of the Movie Graph Pipeline system is based on CurrentTimeStepData which accurately
+		// reflects which frame we're on.
+		{
+			// Update the last sample range to the one we just rendered.
+			CurrentFrameData.LastSampleRange = CurrentFrameData.TemporalRanges[CurrentFrameData.TemporalSampleIndex];
+			if (IsLastTemporalSample())
+			{
+				CurrentFrameData.LastOutputFrameRange = CurrentFrameData.CurrentOutputFrameRange;
+				
+				// Increment the output frame number only on the first temporal sample.
+				CurrentFrameData.OutputFrameNumber++;
+
+			}
+
+			if (CurrentFrameData.TemporalSampleIndex >= CurrentFrameData.TemporalSampleCount - 1)
+			{
+				// If we've rendered the last temporal sub-sample, we've started a new output frame
+				// and we need to reset our temporal sample index.
+				CurrentFrameData.TemporalSampleIndex = 0;
+			}
+			else
+			{
+				// Each tick we increment the temporal sample we're on.
+				CurrentFrameData.TemporalSampleIndex++;
+			}
+		}
 
 
-		if (CurrentCameraCut->ShotInfo.CurrentTickInRoot >= CurrentCameraCut->ShotInfo.TotalOutputRangeRoot.GetUpperBoundValue())
+		if (!ensure(CurrentCameraCut->ShotInfo.CurrentTickInRoot < CurrentCameraCut->ShotInfo.TotalOutputRangeRoot.GetUpperBoundValue()))
 		{
-			CurrentCameraCut->ShotInfo.State = EMovieRenderShotState::Finished;
-			GetOwningGraph()->TeardownShot(CurrentCameraCut);
-			return;
+			UE_LOG(LogMovieRenderPipeline, Error, TEXT("Shot ran past evaluation range, this shouldn't be possible."));
 		}
 	}
+
 }
 
-void UMovieGraphLinearTimeStep::UpdateFrameMetrics()
+
+bool UMovieGraphLinearTimeStep::IsFirstTemporalSample() const
 {
-	FFrameData FrameData;
+	return CurrentFrameData.TemporalSampleIndex == 0;
+}
 
-	// ToDo: This needs to come from the graph
-	int32 TemporalSampleCount = 1;
+bool UMovieGraphLinearTimeStep::IsLastTemporalSample() const
+{
+	return CurrentFrameData.TemporalSampleIndex == CurrentFrameData.TemporalSampleCount - 1;
+}
+
+void UMovieGraphLinearTimeStep::ResetForEndOfOutputFrame()
+{
+	CurrentFrameData.TemporalSampleIndex = 0;
+}
+
+void UMovieGraphLinearTimeStep::UpdateTemporalSampleCount()
+{
+	// ToDo: This needs to come from the config.
+	CurrentFrameData.TemporalSampleCount = 1;
+}
+
+void UMovieGraphLinearTimeStep::Shutdown()
+{
+	// Shut down our custom timestep which reqstores some world settings we modified.
+
+	// ToDo: This takes 200ms because of an arbitrary sleep inside of SetCustomTimeStep when shutting down
+	// but that requires some bigger changes to custom timesteps to fix.
+	GEngine->SetCustomTimeStep(PrevCustomTimeStep);
+}
+
+void UMovieGraphLinearTimeStep::UpdateFrameMetrics()
+{
+	FOutputFrameMetrics FrameData;
 
 	// We inherit a tick resolution from the level sequence so that we can represent the same
 	// range of time that the level sequence does.
@@ -120,7 +313,7 @@ void UMovieGraphLinearTimeStep::UpdateFrameMetrics()
 	FrameData.FrameTimeWhileShutterOpen = FrameData.FrameTimePerOutputFrame * FrameData.MotionBlurAmount;
 
 	// Now that we know how long the shutter is open, figure out how long each temporal sub-sample gets.
-	FrameData.FrameTimePerTemporalSample = FrameData.FrameTimeWhileShutterOpen / TemporalSampleCount;
+	FrameData.FrameTimePerTemporalSample = FrameData.FrameTimeWhileShutterOpen / CurrentFrameData.TemporalSampleCount;
 
 	// The amount of time closed + time open should add up to exactly how long a frame is.
 	FrameData.FrameTimeWhileShutterClosed = FrameData.FrameTimePerOutputFrame - FrameData.FrameTimeWhileShutterOpen;
@@ -148,6 +341,8 @@ void UMovieGraphLinearTimeStep::UpdateFrameMetrics()
 	// centered around the object so we offset our time sampling by half of the
 	// motion blur distance so that the distance blurred represents that time.
 	FrameData.MotionBlurCenteringOffsetTime = FrameData.FrameTimePerTemporalSample / 2.0;
+
+	CurrentFrameMetrics = FrameData;
 }
 
 float UMovieGraphLinearTimeStep::GetBlendedMotionBlurAmount()
@@ -218,4 +413,70 @@ float UMovieGraphLinearTimeStep::GetBlendedMotionBlurAmount()
 	}
 
 	return FinalMotionBlurAmount;
-}
\ No newline at end of file
+}
+
+UMovieGraphEngineTimeStep::UMovieGraphEngineTimeStep()
+	: PrevMinUndilatedFrameTime(0.f)
+	, PrevMaxUndilatedFrameTime(0.f)
+{
+}
+
+bool UMovieGraphEngineTimeStep::Initialize(UEngine* InEngine)
+{
+	// We cache the Min/Max times the world has because we modify these during
+	// renders to match our sample length, otherwise stuff gets out of sync.
+	if (AWorldSettings* WorldSettings = GetWorld()->GetWorldSettings())
+	{
+		PrevMinUndilatedFrameTime = WorldSettings->MinUndilatedFrameTime;
+		PrevMaxUndilatedFrameTime = WorldSettings->MaxUndilatedFrameTime;
+	}
+
+	return true;
+}
+
+void UMovieGraphEngineTimeStep::Shutdown(UEngine* InEngine)
+{
+	if (AWorldSettings* WorldSettings = GetWorld()->GetWorldSettings())
+	{
+		WorldSettings->MinUndilatedFrameTime = PrevMinUndilatedFrameTime;
+		WorldSettings->MaxUndilatedFrameTime = PrevMaxUndilatedFrameTime;
+	}
+}
+
+void UMovieGraphEngineTimeStep::SetCachedFrameTiming(const FTimeStepCache& InTimeCache)
+{
+	if (ensureMsgf(!FMath::IsNearlyZero(InTimeCache.UndilatedDeltaTime), TEXT("An incorrect or uninitialized time step was used! Delta Time of 0 isn't allowed.")))
+	{
+		TimeCache = InTimeCache;
+	}
+	else
+	{
+		UE_LOG(LogMovieRenderPipeline, Error, TEXT("SetCachedFrameTiming called with zero delta time, falling back to 1/24"));
+		TimeCache = FTimeStepCache(1 / 24.0);
+	}
+}
+
+bool UMovieGraphEngineTimeStep::UpdateTimeStep(UEngine* /*InEngine*/)
+{
+	if (ensureMsgf(!FMath::IsNearlyZero(TimeCache.UndilatedDeltaTime), TEXT("An incorrect or uninitialized time step was used! Delta Time of 0 isn't allowed.")))
+	{
+		FApp::UpdateLastTime();
+		FApp::SetDeltaTime(TimeCache.UndilatedDeltaTime);
+		FApp::SetCurrentTime(FApp::GetCurrentTime() + FApp::GetDeltaTime());
+
+		// The UWorldSettings can clamp the delta time inside the Level Tick function, this creates a mismatch between component
+		// velocity and render thread velocity and becomes an issue at high temporal sample counts. 
+		if (AWorldSettings* WorldSettings = GetWorld()->GetWorldSettings())
+		{
+			WorldSettings->MinUndilatedFrameTime = TimeCache.UndilatedDeltaTime;
+			WorldSettings->MaxUndilatedFrameTime = TimeCache.UndilatedDeltaTime;
+		}
+	}
+
+	// Clear our cached time to ensure we're always explicitly telling this what to do and never relying on the last set value.
+	// (This will cause the ensure above to check on the next tick if someone didn't reset our value.)
+	TimeCache = FTimeStepCache();
+
+	// Return false so the engine doesn't run its own logic to overwrite FApp timings.
+	return false;
+}
diff --git a/Engine/Plugins/MovieScene/MovieRenderPipeline/Source/MovieRenderPipelineCore/Private/Graph/MovieGraphPipeline.cpp b/Engine/Plugins/MovieScene/MovieRenderPipeline/Source/MovieRenderPipelineCore/Private/Graph/MovieGraphPipeline.cpp
--- a/Engine/Plugins/MovieScene/MovieRenderPipeline/Source/MovieRenderPipelineCore/Private/Graph/MovieGraphPipeline.cpp
+++ b/Engine/Plugins/MovieScene/MovieRenderPipeline/Source/MovieRenderPipelineCore/Private/Graph/MovieGraphPipeline.cpp
@@ -221,30 +221,30 @@ void UMovieGraphPipeline::TickProducingFrames()
 void UMovieGraphPipeline::TickFinalizeOutputContainers(const bool bInForceFinish)
 {
 	// Tick all containers until they all report that they have finalized.
-	bool bAllContainsFinishedProcessing = true;
+	bool bAllContainsFinishedProcessing;
 
-	//do
-	//{
-	//	bAllContainsFinishedProcessing = true;
-	//
-	//	// Ask the containers if they're all done processing.
-	//	for (UMoviePipelineOutputBase* Container : GetPipelinePrimaryConfig()->GetOutputContainers())
-	//	{
-	//		bAllContainsFinishedProcessing &= Container->HasFinishedProcessing();
-	//	}
-	//
-	//	// If we aren't forcing a finish, early out after one loop to keep
-	//	// the editor/ui responsive.
-	//	if (!bInForceFinish || bAllContainsFinishedProcessing)
-	//	{
-	//		break;
-	//	}
-	//
-	//	// If they've reached here, they're forcing them to finish so we'll sleep for a touch to give
-	//	// everyone a chance to actually do work before asking them if they're done.
-	//	FPlatformProcess::Sleep(1.f);
-	//
-	//} while (true);
+	while(true)
+	{
+		bAllContainsFinishedProcessing = true;
+
+		for (const TObjectPtr<UMovieGraphFileOutputNode>& Node : GetOutputNodesUsed())
+		{
+			bAllContainsFinishedProcessing &= Node->IsFinishedWritingToDisk();
+
+		}
+	
+		// If we aren't forcing a finish, early out after one loop to keep
+		// the editor/ui responsive.
+		if (!bInForceFinish || bAllContainsFinishedProcessing)
+		{
+			break;
+		}
+	
+		// If they've reached here, they're forcing them to finish so we'll sleep for a touch to give
+		// everyone a chance to actually do work before asking them if they're done.
+		FPlatformProcess::Sleep(0.1f);
+	
+	}
 
 	// If an output container is still working, we'll early out to keep the UI responsive.
 	// If they've forced a finish this will have to be true before we can reach this block.
@@ -310,6 +310,15 @@ void UMovieGraphPipeline::TickPostFinalizeExport(const bool bInForceFinish)
 	TransitionToState(EMovieRenderPipelineState::Finished);
 }
 
+void UMovieGraphPipeline::BeginFinalize()
+{
+	// Notify our output nodes that no more frames will be submitted. This allows
+	// them to put fences into queues for file writes, etc.
+	for (const TObjectPtr<UMovieGraphFileOutputNode>& Node : GetOutputNodesUsed())
+	{
+		Node->OnAllFramesSubmitted();
+	}
+}
 
 void UMovieGraphPipeline::SetupShot(UMoviePipelineExecutorShot* InShot)
 {
@@ -361,16 +370,14 @@ void UMovieGraphPipeline::OnEngineTickEndFrame()
 {
 	LLM_SCOPE_BYNAME(TEXT("MovieGraphEndFrame"));
 
-	// Unfortunately, since we can't control when our Initialization function is called
-	// we can end up in a situation where this callback is registered but the matching
-	// OnEngineTickBeginFrame() hasn't been called for that given engine tick. Instead of
-	// changing this registration to hang off of the end of the first OnEngineTickBeginFrame()
-	// we instead just early out here if that hasn't actually been called once. This decision
-	// is designed to minimize places where callbacks are registered and where flow changes.
-	//if (!bHasRunBeginFrameOnce)
-	//{
-	//	return;
-	//}
+	// Don't try to submit anything to the renderer if the shot isn't initialized yet, or has
+	// finished. We tick the engine when transitioning between shot states.
+	EMovieRenderShotState CurrentShotState = GetActiveShotList()[GetCurrentShotIndex()]->ShotInfo.State;
+	if (CurrentShotState == EMovieRenderShotState::Uninitialized || 
+		CurrentShotState == EMovieRenderShotState::Finished)
+	{
+		return;
+	}
 
 	UE_LOG(LogMovieRenderPipeline, VeryVerbose, TEXT("MovieGraph OnEngineTickEndFrame (Start)"));
 
@@ -507,20 +514,22 @@ void UMovieGraphPipeline::TransitionToState(const EMovieRenderPipelineState InNe
 			// to continue to call it to allow ticking the Finalization stage.
 			FCoreDelegates::OnEndFrame.RemoveAll(this);
 
-			// Reset the Custom Timestep because we don't care how long the engine takes now
-			// GEngine->SetCustomTimeStep(CachedPrevCustomTimeStep);
+			// Restore any custom Time Step that may have been set before. We do this here
+			// because the TimeStepInstance is only expected to be having to calculate times
+			// during ProducingFrames.
+			GetTimeStepInstance()->Shutdown();
 
 			// Ensure all frames have been processed by the GPU and sent to the Output Merger
 			FlushRenderingCommands();
 
 			// And then make sure all frames are sent to the Output Containers before we finalize.
-			// ProcessOutstandingFinishedFrames();
+			ProcessOutstandingFinishedFrames();
 
 			// PreviewTexture = nullptr;
 
 			// This is called once notifying output containers that all frames that will be submitted have been submitted.
 			PipelineState = EMovieRenderPipelineState::Finalize;
-			//BeginFinalize();
+			BeginFinalize();
 		}
 		break;
 	case EMovieRenderPipelineState::Finalize:
@@ -531,6 +540,8 @@ void UMovieGraphPipeline::TransitionToState(const EMovieRenderPipelineState InNe
 			// This is called once notifying our export step that they can begin the export.
 			PipelineState = EMovieRenderPipelineState::Export;
 
+
+
 			// Restore the sequence so that the export processes can operate on the original sequence. 
 			// This is also done in the finished state because it's not guaranteed that the Export state 
 			// will be set when the render is canceled early
@@ -539,7 +550,7 @@ void UMovieGraphPipeline::TransitionToState(const EMovieRenderPipelineState InNe
 
 			// Ensure all of our Futures have been converted to the GeneratedOutputData. This has to happen
 			// after finalize finishes, because the futures won't be available until actually written to disk.
-			// ProcessOutstandingFutures();
+			ProcessOutstandingFutures();
 
 			//BeginExport();
 		}
@@ -558,8 +569,7 @@ void UMovieGraphPipeline::TransitionToState(const EMovieRenderPipelineState InNe
 			//	Setting->OnMoviePipelineShutdown(this);
 			//}
 
-			// Restore any custom Time Step that may have been set before.
-			//GEngine->SetCustomTimeStep(CachedPrevCustomTimeStep);
+
 
 			// Ensure our delegates don't get called anymore as we're going to become null soon.
 			FCoreDelegates::OnBeginFrame.RemoveAll(this);
@@ -607,6 +617,8 @@ void UMovieGraphPipeline::TransitionToState(const EMovieRenderPipelineState InNe
 			//	}
 			//}
 
+			OutputNodesDataSentTo.Reset();
+
 			OnMoviePipelineFinishedImpl();
 			UE_LOG(LogMovieRenderPipeline, Log, TEXT("Movie Graph Render completed. Duration: %s"), *(FDateTime::UtcNow() - GraphInitializationTime).ToString());
 		}
@@ -691,8 +703,11 @@ void UMovieGraphPipeline::ProcessOutstandingFinishedFrames()
 		for (UMovieGraphFileOutputNode* CDOInstance : ContainerCDOs)
 		{
 			CDOInstance->OnReceiveImageData(this, &OutputFrame);
-		}
 
+			// Ensure we keep track of which nodes we actually sent data to during this render
+			// so that we can call BeginFinalize/IsFinishedWritingToDisk on them when shutting down.
+			OutputNodesDataSentTo.Add(CDOInstance);
+		}
 	}
 }
 
@@ -802,4 +817,9 @@ FMovieGraphTraversalContext UMovieGraphPipeline::GetCurrentTraversalContext() co
 	CurrentContext.Time = GetTimeStepInstance()->GetCalculatedTimeData();
 
 	return CurrentContext;
-}
\ No newline at end of file
+}
+
+const TSet<TObjectPtr<UMovieGraphFileOutputNode>> UMovieGraphPipeline::GetOutputNodesUsed() const
+{
+	return OutputNodesDataSentTo;
+}
diff --git a/Engine/Plugins/MovieScene/MovieRenderPipeline/Source/MovieRenderPipelineCore/Private/Graph/MovieGraphSequenceDataSource.cpp b/Engine/Plugins/MovieScene/MovieRenderPipeline/Source/MovieRenderPipelineCore/Private/Graph/MovieGraphSequenceDataSource.cpp
--- a/Engine/Plugins/MovieScene/MovieRenderPipeline/Source/MovieRenderPipelineCore/Private/Graph/MovieGraphSequenceDataSource.cpp
+++ b/Engine/Plugins/MovieScene/MovieRenderPipeline/Source/MovieRenderPipelineCore/Private/Graph/MovieGraphSequenceDataSource.cpp
@@ -11,6 +11,11 @@
 #include "MovieRenderPipelineCoreModule.h"
 #include "MoviePipelineBlueprintLibrary.h"
 #include "Engine/GameViewportClient.h"
+
+UMovieGraphSequenceDataSource::UMovieGraphSequenceDataSource()
+{
+	CustomSequenceTimeController = MakeShared<UE::MovieGraph::FMovieGraphSequenceTimeController>();
+}
  
 void UMovieGraphSequenceDataSource::CacheDataPreJob(const FMovieGraphInitConfig& InInitConfig)
 {
@@ -88,7 +93,7 @@ void UMovieGraphSequenceDataSource::CacheLevelSequenceData(ULevelSequence* InSeq
 	// Ensure the (possibly new) Level Sequence Actor uses our sequence
 	LevelSequenceActor->SetSequence(InSequence);
 
-	// LevelSequenceActor->GetSequencePlayer()->SetTimeController(CustomSequenceTimeController);
+	LevelSequenceActor->GetSequencePlayer()->SetTimeController(CustomSequenceTimeController);
 	LevelSequenceActor->GetSequencePlayer()->Stop();
 
 	LevelSequenceActor->GetSequencePlayer()->OnSequenceUpdated().AddUObject(this, &UMovieGraphSequenceDataSource::OnSequenceEvaluated);
@@ -121,4 +126,35 @@ FFrameRate UMovieGraphSequenceDataSource::GetDisplayRate() const
 	}
 
 	return FFrameRate(24, 1);
+}
+
+void UMovieGraphSequenceDataSource::InitializeShot(UMoviePipelineExecutorShot* InShot)
+{
+	//if (!InShot->ShotInfo.bEmulateFirstFrameMotionBlur)
+	{
+		// Real warm up frames walk through the Sequence
+		LevelSequenceActor->GetSequencePlayer()->Play();
+	}
+	//else
+	//{
+	//	// Ensure we don't try to evaluate as we want to sit and wait during warm up and motion blur frames.
+	//	LevelSequenceActor->GetSequencePlayer()->Pause();
+	//}
+}
+
+void UMovieGraphSequenceDataSource::SyncDataSourceTime(const FFrameTime& InTime)
+{
+	FFrameRate TickResolution = LevelSequenceActor->GetSequence()->GetMovieScene()->GetTickResolution();
+	CustomSequenceTimeController->SetCachedFrameTiming(FQualifiedFrameTime(InTime, TickResolution));
+}
+
+namespace UE::MovieGraph
+{
+	FFrameTime FMovieGraphSequenceTimeController::OnRequestCurrentTime(const FQualifiedFrameTime& InCurrentTime, float InPlayRate)
+	{
+		FFrameTime RequestTime = FFrameRate::TransformTime(TimeCache.Time, TimeCache.Rate, InCurrentTime.Rate);
+		UE_LOG(LogMovieRenderPipeline, VeryVerbose, TEXT("[%d] OnRequestCurrentTime: %d %f"), GFrameCounter, RequestTime.FloorToFrame().Value, RequestTime.GetSubFrame());
+
+		return RequestTime;
+	}
 }
\ No newline at end of file
diff --git a/Engine/Plugins/MovieScene/MovieRenderPipeline/Source/MovieRenderPipelineCore/Private/Graph/Nodes/MovieGraphDeferredRenderPassNode.cpp b/Engine/Plugins/MovieScene/MovieRenderPipeline/Source/MovieRenderPipelineCore/Private/Graph/Nodes/MovieGraphDeferredRenderPassNode.cpp
--- a/Engine/Plugins/MovieScene/MovieRenderPipeline/Source/MovieRenderPipelineCore/Private/Graph/Nodes/MovieGraphDeferredRenderPassNode.cpp
+++ b/Engine/Plugins/MovieScene/MovieRenderPipeline/Source/MovieRenderPipelineCore/Private/Graph/Nodes/MovieGraphDeferredRenderPassNode.cpp
@@ -18,6 +18,10 @@
 #include "TextureResource.h"
 #include "MovieRenderOverlappedImage.h"
 #include "MoviePipelineSurfaceReader.h"
+#include "Tasks/Task.h"
+
+// For the 1D Weight table for accumulation
+#include "MovieRenderPipelineDataTypes.h"
 
 void UMovieGraphDeferredRenderPassNode::SetupImpl(const FMovieGraphRenderPassSetupData& InSetupData)
 {
@@ -77,6 +81,12 @@ void UMovieGraphDeferredRenderPassNode::FMovieGraphDeferredRenderPass::Setup(TWe
 	Renderer = InRenderer;
 	RenderPassNode = InRenderPassNode;
 
+	RenderDataIdentifier.RenderLayerName = LayerData.LayerName;
+	RenderDataIdentifier.RendererName = RenderPassNode->GetRendererName();
+	RenderDataIdentifier.SubResourceName = TEXT("beauty");
+	
+	UE::MovieGraph::DefaultRenderer::FCameraInfo CameraInfo = Renderer->GetCameraInfo(LayerData.CameraIdentifier);
+	RenderDataIdentifier.CameraName =  CameraInfo.CameraName;
 
 	// Figure out how big each sub-region (tile) is.
 	// const int32 TileSize = Graph->FindSetting(LayerData.LayerName, "highres.tileSize");
@@ -97,14 +107,6 @@ void UMovieGraphDeferredRenderPassNode::FMovieGraphDeferredRenderPass::Setup(TWe
 
 void UMovieGraphDeferredRenderPassNode::FMovieGraphDeferredRenderPass::GatherOutputPassesImpl(TArray<FMovieGraphRenderDataIdentifier>& OutExpectedPasses) const
 {
-	// ToDo: Store this when the render pass is created, so we don't create it in two places at once.
-	// This is a problem right now due to the camera name not being pre-resolved.
-	FMovieGraphRenderDataIdentifier RenderDataIdentifier;
-	RenderDataIdentifier.RenderLayerName = LayerData.LayerName;
-	RenderDataIdentifier.RendererName = RenderPassNode->GetRendererName();
-	RenderDataIdentifier.SubResourceName = TEXT("beauty");
-	RenderDataIdentifier.CameraName = TEXT("Unnamed_Camera");
-
 	OutExpectedPasses.Add(RenderDataIdentifier);
 }
 void UMovieGraphDeferredRenderPassNode::FMovieGraphDeferredRenderPass::Teardown()
@@ -143,19 +145,7 @@ void UMovieGraphDeferredRenderPassNode::FMovieGraphDeferredRenderPass::Render(co
 	FViewFamilyContextInitData InitData;
 	InitData.RenderTarget = RenderTarget->GameThread_GetRenderTargetResource();
 	InitData.World = Renderer->GetWorld();
-
-	APlayerController* LocalPlayerController = Renderer->GetWorld()->GetFirstPlayerController();
-	if (LocalPlayerController && LocalPlayerController->PlayerCameraManager)
-	{
-		InitData.MinimalViewInfo = LocalPlayerController->PlayerCameraManager->GetCameraCacheView();
-		InitData.ViewActor = LocalPlayerController->GetViewTarget();
-	}
-	else
-	{
-		InitData.MinimalViewInfo = FMinimalViewInfo();
-		InitData.ViewActor = nullptr;
-		UE_LOG(LogMovieRenderPipeline, Error, TEXT("Failed to find Local Player Controller/Camera Manager to get viewpoint!"));
-	}
+	InitData.CameraInfo = Renderer->GetCameraInfo(LayerData.CameraIdentifier);
 
 	// ToDo:
 	InitData.TimeData = InTimeData;
@@ -193,16 +183,6 @@ void UMovieGraphDeferredRenderPassNode::FMovieGraphDeferredRenderPass::Render(co
 	//	return;
 	//}
 
-	FMovieGraphRenderDataIdentifier RenderDataIdentifier;
-	RenderDataIdentifier.RenderLayerName = LayerData.LayerName;
-	RenderDataIdentifier.RendererName = RenderPassNode->GetRendererName();
-	RenderDataIdentifier.SubResourceName = TEXT("beauty");
-	RenderDataIdentifier.CameraName = TEXT("Unnamed_Camera");
-	//if (InitData.ViewActor)
-	//{
-	//	RenderDataIdentifier.CameraName = InitData.ViewActor->GetName();
-	//}
-
 	// Take our per-frame Traversal Context and update it with context specific to this sample.
 	FMovieGraphTraversalContext UpdatedTraversalContext = InFrameTraversalContext;
 	UpdatedTraversalContext.Time = InTimeData;
@@ -212,6 +192,10 @@ void UMovieGraphDeferredRenderPassNode::FMovieGraphDeferredRenderPass::Render(co
 	UE::MovieGraph::FMovieGraphSampleState SampleState;
 	SampleState.TraversalContext = MoveTemp(UpdatedTraversalContext);
 	SampleState.BackbufferResolution = RenderTargetInitParams.Size;
+	SampleState.AccumulatorResolution = RenderTargetInitParams.Size; // ToDo: Overscan
+	SampleState.bWriteSampleToDisk = false; // ToDo: From graph
+	SampleState.bRequiresAccumulator = InTimeData.bRequiresAccumulator;
+	SampleState.bFetchFromAccumulator = InTimeData.bIsLastTemporalSampleForFrame;
 
 	// Readback + Accumulate.
 	PostRendererSubmission(SampleState, RenderTargetInitParams, Canvas);
@@ -244,16 +228,11 @@ void UMovieGraphDeferredRenderPassNode::FMovieGraphDeferredRenderPass::PostRende
 		AccumulationArgs.ImageAccumulator = StaticCastSharedPtr<FImageOverlappedAccumulator>(AccumulatorInstance->Accumulator);
 		AccumulationArgs.bIsFirstSample = InSampleState.TraversalContext.Time.bIsFirstTemporalSampleForFrame; // FramePayload->IsFirstTile() && FramePayload->IsFirstTemporalSample()
 		AccumulationArgs.bIsLastSample = InSampleState.TraversalContext.Time.bIsLastTemporalSampleForFrame; // FramePayload->IsLastTile() && FramePayload->IsLastTemporalSample()
-		AccumulationArgs.TaskPrerequisite = AccumulatorInstance->TaskPrereq;
 	}
 
 	auto OnSurfaceReadbackFinished = [this, InSampleState, AccumulationArgs, AccumulatorInstance](TUniquePtr<FImagePixelData>&& InPixelData)
 	{
-		UE::MovieGraph::DefaultRenderer::FMovieGraphAccumulationTask Task;
-		// There may be other accumulations for this accumulator which need to be processed first
-		Task.LastCompletionEvent = AccumulationArgs.TaskPrerequisite;
-
-		FGraphEventRef Event = Task.Execute([PixelData = MoveTemp(InPixelData), InSampleState, AccumulationArgs, AccumulatorInstance]() mutable
+			UE::Tasks::TTask<void> Task = UE::Tasks::Launch(UE_SOURCE_LOCATION, [PixelData = MoveTemp(InPixelData), InSampleState, AccumulationArgs, AccumulatorInstance]() mutable
 			{
 				// Enqueue a encode for this frame onto our worker thread.
 				UE::MovieGraph::AccumulateSample_TaskThread(MoveTemp(PixelData), InSampleState, AccumulationArgs);
@@ -261,14 +240,15 @@ void UMovieGraphDeferredRenderPassNode::FMovieGraphDeferredRenderPass::PostRende
 				// We have to defer clearing the accumulator until after sample accumulation has finished
 				if (AccumulationArgs.bIsLastSample)
 				{
-					// Final sample has now been executed, break the pre-req chain and free the accumulator for reuse.
+					// Final sample has now been executed, free the accumulator for reuse.
 					AccumulatorInstance->bIsActive = false;
-					AccumulatorInstance->TaskPrereq = nullptr;
 				}
-			});
-		AccumulatorInstance->TaskPrereq = Event;
+			}, AccumulatorInstance->TaskPrereq);
 
-		this->Renderer->AddOutstandingRenderTask_AnyThread(Event);
+		// Make the next accumulation task that uses this accumulator use the task we just created as a pre-req.
+		AccumulatorInstance->TaskPrereq = Task;
+
+		this->Renderer->AddOutstandingRenderTask_AnyThread(Task);
 	};
 
 	FRenderTarget* RenderTarget = InCanvas.GetRenderTarget();
@@ -289,7 +269,7 @@ TSharedRef<FSceneViewFamilyContext> UMovieGraphDeferredRenderPassNode::FMovieGra
 	FEngineShowFlags ShowFlags = FEngineShowFlags(EShowFlagInitMode::ESFIM_Game);
 	EViewModeIndex ViewModeIndex = EViewModeIndex::VMI_Lit;
 
-	const bool bIsPerspective = InInitData.MinimalViewInfo.ProjectionMode == ECameraProjectionMode::Type::Perspective;
+	const bool bIsPerspective = InInitData.CameraInfo.ViewInfo.ProjectionMode == ECameraProjectionMode::Type::Perspective;
 
 	// Allow the Engine Showflag system to override our engine showflags, based on our view mode index.
 	// This is required for certain debug view modes (to have matching show flags set for rendering).
@@ -327,10 +307,10 @@ FSceneView* UMovieGraphDeferredRenderPassNode::FMovieGraphDeferredRenderPass::Al
 {
 	FSceneViewInitOptions ViewInitOptions;
 	ViewInitOptions.ViewFamily = InViewFamilyContext.Get();
-	ViewInitOptions.ViewOrigin = InInitData.MinimalViewInfo.Location;
+	ViewInitOptions.ViewOrigin = InInitData.CameraInfo.ViewInfo.Location;
 	ViewInitOptions.SetViewRectangle(FIntRect(FIntPoint(0, 0), FIntPoint(InInitData.RenderTarget->GetSizeXY().X, InInitData.RenderTarget->GetSizeXY().Y)));
-	ViewInitOptions.ViewRotationMatrix = FInverseRotationMatrix(InInitData.MinimalViewInfo.Rotation);
-	ViewInitOptions.ViewActor = InInitData.ViewActor;
+	ViewInitOptions.ViewRotationMatrix = FInverseRotationMatrix(InInitData.CameraInfo.ViewInfo.Rotation);
+	ViewInitOptions.ViewActor = InInitData.CameraInfo.ViewActor;
 
 	// Rotate the view 90 degrees to match the rest of the engine.
 	ViewInitOptions.ViewRotationMatrix = ViewInitOptions.ViewRotationMatrix * FMatrix(
@@ -338,7 +318,7 @@ FSceneView* UMovieGraphDeferredRenderPassNode::FMovieGraphDeferredRenderPass::Al
 		FPlane(1, 0, 0, 0),
 		FPlane(0, 1, 0, 0),
 		FPlane(0, 0, 0, 1));
-	float ViewFOV = InInitData.MinimalViewInfo.FOV;
+	float ViewFOV = InInitData.CameraInfo.ViewInfo.FOV;
 
 	// Inflate our FOV to support the overscan 
 	ViewFOV = 2.0f * FMath::RadiansToDegrees(FMath::Atan((1.0f + InInitData.OverscanFraction) * FMath::Tan(FMath::DegreesToRadians(ViewFOV * 0.5f))));
@@ -347,7 +327,7 @@ FSceneView* UMovieGraphDeferredRenderPassNode::FMovieGraphDeferredRenderPass::Al
 	// ToDo: This isn't great but this appears to be how the system works... ideally we could fetch this from the camera itself
 	const EAspectRatioAxisConstraint AspectRatioAxisConstraint = GetDefault<ULocalPlayer>()->AspectRatioAxisConstraint;
 	FIntRect ViewportRect = FIntRect(FIntPoint(0, 0), FIntPoint(InInitData.RenderTarget->GetSizeXY().X, InInitData.RenderTarget->GetSizeXY().Y));
-	FMinimalViewInfo::CalculateProjectionMatrixGivenViewRectangle(InInitData.MinimalViewInfo, AspectRatioAxisConstraint, ViewportRect, /*InOut*/ ViewInitOptions);
+	FMinimalViewInfo::CalculateProjectionMatrixGivenViewRectangle(InInitData.CameraInfo.ViewInfo, AspectRatioAxisConstraint, ViewportRect, /*InOut*/ ViewInitOptions);
 
 	// ToDo: High Res Tiling, Overscan support, letterboxing
 	ViewInitOptions.SceneViewStateInterface = InInitData.SceneViewStateReference;
@@ -357,11 +337,6 @@ FSceneView* UMovieGraphDeferredRenderPassNode::FMovieGraphDeferredRenderPass::Al
 	FSceneView* View = new FSceneView(ViewInitOptions);
 	InViewFamilyContext->Views.Add(View);
 
-	//View->ViewLocation = CameraInfo.ViewInfo.Location;
-	//View->ViewRotation = CameraInfo.ViewInfo.Rotation;
-	// Override previous/current view transforms so that tiled renders don't use the wrong occlusion/motion blur information.
-	//View->PreviousViewTransform = CameraInfo.ViewInfo.PreviousViewTransform;
-
 	View->StartFinalPostprocessSettings(View->ViewLocation);
 	// BlendPostProcessSettings(View, InOutSampleState, OptPayload);
 
@@ -507,8 +482,8 @@ namespace UE::MovieGraph
 		TSharedPtr<UE::MovieGraph::FMovieGraphSampleState> SampleStatePayload = MakeShared<UE::MovieGraph::FMovieGraphSampleState>(InSampleState);
 		SamplePixelData->SetPayload(StaticCastSharedPtr<IImagePixelDataPayload>(SampleStatePayload));
 
-		TSharedPtr<IMovieGraphOutputMerger, ESPMode::ThreadSafe> AccumulatorPin = InAccumulationParams.OutputMerger.Pin();
-		if (!AccumulatorPin.IsValid())
+		TSharedPtr<IMovieGraphOutputMerger, ESPMode::ThreadSafe> OutputMergerPin = InAccumulationParams.OutputMerger.Pin();
+		if (!OutputMergerPin.IsValid())
 		{
 			return;
 		}
@@ -516,16 +491,106 @@ namespace UE::MovieGraph
 		const bool bIsWellFormed = SamplePixelData->IsDataWellFormed();
 		check(bIsWellFormed);
 
-		// ToDo:
-		bool bWriteSampleToDisk = false;
-		if (bWriteSampleToDisk)
+		if (SampleStatePayload->bWriteSampleToDisk)
 		{
 			// Debug Feature: Write the raw sample to disk for debugging purposes. We copy the data here,
 			// as we don't want to disturb the memory flow below.
 			TUniquePtr<FImagePixelData> SampleData = SamplePixelData->CopyImageData();
-			AccumulatorPin->OnSingleSampleDataAvailable_AnyThread(MoveTemp(SampleData));
+			OutputMergerPin->OnSingleSampleDataAvailable_AnyThread(MoveTemp(SampleData));
+		}
+
+		// Optimization! If we don't need the accumulator (no tiling, no sub-sampling) then we'll skip it and just send it straight to the output stage.
+		// This reduces memory requirements and improves performance in the baseline case.
+		if (!SampleStatePayload->bRequiresAccumulator)
+		{
+			OutputMergerPin->OnCompleteRenderPassDataAvailable_AnyThread(MoveTemp(SamplePixelData));
+			return;
+		}
+
+		TSharedPtr<FImageOverlappedAccumulator> AccumulatorPin = InAccumulationParams.ImageAccumulator.Pin();
+		if (AccumulatorPin->NumChannels == 0)
+		{
+			LLM_SCOPE_BYNAME(TEXT("MoviePipeline/ImageAccumulatorInitMemory"));
+			const int32 ChannelCount = 4;
+			AccumulatorPin->InitMemory(SampleStatePayload->AccumulatorResolution, ChannelCount);
+			AccumulatorPin->ZeroPlanes();
+			AccumulatorPin->AccumulationGamma = 1.f;
+		}
+
+		// Accumulate the new sample to our target
+		{
+			// ToDo: Handle incorrectly sized samples (Post Process materials using r.screenpercentage)
+			
+			{
+				TRACE_CPUPROFILER_EVENT_SCOPE(MoviePipeline_AccumulatePixelData);
+				
+				// ToDo: High-res Tiling Support
+				const FIntPoint OverlappedPad = FIntPoint(0, 0);
+				const FIntPoint TileSize = SampleStatePayload->BackbufferResolution;
+				const FIntPoint OverlappedOffset = FIntPoint(0, 0);
+				const FVector2D OverlappedSubpixelShift = FVector2D(0.5f, 0.5f);
+				::MoviePipeline::FTileWeight1D WeightFunctionX;
+				::MoviePipeline::FTileWeight1D WeightFunctionY;
+				WeightFunctionX.InitHelper(OverlappedPad.X, TileSize.X, OverlappedPad.X);
+				WeightFunctionY.InitHelper(OverlappedPad.Y, TileSize.Y, OverlappedPad.Y);
+				
+				AccumulatorPin->AccumulatePixelData(*SamplePixelData, OverlappedOffset, OverlappedSubpixelShift, WeightFunctionX, WeightFunctionY);
+			}
+
+		}
+
+		// Finally on our last sample, we fetch the data out of the accumulator
+		// and move it to the Output Merger.
+		if (SampleStatePayload->bFetchFromAccumulator)
+		{
+			TRACE_CPUPROFILER_EVENT_SCOPE(MoviePipeline_FetchAccumulatedPixelData);
+			int32 FullSizeX = AccumulatorPin->PlaneSize.X;
+			int32 FullSizeY = AccumulatorPin->PlaneSize.Y;
+
+			// Now that a tile is fully built and accumulated we can notify the output builder that the
+			// data is ready so it can pass that onto the output containers (if needed).
+			if (SamplePixelData->GetType() == EImagePixelType::Float32)
+			{
+				// 32 bit FLinearColor
+				TUniquePtr<TImagePixelData<FLinearColor> > FinalPixelData = MakeUnique<TImagePixelData<FLinearColor>>(FIntPoint(FullSizeX, FullSizeY), SampleStatePayload);
+				AccumulatorPin->FetchFinalPixelDataLinearColor(FinalPixelData->Pixels);
+
+				// Send the data to the Output Builder
+				OutputMergerPin->OnCompleteRenderPassDataAvailable_AnyThread(MoveTemp(FinalPixelData));
+			}
+			else if (SamplePixelData->GetType() == EImagePixelType::Float16)
+			{
+				// 16 bit FLinearColor
+				TUniquePtr<TImagePixelData<FFloat16Color> > FinalPixelData = MakeUnique<TImagePixelData<FFloat16Color>>(FIntPoint(FullSizeX, FullSizeY), SampleStatePayload);
+				AccumulatorPin->FetchFinalPixelDataHalfFloat(FinalPixelData->Pixels);
+
+				// Send the data to the Output Builder
+				OutputMergerPin->OnCompleteRenderPassDataAvailable_AnyThread(MoveTemp(FinalPixelData));
+			}
+			else if (SamplePixelData->GetType() == EImagePixelType::Color)
+			{
+				// 8bit FColors
+				TUniquePtr<TImagePixelData<FColor>> FinalPixelData = MakeUnique<TImagePixelData<FColor>>(FIntPoint(FullSizeX, FullSizeY), SampleStatePayload);
+				AccumulatorPin->FetchFinalPixelDataByte(FinalPixelData->Pixels);
+
+				// Send the data to the Output Builder
+				OutputMergerPin->OnCompleteRenderPassDataAvailable_AnyThread(MoveTemp(FinalPixelData));
+			}
+			else
+			{
+				check(0);
+			}
+
+			{
+				TRACE_CPUPROFILER_EVENT_SCOPE(MoviePipeline_FreeAccumulatedPixelData);
+				// Free the memory in the accumulator.
+				AccumulatorPin->Reset();
+			}
 		}
 
-		AccumulatorPin->OnCompleteRenderPassDataAvailable_AnyThread(MoveTemp(SamplePixelData));
+		{
+			TRACE_CPUPROFILER_EVENT_SCOPE(ReleaseSampleData);
+			SamplePixelData.Reset();
+		}
 	}
 }
\ No newline at end of file
diff --git a/Engine/Plugins/MovieScene/MovieRenderPipeline/Source/MovieRenderPipelineCore/Private/Graph/Nodes/MovieGraphImageSequenceOutputNode.cpp b/Engine/Plugins/MovieScene/MovieRenderPipeline/Source/MovieRenderPipelineCore/Private/Graph/Nodes/MovieGraphImageSequenceOutputNode.cpp
--- a/Engine/Plugins/MovieScene/MovieRenderPipeline/Source/MovieRenderPipelineCore/Private/Graph/Nodes/MovieGraphImageSequenceOutputNode.cpp
+++ b/Engine/Plugins/MovieScene/MovieRenderPipeline/Source/MovieRenderPipelineCore/Private/Graph/Nodes/MovieGraphImageSequenceOutputNode.cpp
@@ -13,6 +13,17 @@ UMovieGraphImageSequenceOutputNode::UMovieGraphImageSequenceOutputNode()
 	ImageWriteQueue = &FModuleManager::Get().LoadModuleChecked<IImageWriteQueueModule>("ImageWriteQueue").GetWriteQueue();
 }
 
+void UMovieGraphImageSequenceOutputNode::OnAllFramesSubmittedImpl()
+{
+	FinalizeFence = ImageWriteQueue->CreateFence();
+}
+
+bool UMovieGraphImageSequenceOutputNode::IsFinishedWritingToDiskImpl() const
+{
+	// Wait until the finalization fence is reached meaning we've written everything to disk.
+	return Super::IsFinishedWritingToDiskImpl() && (!FinalizeFence.IsValid() || FinalizeFence.WaitFor(0));
+}
+
 void UMovieGraphImageSequenceOutputNode::OnReceiveImageDataImpl(UMovieGraphPipeline* InPipeline, UE::MovieGraph::FMovieGraphOutputMergerFrame* InRawFrameData)
 {
 	check(InRawFrameData);
diff --git a/Engine/Plugins/MovieScene/MovieRenderPipeline/Source/MovieRenderPipelineCore/Public/Graph/MovieGraphDataTypes.h b/Engine/Plugins/MovieScene/MovieRenderPipeline/Source/MovieRenderPipelineCore/Public/Graph/MovieGraphDataTypes.h
--- a/Engine/Plugins/MovieScene/MovieRenderPipeline/Source/MovieRenderPipelineCore/Public/Graph/MovieGraphDataTypes.h
+++ b/Engine/Plugins/MovieScene/MovieRenderPipeline/Source/MovieRenderPipelineCore/Public/Graph/MovieGraphDataTypes.h
@@ -96,8 +96,19 @@ class MOVIERENDERPIPELINECORE_API UMovieGraphTimeStepBase : public UObject
 {
 	GENERATED_BODY()
 public:
+	/** Called each frame while the Movie Graph Pipeline is in a producing frames state. */
 	virtual void TickProducingFrames() {}
+
+	/** Called when the Movie Graph Pipeline is shutting down, use this to restore any changes. */
+	virtual void Shutdown() {}
+
+	/** 
+	* TickProducingFrames will be called for a frame (before the frame starts) and then this will be
+	* called at the end of the frame when we kick off the renders for the frame. Should return data
+	* needed to calculate the correct rendering timestep.
+	*/
 	virtual FMovieGraphTimeStepData GetCalculatedTimeData() const { return FMovieGraphTimeStepData(); }
+
 	UMovieGraphPipeline* GetOwningGraph() const;
 };
 
@@ -129,8 +140,13 @@ class MOVIERENDERPIPELINECORE_API UMovieGraphDataSourceBase : public UObject
 public:
 	/** An internal, high resolution framerate that seeks, etc. will be returned in. (ie: 24,000/1) */
 	virtual FFrameRate GetTickResolution() const { return FFrameRate(); }
+
 	/** A lower res, human readable Frame Rate. We convert to the Tick Resolution internally. (ie: 24/1) */
 	virtual FFrameRate GetDisplayRate() const { return FFrameRate(); }
+
+	/** Called by the Time Step system when it wants the external data source to update. Time is in TickResolution scale. */
+	virtual void SyncDataSourceTime(const FFrameTime& InTime) {}
+
 	virtual void BuildTimeRanges() { }
 	/** 
 	* Called when the Movie Graph Pipeline starts before anything has happened, allowing you to 
@@ -147,6 +163,7 @@ public:
 struct FMovieGraphRenderPassLayerData
 {
 	FString LayerName;
+	FGuid CameraIdentifier;
 };
 // ToDo: Both of these can probably go into the Default Renderer implementation.
 
@@ -196,11 +213,32 @@ namespace UE::MovieGraph
 
 	struct FMovieGraphSampleState : IImagePixelDataPayload, public TSharedFromThis<FMovieGraphSampleState>
 	{
+		FMovieGraphSampleState()
+			: bWriteSampleToDisk(false)
+			, bRequiresAccumulator(false)
+			, bFetchFromAccumulator(false)
+		{}
+
 		/** The traversal context used to read graph values at the time of submission. */
 		FMovieGraphTraversalContext TraversalContext;
 
-		/** The backbuffer resolution of this render resource(ie: the resolution the sample was rendered at). */
+		/** The backbuffer resolution of this render resource (ie: the resolution the sample was rendered at). */
 		FIntPoint BackbufferResolution;
+
+		/** 
+		* The resolution the accumulation buffer is accumulating at. This may be different than our 
+		* backbuffer resolution (due to tiling), and our final output resolution (due to overscan + crop). 
+		*/
+		FIntPoint AccumulatorResolution;
+
+		/** Debug feature, should this sample be written to disk, not being accumulated? */
+		bool bWriteSampleToDisk;
+
+		/** Set this to true for every sample if an accumulator is required (ie: has temporal/spatial sub-samples, or high res tiling). */
+		bool bRequiresAccumulator;
+
+		/** Set this to true for the last sample for a given frame (if it requires accumulation) to fetch the accumulated data out of the accumulator. */
+		bool bFetchFromAccumulator;
 	};
 
 	/**
diff --git a/Engine/Plugins/MovieScene/MovieRenderPipeline/Source/MovieRenderPipelineCore/Public/Graph/MovieGraphDefaultRenderer.h b/Engine/Plugins/MovieScene/MovieRenderPipeline/Source/MovieRenderPipelineCore/Public/Graph/MovieGraphDefaultRenderer.h
--- a/Engine/Plugins/MovieScene/MovieRenderPipeline/Source/MovieRenderPipelineCore/Public/Graph/MovieGraphDefaultRenderer.h
+++ b/Engine/Plugins/MovieScene/MovieRenderPipeline/Source/MovieRenderPipelineCore/Public/Graph/MovieGraphDefaultRenderer.h
@@ -5,6 +5,8 @@
 #include "PixelFormat.h"
 #include "MovieRenderPipelineCoreModule.h"
 #include "Async/TaskGraphInterfaces.h"
+#include "Camera/CameraTypes.h"
+#include "Tasks/Task.h"
 #include "MovieGraphDefaultRenderer.generated.h"
 
 // Forward Declares
@@ -20,6 +22,17 @@ typedef TSharedPtr<UE::MovieGraph::DefaultRenderer::FSurfaceAccumulatorPool, ESP
 
 namespace UE::MovieGraph::DefaultRenderer
 {
+	struct FCameraInfo
+	{
+		FCameraInfo()
+			: ViewActor(nullptr)
+		{}
+
+		FMinimalViewInfo ViewInfo;
+		FString CameraName;
+		class AActor* ViewActor;
+	};
+
 	struct FRenderTargetInitParams
 	{
 		FRenderTargetInitParams()
@@ -68,7 +81,7 @@ namespace UE::MovieGraph::DefaultRenderer
 			int32 ActiveFrameNumber;
 			FMovieGraphRenderDataIdentifier ActivePassIdentifier;
 			FThreadSafeBool bIsActive;
-			FGraphEventRef TaskPrereq;
+			UE::Tasks::FTask TaskPrereq;
 		};
 		typedef TSharedPtr<FInstance, ESPMode::ThreadSafe> FInstancePtr;
 
@@ -83,6 +96,7 @@ namespace UE::MovieGraph::DefaultRenderer
 	{
 	public:
 		FGraphEventRef LastCompletionEvent;
+		UE::Tasks::FTask TaskPrereq;
 
 	public:
 		FGraphEventRef Execute(TUniqueFunction<void()> InFunctor)
@@ -136,8 +150,10 @@ public:
 	static void AddReferencedObjects(UObject* InThis, FReferenceCollector& Collector);
 	// ~UObject Interface
 
-	void AddOutstandingRenderTask_AnyThread(FGraphEventRef InTask);
-	
+	void AddOutstandingRenderTask_AnyThread(UE::Tasks::FTask InTask);
+	UE::MovieGraph::DefaultRenderer::FCameraInfo GetCameraInfo(const FGuid& InCameraIdentifier) const;
+
+public:
 	UTextureRenderTarget2D* GetOrCreateViewRenderTarget(const UE::MovieGraph::DefaultRenderer::FRenderTargetInitParams& InInitParams);
 	FMoviePipelineSurfaceQueuePtr GetOrCreateSurfaceQueue(const UE::MovieGraph::DefaultRenderer::FRenderTargetInitParams& InInitParams);
 	
@@ -199,5 +215,5 @@ protected:
 	/** Accessed by the Render Thread when starting up a new task.Makes sure we don't add tasks to the array while we're removing finished ones. */
 	FCriticalSection OutstandingTasksMutex;
 	/** Array of outstanding accumulation / blending tasks that are currently being worked on. */
-	FGraphEventArray OutstandingTasks;
+	TArray<UE::Tasks::FTask> OutstandingTasks;
 };
\ No newline at end of file
diff --git a/Engine/Plugins/MovieScene/MovieRenderPipeline/Source/MovieRenderPipelineCore/Public/Graph/MovieGraphLinearTimeStep.h b/Engine/Plugins/MovieScene/MovieRenderPipeline/Source/MovieRenderPipelineCore/Public/Graph/MovieGraphLinearTimeStep.h
--- a/Engine/Plugins/MovieScene/MovieRenderPipeline/Source/MovieRenderPipelineCore/Public/Graph/MovieGraphLinearTimeStep.h
+++ b/Engine/Plugins/MovieScene/MovieRenderPipeline/Source/MovieRenderPipelineCore/Public/Graph/MovieGraphLinearTimeStep.h
@@ -4,8 +4,12 @@
 #include "Graph/MovieGraphDataTypes.h"
 #include "Misc/FrameRate.h"
 #include "Misc/FrameTime.h"
+#include "Engine/EngineCustomTimeStep.h"
 #include "MovieGraphLinearTimeStep.generated.h"
 
+// Forward Declares
+class UMovieGraphEngineTimeStep;
+
 /**
 * This class is responsible for calculating the time step of each tick of the engine during a
 * Movie Render Queue render using a linear strategy - the number of temporal sub-samples is
@@ -21,23 +25,32 @@ class MOVIERENDERPIPELINECORE_API UMovieGraphLinearTimeStep : public UMovieGraph
 {
 	GENERATED_BODY()
 public:
+	UMovieGraphLinearTimeStep();
+
+	// UMovieGraphTimeStepBase Interface
 	virtual void TickProducingFrames() override;
+	virtual void Shutdown() override;
 	virtual FMovieGraphTimeStepData GetCalculatedTimeData() const override { return CurrentTimeStepData; }
+	// ~UMovieGraphTimeStepBase Interface
 
 protected:
 	virtual void UpdateFrameMetrics();
+	virtual bool IsFirstTemporalSample() const;
+	virtual bool IsLastTemporalSample() const;
+	virtual void UpdateTemporalSampleCount();
+	virtual void ResetForEndOfOutputFrame();
 	virtual float GetBlendedMotionBlurAmount();
 
 protected:
 	/** This is the output data needed by the rest of MRQ to produce a frame. */
 	FMovieGraphTimeStepData CurrentTimeStepData;
 
-	struct FFrameData
+	struct FOutputFrameMetrics
 	{
-		FFrameData()
+		FOutputFrameMetrics()
 			: MotionBlurAmount(0.f)
 		{}
-		
+
 		/** The human readable frame rate, adjusted by the config file (24fps, 30fps, etc.) */
 		FFrameRate FrameRate;
 
@@ -70,9 +83,104 @@ protected:
 		FFrameTime MotionBlurCenteringOffsetTime;
 	};
 
+	struct FCurrentFrameData
+	{
+		FCurrentFrameData()
+			: TemporalSampleIndex(0)
+			, TemporalSampleCount(0)
+			, OutputFrameNumber(0)
+		{
+		}
+
+		/** Which temporal sub-sample are we working on? [0,TemporalSampleCount) */
+		int32 TemporalSampleIndex;
+		
+		/** How many temporal sub-samples are there for the current output frame? */
+		int32 TemporalSampleCount;
+		
+		/** Which output frame are we working on, relative to zero.*/
+		int32 OutputFrameNumber;
+
+		/** A range of time (in Tick Resolution) for the last output frame being worked on. Updated before first TS of next frame. */
+		TRange<FFrameTime> LastOutputFrameRange;
+
+		/** A range of time (in Tick Resolution) for the current frame being worked on. Represents the whole output frame (shutter closed + open). */
+		TRange<FFrameTime> CurrentOutputFrameRange;
+		
+		/** A range of time (in Tick Resolution) that the shutter is opened. */
+		TRange<FFrameTime> RangeShutterOpen;
+		
+		/** A range of time (in Tick Resolution) that the shutter is closed. Can be an empty range (when using MotionBlurAmount=1.0) */
+		TRange<FFrameTime> RangeShutterClosed;
+
+		/**
+		* A range of time (in Tick Resolution) for the last temporal sub-sample (ie: what the last tick rendered with). We don't track
+		* CurrentSampleRange as it's just TemporalRanges[TemporalSampleIndex].
+		*/
+		TRange<FFrameTime> LastSampleRange;
+
+		/** A pre-calculated (when the output frame starts) range of time, breaking RangeShutterOpen into smaller chunks, one per temporal sub-sample. */
+		TArray<TRange<FFrameTime>> TemporalRanges;
+	};
+
 	/**
 	* A set of cached values that are true for the current output frame. They will be recalculated
 	* at the start of the next output frame.
 	*/
-	FFrameData CurrentFrameMetrics;
+	FOutputFrameMetrics CurrentFrameMetrics;
+
+	/**
+	* Frame Metrics are constant data updated once per output frame, while CurrentFrameData is keeping track
+	* across engine ticks to work towards producing a single frame.
+	*/
+	FCurrentFrameData CurrentFrameData;
+
+	/**
+	* A custom timestep owned by this object that is used to inform the engine what the delta time for each
+	* frame should be. 
+	*/
+	UPROPERTY(Transient)
+	TObjectPtr<UMovieGraphEngineTimeStep> CustomTimeStep;
+
+	/** The previous custom timestep the engine was using, if any. */
+	UPROPERTY(Transient)
+	TObjectPtr<UEngineCustomTimeStep> PrevCustomTimeStep;
+};
+
+UCLASS()
+class MOVIERENDERPIPELINECORE_API UMovieGraphEngineTimeStep : public UEngineCustomTimeStep
+{
+	GENERATED_BODY()
+public:
+	UMovieGraphEngineTimeStep();
+
+	struct FTimeStepCache
+	{
+		FTimeStepCache() 
+			: UndilatedDeltaTime(0.0)
+		{}
+
+		FTimeStepCache(double InUndilatedDeltaTime)
+			: UndilatedDeltaTime(InUndilatedDeltaTime)
+		{}
+
+		double UndilatedDeltaTime;
+	};
+
+public:
+	void SetCachedFrameTiming(const FTimeStepCache& InTimeCache);
+	
+	// UEngineCustomTimeStep Interface
+	virtual bool Initialize(UEngine* InEngine) override;
+	virtual void Shutdown(UEngine* InEngine) override;
+	virtual bool UpdateTimeStep(UEngine* InEngine) override;
+	virtual ECustomTimeStepSynchronizationState GetSynchronizationState() const override { return ECustomTimeStepSynchronizationState::Synchronized; }
+	// ~UEngineCustomTimeStep Interface
+
+	/** We don't do any thinking on our own, instead we just spit out the numbers stored in our time cache. */
+	FTimeStepCache TimeCache;
+
+	// Not cached in TimeCache as TimeCache is reset every frame.
+	float PrevMinUndilatedFrameTime;
+	float PrevMaxUndilatedFrameTime;
 };
diff --git a/Engine/Plugins/MovieScene/MovieRenderPipeline/Source/MovieRenderPipelineCore/Public/Graph/MovieGraphPipeline.h b/Engine/Plugins/MovieScene/MovieRenderPipeline/Source/MovieRenderPipelineCore/Public/Graph/MovieGraphPipeline.h
--- a/Engine/Plugins/MovieScene/MovieRenderPipeline/Source/MovieRenderPipelineCore/Public/Graph/MovieGraphPipeline.h
+++ b/Engine/Plugins/MovieScene/MovieRenderPipeline/Source/MovieRenderPipelineCore/Public/Graph/MovieGraphPipeline.h
@@ -11,6 +11,7 @@
 // Forward Declares
 class UMoviePipelineExecutorJob;
 class UMoviePipelineExecutorShot;
+class UMovieGraphFileOutputNode;
 class IImageWriteQueue;
 
 namespace UE::MovieGraph
@@ -72,6 +73,8 @@ protected:
 	virtual void TickPostFinalizeExport(const bool bInForceFinish);
 	virtual void TickFinalizeOutputContainers(const bool bInForceFinish);
 	virtual void TransitionToState(const EMovieRenderPipelineState InNewState);
+	virtual const TSet<TObjectPtr<UMovieGraphFileOutputNode>> GetOutputNodesUsed() const;
+	virtual void BeginFinalize();
 
 	// UMoviePipelineBase Interface
 	virtual void RequestShutdownImpl(bool bIsError) override;
@@ -98,6 +101,14 @@ protected:
 	UPROPERTY(Transient)
 	TArray<TObjectPtr<UMoviePipelineExecutorShot>> ActiveShotList;
 
+	/**
+	* An array of Node CDOs that we sent data through to write data to disk.
+	* This is the list of output nodes that will have OnAllFramesSubmitted/IsFinishedWritingToDisk
+	* on them.
+	*/
+	UPROPERTY(Transient)
+	TSet<TObjectPtr<UMovieGraphFileOutputNode>> OutputNodesDataSentTo;
+
 	int32 CurrentShotIndex;
 
 	/** True if we're in a TransitionToState call. Used to prevent reentrancy. */
diff --git a/Engine/Plugins/MovieScene/MovieRenderPipeline/Source/MovieRenderPipelineCore/Public/Graph/MovieGraphSequenceDataSource.h b/Engine/Plugins/MovieScene/MovieRenderPipeline/Source/MovieRenderPipelineCore/Public/Graph/MovieGraphSequenceDataSource.h
--- a/Engine/Plugins/MovieScene/MovieRenderPipeline/Source/MovieRenderPipelineCore/Public/Graph/MovieGraphSequenceDataSource.h
+++ b/Engine/Plugins/MovieScene/MovieRenderPipeline/Source/MovieRenderPipelineCore/Public/Graph/MovieGraphSequenceDataSource.h
@@ -3,6 +3,8 @@
 
 #include "Graph/MovieGraphDataTypes.h"
 #include "Misc/FrameTime.h"
+#include "MovieSceneTimeController.h"
+
 #include "MovieGraphSequenceDataSource.generated.h"
 
 // Forward Declares
@@ -10,6 +12,19 @@ class ULevelSequence;
 class ALevelSequenceActor;
 class UMovieSceneSequencePlayer;
 
+namespace UE::MovieGraph
+{
+	struct FMovieGraphSequenceTimeController : public FMovieSceneTimeController
+	{
+		virtual FFrameTime OnRequestCurrentTime(const FQualifiedFrameTime& InCurrentTime, float InPlayRate) override;
+		void SetCachedFrameTiming(const FQualifiedFrameTime& InTimeCache) { TimeCache = InTimeCache; }
+
+	private:
+		/** Simply store the number calculated and return it when requested. */
+		FQualifiedFrameTime TimeCache;
+	};
+}
+
 /**
 * The UMovieGraphSequenceDataSource allows using a ULevelSequence as the external datasource for the Movie Graph.
 * It will build the range of shots based on the contents of the level sequence (one shot per camera cut found inside
@@ -22,12 +37,14 @@ class MOVIERENDERPIPELINECORE_API UMovieGraphSequenceDataSource : public UMovieG
 	GENERATED_BODY()
 
 public:
+	UMovieGraphSequenceDataSource();
 	virtual void CacheDataPreJob(const FMovieGraphInitConfig& InInitConfig) override;
 	virtual void RestoreCachedDataPostJob() override;
 	virtual void UpdateShotList() override;
 	virtual FFrameRate GetTickResolution() const override;
 	virtual FFrameRate GetDisplayRate() const override;
-
+	virtual void SyncDataSourceTime(const FFrameTime& InTime) override;
+	virtual void InitializeShot(UMoviePipelineExecutorShot* InShot) override;
 protected:
 	void CacheLevelSequenceData(ULevelSequence* InSequence);
 	void OnSequenceEvaluated(const UMovieSceneSequencePlayer& Player, FFrameTime CurrentTime, FFrameTime PreviousTime);
@@ -36,4 +53,6 @@ protected:
 	UPROPERTY(Transient)
 	TObjectPtr<ALevelSequenceActor> LevelSequenceActor;
 
+	/** Custom Time Controller for the Sequence Player, used to match Custom TimeStep without any floating point accumulation errors. */
+	TSharedPtr<UE::MovieGraph::FMovieGraphSequenceTimeController> CustomSequenceTimeController;
 };
\ No newline at end of file
diff --git a/Engine/Plugins/MovieScene/MovieRenderPipeline/Source/MovieRenderPipelineCore/Public/Graph/MovieGraphTimeStepData.h b/Engine/Plugins/MovieScene/MovieRenderPipeline/Source/MovieRenderPipelineCore/Public/Graph/MovieGraphTimeStepData.h
--- a/Engine/Plugins/MovieScene/MovieRenderPipeline/Source/MovieRenderPipelineCore/Public/Graph/MovieGraphTimeStepData.h
+++ b/Engine/Plugins/MovieScene/MovieRenderPipeline/Source/MovieRenderPipelineCore/Public/Graph/MovieGraphTimeStepData.h
@@ -54,4 +54,12 @@ public:
 	*/
 	UPROPERTY(BlueprintReadWrite, EditAnywhere, Category = "Movie Graph")
 	bool bIsLastTemporalSampleForFrame;
+
+	/**
+	* Should be set to true for every sample if there is more than one temporal sample
+	* making up this render. This will cause the renderer to allocate accumulators
+	* to store the multi-frame data into.
+	*/
+	UPROPERTY(BlueprintReadWrite, EditAnywhere, Category = "Movie Graph")
+	bool bRequiresAccumulator;
 };
diff --git a/Engine/Plugins/MovieScene/MovieRenderPipeline/Source/MovieRenderPipelineCore/Public/Graph/Nodes/MovieGraphDeferredRenderPassNode.h b/Engine/Plugins/MovieScene/MovieRenderPipeline/Source/MovieRenderPipelineCore/Public/Graph/Nodes/MovieGraphDeferredRenderPassNode.h
--- a/Engine/Plugins/MovieScene/MovieRenderPipeline/Source/MovieRenderPipelineCore/Public/Graph/Nodes/MovieGraphDeferredRenderPassNode.h
+++ b/Engine/Plugins/MovieScene/MovieRenderPipeline/Source/MovieRenderPipelineCore/Public/Graph/Nodes/MovieGraphDeferredRenderPassNode.h
@@ -2,6 +2,7 @@
 
 #pragma once
 #include "Graph/Nodes/MovieGraphRenderPassNode.h"
+#include "Graph/MovieGraphDefaultRenderer.h" // For CameraInfo
 #include "Graph/MovieGraphDataTypes.h"
 #include "SceneTypes.h"
 #include "Camera/CameraTypes.h"
@@ -9,6 +10,7 @@
 #include "Engine/EngineTypes.h"
 #include "Async/TaskGraphFwd.h"
 #include "Styling/AppStyle.h"
+#include "Tasks/Task.h"
 #include "MovieGraphDeferredRenderPassNode.generated.h"
 
 // Forward Declares
@@ -42,8 +44,6 @@ namespace UE::MovieGraph
 		bool bIsFirstSample;
 		// If it's the last sample, then we will trigger moving the data to the output merger. It can be both the first and last sample at the same time.
 		bool bIsLastSample;
-		// Does this accumulation event depend on a previous accumulation event before it can be executed?
-		FGraphEventRef TaskPrerequisite;
 	};
 
 	void AccumulateSample_TaskThread(TUniquePtr<FImagePixelData>&& InPixelData, const UE::MovieGraph::FMovieGraphSampleState InSampleState, const UE::MovieGraph::FMovieGraphRenderDataAccumulationArgs& InAccumulationParams);
@@ -98,7 +98,6 @@ protected:
 			, AntiAliasingMethod(EAntiAliasingMethod::AAM_None)
 			, View(nullptr)
 			, SceneViewStateReference(nullptr)
-			, ViewActor(nullptr)
 		{
 		}
 
@@ -119,8 +118,8 @@ protected:
 		FSceneViewStateInterface* SceneViewStateReference;
 
 		// Camera Setup
-		FMinimalViewInfo MinimalViewInfo;
-		class AActor* ViewActor;
+		UE::MovieGraph::DefaultRenderer::FCameraInfo CameraInfo;
+
 	};
 
 	struct FMovieGraphDeferredRenderPass
@@ -142,6 +141,9 @@ protected:
 	protected:
 		FMovieGraphRenderPassLayerData LayerData;
 
+		/** Unique identifier passed in GatherOutputPasses and with each render that identifies the data produced by this renderer. */
+		FMovieGraphRenderDataIdentifier RenderDataIdentifier;
+
 		// Scene View history used by the renderer 
 		FSceneViewStateReference SceneViewState;
 
diff --git a/Engine/Plugins/MovieScene/MovieRenderPipeline/Source/MovieRenderPipelineCore/Public/Graph/Nodes/MovieGraphFileOutputNode.h b/Engine/Plugins/MovieScene/MovieRenderPipeline/Source/MovieRenderPipelineCore/Public/Graph/Nodes/MovieGraphFileOutputNode.h
--- a/Engine/Plugins/MovieScene/MovieRenderPipeline/Source/MovieRenderPipelineCore/Public/Graph/Nodes/MovieGraphFileOutputNode.h
+++ b/Engine/Plugins/MovieScene/MovieRenderPipeline/Source/MovieRenderPipelineCore/Public/Graph/Nodes/MovieGraphFileOutputNode.h
@@ -29,6 +29,8 @@ public:
 	* memory until the end of a render.
 	*/
 	void OnReceiveImageData(UMovieGraphPipeline* InPipeline, UE::MovieGraph::FMovieGraphOutputMergerFrame* InRawFrameData) { OnReceiveImageDataImpl(InPipeline, InRawFrameData); }
+	void OnAllFramesSubmitted() { OnAllFramesSubmittedImpl(); }
+	bool IsFinishedWritingToDisk() const { return IsFinishedWritingToDiskImpl(); }
 	// ~UMovieGraphFileOutputNode Interface
 
 	virtual TArray<FMovieGraphPinProperties> GetInputPinProperties() const override
@@ -54,4 +56,6 @@ public:
 
 protected:
 	virtual void OnReceiveImageDataImpl(UMovieGraphPipeline* InPipeline, UE::MovieGraph::FMovieGraphOutputMergerFrame* InRawFrameData) {}
+	virtual void OnAllFramesSubmittedImpl() {}
+	virtual bool IsFinishedWritingToDiskImpl() const { return true; }
 };
\ No newline at end of file
diff --git a/Engine/Plugins/MovieScene/MovieRenderPipeline/Source/MovieRenderPipelineCore/Public/Graph/Nodes/MovieGraphImageSequenceOutputNode.h b/Engine/Plugins/MovieScene/MovieRenderPipeline/Source/MovieRenderPipelineCore/Public/Graph/Nodes/MovieGraphImageSequenceOutputNode.h
--- a/Engine/Plugins/MovieScene/MovieRenderPipeline/Source/MovieRenderPipelineCore/Public/Graph/Nodes/MovieGraphImageSequenceOutputNode.h
+++ b/Engine/Plugins/MovieScene/MovieRenderPipeline/Source/MovieRenderPipelineCore/Public/Graph/Nodes/MovieGraphImageSequenceOutputNode.h
@@ -4,6 +4,7 @@
 #include "Graph/Nodes/MovieGraphFileOutputNode.h"
 #include "IImageWrapper.h"
 #include "Styling/AppStyle.h"
+#include "Async/Future.h"
 #include "MovieGraphImageSequenceOutputNode.generated.h"
 
 // Forward Declare
@@ -23,6 +24,8 @@ public:
 	
 	// UMovieGraphFileOutputNode Interface
 	virtual void OnReceiveImageDataImpl(UMovieGraphPipeline* InPipeline, UE::MovieGraph::FMovieGraphOutputMergerFrame* InRawFrameData) override;
+	virtual void OnAllFramesSubmittedImpl() override;
+	virtual bool IsFinishedWritingToDiskImpl() const override;
 	// ~UMovieGraphFileOutputNode Interface
 
 protected:
@@ -31,6 +34,9 @@ protected:
 
 	/** A pointer to the image write queue used for asynchronously writing images */
 	IImageWriteQueue* ImageWriteQueue;
+
+	/** A fence to keep track of when the Image Write queue has fully flushed. */
+	TFuture<void> FinalizeFence;
 };
 
 /**
diff --git a/Engine/Plugins/MovieScene/MovieRenderPipeline/Source/MovieRenderPipelineCore/Public/MovieRenderPipelineDataTypes.h b/Engine/Plugins/MovieScene/MovieRenderPipeline/Source/MovieRenderPipelineCore/Public/MovieRenderPipelineDataTypes.h
--- a/Engine/Plugins/MovieScene/MovieRenderPipeline/Source/MovieRenderPipelineCore/Public/MovieRenderPipelineDataTypes.h
+++ b/Engine/Plugins/MovieScene/MovieRenderPipeline/Source/MovieRenderPipelineCore/Public/MovieRenderPipelineDataTypes.h
@@ -529,6 +529,9 @@ public:
 	/** The current tick of this shot that we're on in root space */
 	FFrameNumber CurrentTickInRoot;
 
+	/** The current time of this shot that we're on in root space (for the Graph version. Tracks sub-frames too.) */
+	FFrameTime CurrentTimeInRoot;
+
 	/** Converts from the outermost space into the innermost space. Only works with linear transforms. */
 	FMovieSceneTimeTransform OuterToInnerTransform;
 	
