CIS: Fixed uninitialized variable in struct.

diff --git a/Engine/Plugins/MovieScene/MovieRenderPipeline/Source/MovieRenderPipelineCore/Public/Graph/MovieGraphTimeStepData.h b/Engine/Plugins/MovieScene/MovieRenderPipeline/Source/MovieRenderPipelineCore/Public/Graph/MovieGraphTimeStepData.h
--- a/Engine/Plugins/MovieScene/MovieRenderPipeline/Source/MovieRenderPipelineCore/Public/Graph/MovieGraphTimeStepData.h
+++ b/Engine/Plugins/MovieScene/MovieRenderPipeline/Source/MovieRenderPipelineCore/Public/Graph/MovieGraphTimeStepData.h
@@ -23,6 +23,7 @@ public:
 		, MotionBlurFraction(0.f)
 		, bIsFirstTemporalSampleForFrame(false)
 		, bIsLastTemporalSampleForFrame(false)
+		, bRequiresAccumulator(false)
 	{}
 
 	UPROPERTY(BlueprintReadWrite, EditAnywhere, Category = "Movie Graph")
