MoviePipeline: Updated graph members to utilize property bags for variable data storage, and updated jobs to store graph variable overrides in property bags. This lets variables contain data of any type.

Property bags were chosen for this because they:
* Allow for the type of a property to be easily changed in the details panel
* Takes care of type-agnostic data storage
* Allow for seting the value of the property in both C++ and Python
* Allow changing the property at runtime

Adjusted the details pane so jobs display the new variable overrides correctly.

diff --git a/Engine/Plugins/MovieScene/MovieRenderPipeline/Source/MovieRenderPipelineCore/Private/Graph/MovieGraphConfig.cpp b/Engine/Plugins/MovieScene/MovieRenderPipeline/Source/MovieRenderPipelineCore/Private/Graph/MovieGraphConfig.cpp
--- a/Engine/Plugins/MovieScene/MovieRenderPipeline/Source/MovieRenderPipelineCore/Private/Graph/MovieGraphConfig.cpp
+++ b/Engine/Plugins/MovieScene/MovieRenderPipeline/Source/MovieRenderPipelineCore/Private/Graph/MovieGraphConfig.cpp
@@ -85,6 +85,17 @@ void UMovieGraphConfig::PostLoad()
 		// upgrades/deprecations. For now, we assume that each load of the graph should re-initialize all default
 		// members.
 		AddDefaultMembers();
+
+		// Fire OnGraphVariablesChangedDelegate when the variable changes (name, value, type, etc)
+#if WITH_EDITOR
+		for (const TObjectPtr<UMovieGraphVariable>& Variable : Variables)
+		{
+			Variable->OnMovieGraphVariableChangedDelegate.AddWeakLambda(this, [this](UMovieGraphMember*)
+			{
+				OnGraphVariablesChangedDelegate.Broadcast();
+			});
+		}
+#endif
 	}
 }
 
@@ -424,7 +435,6 @@ T* UMovieGraphConfig::AddMember(TArray<TObjectPtr<T>>& InMemberArray, const FNam
 
 	InMemberArray.Add(NewMember);
 	NewMember->SetFlags(RF_Transactional);
-	NewMember->Type = EMovieGraphMemberType::Float;
 	NewMember->SetGuid(FGuid::NewGuid());
 
 	// Generate and set a unique name
@@ -442,6 +452,17 @@ UMovieGraphVariable* UMovieGraphConfig::AddVariable(const FName InCustomBaseName
 	UMovieGraphVariable* NewVariable = AddMember(
 		Variables, !InCustomBaseName.IsNone() ? InCustomBaseName : FName(*VariableBaseName.ToString()));
 
+	if (NewVariable)
+	{
+		// Fire OnGraphVariablesChangedDelegate when the variable changes (name, value, type, etc)
+#if WITH_EDITOR
+		NewVariable->OnMovieGraphVariableChangedDelegate.AddWeakLambda(this, [this](UMovieGraphMember*)
+		{
+			OnGraphVariablesChangedDelegate.Broadcast();
+		});
+#endif
+	}
+
 #if WITH_EDITOR
 	OnGraphVariablesChangedDelegate.Broadcast();
 #endif
diff --git a/Engine/Plugins/MovieScene/MovieRenderPipeline/Source/MovieRenderPipelineCore/Private/Graph/MovieGraphNode.cpp b/Engine/Plugins/MovieScene/MovieRenderPipeline/Source/MovieRenderPipelineCore/Private/Graph/MovieGraphNode.cpp
--- a/Engine/Plugins/MovieScene/MovieRenderPipeline/Source/MovieRenderPipelineCore/Private/Graph/MovieGraphNode.cpp
+++ b/Engine/Plugins/MovieScene/MovieRenderPipeline/Source/MovieRenderPipelineCore/Private/Graph/MovieGraphNode.cpp
@@ -17,7 +17,7 @@ TArray<FMovieGraphPinProperties> UMovieGraphNode::GetExposedDynamicPinProperties
 	{
 		// Note: Currently hardcoded to not allow multiple connections, but this may need to change in the future
 		const bool bInAllowMultipleConnections = false;
-		Properties.Add(FMovieGraphPinProperties(PropertyName, EMovieGraphMemberType::Float, bInAllowMultipleConnections));
+		Properties.Add(FMovieGraphPinProperties(PropertyName, EMovieGraphValueType::Float, bInAllowMultipleConnections));
 	}
 
 	return Properties;
diff --git a/Engine/Plugins/MovieScene/MovieRenderPipeline/Source/MovieRenderPipelineCore/Private/Graph/Nodes/MovieGraphInputNode.cpp b/Engine/Plugins/MovieScene/MovieRenderPipeline/Source/MovieRenderPipelineCore/Private/Graph/Nodes/MovieGraphInputNode.cpp
--- a/Engine/Plugins/MovieScene/MovieRenderPipeline/Source/MovieRenderPipelineCore/Private/Graph/Nodes/MovieGraphInputNode.cpp
+++ b/Engine/Plugins/MovieScene/MovieRenderPipeline/Source/MovieRenderPipelineCore/Private/Graph/Nodes/MovieGraphInputNode.cpp
@@ -1,4 +1,4 @@
-﻿// Copyright Epic Games, Inc. All Rights Reserved.
+// Copyright Epic Games, Inc. All Rights Reserved.
 
 #include "Graph/Nodes/MovieGraphInputNode.h"
 
@@ -27,7 +27,7 @@ TArray<FMovieGraphPinProperties> UMovieGraphInputNode::GetOutputPinProperties()
 	{
 		for (const UMovieGraphInput* Input : ParentGraph->GetInputs())
 		{
-			Properties.Add(FMovieGraphPinProperties(FName(Input->Name), EMovieGraphMemberType::Branch, false));
+			Properties.Add(FMovieGraphPinProperties(FName(Input->Name), EMovieGraphValueType::Branch, false));
 		}
 	}
 	
diff --git a/Engine/Plugins/MovieScene/MovieRenderPipeline/Source/MovieRenderPipelineCore/Private/Graph/Nodes/MovieGraphOutputNode.cpp b/Engine/Plugins/MovieScene/MovieRenderPipeline/Source/MovieRenderPipelineCore/Private/Graph/Nodes/MovieGraphOutputNode.cpp
--- a/Engine/Plugins/MovieScene/MovieRenderPipeline/Source/MovieRenderPipelineCore/Private/Graph/Nodes/MovieGraphOutputNode.cpp
+++ b/Engine/Plugins/MovieScene/MovieRenderPipeline/Source/MovieRenderPipelineCore/Private/Graph/Nodes/MovieGraphOutputNode.cpp
@@ -27,7 +27,7 @@ TArray<FMovieGraphPinProperties> UMovieGraphOutputNode::GetInputPinProperties()
 	{
 		for (const UMovieGraphOutput* Output : ParentGraph->GetOutputs())
 		{
-			Properties.Add(FMovieGraphPinProperties(FName(Output->Name), EMovieGraphMemberType::Branch, false));
+			Properties.Add(FMovieGraphPinProperties(FName(Output->Name), EMovieGraphValueType::Branch, false));
 		}
 	}
 	
diff --git a/Engine/Plugins/MovieScene/MovieRenderPipeline/Source/MovieRenderPipelineCore/Private/Graph/Nodes/MovieGraphRenderLayerNode.cpp b/Engine/Plugins/MovieScene/MovieRenderPipeline/Source/MovieRenderPipelineCore/Private/Graph/Nodes/MovieGraphRenderLayerNode.cpp
--- a/Engine/Plugins/MovieScene/MovieRenderPipeline/Source/MovieRenderPipelineCore/Private/Graph/Nodes/MovieGraphRenderLayerNode.cpp
+++ b/Engine/Plugins/MovieScene/MovieRenderPipeline/Source/MovieRenderPipelineCore/Private/Graph/Nodes/MovieGraphRenderLayerNode.cpp
@@ -13,14 +13,14 @@ UMovieGraphRenderLayerNode::UMovieGraphRenderLayerNode()
 TArray<FMovieGraphPinProperties> UMovieGraphRenderLayerNode::GetInputPinProperties() const
 {
 	TArray<FMovieGraphPinProperties> Properties;
-	Properties.Add(FMovieGraphPinProperties(NAME_None, EMovieGraphMemberType::Branch, false));
+	Properties.Add(FMovieGraphPinProperties(NAME_None, EMovieGraphValueType::Branch, false));
 	return Properties;
 }
 
 TArray<FMovieGraphPinProperties> UMovieGraphRenderLayerNode::GetOutputPinProperties() const
 {
 	TArray<FMovieGraphPinProperties> Properties;
-	Properties.Add(FMovieGraphPinProperties(NAME_None, EMovieGraphMemberType::Branch, false));
+	Properties.Add(FMovieGraphPinProperties(NAME_None, EMovieGraphValueType::Branch, false));
 	return Properties;
 }
 
diff --git a/Engine/Plugins/MovieScene/MovieRenderPipeline/Source/MovieRenderPipelineCore/Private/Graph/Nodes/MovieGraphVariableNode.cpp b/Engine/Plugins/MovieScene/MovieRenderPipeline/Source/MovieRenderPipelineCore/Private/Graph/Nodes/MovieGraphVariableNode.cpp
--- a/Engine/Plugins/MovieScene/MovieRenderPipeline/Source/MovieRenderPipelineCore/Private/Graph/Nodes/MovieGraphVariableNode.cpp
+++ b/Engine/Plugins/MovieScene/MovieRenderPipeline/Source/MovieRenderPipelineCore/Private/Graph/Nodes/MovieGraphVariableNode.cpp
@@ -19,11 +19,11 @@ TArray<FMovieGraphPinProperties> UMovieGraphVariableNode::GetOutputPinProperties
 	
 	if (GraphVariable)
 	{
-		Properties.Add(FMovieGraphPinProperties(FName(GraphVariable->Name), GraphVariable->Type, false));
+		Properties.Add(FMovieGraphPinProperties(FName(GraphVariable->Name), GraphVariable->GetValueType(), false));
 	}
 	else
 	{
-		Properties.Add(FMovieGraphPinProperties(TEXT("Unknown"), EMovieGraphMemberType::Float, false));
+		Properties.Add(FMovieGraphPinProperties(TEXT("Unknown"), EMovieGraphValueType::Unknown, false));
 	}
 	
 	return Properties;
@@ -87,7 +87,7 @@ void UMovieGraphVariableNode::UpdateOutputPin(UMovieGraphMember* ChangedVariable
 	{
 		// Update the output pin to reflect the variable data model
 		OutputPins[0]->Properties.Label = FName(ChangedVariable->Name);
-		OutputPins[0]->Properties.Type = ChangedVariable->Type;
+		OutputPins[0]->Properties.Type = ChangedVariable->GetValueType();
 	}
 
 	OnNodeChangedDelegate.Broadcast(this);
diff --git a/Engine/Plugins/MovieScene/MovieRenderPipeline/Source/MovieRenderPipelineCore/Private/MovieJobVariableAssignmentContainer.cpp b/Engine/Plugins/MovieScene/MovieRenderPipeline/Source/MovieRenderPipelineCore/Private/MovieJobVariableAssignmentContainer.cpp
new file mode 100644
--- /dev/null
+++ b/Engine/Plugins/MovieScene/MovieRenderPipeline/Source/MovieRenderPipelineCore/Private/MovieJobVariableAssignmentContainer.cpp
@@ -0,0 +1,465 @@
+// Copyright Epic Games, Inc. All Rights Reserved.
+
+#include "MovieJobVariableAssignmentContainer.h"
+
+#include "Graph/MovieGraphCommon.h"
+#include "Graph/MovieGraphConfig.h"
+
+// The prefix for property descriptions that are in charge of variable assignment EditCondition enable/disable state
+static const TCHAR* EditConditionPrefix = TEXT("bOverride_");
+
+void UMovieJobVariableAssignmentContainer::SetGraphConfig(const TSoftObjectPtr<UMovieGraphConfig>& InGraphConfig)
+{
+	GraphPreset = InGraphConfig;
+}
+
+uint32 UMovieJobVariableAssignmentContainer::GetNumAssignments() const
+{
+	uint32 NumAssignments = 0;
+	
+	if (const UPropertyBag* PropertyBag = Value.GetPropertyBagStruct())
+	{
+		// Only count the descs that aren't in charge of the enable/disable state of a variable assignment
+		for (const FPropertyBagPropertyDesc& Desc : PropertyBag->GetPropertyDescs())
+		{
+			if (!Desc.Name.ToString().StartsWith(EditConditionPrefix))
+			{
+				NumAssignments++;
+			}
+		}
+	}
+
+	return NumAssignments;
+}
+
+bool UMovieJobVariableAssignmentContainer::GetValueBool(const FName& PropertyName, bool& bOutValue) const
+{
+	TValueOrError<bool, EPropertyBagResult> Result = Value.GetValueBool(PropertyName);
+	return UE::MovieGraph::Private::GetOptionalValue<bool>(Result, bOutValue);
+}
+
+bool UMovieJobVariableAssignmentContainer::GetValueByte(const FName& PropertyName, uint8& OutValue) const
+{
+	TValueOrError<uint8, EPropertyBagResult> Result = Value.GetValueByte(PropertyName);
+	return UE::MovieGraph::Private::GetOptionalValue<uint8>(Result, OutValue);
+}
+
+bool UMovieJobVariableAssignmentContainer::GetValueInt32(const FName& PropertyName, int32& OutValue) const
+{
+	TValueOrError<int32, EPropertyBagResult> Result = Value.GetValueInt32(PropertyName);
+	return UE::MovieGraph::Private::GetOptionalValue<int32>(Result, OutValue);
+}
+
+bool UMovieJobVariableAssignmentContainer::GetValueInt64(const FName& PropertyName, int64& OutValue) const
+{
+	TValueOrError<int64, EPropertyBagResult> Result = Value.GetValueInt64(PropertyName);
+	return UE::MovieGraph::Private::GetOptionalValue<int64>(Result, OutValue);
+}
+
+bool UMovieJobVariableAssignmentContainer::GetValueFloat(const FName& PropertyName, float& OutValue) const
+{
+	TValueOrError<float, EPropertyBagResult> Result = Value.GetValueFloat(PropertyName);
+	return UE::MovieGraph::Private::GetOptionalValue<float>(Result, OutValue);
+}
+
+bool UMovieJobVariableAssignmentContainer::GetValueDouble(const FName& PropertyName, double& OutValue) const
+{
+	TValueOrError<double, EPropertyBagResult> Result = Value.GetValueDouble(PropertyName);
+	return UE::MovieGraph::Private::GetOptionalValue<double>(Result, OutValue);
+}
+
+bool UMovieJobVariableAssignmentContainer::GetValueName(const FName& PropertyName, FName& OutValue) const
+{
+	TValueOrError<FName, EPropertyBagResult> Result = Value.GetValueName(PropertyName);
+	return UE::MovieGraph::Private::GetOptionalValue<FName>(Result, OutValue);
+}
+
+bool UMovieJobVariableAssignmentContainer::GetValueString(const FName& PropertyName, FString& OutValue) const
+{
+	TValueOrError<FString, EPropertyBagResult> Result = Value.GetValueString(PropertyName);
+	return UE::MovieGraph::Private::GetOptionalValue<FString>(Result, OutValue);
+}
+
+bool UMovieJobVariableAssignmentContainer::GetValueText(const FName& PropertyName, FText& OutValue) const
+{
+	TValueOrError<FText, EPropertyBagResult> Result = Value.GetValueText(PropertyName);
+	return UE::MovieGraph::Private::GetOptionalValue<FText>(Result, OutValue);
+}
+
+bool UMovieJobVariableAssignmentContainer::GetValueEnum(const FName& PropertyName, uint8& OutValue, const UEnum* RequestedEnum) const
+{
+	TValueOrError<uint8, EPropertyBagResult> Result = Value.GetValueEnum(PropertyName, RequestedEnum);
+	return UE::MovieGraph::Private::GetOptionalValue<uint8>(Result, OutValue);
+}
+
+bool UMovieJobVariableAssignmentContainer::GetValueStruct(const FName& PropertyName, FStructView& OutValue, const UScriptStruct* RequestedStruct) const
+{
+	TValueOrError<FStructView, EPropertyBagResult> Result = Value.GetValueStruct(PropertyName, RequestedStruct);
+	return UE::MovieGraph::Private::GetOptionalValue<FStructView>(Result, OutValue);
+}
+
+bool UMovieJobVariableAssignmentContainer::GetValueObject(const FName& PropertyName, UObject* OutValue, const UClass* RequestedClass) const
+{
+	TValueOrError<UObject*, EPropertyBagResult> Result = Value.GetValueObject(PropertyName, RequestedClass);
+	return UE::MovieGraph::Private::GetOptionalValue<UObject*>(Result, OutValue);
+}
+
+bool UMovieJobVariableAssignmentContainer::GetValueClass(const FName& PropertyName, UClass* OutValue) const
+{
+	TValueOrError<UClass*, EPropertyBagResult> Result = Value.GetValueClass(PropertyName);
+	return UE::MovieGraph::Private::GetOptionalValue<UClass*>(Result, OutValue);
+}
+
+FString UMovieJobVariableAssignmentContainer::GetValueSerializedString(const FName& PropertyName)
+{
+	TValueOrError<FString, EPropertyBagResult> Result = Value.GetValueSerializedString(PropertyName);
+	FString ResultString;
+	UE::MovieGraph::Private::GetOptionalValue<FString>(Result, ResultString);
+	return ResultString;
+}
+
+bool UMovieJobVariableAssignmentContainer::SetValueBool(const FName& PropertyName, const bool bInValue)
+{
+	return Value.SetValueBool(PropertyName, bInValue) == EPropertyBagResult::Success;
+}
+
+bool UMovieJobVariableAssignmentContainer::SetValueByte(const FName& PropertyName, const uint8 InValue)
+{
+	return Value.SetValueByte(PropertyName, InValue) == EPropertyBagResult::Success;
+}
+
+bool UMovieJobVariableAssignmentContainer::SetValueInt32(const FName& PropertyName, const int32 InValue)
+{
+	return Value.SetValueInt32(PropertyName, InValue) == EPropertyBagResult::Success;
+}
+
+bool UMovieJobVariableAssignmentContainer::SetValueInt64(const FName& PropertyName, const int64 InValue)
+{
+	return Value.SetValueInt64(PropertyName, InValue) == EPropertyBagResult::Success;
+}
+
+bool UMovieJobVariableAssignmentContainer::SetValueFloat(const FName& PropertyName, const float InValue)
+{
+	return Value.SetValueFloat(PropertyName, InValue) == EPropertyBagResult::Success;
+}
+
+bool UMovieJobVariableAssignmentContainer::SetValueDouble(const FName& PropertyName, const double InValue)
+{
+	return Value.SetValueDouble(PropertyName, InValue) == EPropertyBagResult::Success;
+}
+
+bool UMovieJobVariableAssignmentContainer::SetValueName(const FName& PropertyName, const FName InValue)
+{
+	return Value.SetValueName(PropertyName, InValue) == EPropertyBagResult::Success;
+}
+
+bool UMovieJobVariableAssignmentContainer::SetValueString(const FName& PropertyName, const FString& InValue)
+{
+	return Value.SetValueString(PropertyName, InValue) == EPropertyBagResult::Success;
+}
+
+bool UMovieJobVariableAssignmentContainer::SetValueText(const FName& PropertyName, const FText& InValue)
+{
+	return Value.SetValueText(PropertyName, InValue) == EPropertyBagResult::Success;
+}
+
+bool UMovieJobVariableAssignmentContainer::SetValueEnum(const FName& PropertyName, const uint8 InValue, const UEnum* Enum)
+{
+	return Value.SetValueEnum(PropertyName, InValue, Enum) == EPropertyBagResult::Success;
+}
+
+bool UMovieJobVariableAssignmentContainer::SetValueStruct(const FName& PropertyName, FConstStructView InValue)
+{
+	return Value.SetValueStruct(PropertyName, InValue) == EPropertyBagResult::Success;
+}
+
+bool UMovieJobVariableAssignmentContainer::SetValueObject(const FName& PropertyName, UObject* InValue)
+{
+	return Value.SetValueObject(PropertyName, InValue) == EPropertyBagResult::Success;
+}
+
+bool UMovieJobVariableAssignmentContainer::SetValueClass(const FName& PropertyName, UClass* InValue)
+{
+	return Value.SetValueClass(PropertyName, InValue) == EPropertyBagResult::Success;
+}
+
+bool UMovieJobVariableAssignmentContainer::SetValueSerializedString(const FName& PropertyName, const FString& NewValue)
+{
+	return Value.SetValueSerializedString(PropertyName, NewValue) == EPropertyBagResult::Success;
+}
+
+EMovieGraphValueType UMovieJobVariableAssignmentContainer::GetValueType(const FName& PropertyName) const
+{
+	if (const FPropertyBagPropertyDesc* Desc = Value.FindPropertyDescByName(PropertyName))
+	{
+		return static_cast<EMovieGraphValueType>(Desc->ValueType);
+	}
+
+	return EMovieGraphValueType::None;
+}
+
+const UObject* UMovieJobVariableAssignmentContainer::GetValueTypeObject(const FName& PropertyName) const
+{
+	if (const FPropertyBagPropertyDesc* Desc = Value.FindPropertyDescByName(PropertyName))
+	{
+		return Desc->ValueTypeObject;
+	}
+	
+	return nullptr;
+}
+
+EMovieGraphContainerType UMovieJobVariableAssignmentContainer::GetValueContainerType(const FName& PropertyName) const
+{
+	if (const FPropertyBagPropertyDesc* Desc = Value.FindPropertyDescByName(PropertyName))
+	{
+		return static_cast<EMovieGraphContainerType>(Desc->ContainerType);
+	}
+
+	return EMovieGraphContainerType::None;
+}
+
+bool UMovieJobVariableAssignmentContainer::FindOrGenerateVariableOverride(
+	const UMovieGraphVariable* InGraphVariable, FPropertyBagPropertyDesc* OutPropDesc, FPropertyBagPropertyDesc* OutEditConditionPropDesc, bool bGenerateIfNotExists)
+{
+	if (InGraphVariable)
+	{
+		const FString EditCondPropName = FString::Format(TEXT("{0}{1}"), {EditConditionPrefix, InGraphVariable->Name});
+		
+		if (const FPropertyBagPropertyDesc* Desc = Value.FindPropertyDescByName(FName(InGraphVariable->Name)))
+		{
+			*OutPropDesc = *Desc;
+
+			if (const FPropertyBagPropertyDesc* EditCondPropDesc = Value.FindPropertyDescByName(FName(EditCondPropName)))
+			{
+				*OutEditConditionPropDesc = *EditCondPropDesc;
+			}
+			
+			return true;
+		}
+
+		// Didn't find an existing one; add a new assignment
+		if (bGenerateIfNotExists)
+		{
+			return GenerateVariableOverride(InGraphVariable, OutPropDesc, OutEditConditionPropDesc);
+		}
+	}
+
+	return false;
+}
+
+bool UMovieJobVariableAssignmentContainer::GenerateVariableOverride(const UMovieGraphVariable* InGraphVariable,
+	FPropertyBagPropertyDesc* OutPropDesc, FPropertyBagPropertyDesc* OutEditConditionPropDesc)
+{
+	// Create the description for the variable override
+	FPropertyBagPropertyDesc NewProperty = FPropertyBagPropertyDesc();
+	NewProperty.ValueType = static_cast<EPropertyBagPropertyType>(InGraphVariable->GetValueType());
+	NewProperty.ValueTypeObject = InGraphVariable->GetValueTypeObject();
+	NewProperty.ContainerType = static_cast<EPropertyBagContainerType>(InGraphVariable->GetValueContainerType());
+	NewProperty.Name = FName(InGraphVariable->Name);
+	NewProperty.ID = FGuid::NewGuid();
+#if WITH_EDITOR
+	NewProperty.MetaData.Add(FPropertyBagPropertyDescMetaData("VariableGUID", InGraphVariable->GetGuid().ToString()));
+#endif
+
+	// Track a separate EditCondition property that can enable/disable the above property. Since the variable can be
+	// renamed, base this property's name on the guid rather than the desc name.
+	const FString EditCondPropName = FString::Format(TEXT("{0}{1}"), {EditConditionPrefix, NewProperty.ID.ToString()});
+	FPropertyBagPropertyDesc NewPropertyEditCondition = FPropertyBagPropertyDesc(FName(EditCondPropName), EPropertyBagPropertyType::Bool);
+	NewPropertyEditCondition.ID = FGuid::NewGuid();
+#if WITH_EDITOR
+	NewPropertyEditCondition.MetaData.Add(FPropertyBagPropertyDescMetaData("VariableGUID", InGraphVariable->GetGuid().ToString()));
+#endif
+
+#if WITH_EDITOR
+	NewPropertyEditCondition.MetaData.Add(FPropertyBagPropertyDescMetaData("InlineEditConditionToggle", "true"));
+	NewProperty.MetaData.Add(FPropertyBagPropertyDescMetaData("EditCondition", EditCondPropName));
+#endif
+
+	*OutPropDesc = NewProperty;
+	*OutEditConditionPropDesc = NewPropertyEditCondition;
+	
+	return true;
+}
+
+#if WITH_EDITOR
+void UMovieJobVariableAssignmentContainer::UpdateGraphVariableOverrides()
+{
+	auto GetVariableGuidFromDesc = [](const FPropertyBagPropertyDesc& Desc) -> FGuid
+	{
+		const FPropertyBagPropertyDescMetaData* Meta = Desc.MetaData.FindByPredicate([](const FPropertyBagPropertyDescMetaData& MetaData)
+		{
+			return MetaData.Key == "VariableGUID";
+		});
+		
+		if (Meta)
+		{
+			FGuid VariableGUID;
+			if (FGuid::Parse(Meta->Value, VariableGUID))
+			{
+				return VariableGUID;
+			}
+		}
+		
+		return FGuid();
+	};
+	
+	UMovieGraphConfig* LoadedPreset = GraphPreset.LoadSynchronous();
+	if (!LoadedPreset)
+	{
+		return;
+	}
+
+	// Get the existing property descs. Modify (update, add, remove) the descs as needed to reflect the variables in the
+	// graph preset. If the descs were changed, flag bNeedsToRegenerate to trigger rebuilding the property bag.
+	const UPropertyBag* PropertyBag = Value.GetPropertyBagStruct();
+	bool bNeedsToRegenerate = false;
+	TArray<FPropertyBagPropertyDesc> ModifiedDescs;
+	if (PropertyBag)
+	{
+		ModifiedDescs = PropertyBag->GetPropertyDescs();
+	}
+
+	// First, do a sweep to remove any variable overrides which no longer correspond to graph variables. Bag migration will
+	// not remove properties from the bag, so this needs to be done as another step.
+	if (PropertyBag)
+	{
+		const int32 NumDescsRemoved = ModifiedDescs.RemoveAll([&GetVariableGuidFromDesc, LoadedPreset](const FPropertyBagPropertyDesc& Desc)
+		{
+			const FGuid VariableGUID = GetVariableGuidFromDesc(Desc);
+		
+			const bool bVariableExists = LoadedPreset->GetVariables().ContainsByPredicate([&VariableGUID](const UMovieGraphVariable* Variable)
+				{ return Variable && (Variable->GetGuid() == VariableGUID); });
+
+			return !bVariableExists;
+		});
+
+		if (NumDescsRemoved > 0)
+		{
+			bNeedsToRegenerate = true;
+		}
+	}
+	
+	TArray<UMovieGraphVariable*> GraphVariables = LoadedPreset->GetVariables();
+
+	// Second, ensure all variables in the bag match all variables in the graph; if something doesn't match,
+	// update the desc and flag that the bag needs to be regenerated.
+	if (PropertyBag)
+	{		
+		for (FPropertyBagPropertyDesc& Desc : ModifiedDescs)
+		{
+			// Skip descs that are responsible for the override checkmark; nothing about them will need to be updated
+			if (Desc.Name.ToString().StartsWith(EditConditionPrefix))
+			{
+				continue;
+			}
+		
+			const FGuid VariableGUID = GetVariableGuidFromDesc(Desc);
+		
+			UMovieGraphVariable** FoundVariable = GraphVariables.FindByPredicate([&VariableGUID](const UMovieGraphVariable* Variable)
+				{ return Variable && (Variable->GetGuid() == VariableGUID); });
+
+			// This shouldn't happen since descs that don't correspond to a real variable were removed above
+			if (!ensure(FoundVariable != nullptr))
+			{
+				continue;
+			}
+
+			const UMovieGraphVariable* Variable = *FoundVariable;
+
+			if (Desc.Name != FName(Variable->Name))
+			{
+				bNeedsToRegenerate = true;
+				Desc.Name = FName(Variable->Name);
+			}
+
+			if ((uint8)Desc.ValueType != (uint8)(Variable->GetValueType()))
+			{
+				bNeedsToRegenerate = true;
+				Desc.ValueType = static_cast<EPropertyBagPropertyType>(Variable->GetValueType());
+			}
+
+			if (Desc.ValueTypeObject != Variable->GetValueTypeObject())
+			{
+				bNeedsToRegenerate = true;
+				Desc.ValueTypeObject = Variable->GetValueTypeObject();
+			}
+
+			if ((uint8)Desc.ContainerType != (uint8)(Variable->GetValueContainerType()))
+			{
+				bNeedsToRegenerate = true;
+				Desc.ContainerType = static_cast<EPropertyBagContainerType>(Variable->GetValueContainerType());
+			}
+		}
+	}
+
+	// Third, ensure that each variable has a corresponding property in the bag. If not, generate one and flag that the
+	// bag needs to be regenerated.
+	for (const UMovieGraphVariable* Variable : GraphVariables)
+	{
+		bool bFoundMatchingDesc = false;
+		
+		for (const FPropertyBagPropertyDesc& Desc : ModifiedDescs)
+		{
+			if (Variable->GetGuid() == GetVariableGuidFromDesc(Desc))
+			{
+				bFoundMatchingDesc = true;
+				break;
+			}
+		}
+
+		if (!bFoundMatchingDesc)
+		{
+			// Generate the desc for both the property and the associated EditCondition
+			FPropertyBagPropertyDesc NewPropertyDesc;
+			FPropertyBagPropertyDesc EditConditionPropertyDesc;
+			if (GenerateVariableOverride(Variable, &NewPropertyDesc, &EditConditionPropertyDesc))
+			{
+				ModifiedDescs.Add(NewPropertyDesc);
+				ModifiedDescs.Add(EditConditionPropertyDesc);
+				bNeedsToRegenerate = true;
+			}
+		}
+	}
+
+	if (bNeedsToRegenerate)
+	{
+		const UPropertyBag* NewBagStruct = UPropertyBag::GetOrCreateFromDescs(ModifiedDescs);
+		Value.MigrateToNewBagStruct(NewBagStruct);
+	}
+}
+#endif
+
+bool UMovieJobVariableAssignmentContainer::AddVariableAssignment(const UMovieGraphVariable* InGraphVariable)
+{
+	return InGraphVariable && FindOrGenerateVariableOverride(InGraphVariable);
+}
+
+bool UMovieJobVariableAssignmentContainer::SetVariableAssignmentEnableState(const UMovieGraphVariable* InGraphVariable, bool bIsEnabled)
+{
+	if (FindOrGenerateVariableOverride(InGraphVariable))
+	{
+		const FString EditCondPropName = FString::Format(TEXT("{0}{1}"), {EditConditionPrefix, InGraphVariable->Name});
+		const EPropertyBagResult Result = Value.SetValueBool(FName(EditCondPropName), bIsEnabled);
+		
+		return Result == EPropertyBagResult::Success;
+	}
+
+	return false;
+}
+
+bool UMovieJobVariableAssignmentContainer::GetVariableAssignmentEnableState(const UMovieGraphVariable* InGraphVariable, bool& bOutIsEnabled)
+{
+	FPropertyBagPropertyDesc* PropDesc = nullptr;
+	FPropertyBagPropertyDesc* EditConditionPropDesc = nullptr;
+	const bool bAddIfNotExists = false;
+	if (FindOrGenerateVariableOverride(InGraphVariable, PropDesc, EditConditionPropDesc, bAddIfNotExists))
+	{
+		const FString EditCondPropName = FString::Format(TEXT("{0}{1}"), {EditConditionPrefix, InGraphVariable->Name});
+		TValueOrError<bool, EPropertyBagResult> Result = Value.GetValueBool(FName(EditCondPropName));
+		bOutIsEnabled = Result.GetValue();
+		
+		return !Result.HasError();
+	}
+
+	return false;
+}
\ No newline at end of file
diff --git a/Engine/Plugins/MovieScene/MovieRenderPipeline/Source/MovieRenderPipelineCore/Private/MoviePipelineQueue.cpp b/Engine/Plugins/MovieScene/MovieRenderPipeline/Source/MovieRenderPipelineCore/Private/MoviePipelineQueue.cpp
--- a/Engine/Plugins/MovieScene/MovieRenderPipeline/Source/MovieRenderPipelineCore/Private/MoviePipelineQueue.cpp
+++ b/Engine/Plugins/MovieScene/MovieRenderPipeline/Source/MovieRenderPipelineCore/Private/MoviePipelineQueue.cpp
@@ -199,62 +199,30 @@ void UMoviePipelineExecutorJob::PreSave(FObjectPreSaveContext ObjectSaveContext)
 {
 	Super::PreSave(ObjectSaveContext);
 
-	CleanUpInvalidVariableAssignments();
-}
-
-bool UMoviePipelineExecutorJob::AddVariableAssignment(const UMovieGraphVariable* InGraphVariable)
-{
-	return InGraphVariable && (FindOrAddVariableAssignment(InGraphVariable) != nullptr);
-}
-
-bool UMoviePipelineExecutorJob::SetVariableAssignmentValue(const UMovieGraphVariable* InGraphVariable, float NewValue)
-{
-	const bool bAddIfNotExists = true;
-	if (FMoviePipelineGraphVariableValueAssignment* Assignment = FindOrAddVariableAssignment(InGraphVariable, bAddIfNotExists))
-	{
-		Modify();
-		Assignment->VariableValue = NewValue;
-		return true;
-	}
-
-	return false;
-}
-
-bool UMoviePipelineExecutorJob::SetVariableAssignmentEnableState(const UMovieGraphVariable* InGraphVariable, bool bIsEnabled)
-{
-	const bool bAddIfNotExists = true;
-	if (FMoviePipelineGraphVariableValueAssignment* Assignment = FindOrAddVariableAssignment(InGraphVariable, bAddIfNotExists))
-	{
-		Modify();
-		Assignment->bIsValueSet = bIsEnabled;
-		return true;
-	}
-
-	return false;
+#if WITH_EDITOR
+	VariableAssignments->UpdateGraphVariableOverrides();
+#endif
 }
 
-bool UMoviePipelineExecutorJob::GetVariableAssignmentValue(const UMovieGraphVariable* InGraphVariable, float& OutValue)
+void UMoviePipelineExecutorJob::PostLoad()
 {
-	const bool bAddIfNotExists = false;
-	if (const FMoviePipelineGraphVariableValueAssignment* Assignment = FindOrAddVariableAssignment(InGraphVariable, bAddIfNotExists))
+	Super::PostLoad();
+	
+	if (VariableAssignments)
 	{
-		OutValue = Assignment->VariableValue;
-		return true;
+		VariableAssignments->SetGraphConfig(GraphPreset.LoadSynchronous());
 	}
 
-	return false;
-}
-
-bool UMoviePipelineExecutorJob::GetVariableAssignmentEnableState(const UMovieGraphVariable* InGraphVariable, bool& bOutIsEnabled)
-{
-	const bool bAddIfNotExists = false;
-	if (const FMoviePipelineGraphVariableValueAssignment* Assignment = FindOrAddVariableAssignment(InGraphVariable, bAddIfNotExists))
+	// Update the job's variable overrides whenever the graph's variables change
+	if (UMovieGraphConfig* Config = GraphPreset.LoadSynchronous())
 	{
-		bOutIsEnabled = Assignment->bIsValueSet;
-		return true;
+#if WITH_EDITOR
+		Config->OnGraphVariablesChangedDelegate.AddUObject(VariableAssignments, &UMovieJobVariableAssignmentContainer::UpdateGraphVariableOverrides);
+#endif
 	}
 
-	return false;
+	// TODO: The variable assignments need to be updated here, but currently GraphPreset.LoadSynchronous() does not
+	// provide a fully-loaded object (ie, GraphPreset will only be partially loaded at this point).
 }
 
 void UMoviePipelineExecutorJob::SetSequence(FSoftObjectPath InSequence)
@@ -305,47 +273,6 @@ void UMoviePipelineExecutorJob::OnDuplicated_Implementation()
 	SetConsumed(false);
 }
 
-FMoviePipelineGraphVariableValueAssignment* UMoviePipelineExecutorJob::FindOrAddVariableAssignment(const UMovieGraphVariable* InGraphVariable, bool bAddIfNotExists)
-{
-	if (InGraphVariable)
-	{
-		for (FMoviePipelineGraphVariableValueAssignment& Assignment : GraphVariableAssignments)
-		{
-			if (Assignment.VariableGuid == InGraphVariable->GetGuid())
-			{
-				return &Assignment;
-			}
-		}
-
-		// Didn't find an existing one; add a new assignment
-		if (bAddIfNotExists)
-		{
-			Modify();
-			
-			FMoviePipelineGraphVariableValueAssignment NewAssignment;
-			NewAssignment.VariableGuid = InGraphVariable->GetGuid();
-			return &GraphVariableAssignments.Add_GetRef(MoveTemp(NewAssignment));
-		}
-	}
-
-	return nullptr;
-}
-
-void UMoviePipelineExecutorJob::CleanUpInvalidVariableAssignments()
-{
-	// TODO: This should get the graph preset OR config
-	const UMovieGraphConfig* Graph = GetGraphPreset();
-	if (!Graph)
-	{
-		return;
-	}
-	
-	GraphVariableAssignments.RemoveAll([Graph](const FMoviePipelineGraphVariableValueAssignment& Assignment)
-	{
-		return Graph->GetVariableByGuid(Assignment.VariableGuid) == nullptr;
-	});
-}
-
 UMoviePipelineShotConfig* UMoviePipelineExecutorShot::AllocateNewShotOverrideConfig(TSubclassOf<UMoviePipelineShotConfig> InConfigType)
 {
 	if (!ensureAlwaysMsgf(InConfigType, TEXT("Failed to specify a Config Type. Use the default in project setting or UMoviePipelineShotConfig.")))
@@ -400,3 +327,32 @@ void UMoviePipelineExecutorShot::SetShotOverridePresetOrigin(UMoviePipelineShotC
 	}
 }
 
+void UMoviePipelineExecutorShot::PreSave(FObjectPreSaveContext ObjectSaveContext)
+{
+	Super::PreSave(ObjectSaveContext);
+
+#if WITH_EDITOR
+	VariableAssignments->UpdateGraphVariableOverrides();
+#endif
+}
+
+void UMoviePipelineExecutorShot::PostLoad()
+{
+	Super::PostLoad();
+
+	if (VariableAssignments)
+	{
+		VariableAssignments->SetGraphConfig(GraphPreset.LoadSynchronous());
+	}
+
+	// Update the job's variable overrides whenever the graph's variables change
+	if (UMovieGraphConfig* Config = GraphPreset.LoadSynchronous())
+	{
+#if WITH_EDITOR
+		Config->OnGraphVariablesChangedDelegate.AddUObject(VariableAssignments, &UMovieJobVariableAssignmentContainer::UpdateGraphVariableOverrides);
+#endif
+	}
+
+	// TODO: The variable assignments need to be updated here, but currently GraphPreset.LoadSynchronous() does not
+	// provide a fully-loaded object (ie, GraphPreset will only be partially loaded at this point).
+}
diff --git a/Engine/Plugins/MovieScene/MovieRenderPipeline/Source/MovieRenderPipelineCore/Public/Graph/MovieGraphCommon.h b/Engine/Plugins/MovieScene/MovieRenderPipeline/Source/MovieRenderPipelineCore/Public/Graph/MovieGraphCommon.h
--- a/Engine/Plugins/MovieScene/MovieRenderPipeline/Source/MovieRenderPipelineCore/Public/Graph/MovieGraphCommon.h
+++ b/Engine/Plugins/MovieScene/MovieRenderPipeline/Source/MovieRenderPipelineCore/Public/Graph/MovieGraphCommon.h
@@ -2,16 +2,52 @@
 
 #pragma once
 
+#include "PropertyBag.h"
+
 #include "MovieGraphCommon.generated.h"
 
-/** Types of members that can be used in the graph. */
+// Note: This is a copy of the property bag's types so the implementation details of the graph
+// members don't leak to the external API. Not ideal, but UHT doesn't let us alias types.
+/** The type of a graph member's value. */
 UENUM()
-enum class EMovieGraphMemberType : uint8
+enum class EMovieGraphValueType : uint8
 {
-	Branch,
+	None UMETA(Hidden),
 	Bool,
+	Byte,
+	Int32,
+	Int64,
 	Float,
-	Int,
-	IntPoint,
-	String
-};
\ No newline at end of file
+	Double,
+	Name,
+	String,
+	Text,
+	Enum UMETA(Hidden),
+	Struct UMETA(Hidden),
+	Object UMETA(Hidden),
+	SoftObject UMETA(Hidden),
+	Class UMETA(Hidden),
+	SoftClass UMETA(Hidden),
+
+	Branch UMETA(Hidden),	// Added for MRQ purposes
+	Unknown UMETA(Hidden),	// Added for MRQ purposes
+
+	Count UMETA(Hidden)
+};
+
+// Note: The +2 is to account for the MRQ-added items in the enum
+// TODO: We may want a method which converts between these enum types instead
+static_assert((uint8)EMovieGraphValueType::Count == (uint8)EPropertyBagPropertyType::Count + 2);
+
+// Note: Also a copy of the property bag's container types for the same reason as EMovieGraphValueType.
+/** The container type of a graph member's value. */
+UENUM()
+enum class EMovieGraphContainerType : uint8
+{
+	None UMETA(Hidden),
+	Array,
+
+	Count UMETA(Hidden)
+};
+
+static_assert((uint8)EMovieGraphContainerType::Count == (uint8)EPropertyBagContainerType::Count);
\ No newline at end of file
diff --git a/Engine/Plugins/MovieScene/MovieRenderPipeline/Source/MovieRenderPipelineCore/Public/Graph/MovieGraphConfig.h b/Engine/Plugins/MovieScene/MovieRenderPipeline/Source/MovieRenderPipelineCore/Public/Graph/MovieGraphConfig.h
--- a/Engine/Plugins/MovieScene/MovieRenderPipeline/Source/MovieRenderPipelineCore/Public/Graph/MovieGraphConfig.h
+++ b/Engine/Plugins/MovieScene/MovieRenderPipeline/Source/MovieRenderPipelineCore/Public/Graph/MovieGraphConfig.h
@@ -5,6 +5,7 @@
 #include "EdGraph/EdGraph.h"
 #include "Graph/MovieGraphNode.h"
 #include "Graph/MovieGraphTraversalContext.h"
+#include "MovieGraphValueContainer.h"
 
 #include "MovieGraphConfig.generated.h"
 
@@ -18,7 +19,7 @@ class UMovieGraphNode;
 #endif
 
 UCLASS(Abstract)
-class MOVIERENDERPIPELINECORE_API UMovieGraphMember : public UObject
+class MOVIERENDERPIPELINECORE_API UMovieGraphMember : public UMovieGraphValueContainer
 {
 	// The graph needs to set private flags during construction time
 	friend class UMovieGraphConfig;
@@ -38,10 +39,6 @@ public:
 	virtual bool IsDeletable() const { return true; }
 
 public:
-	/** The type of data associated with this member. */
-	UPROPERTY(EditAnywhere, Category = "General", meta=(EditCondition="bIsEditable", HideEditConditionToggle))
-	EMovieGraphMemberType Type = EMovieGraphMemberType::Float;
-
 	// TODO: Need a details customization that validates whether or not the name is valid/unique
 	/** The name of this member, which is user-facing. */
 	UPROPERTY(EditAnywhere, Category = "General", meta=(EditCondition="bIsEditable", HideEditConditionToggle))
@@ -51,12 +48,6 @@ public:
 	UPROPERTY(EditAnywhere, Category = "General", meta=(EditCondition="bIsEditable", HideEditConditionToggle))
 	FString Description;
 
-	// TODO: This needs to eventually be capable of storing multiple types. Using TVariant looks promising, but it will
-	// still require extensive details panel customizations
-	/** The default value of this member. */
-	UPROPERTY(EditAnywhere, Category = "General", meta=(EditCondition="bIsEditable", HideEditConditionToggle))
-	float Default = 0.f;
-
 private:
 	/** A GUID that uniquely identifies this member within its graph. */
 	UPROPERTY()
@@ -197,6 +188,7 @@ public:
 	 * Adds a new variable member with default values to the graph. The new variable will have a base name of
 	 * "Variable" unless specified in InCustomBaseName. Returns the new variable on success, else nullptr.
 	 */
+	UFUNCTION(BlueprintCallable, Category="Experimental")
 	UMovieGraphVariable* AddVariable(const FName InCustomBaseName = NAME_None);
 
 	/** Adds a new input member to the graph. Returns the new input on success, else nullptr. */
@@ -212,6 +204,7 @@ public:
 	 * Gets all variables that are available to be used in the graph. Global variables can optionally be included if
 	 * bIncludeGlobal is set to true.
 	 */
+	UFUNCTION(BlueprintCallable, Category="Experimental")
 	TArray<UMovieGraphVariable*> GetVariables(const bool bIncludeGlobal = false) const;
 
 	/** Gets all inputs that have been defined on the graph. */
diff --git a/Engine/Plugins/MovieScene/MovieRenderPipeline/Source/MovieRenderPipelineCore/Public/Graph/MovieGraphNode.h b/Engine/Plugins/MovieScene/MovieRenderPipeline/Source/MovieRenderPipelineCore/Public/Graph/MovieGraphNode.h
--- a/Engine/Plugins/MovieScene/MovieRenderPipeline/Source/MovieRenderPipelineCore/Public/Graph/MovieGraphNode.h
+++ b/Engine/Plugins/MovieScene/MovieRenderPipeline/Source/MovieRenderPipelineCore/Public/Graph/MovieGraphNode.h
@@ -139,7 +139,7 @@ protected:
 	UPROPERTY()
 	TArray<TObjectPtr<UMovieGraphPin>> OutputPins;
 
-	UPROPERTY(EditAnywhere, meta=(FixedLayout), Category = "Node")
+	UPROPERTY(EditAnywhere, meta=(FixedLayout, ShowOnlyInnerProperties), Category = "Node")
 	FInstancedPropertyBag DynamicProperties;
 
 	UPROPERTY()
diff --git a/Engine/Plugins/MovieScene/MovieRenderPipeline/Source/MovieRenderPipelineCore/Public/Graph/MovieGraphPin.h b/Engine/Plugins/MovieScene/MovieRenderPipeline/Source/MovieRenderPipelineCore/Public/Graph/MovieGraphPin.h
--- a/Engine/Plugins/MovieScene/MovieRenderPipeline/Source/MovieRenderPipelineCore/Public/Graph/MovieGraphPin.h
+++ b/Engine/Plugins/MovieScene/MovieRenderPipeline/Source/MovieRenderPipelineCore/Public/Graph/MovieGraphPin.h
@@ -15,7 +15,7 @@ struct FMovieGraphPinProperties
 	GENERATED_BODY()
 
 	FMovieGraphPinProperties() = default;
-	explicit FMovieGraphPinProperties(const FName& InLabel, const EMovieGraphMemberType PinType, bool bInAllowMultipleConnections)
+	explicit FMovieGraphPinProperties(const FName& InLabel, const EMovieGraphValueType PinType, bool bInAllowMultipleConnections)
 		: Label(InLabel)
 		, Type(PinType)
 		, bAllowMultipleConnections(bInAllowMultipleConnections)
@@ -25,7 +25,7 @@ struct FMovieGraphPinProperties
 	FName Label = NAME_None;
 
 	UPROPERTY(BlueprintReadWrite, EditAnywhere, Category = "Settings")
-	EMovieGraphMemberType Type = EMovieGraphMemberType::Float;
+	EMovieGraphValueType Type = EMovieGraphValueType::Float;
 
 	UPROPERTY(BlueprintReadWrite, EditAnywhere, Category = "Settings")
 	bool bAllowMultipleConnections = true;
diff --git a/Engine/Plugins/MovieScene/MovieRenderPipeline/Source/MovieRenderPipelineCore/Public/Graph/MovieGraphPlaceholderNodes.cpp b/Engine/Plugins/MovieScene/MovieRenderPipeline/Source/MovieRenderPipelineCore/Public/Graph/MovieGraphPlaceholderNodes.cpp
--- a/Engine/Plugins/MovieScene/MovieRenderPipeline/Source/MovieRenderPipelineCore/Public/Graph/MovieGraphPlaceholderNodes.cpp
+++ b/Engine/Plugins/MovieScene/MovieRenderPipeline/Source/MovieRenderPipelineCore/Public/Graph/MovieGraphPlaceholderNodes.cpp
@@ -17,14 +17,14 @@ static const FText NodeCategory_Settings = LOCTEXT("NodeCategory_Settings", "Set
 TArray<FMovieGraphPinProperties> UMoviePipelineCollectionNode::GetInputPinProperties() const
 {
 	TArray<FMovieGraphPinProperties> Properties;
-	Properties.Add(FMovieGraphPinProperties(NAME_None, EMovieGraphMemberType::Branch, false));
+	Properties.Add(FMovieGraphPinProperties(NAME_None, EMovieGraphValueType::Branch, false));
 	return Properties;
 }
 
 TArray<FMovieGraphPinProperties> UMoviePipelineCollectionNode::GetOutputPinProperties() const
 {
 	TArray<FMovieGraphPinProperties> Properties;
-	Properties.Add(FMovieGraphPinProperties(NAME_None, EMovieGraphMemberType::Branch, false));
+	Properties.Add(FMovieGraphPinProperties(NAME_None, EMovieGraphValueType::Branch, false));
 	return Properties;
 }
 
@@ -75,14 +75,14 @@ void UMoviePipelineCollectionNode::PostEditChangeProperty(FPropertyChangedEvent&
 TArray<FMovieGraphPinProperties> UMovieGraphModifierNode::GetInputPinProperties() const
 {
 	TArray<FMovieGraphPinProperties> Properties;
-	Properties.Add(FMovieGraphPinProperties(NAME_None, EMovieGraphMemberType::Branch, false));
+	Properties.Add(FMovieGraphPinProperties(NAME_None, EMovieGraphValueType::Branch, false));
 	return Properties;
 }
 
 TArray<FMovieGraphPinProperties> UMovieGraphModifierNode::GetOutputPinProperties() const
 {
 	TArray<FMovieGraphPinProperties> Properties;
-	Properties.Add(FMovieGraphPinProperties(NAME_None, EMovieGraphMemberType::Branch, false));
+	Properties.Add(FMovieGraphPinProperties(NAME_None, EMovieGraphValueType::Branch, false));
 	return Properties;
 }
 
@@ -134,14 +134,14 @@ void UMovieGraphModifierNode::PostEditChangeProperty(FPropertyChangedEvent& Prop
 TArray<FMovieGraphPinProperties> UMovieGraphGlobalGameOverridesNode::GetInputPinProperties() const
 {
 	TArray<FMovieGraphPinProperties> Properties;
-	Properties.Add(FMovieGraphPinProperties(NAME_None, EMovieGraphMemberType::Branch, false));
+	Properties.Add(FMovieGraphPinProperties(NAME_None, EMovieGraphValueType::Branch, false));
 	return Properties;
 }
 
 TArray<FMovieGraphPinProperties> UMovieGraphGlobalGameOverridesNode::GetOutputPinProperties() const
 {
 	TArray<FMovieGraphPinProperties> Properties;
-	Properties.Add(FMovieGraphPinProperties(NAME_None, EMovieGraphMemberType::Branch, false));
+	Properties.Add(FMovieGraphPinProperties(NAME_None, EMovieGraphValueType::Branch, false));
 	return Properties;
 }
 
@@ -175,14 +175,14 @@ FSlateIcon UMovieGraphGlobalGameOverridesNode::GetIconAndTint(FLinearColor& OutC
 TArray<FMovieGraphPinProperties> UMovieGraphPathTracedRendererNode::GetInputPinProperties() const
 {
 	TArray<FMovieGraphPinProperties> Properties;
-	Properties.Add(FMovieGraphPinProperties(NAME_None, EMovieGraphMemberType::Branch, false));
+	Properties.Add(FMovieGraphPinProperties(NAME_None, EMovieGraphValueType::Branch, false));
 	return Properties;
 }
 
 TArray<FMovieGraphPinProperties> UMovieGraphPathTracedRendererNode::GetOutputPinProperties() const
 {
 	TArray<FMovieGraphPinProperties> Properties;
-	Properties.Add(FMovieGraphPinProperties(NAME_None, EMovieGraphMemberType::Branch, false));
+	Properties.Add(FMovieGraphPinProperties(NAME_None, EMovieGraphValueType::Branch, false));
 	return Properties;
 }
 
@@ -216,14 +216,14 @@ FSlateIcon UMovieGraphPathTracedRendererNode::GetIconAndTint(FLinearColor& OutCo
 TArray<FMovieGraphPinProperties> UMovieGraphEXRSequenceNode::GetInputPinProperties() const
 {
 	TArray<FMovieGraphPinProperties> Properties;
-	Properties.Add(FMovieGraphPinProperties(NAME_None, EMovieGraphMemberType::Branch, false));
+	Properties.Add(FMovieGraphPinProperties(NAME_None, EMovieGraphValueType::Branch, false));
 	return Properties;
 }
 
 TArray<FMovieGraphPinProperties> UMovieGraphEXRSequenceNode::GetOutputPinProperties() const
 {
 	TArray<FMovieGraphPinProperties> Properties;
-	Properties.Add(FMovieGraphPinProperties(NAME_None, EMovieGraphMemberType::Branch, false));
+	Properties.Add(FMovieGraphPinProperties(NAME_None, EMovieGraphValueType::Branch, false));
 	return Properties;
 }
 
@@ -261,16 +261,16 @@ TArray<FMovieGraphPinProperties> UMovieGraphBranchNode::GetInputPinProperties()
 	static const FName Condition("Condition");
 	
 	TArray<FMovieGraphPinProperties> Properties;
-	Properties.Add(FMovieGraphPinProperties(TrueBranch, EMovieGraphMemberType::Branch, false));
-	Properties.Add(FMovieGraphPinProperties(FalseBranch, EMovieGraphMemberType::Branch, false));
-	Properties.Add(FMovieGraphPinProperties(Condition, EMovieGraphMemberType::Bool, false));
+	Properties.Add(FMovieGraphPinProperties(TrueBranch, EMovieGraphValueType::Branch, false));
+	Properties.Add(FMovieGraphPinProperties(FalseBranch, EMovieGraphValueType::Branch, false));
+	Properties.Add(FMovieGraphPinProperties(Condition, EMovieGraphValueType::Bool, false));
 	return Properties;
 }
 
 TArray<FMovieGraphPinProperties> UMovieGraphBranchNode::GetOutputPinProperties() const
 {
 	TArray<FMovieGraphPinProperties> Properties;
-	Properties.Add(FMovieGraphPinProperties(NAME_None, EMovieGraphMemberType::Branch, false));
+	Properties.Add(FMovieGraphPinProperties(NAME_None, EMovieGraphValueType::Branch, false));
 	return Properties;
 }
 
@@ -308,10 +308,10 @@ TArray<FMovieGraphPinProperties> UMovieGraphSelectNode::GetInputPinProperties()
 	// Generate branch pins for each option
 	for (const FString& SelectOption : SelectOptions)
 	{
-		Properties.Add(FMovieGraphPinProperties(FName(SelectOption), EMovieGraphMemberType::Branch, false));
+		Properties.Add(FMovieGraphPinProperties(FName(SelectOption), EMovieGraphValueType::Branch, false));
 	}
 
-	Properties.Add(FMovieGraphPinProperties(NAME_None, EMovieGraphMemberType::String, false));
+	Properties.Add(FMovieGraphPinProperties(NAME_None, EMovieGraphValueType::String, false));
 	
 	return Properties;
 }
@@ -319,7 +319,7 @@ TArray<FMovieGraphPinProperties> UMovieGraphSelectNode::GetInputPinProperties()
 TArray<FMovieGraphPinProperties> UMovieGraphSelectNode::GetOutputPinProperties() const
 {
 	TArray<FMovieGraphPinProperties> Properties;
-	Properties.Add(FMovieGraphPinProperties(NAME_None, EMovieGraphMemberType::Branch, false));
+	Properties.Add(FMovieGraphPinProperties(NAME_None, EMovieGraphValueType::Branch, false));
 	return Properties;
 }
 
@@ -375,14 +375,14 @@ void UMovieGraphSelectNode::PostEditChangeProperty(FPropertyChangedEvent& Proper
 TArray<FMovieGraphPinProperties> UMovieGraphOutputSettingsNode::GetInputPinProperties() const
 {
 	TArray<FMovieGraphPinProperties> Properties;
-	Properties.Add(FMovieGraphPinProperties(NAME_None, EMovieGraphMemberType::Branch, false));
+	Properties.Add(FMovieGraphPinProperties(NAME_None, EMovieGraphValueType::Branch, false));
 	return Properties;
 }
 
 TArray<FMovieGraphPinProperties> UMovieGraphOutputSettingsNode::GetOutputPinProperties() const
 {
 	TArray<FMovieGraphPinProperties> Properties;
-	Properties.Add(FMovieGraphPinProperties(NAME_None, EMovieGraphMemberType::Branch, false));
+	Properties.Add(FMovieGraphPinProperties(NAME_None, EMovieGraphValueType::Branch, false));
 	return Properties;
 }
 
@@ -416,14 +416,14 @@ FSlateIcon UMovieGraphOutputSettingsNode::GetIconAndTint(FLinearColor& OutColor)
 TArray<FMovieGraphPinProperties> UMovieGraphAntiAliasingNode::GetInputPinProperties() const
 {
 	TArray<FMovieGraphPinProperties> Properties;
-	Properties.Add(FMovieGraphPinProperties(NAME_None, EMovieGraphMemberType::Branch, false));
+	Properties.Add(FMovieGraphPinProperties(NAME_None, EMovieGraphValueType::Branch, false));
 	return Properties;
 }
 
 TArray<FMovieGraphPinProperties> UMovieGraphAntiAliasingNode::GetOutputPinProperties() const
 {
 	TArray<FMovieGraphPinProperties> Properties;
-	Properties.Add(FMovieGraphPinProperties(NAME_None, EMovieGraphMemberType::Branch, false));
+	Properties.Add(FMovieGraphPinProperties(NAME_None, EMovieGraphValueType::Branch, false));
 	return Properties;
 }
 
diff --git a/Engine/Plugins/MovieScene/MovieRenderPipeline/Source/MovieRenderPipelineCore/Public/Graph/MovieGraphValueContainer.cpp b/Engine/Plugins/MovieScene/MovieRenderPipeline/Source/MovieRenderPipelineCore/Public/Graph/MovieGraphValueContainer.cpp
new file mode 100644
--- /dev/null
+++ b/Engine/Plugins/MovieScene/MovieRenderPipeline/Source/MovieRenderPipelineCore/Public/Graph/MovieGraphValueContainer.cpp
@@ -0,0 +1,242 @@
+// Copyright Epic Games, Inc. All Rights Reserved.
+
+#include "MovieGraphValueContainer.h"
+
+#include "MovieGraphCommon.h"
+
+// The property bag has one "default" property in it
+const FName UMovieGraphValueContainer::PropertyBagDefaultPropertyName("Default");
+
+UMovieGraphValueContainer::UMovieGraphValueContainer()
+{
+	if (!HasAnyFlags(RF_ClassDefaultObject))
+	{
+		// Add a default float property if one does not exist already
+		if (Value.GetNumPropertiesInBag() == 0)
+		{
+			Value.AddProperty(PropertyBagDefaultPropertyName, EPropertyBagPropertyType::Float);
+		}
+	}
+}
+
+bool UMovieGraphValueContainer::GetValueBool(bool& bOutValue) const
+{
+	TValueOrError<bool, EPropertyBagResult> Result = Value.GetValueBool(PropertyBagDefaultPropertyName);
+	return UE::MovieGraph::Private::GetOptionalValue<bool>(Result, bOutValue);
+}
+
+bool UMovieGraphValueContainer::GetValueByte(uint8& OutValue) const
+{
+	TValueOrError<uint8, EPropertyBagResult> Result = Value.GetValueByte(PropertyBagDefaultPropertyName);
+	return UE::MovieGraph::Private::GetOptionalValue<uint8>(Result, OutValue);
+}
+
+bool UMovieGraphValueContainer::GetValueInt32(int32& OutValue) const
+{
+	TValueOrError<int32, EPropertyBagResult> Result = Value.GetValueInt32(PropertyBagDefaultPropertyName);
+	return UE::MovieGraph::Private::GetOptionalValue<int32>(Result, OutValue);
+}
+
+bool UMovieGraphValueContainer::GetValueInt64(int64& OutValue) const
+{
+	TValueOrError<int64, EPropertyBagResult> Result = Value.GetValueInt64(PropertyBagDefaultPropertyName);
+	return UE::MovieGraph::Private::GetOptionalValue<int64>(Result, OutValue);
+}
+
+bool UMovieGraphValueContainer::GetValueFloat(float& OutValue) const
+{
+	TValueOrError<float, EPropertyBagResult> Result = Value.GetValueFloat(PropertyBagDefaultPropertyName);
+	return UE::MovieGraph::Private::GetOptionalValue<float>(Result, OutValue);
+}
+
+bool UMovieGraphValueContainer::GetValueDouble(double& OutValue) const
+{
+	TValueOrError<double, EPropertyBagResult> Result = Value.GetValueDouble(PropertyBagDefaultPropertyName);
+	return UE::MovieGraph::Private::GetOptionalValue<double>(Result, OutValue);
+}
+
+bool UMovieGraphValueContainer::GetValueName(FName& OutValue) const
+{
+	TValueOrError<FName, EPropertyBagResult> Result = Value.GetValueName(PropertyBagDefaultPropertyName);
+	return UE::MovieGraph::Private::GetOptionalValue<FName>(Result, OutValue);
+}
+
+bool UMovieGraphValueContainer::GetValueString(FString& OutValue) const
+{
+	TValueOrError<FString, EPropertyBagResult> Result = Value.GetValueString(PropertyBagDefaultPropertyName);
+	return UE::MovieGraph::Private::GetOptionalValue<FString>(Result, OutValue);
+}
+
+bool UMovieGraphValueContainer::GetValueText(FText& OutValue) const
+{
+	TValueOrError<FText, EPropertyBagResult> Result = Value.GetValueText(PropertyBagDefaultPropertyName);
+	return UE::MovieGraph::Private::GetOptionalValue<FText>(Result, OutValue);
+}
+
+bool UMovieGraphValueContainer::GetValueEnum(uint8& OutValue, const UEnum* RequestedEnum) const
+{
+	TValueOrError<uint8, EPropertyBagResult> Result = Value.GetValueEnum(PropertyBagDefaultPropertyName, RequestedEnum);
+	return UE::MovieGraph::Private::GetOptionalValue<uint8>(Result, OutValue);
+}
+
+bool UMovieGraphValueContainer::GetValueStruct(FStructView& OutValue, const UScriptStruct* RequestedStruct) const
+{
+	TValueOrError<FStructView, EPropertyBagResult> Result = Value.GetValueStruct(PropertyBagDefaultPropertyName, RequestedStruct);
+	return UE::MovieGraph::Private::GetOptionalValue<FStructView>(Result, OutValue);
+}
+
+bool UMovieGraphValueContainer::GetValueObject(UObject* OutValue, const UClass* RequestedClass) const
+{
+	TValueOrError<UObject*, EPropertyBagResult> Result = Value.GetValueObject(PropertyBagDefaultPropertyName, RequestedClass);
+	return UE::MovieGraph::Private::GetOptionalValue<UObject*>(Result, OutValue);
+}
+
+bool UMovieGraphValueContainer::GetValueClass(UClass* OutValue) const
+{
+	TValueOrError<UClass*, EPropertyBagResult> Result = Value.GetValueClass(PropertyBagDefaultPropertyName);
+	return UE::MovieGraph::Private::GetOptionalValue<UClass*>(Result, OutValue);
+}
+
+FString UMovieGraphValueContainer::GetValueSerializedString()
+{
+	TValueOrError<FString, EPropertyBagResult> Result = Value.GetValueSerializedString(PropertyBagDefaultPropertyName);
+	FString ResultString;
+	UE::MovieGraph::Private::GetOptionalValue<FString>(Result, ResultString);
+	return ResultString;
+}
+
+bool UMovieGraphValueContainer::SetValueBool(const bool bInValue)
+{
+	return Value.SetValueBool(PropertyBagDefaultPropertyName, bInValue) == EPropertyBagResult::Success;
+}
+
+bool UMovieGraphValueContainer::SetValueByte(const uint8 InValue)
+{
+	return Value.SetValueByte(PropertyBagDefaultPropertyName, InValue) == EPropertyBagResult::Success;
+}
+
+bool UMovieGraphValueContainer::SetValueInt32(const int32 InValue)
+{
+	return Value.SetValueInt32(PropertyBagDefaultPropertyName, InValue) == EPropertyBagResult::Success;
+}
+
+bool UMovieGraphValueContainer::SetValueInt64(const int64 InValue)
+{
+	return Value.SetValueInt64(PropertyBagDefaultPropertyName, InValue) == EPropertyBagResult::Success;
+}
+
+bool UMovieGraphValueContainer::SetValueFloat(const float InValue)
+{
+	return Value.SetValueFloat(PropertyBagDefaultPropertyName, InValue) == EPropertyBagResult::Success;
+}
+
+bool UMovieGraphValueContainer::SetValueDouble(const double InValue)
+{
+	return Value.SetValueDouble(PropertyBagDefaultPropertyName, InValue) == EPropertyBagResult::Success;
+}
+
+bool UMovieGraphValueContainer::SetValueName(const FName InValue)
+{
+	return Value.SetValueName(PropertyBagDefaultPropertyName, InValue) == EPropertyBagResult::Success;
+}
+
+bool UMovieGraphValueContainer::SetValueString(const FString& InValue)
+{
+	return Value.SetValueString(PropertyBagDefaultPropertyName, InValue) == EPropertyBagResult::Success;
+}
+
+bool UMovieGraphValueContainer::SetValueText(const FText& InValue)
+{
+	return Value.SetValueText(PropertyBagDefaultPropertyName, InValue) == EPropertyBagResult::Success;
+}
+
+bool UMovieGraphValueContainer::SetValueEnum(const uint8 InValue, const UEnum* Enum)
+{
+	return Value.SetValueEnum(PropertyBagDefaultPropertyName, InValue, Enum) == EPropertyBagResult::Success;
+}
+
+bool UMovieGraphValueContainer::SetValueStruct(FConstStructView InValue)
+{
+	return Value.SetValueStruct(PropertyBagDefaultPropertyName, InValue) == EPropertyBagResult::Success;
+}
+
+bool UMovieGraphValueContainer::SetValueObject(UObject* InValue)
+{
+	return Value.SetValueObject(PropertyBagDefaultPropertyName, InValue) == EPropertyBagResult::Success;
+}
+
+bool UMovieGraphValueContainer::SetValueClass(UClass* InValue)
+{
+	return Value.SetValueClass(PropertyBagDefaultPropertyName, InValue) == EPropertyBagResult::Success;
+}
+
+bool UMovieGraphValueContainer::SetValueSerializedString(const FString& NewValue)
+{
+	return Value.SetValueSerializedString(PropertyBagDefaultPropertyName, NewValue) == EPropertyBagResult::Success;
+}
+
+EMovieGraphValueType UMovieGraphValueContainer::GetValueType() const
+{
+	if (const FPropertyBagPropertyDesc* Desc = Value.FindPropertyDescByName(PropertyBagDefaultPropertyName))
+	{
+		return static_cast<EMovieGraphValueType>(Desc->ValueType);
+	}
+
+	return EMovieGraphValueType::None;
+}
+
+void UMovieGraphValueContainer::SetValueType(EMovieGraphValueType ValueType)
+{
+	if (const FPropertyBagPropertyDesc* Desc = Value.FindPropertyDescByName(PropertyBagDefaultPropertyName))
+	{
+		FPropertyBagPropertyDesc NewDesc(*Desc);
+		NewDesc.ValueType = static_cast<EPropertyBagPropertyType>(ValueType);
+
+		const UPropertyBag* NewPropBag = UPropertyBag::GetOrCreateFromDescs({NewDesc});
+		Value.MigrateToNewBagStruct(NewPropBag);
+	}
+}
+
+const UObject* UMovieGraphValueContainer::GetValueTypeObject() const
+{
+	if (const FPropertyBagPropertyDesc* Desc = Value.FindPropertyDescByName(PropertyBagDefaultPropertyName))
+	{
+		return Desc->ValueTypeObject;
+	}
+	
+	return nullptr;
+}
+
+void UMovieGraphValueContainer::SetValueTypeObject(const UObject* ValueTypeObject)
+{
+	if (const FPropertyBagPropertyDesc* Desc = Value.FindPropertyDescByName(PropertyBagDefaultPropertyName))
+	{
+		FPropertyBagPropertyDesc NewDesc(*Desc);
+		NewDesc.ValueTypeObject = ValueTypeObject;
+
+		const UPropertyBag* NewPropBag = UPropertyBag::GetOrCreateFromDescs({NewDesc});
+		Value.MigrateToNewBagStruct(NewPropBag);
+	}
+}
+
+EMovieGraphContainerType UMovieGraphValueContainer::GetValueContainerType() const
+{
+	if (const FPropertyBagPropertyDesc* Desc = Value.FindPropertyDescByName(PropertyBagDefaultPropertyName))
+	{
+		return static_cast<EMovieGraphContainerType>(Desc->ContainerType);
+	}
+
+	return EMovieGraphContainerType::None;
+}
+
+void UMovieGraphValueContainer::SetValueContainerType(EMovieGraphContainerType ContainerType)
+{
+	if (const FPropertyBagPropertyDesc* Desc = Value.FindPropertyDescByName(PropertyBagDefaultPropertyName))
+	{
+		FPropertyBagPropertyDesc NewDesc(*Desc);
+		NewDesc.ContainerType = static_cast<EPropertyBagContainerType>(ContainerType);
+
+		const UPropertyBag* NewPropBag = UPropertyBag::GetOrCreateFromDescs({NewDesc});
+		Value.MigrateToNewBagStruct(NewPropBag);
+	}
+}
\ No newline at end of file
diff --git a/Engine/Plugins/MovieScene/MovieRenderPipeline/Source/MovieRenderPipelineCore/Public/Graph/MovieGraphValueContainer.h b/Engine/Plugins/MovieScene/MovieRenderPipeline/Source/MovieRenderPipelineCore/Public/Graph/MovieGraphValueContainer.h
new file mode 100644
--- /dev/null
+++ b/Engine/Plugins/MovieScene/MovieRenderPipeline/Source/MovieRenderPipelineCore/Public/Graph/MovieGraphValueContainer.h
@@ -0,0 +1,231 @@
+// Copyright Epic Games, Inc. All Rights Reserved.
+
+#pragma once
+
+#include "CoreMinimal.h"
+#include "MovieGraphCommon.h"
+#include "PropertyBag.h"
+
+#include "MovieGraphValueContainer.generated.h"
+
+namespace UE::MovieGraph::Private
+{
+	template<typename ReturnType>
+	bool GetOptionalValue(TValueOrError<ReturnType, EPropertyBagResult>& PropertyBagValue, ReturnType& OutValue)
+	{
+		// Convert the property bag-provided TValueOrError, which contains an EPropertyBagResult, to an output value and
+		// a bool (signifying if there was an error or not). EPropertyBagResult shouldn't be exposed on the MRQ API.
+		if (PropertyBagValue.HasValue())
+		{
+			OutValue = PropertyBagValue.StealValue();
+			return true;
+		}
+
+		return false;
+	}
+}
+
+/**
+ * Holds a generic value, with an API for getting/setting the value, as well as getting/setting its type
+ * and container (eg, array).
+ */
+UCLASS(Abstract)
+class MOVIERENDERPIPELINECORE_API UMovieGraphValueContainer : public UObject
+{
+	GENERATED_BODY()
+
+public:
+	UMovieGraphValueContainer();
+
+	/** Gets the bool value of the held data. Returns true on success, else false. */
+	UFUNCTION(BlueprintCallable, Category="Config")
+	bool GetValueBool(bool& bOutValue) const;
+
+	/** Gets the byte value of the held data. Returns true on success, else false. */
+	UFUNCTION(BlueprintCallable, Category="Config")
+	bool GetValueByte(uint8& OutValue) const;
+
+	/** Gets the int32 value of the held data. Returns true on success, else false. */
+	UFUNCTION(BlueprintCallable, Category="Config")
+	bool GetValueInt32(int32& OutValue) const;
+
+	/** Gets the int64 value of the held data. Returns true on success, else false. */
+	UFUNCTION(BlueprintCallable, Category="Config")
+	bool GetValueInt64(int64& OutValue) const;
+
+	/** Gets the float value of the held data. Returns true on success, else false. */
+	UFUNCTION(BlueprintCallable, Category="Config")
+	bool GetValueFloat(float& OutValue) const;
+
+	/** Gets the double value of the held data. Returns true on success, else false. */
+	UFUNCTION(BlueprintCallable, Category="Config")
+	bool GetValueDouble(double& OutValue) const;
+
+	/** Gets the FName value of the held data. Returns true on success, else false. */
+	UFUNCTION(BlueprintCallable, Category="Config")
+	bool GetValueName(FName& OutValue) const;
+
+	/** Gets the FString value of the held data. Returns true on success, else false. */
+	UFUNCTION(BlueprintCallable, Category="Config")
+	bool GetValueString(FString& OutValue) const;
+
+	/** Gets the FText value of the held data. Returns true on success, else false. */
+	UFUNCTION(BlueprintCallable, Category="Config")
+	bool GetValueText(FText& OutValue) const;
+
+	/** Gets the enum value (for a specific enum) of the held data. Returns true on success, else false. */
+	UFUNCTION(BlueprintCallable, Category="Config")
+	bool GetValueEnum(uint8& OutValue, const UEnum* RequestedEnum = nullptr) const;
+
+	/** Gets the struct value (for a specific struct) of the held data. Returns true on success, else false. */
+	bool GetValueStruct(FStructView& OutValue, const UScriptStruct* RequestedStruct = nullptr) const;
+
+	/** Gets the object value (for a specific class) of the held data. Returns true on success, else false. */
+	UFUNCTION(BlueprintCallable, Category="Config")
+	bool GetValueObject(UObject* OutValue, const UClass* RequestedClass = nullptr) const;
+
+	/** Gets the UClass value of the held data. Returns true on success, else false. */
+	UFUNCTION(BlueprintCallable, Category="Config")
+	bool GetValueClass(UClass* OutValue) const;
+
+	/** Gets the serialized string value of the held data. */
+	UFUNCTION(BlueprintCallable, Category="Config")
+	FString GetValueSerializedString();
+
+	/** Gets the enum value of the held data. Returns true on success, else false. */
+	template <typename T>
+	bool GetValueEnum(T& OutValue) const
+	{
+		TValueOrError<T, EPropertyBagResult> Result = Value.GetValueEnum<T>(PropertyBagDefaultPropertyName);
+		return UE::MovieGraph::Private::GetOptionalValue<T>(Result, OutValue);
+	}
+
+	/** Gets the struct value of the held data. Returns true on success, else false. */
+	template <typename T>
+	bool GetValueStruct(T* OutValue) const
+	{
+		TValueOrError<T*, EPropertyBagResult> Result = Value.GetValueStruct<T*>(PropertyBagDefaultPropertyName);
+		return UE::MovieGraph::Private::GetOptionalValue<T*>(Result, OutValue);
+	}
+
+	/** Gets the object value of the held data. Returns true on success, else false. */
+	template <typename T>
+	bool GetValueObject(T* OutValue) const
+	{
+		TValueOrError<T*, EPropertyBagResult> Result = Value.GetValueObject<T*>(PropertyBagDefaultPropertyName);
+		return UE::MovieGraph::Private::GetOptionalValue<T*>(Result, OutValue);
+	}
+
+	/** Sets the bool value of the held data. Returns true on success, else false. */
+	UFUNCTION(BlueprintCallable, Category="Config")
+	bool SetValueBool(const bool bInValue);
+
+	/** Sets the byte value of the held data. Returns true on success, else false. */
+	UFUNCTION(BlueprintCallable, Category="Config")
+	bool SetValueByte(const uint8 InValue);
+
+	/** Sets the int32 value of the held data. Returns true on success, else false. */
+	UFUNCTION(BlueprintCallable, Category="Config")
+	bool SetValueInt32(const int32 InValue);
+
+	/** Sets the int64 value of the held data. Returns true on success, else false. */
+	UFUNCTION(BlueprintCallable, Category="Config")
+	bool SetValueInt64(const int64 InValue);
+
+	/** Sets the float value of the held data. Returns true on success, else false. */
+	UFUNCTION(BlueprintCallable, Category="Config")
+	bool SetValueFloat(const float InValue);
+
+	/** Sets the double value of the held data. Returns true on success, else false. */
+	UFUNCTION(BlueprintCallable, Category="Config")
+	bool SetValueDouble(const double InValue);
+
+	/** Sets the FName value of the held data. Returns true on success, else false. */
+	UFUNCTION(BlueprintCallable, Category="Config")
+	bool SetValueName(const FName InValue);
+
+	/** Sets the FString value of the held data. Returns true on success, else false. */
+	UFUNCTION(BlueprintCallable, Category="Config")
+	bool SetValueString(const FString& InValue);
+
+	/** Sets the FText value of the held data. Returns true on success, else false. */
+	UFUNCTION(BlueprintCallable, Category="Config")
+	bool SetValueText(const FText& InValue);
+
+	/** Sets the enum value of the held data. Returns true on success, else false. */
+	UFUNCTION(BlueprintCallable, Category="Config")
+	bool SetValueEnum(const uint8 InValue, const UEnum* Enum);
+
+	/** Sets the struct value of the held data. Returns true on success, else false. */
+	bool SetValueStruct(FConstStructView InValue);
+
+	/** Sets the object value of the held data. Returns true on success, else false. */
+	UFUNCTION(BlueprintCallable, Category="Config")
+	bool SetValueObject(UObject* InValue);
+
+	/** Sets the class value of the held data. Returns true on success, else false. */
+	UFUNCTION(BlueprintCallable, Category="Config")
+	bool SetValueClass(UClass* InValue);
+
+	/** Sets the serialized value of the held data. The string should be the serialized representation of the value. Returns true on success, else false. */
+	UFUNCTION(BlueprintCallable, Category="Config")
+	bool SetValueSerializedString(const FString& NewValue);
+
+	/** Sets the enum value of the held data. Returns true on success, else false. */
+	template <typename T>
+	bool SetValueEnum(const T InValue)
+	{
+		return Value.SetValueEnum<T>(PropertyBagDefaultPropertyName, InValue) == EPropertyBagResult::Success;
+	}
+
+	/** Sets the struct value of the held data. Returns true on success, else false. */
+	template <typename T>
+	bool SetValueStruct(const T& InValue)
+	{
+		return Value.SetValueStruct<T>(PropertyBagDefaultPropertyName, InValue) == EPropertyBagResult::Success;
+	}
+
+	/** Sets the object value of the held data. Returns true on success, else false. */
+	template <typename T>
+	bool SetValueObject(T* InValue)
+	{
+		return Value.SetValueObject<T>(PropertyBagDefaultPropertyName, InValue) == EPropertyBagResult::Success;
+	}
+
+	/** Gets the type of the stored data. */
+	UFUNCTION(BlueprintCallable, Category="Config")
+	EMovieGraphValueType GetValueType() const;
+
+	/** Sets the type of the stored data. */
+	UFUNCTION(BlueprintCallable, Category="Config")
+	void SetValueType(EMovieGraphValueType ValueType);
+
+	/** Gets the object that defines the enum, struct, or class. */
+	UFUNCTION(BlueprintCallable, Category="Config")
+	const UObject* GetValueTypeObject() const;
+
+	/** Sets the object that defines the enum, struct, or class. */
+	UFUNCTION(BlueprintCallable, Category="Config")
+	void SetValueTypeObject(const UObject* ValueTypeObject);
+
+	/** Gets the container type of the stored value. */
+	UFUNCTION(BlueprintCallable, Category="Config")
+	EMovieGraphContainerType GetValueContainerType() const;
+
+	/** Sets the container type of the stored value. */
+	UFUNCTION(BlueprintCallable, Category="Config")
+	void SetValueContainerType(EMovieGraphContainerType ContainerType);
+
+private:
+	/** The name of the single property stored in the property bag. The name in the property bag does not change. */
+	static const FName PropertyBagDefaultPropertyName;
+	
+	// Note: The property bag only stores one property, since the object only needs to store one value. This is an odd
+	// use of a property bag, but the property bag solves a number of issues that would be very difficult to solve
+	// otherwise. 1) Presentation of a property that can have its type changed within the details pane, 2) data storage
+	// of a property which can have its type changed, 3) the ability to set the value of the property from both Python
+	// and C++, and 4) the ability to change the property at runtime.
+	/** The value held by this object. */
+	UPROPERTY(EditAnywhere, meta=(ShowOnlyInnerProperties), Category = "General")
+	FInstancedPropertyBag Value;
+};
\ No newline at end of file
diff --git a/Engine/Plugins/MovieScene/MovieRenderPipeline/Source/MovieRenderPipelineCore/Public/Graph/Nodes/MovieGraphFileOutputNode.h b/Engine/Plugins/MovieScene/MovieRenderPipeline/Source/MovieRenderPipelineCore/Public/Graph/Nodes/MovieGraphFileOutputNode.h
--- a/Engine/Plugins/MovieScene/MovieRenderPipeline/Source/MovieRenderPipelineCore/Public/Graph/Nodes/MovieGraphFileOutputNode.h
+++ b/Engine/Plugins/MovieScene/MovieRenderPipeline/Source/MovieRenderPipelineCore/Public/Graph/Nodes/MovieGraphFileOutputNode.h
@@ -36,14 +36,14 @@ public:
 	virtual TArray<FMovieGraphPinProperties> GetInputPinProperties() const override
 	{
 		TArray<FMovieGraphPinProperties> Properties;
-		Properties.Add(FMovieGraphPinProperties(NAME_None, EMovieGraphMemberType::Branch, false));
+		Properties.Add(FMovieGraphPinProperties(NAME_None, EMovieGraphValueType::Branch, false));
 		return Properties;
 	}
 
 	virtual TArray<FMovieGraphPinProperties> GetOutputPinProperties() const override
 	{
 		TArray<FMovieGraphPinProperties> Properties;
-		Properties.Add(FMovieGraphPinProperties(NAME_None, EMovieGraphMemberType::Branch, false));
+		Properties.Add(FMovieGraphPinProperties(NAME_None, EMovieGraphValueType::Branch, false));
 		return Properties;
 	}
 
diff --git a/Engine/Plugins/MovieScene/MovieRenderPipeline/Source/MovieRenderPipelineCore/Public/Graph/Nodes/MovieGraphRenderPassNode.h b/Engine/Plugins/MovieScene/MovieRenderPipeline/Source/MovieRenderPipelineCore/Public/Graph/Nodes/MovieGraphRenderPassNode.h
--- a/Engine/Plugins/MovieScene/MovieRenderPipeline/Source/MovieRenderPipelineCore/Public/Graph/Nodes/MovieGraphRenderPassNode.h
+++ b/Engine/Plugins/MovieScene/MovieRenderPipeline/Source/MovieRenderPipelineCore/Public/Graph/Nodes/MovieGraphRenderPassNode.h
@@ -46,14 +46,14 @@ public:
 	virtual TArray<FMovieGraphPinProperties> GetInputPinProperties() const override
 	{
 		TArray<FMovieGraphPinProperties> Properties;
-		Properties.Add(FMovieGraphPinProperties(NAME_None, EMovieGraphMemberType::Branch, false));
+		Properties.Add(FMovieGraphPinProperties(NAME_None, EMovieGraphValueType::Branch, false));
 		return Properties;
 	}
 
 	virtual TArray<FMovieGraphPinProperties> GetOutputPinProperties() const override
 	{
 		TArray<FMovieGraphPinProperties> Properties;
-		Properties.Add(FMovieGraphPinProperties(NAME_None, EMovieGraphMemberType::Branch, false));
+		Properties.Add(FMovieGraphPinProperties(NAME_None, EMovieGraphValueType::Branch, false));
 		return Properties;
 	}
 
diff --git a/Engine/Plugins/MovieScene/MovieRenderPipeline/Source/MovieRenderPipelineCore/Public/MovieJobVariableAssignmentContainer.h b/Engine/Plugins/MovieScene/MovieRenderPipeline/Source/MovieRenderPipelineCore/Public/MovieJobVariableAssignmentContainer.h
new file mode 100644
--- /dev/null
+++ b/Engine/Plugins/MovieScene/MovieRenderPipeline/Source/MovieRenderPipelineCore/Public/MovieJobVariableAssignmentContainer.h
@@ -0,0 +1,256 @@
+// Copyright Epic Games, Inc. All Rights Reserved.
+
+#pragma once
+
+#include "CoreMinimal.h"
+#include "Graph/MovieGraphCommon.h"
+#include "Graph/MovieGraphValueContainer.h"	// For GetOptionalValue()
+#include "PropertyBag.h"
+
+#include "MovieJobVariableAssignmentContainer.generated.h"
+
+class UMovieGraphConfig;
+class UMovieGraphVariable;
+
+/**
+ * Holds a group of properties which override variable values on the job's associated graph (if any). Overrides are not
+ * added manually. Instead, UpdateGraphVariableOverrides() should be called which will update, add, or remove overrides
+ * as appropriate. After the update, overrides can have their values retrieved and set.
+ */
+UCLASS(BlueprintType)
+class MOVIERENDERPIPELINECORE_API UMovieJobVariableAssignmentContainer : public UObject
+{
+	GENERATED_BODY()
+
+public:
+	UMovieJobVariableAssignmentContainer() = default;
+
+	/**
+	 * Sets the graph config associated with the variable assignments. Calls UpdateGraphVariableOverrides() to ensure
+	 * that the overrides reflect the specified graph config.
+	 */
+	UFUNCTION(BlueprintCallable, Category="Experimental")
+	void SetGraphConfig(const TSoftObjectPtr<UMovieGraphConfig>& InGraphConfig);
+
+	/** Gets the number of variable assignments present in this container. Assignments that are disabled are counted. */
+	uint32 GetNumAssignments() const;
+
+	/** Gets the bool value of the specified property. Returns true on success, else false. */
+	UFUNCTION(BlueprintCallable, Category="Experimental")
+	bool GetValueBool(const FName& PropertyName, bool& bOutValue) const;
+
+	/** Gets the byte value of the specified property. Returns true on success, else false. */
+	UFUNCTION(BlueprintCallable, Category="Experimental")
+	bool GetValueByte(const FName& PropertyName, uint8& OutValue) const;
+
+	/** Gets the int32 value of the specified property. Returns true on success, else false. */
+	UFUNCTION(BlueprintCallable, Category="Experimental")
+	bool GetValueInt32(const FName& PropertyName, int32& OutValue) const;
+
+	/** Gets the int64 value of the specified property. Returns true on success, else false. */
+	UFUNCTION(BlueprintCallable, Category="Experimental")
+	bool GetValueInt64(const FName& PropertyName, int64& OutValue) const;
+
+	/** Gets the float value of the specified property. Returns true on success, else false. */
+	UFUNCTION(BlueprintCallable, Category="Experimental")
+	bool GetValueFloat(const FName& PropertyName, float& OutValue) const;
+
+	/** Gets the double value of the specified property. Returns true on success, else false. */
+	UFUNCTION(BlueprintCallable, Category="Experimental")
+	bool GetValueDouble(const FName& PropertyName, double& OutValue) const;
+
+	/** Gets the FName value of the specified property. Returns true on success, else false. */
+	UFUNCTION(BlueprintCallable, Category="Experimental")
+	bool GetValueName(const FName& PropertyName, FName& OutValue) const;
+
+	/** Gets the FString value of the specified property. Returns true on success, else false. */
+	UFUNCTION(BlueprintCallable, Category="Experimental")
+	bool GetValueString(const FName& PropertyName, FString& OutValue) const;
+
+	/** Gets the FText value of the specified property. Returns true on success, else false. */
+	UFUNCTION(BlueprintCallable, Category="Experimental")
+	bool GetValueText(const FName& PropertyName, FText& OutValue) const;
+
+	/** Gets the enum value (for a specific enum) of the specified property. Returns true on success, else false. */
+	UFUNCTION(BlueprintCallable, Category="Experimental")
+	bool GetValueEnum(const FName& PropertyName, uint8& OutValue, const UEnum* RequestedEnum = nullptr) const;
+
+	/** Gets the struct value (for a specific struct) of the specified property. Returns true on success, else false. */
+	bool GetValueStruct(const FName& PropertyName, FStructView& OutValue, const UScriptStruct* RequestedStruct = nullptr) const;
+
+	/** Gets the object value (for a specific class) of the specified property. Returns true on success, else false. */
+	UFUNCTION(BlueprintCallable, Category="Experimental")
+	bool GetValueObject(const FName& PropertyName, UObject* OutValue, const UClass* RequestedClass = nullptr) const;
+
+	/** Gets the UClass value of the specified property. Returns true on success, else false. */
+	UFUNCTION(BlueprintCallable, Category="Experimental")
+	bool GetValueClass(const FName& PropertyName, UClass* OutValue) const;
+
+	/** Gets the serialized string value of the specified property. */
+	UFUNCTION(BlueprintCallable, Category="Experimental")
+	FString GetValueSerializedString(const FName& PropertyName);
+
+	/** Gets the enum value of the specified property. Returns true on success, else false. */
+	template <typename T>
+	bool GetValueEnum(const FName& PropertyName, T& OutValue) const
+	{
+		TValueOrError<T, EPropertyBagResult> Result = Value.GetValueEnum<T>(PropertyName);
+		return UE::MovieGraph::Private::GetOptionalValue<T>(Result, OutValue);
+	}
+
+	/** Gets the struct value of the specified property. Returns true on success, else false. */
+	template <typename T>
+	bool GetValueStruct(const FName& PropertyName, T* OutValue) const
+	{
+		TValueOrError<T*, EPropertyBagResult> Result = Value.GetValueStruct<T*>(PropertyName);
+		return UE::MovieGraph::Private::GetOptionalValue<T*>(Result, OutValue);
+	}
+
+	/** Gets the object value of the specified property. Returns true on success, else false. */
+	template <typename T>
+	bool GetValueObject(const FName& PropertyName, T* OutValue) const
+	{
+		TValueOrError<T*, EPropertyBagResult> Result = Value.GetValueObject<T*>(PropertyName);
+		return UE::MovieGraph::Private::GetOptionalValue<T*>(Result, OutValue);
+	}
+
+	/** Sets the bool value of the specified property. Returns true on success, else false. */
+	UFUNCTION(BlueprintCallable, Category="Experimental")
+	bool SetValueBool(const FName& PropertyName, const bool bInValue);
+
+	/** Sets the byte value of the specified property. Returns true on success, else false. */
+	UFUNCTION(BlueprintCallable, Category="Experimental")
+	bool SetValueByte(const FName& PropertyName, const uint8 InValue);
+
+	/** Sets the int32 value of the specified property. Returns true on success, else false. */
+	UFUNCTION(BlueprintCallable, Category="Experimental")
+	bool SetValueInt32(const FName& PropertyName, const int32 InValue);
+
+	/** Sets the int64 value of the specified property. Returns true on success, else false. */
+	UFUNCTION(BlueprintCallable, Category="Experimental")
+	bool SetValueInt64(const FName& PropertyName, const int64 InValue);
+
+	/** Sets the float value of the specified property. Returns true on success, else false. */
+	UFUNCTION(BlueprintCallable, Category="Experimental")
+	bool SetValueFloat(const FName& PropertyName, const float InValue);
+
+	/** Sets the double value of the specified property. Returns true on success, else false. */
+	UFUNCTION(BlueprintCallable, Category="Experimental")
+	bool SetValueDouble(const FName& PropertyName, const double InValue);
+
+	/** Sets the FName value of the specified property. Returns true on success, else false. */
+	UFUNCTION(BlueprintCallable, Category="Experimental")
+	bool SetValueName(const FName& PropertyName, const FName InValue);
+
+	/** Sets the FString value of the specified property. Returns true on success, else false. */
+	UFUNCTION(BlueprintCallable, Category="Experimental")
+	bool SetValueString(const FName& PropertyName, const FString& InValue);
+
+	/** Sets the FText value of the specified property. Returns true on success, else false. */
+	UFUNCTION(BlueprintCallable, Category="Experimental")
+	bool SetValueText(const FName& PropertyName, const FText& InValue);
+
+	/** Sets the enum value of the specified property. Returns true on success, else false. */
+	UFUNCTION(BlueprintCallable, Category="Experimental")
+	bool SetValueEnum(const FName& PropertyName, const uint8 InValue, const UEnum* Enum);
+
+	/** Sets the struct value of the specified property. Returns true on success, else false. */
+	bool SetValueStruct(const FName& PropertyName, FConstStructView InValue);
+
+	/** Sets the object value of the specified property. Returns true on success, else false. */
+	UFUNCTION(BlueprintCallable, Category="Experimental")
+	bool SetValueObject(const FName& PropertyName, UObject* InValue);
+
+	/** Sets the class value of the specified property. Returns true on success, else false. */
+	UFUNCTION(BlueprintCallable, Category="Experimental")
+	bool SetValueClass(const FName& PropertyName, UClass* InValue);
+
+	/** Sets the serialized value of this member. The string should be the serialized representation of the value. Returns true on success, else false.*/
+	UFUNCTION(BlueprintCallable, Category="Experimental")
+	bool SetValueSerializedString(const FName& PropertyName, const FString& NewValue);
+
+	/** Sets the enum value of the specified property. Returns true on success, else false. */
+	template <typename T>
+	bool SetValueEnum(const FName& PropertyName, const T InValue)
+	{
+		return Value.SetValueEnum<T>(PropertyName, InValue) == EPropertyBagResult::Success;
+	}
+
+	/** Sets the struct value of the specified property. Returns true on success, else false. */
+	template <typename T>
+	bool SetValueStruct(const FName& PropertyName, const T& InValue)
+	{
+		return Value.SetValueStruct<T>(PropertyName, InValue) == EPropertyBagResult::Success;
+	}
+
+	/** Sets the object value of the specified property. Returns true on success, else false. */
+	template <typename T>
+	bool SetValueObject(const FName& PropertyName, T* InValue)
+	{
+		return Value.SetValueObject<T>(PropertyName, InValue) == EPropertyBagResult::Success;
+	}
+
+	/** Gets the type of the value stored in the specified property. */
+	UFUNCTION(BlueprintCallable, Category="Experimental")
+	EMovieGraphValueType GetValueType(const FName& PropertyName) const;
+
+	/** Gets the object that defines the enum, struct, or class stored in the specified property. */
+	UFUNCTION(BlueprintCallable, Category="Experimental")
+	const UObject* GetValueTypeObject(const FName& PropertyName) const;
+
+	/** Gets the container type of the stored value in the specified property. */
+	UFUNCTION(BlueprintCallable, Category="Experimental")
+	EMovieGraphContainerType GetValueContainerType(const FName& PropertyName) const;
+
+	/**
+	 * Updates an existing variable assignment for the provided graph variable to a new enable state, or adds a new
+	 * assignment and updates its enable state. Returns true on success, else false.
+	 */
+	UFUNCTION(BlueprintCallable, Category = "Experimental")
+	bool SetVariableAssignmentEnableState(const UMovieGraphVariable* InGraphVariable, bool bIsEnabled);
+
+	/**
+	 * Gets the enable state of the variable assignment for the provided graph variable. The enable state is provided
+	 * via bOutIsEnabled. Returns true if an enable state was set on the variable and bOutIsEnabled was changed, else false.
+	 */
+	UFUNCTION(BlueprintCallable, Category = "Experimental")
+	bool GetVariableAssignmentEnableState(const UMovieGraphVariable* InGraphVariable, bool& bOutIsEnabled);
+
+	/**
+	 * Updates the stored variable overrides to reflect the graph preset. Existing overrides will be updated to match
+	 * the graph variable name, value type, object type, and container type. Additionally, stale overrides that have no
+	 * corresponding graph variable will be removed, and overrides will be created for graph variables which do not have
+	 * existing overrides.
+	 */
+#if WITH_EDITOR
+	UFUNCTION(BlueprintCallable, Category = "Experimental")
+	void UpdateGraphVariableOverrides();
+#endif
+
+protected:
+	/**
+	 * Add a variable assignment for the graph. Returns true if the variable was added or it already exists,
+	 * else false.
+	 */
+	bool AddVariableAssignment(const UMovieGraphVariable* InGraphVariable);
+	
+	/**
+	 * Finds a variable assignment for the provided variable, or adds one if one does not already exist (and
+	 * bAddIfNotExists is set to true). Returns nullptr if the operation failed. 
+	 */
+	bool FindOrGenerateVariableOverride(const UMovieGraphVariable* InGraphVariable, FPropertyBagPropertyDesc* OutPropDesc = nullptr,
+		FPropertyBagPropertyDesc* OutEditConditionPropDesc = nullptr, bool bGenerateIfNotExists = true);
+
+	/** Generates a variable override for the provided variable, as well as the associated EditCondition for it. */
+	bool GenerateVariableOverride(const UMovieGraphVariable* InGraphVariable, FPropertyBagPropertyDesc* OutPropDesc = nullptr,
+		FPropertyBagPropertyDesc* OutEditConditionPropDesc = nullptr);
+
+private:
+	/** The properties managed by this object. */
+	UPROPERTY(EditAnywhere, meta=(ShowOnlyInnerProperties, FixedLayout), Category = "Value")
+	FInstancedPropertyBag Value;
+
+	/** The graph preset associated with the variable overrides. */
+	UPROPERTY()
+	TSoftObjectPtr<UMovieGraphConfig> GraphPreset;
+};
\ No newline at end of file
diff --git a/Engine/Plugins/MovieScene/MovieRenderPipeline/Source/MovieRenderPipelineCore/Public/MoviePipelineQueue.h b/Engine/Plugins/MovieScene/MovieRenderPipeline/Source/MovieRenderPipelineCore/Public/MoviePipelineQueue.h
--- a/Engine/Plugins/MovieScene/MovieRenderPipeline/Source/MovieRenderPipelineCore/Public/MoviePipelineQueue.h
+++ b/Engine/Plugins/MovieScene/MovieRenderPipeline/Source/MovieRenderPipelineCore/Public/MoviePipelineQueue.h
@@ -4,6 +4,7 @@
 #include "UObject/Object.h"
 #include "MovieRenderPipelineDataTypes.h"
 #include "LevelSequence.h"
+#include "MovieJobVariableAssignmentContainer.h"
 #include "MoviePipelinePrimaryConfig.h"
 #include "MoviePipelineShotConfig.h"
 #include "MoviePipelineConfigBase.h"
@@ -52,6 +53,7 @@ public:
 		: bEnabled(true)
 	{
 		Progress = 0.f;
+		VariableAssignments = CreateDefaultSubobject<UMovieJobVariableAssignmentContainer>("VariableAssignments");
 	}
 
 public:
@@ -170,6 +172,11 @@ public:
 	void SetGraphPreset(const UMovieGraphConfig* InGraphPreset)
 	{
 		GraphPreset = InGraphPreset;
+		VariableAssignments->SetGraphConfig(GraphPreset);
+
+#if WITH_EDITOR
+		VariableAssignments->UpdateGraphVariableOverrides();
+#endif
 	}
 
 	/**
@@ -199,6 +206,12 @@ public:
 		}
 		return InnerName;
 	}
+
+	// UObject Interface
+	virtual void PreSave(FObjectPreSaveContext ObjectSaveContext) override;
+	virtual void PostLoad() override;
+	// ~UObject Interface
+	
 protected:
 	// UMoviePipipelineExecutorShot Interface
 	virtual void SetStatusMessage_Implementation(const FString& InMessage) { StatusMessage = InMessage; }
@@ -224,6 +237,11 @@ public:
 	/** List of cameras to render for this shot. Only used if the setting flag is set in the Camera setting. */
 	UPROPERTY(BlueprintReadWrite, EditAnywhere, Category = "Movie Render Pipeline")
 	TArray<FMoviePipelineSidecarCamera> SidecarCameras;
+
+	/** Overrides on the variables in graph preset associated with this job. */
+	UPROPERTY(BlueprintReadWrite, EditAnywhere, Category = "Variable Assignments")
+	TObjectPtr<UMovieJobVariableAssignmentContainer> VariableAssignments;
+
 public:
 	/** Transient information used by the active Movie Pipeline working on this shot. */
 	FMoviePipelineCameraCutInfo ShotInfo;
@@ -250,25 +268,6 @@ private:
 	TSoftObjectPtr<UMovieGraphConfig> GraphPreset;
 };
 
-/** An override to a graph variable's value from the job level. */
-USTRUCT(BlueprintType)
-struct FMoviePipelineGraphVariableValueAssignment
-{
-	GENERATED_BODY()
-	
-	/** The GUID of the variable that is having its value changed. */
-	UPROPERTY()
-	FGuid VariableGuid;
-
-	/** The new value of the variable, if one was set. */
-	UPROPERTY()
-	float VariableValue = 0.f;
-
-	/** Whether a value override has been set. */
-	UPROPERTY()
-	bool bIsValueSet = false;
-};
-
 /**
 * A particular job within the Queue
 */
@@ -283,6 +282,7 @@ public:
 		StatusProgress = 0.f;
 		bIsConsumed = false;
 		Configuration = CreateDefaultSubobject<UMoviePipelinePrimaryConfig>("DefaultConfig");
+		VariableAssignments = CreateDefaultSubobject<UMovieJobVariableAssignmentContainer>("VariableAssignments");
 	}
 
 public:	
@@ -481,6 +481,11 @@ public:
 	void SetGraphPreset(const UMovieGraphConfig* InGraphPreset)
 	{
 		GraphPreset = InGraphPreset;
+		VariableAssignments->SetGraphConfig(InGraphPreset);
+
+#if WITH_EDITOR
+		VariableAssignments->UpdateGraphVariableOverrides();
+#endif
 	}
 
 	/**
@@ -494,41 +499,6 @@ public:
 		GraphConfig = InGraphConfig;
 	}
 
-	/**
-	 * Add a variable assignment for the graph. Returns true if the variable was added or it already exists,
-	 * else false.
-	 */
-	UFUNCTION(BlueprintCallable, Category = "Movie Render Pipeline|Experimental")
-	bool AddVariableAssignment(const UMovieGraphVariable* InGraphVariable);
-
-	/**
-	 * Updates an existing variable assignment for the provided graph variable to a new value, or adds a new assignment
-	 * and updates its value. Returns true on success, else false.
-	 */
-	UFUNCTION(BlueprintCallable, Category = "Movie Render Pipeline|Experimental")
-	bool SetVariableAssignmentValue(const UMovieGraphVariable* InGraphVariable, float NewValue);
-
-	/**
-	 * Updates an existing variable assignment for the provided graph variable to a new enable state, or adds a new
-	 * assignment and updates its enable state. Returns true on success, else false.
-	 */
-	UFUNCTION(BlueprintCallable, Category = "Movie Render Pipeline|Experimental")
-	bool SetVariableAssignmentEnableState(const UMovieGraphVariable* InGraphVariable, bool bIsEnabled);
-
-	/**
-	 * Gets the value of the variable assignment for the provided graph variable. The value is provided via OutValue.
-	 * Returns true if a value was set on the variable and OutValue was changed, else false.
-	 */
-	UFUNCTION(BlueprintCallable, Category = "Movie Render Pipeline|Experimental")
-	bool GetVariableAssignmentValue(const UMovieGraphVariable* InGraphVariable, float& OutValue);
-
-	/**
-	 * Gets the enable state of the variable assignment for the provided graph variable. The enable state is provided
-	 * via bOutIsEnabled. Returns true if an enable state was set on the variable and bOutIsEnabled was changed, else false.
-	 */
-	UFUNCTION(BlueprintCallable, Category = "Movie Render Pipeline|Experimental")
-	bool GetVariableAssignmentEnableState(const UMovieGraphVariable* InGraphVariable, bool& bOutIsEnabled);
-
 	UFUNCTION(BlueprintSetter, Category = "Movie Render Pipeline")
 	void SetSequence(FSoftObjectPath InSequence);
 
@@ -538,6 +508,7 @@ public:
 	void PostEditChangeProperty(FPropertyChangedEvent& PropertyChangedEvent);
 #endif
 	virtual void PreSave(FObjectPreSaveContext ObjectSaveContext) override;
+	virtual void PostLoad() override;
 	// ~UObject Interface
 
 protected:
@@ -559,19 +530,6 @@ protected:
 	virtual bool IsEnabled_Implementation() const { return bEnabled; }
 	// ~UMoviePipelineExecutorJob Interface
 
-	// Note: This is not exposed to the public API because the underlying struct should remain private.
-	/**
-	 * Finds a variable assignment for the provided variable, or adds one if one does not already exist (and
-	 * bAddIfNotExists is set to true). Returns nullptr if the operation failed. 
-	 */
-	FMoviePipelineGraphVariableValueAssignment* FindOrAddVariableAssignment(const UMovieGraphVariable* InGraphVariable, bool bAddIfNotExists = true);
-
-	/**
-	 * Removes variable assignments for variables which no longer exist on the graph that the job currently uses. Old
-	 * and invalid assignments are not necessarily harmful, but they can bloat the asset.
-	 */
-	void CleanUpInvalidVariableAssignments();
-
 public:
 	/** (Optional) Name of the job. Shown on the default burn-in. */
 	UPROPERTY(BlueprintReadWrite, EditAnywhere, Category = "Movie Render Pipeline")
@@ -604,6 +562,11 @@ public:
 	*/
 	UPROPERTY(BlueprintReadWrite, Category = "Movie Render Pipeline")
 	FString UserData;
+
+	/** Overrides on the variables in graph preset associated with this job. */
+	UPROPERTY(BlueprintReadWrite, EditAnywhere, Instanced, Category = "Variable Assignments")
+	TObjectPtr<UMovieJobVariableAssignmentContainer> VariableAssignments;
+
 private:
 	UPROPERTY(Transient)
 	FString StatusMessage;
@@ -614,7 +577,6 @@ private:
 	UPROPERTY(Transient)
 	bool bIsConsumed;
 
-private:
 	/** 
 	*/
 	UPROPERTY(Instanced)
@@ -636,10 +598,6 @@ private:
 	/** The graph-based configuration preset that this job is using. Can be nullptr. */
 	UPROPERTY()
 	TSoftObjectPtr<UMovieGraphConfig> GraphPreset;
-
-	/** Values that the job has set on variables for the graph it is using (if any). */
-	UPROPERTY()
-	TArray<FMoviePipelineGraphVariableValueAssignment> GraphVariableAssignments;
 };
 
 /**
diff --git a/Engine/Plugins/MovieScene/MovieRenderPipeline/Source/MovieRenderPipelineEditor/MovieRenderPipelineEditor.Build.cs b/Engine/Plugins/MovieScene/MovieRenderPipeline/Source/MovieRenderPipelineEditor/MovieRenderPipelineEditor.Build.cs
--- a/Engine/Plugins/MovieScene/MovieRenderPipeline/Source/MovieRenderPipelineEditor/MovieRenderPipelineEditor.Build.cs
+++ b/Engine/Plugins/MovieScene/MovieRenderPipeline/Source/MovieRenderPipelineEditor/MovieRenderPipelineEditor.Build.cs
@@ -68,6 +68,7 @@ public class MovieRenderPipelineEditor : ModuleRules
 				"ConsoleVariablesEditor",
 				"StructUtils",
 				"GraphEditor",
+				"StructUtilsEditor",
 			}
         );
 
diff --git a/Engine/Plugins/MovieScene/MovieRenderPipeline/Source/MovieRenderPipelineEditor/Private/Customizations/JobCustomization.h b/Engine/Plugins/MovieScene/MovieRenderPipeline/Source/MovieRenderPipelineEditor/Private/Customizations/JobCustomization.h
--- a/Engine/Plugins/MovieScene/MovieRenderPipeline/Source/MovieRenderPipelineEditor/Private/Customizations/JobCustomization.h
+++ b/Engine/Plugins/MovieScene/MovieRenderPipeline/Source/MovieRenderPipelineEditor/Private/Customizations/JobCustomization.h
@@ -26,16 +26,11 @@ protected:
 	//~ Begin IDetailCustomization interface
 	virtual void CustomizeDetails(const TSharedPtr<IDetailLayoutBuilder>& DetailBuilder) override
 	{
-		WeakDetailBuilder = DetailBuilder;
 		CustomizeDetails(*DetailBuilder);
 	}
 
 	virtual void CustomizeDetails(IDetailLayoutBuilder& DetailBuilder) override
 	{
-		// Add a new "Graph Variables" category. Set as "Uncommon" priority to push variables down below the other properties.
-		IDetailCategoryBuilder& GraphVariablesCategory = DetailBuilder.EditCategory(
-"GraphVariables", LOCTEXT("GraphVariablesCategory", "Graph Variables"), ECategoryPriority::Uncommon);
-
 		TArray<TWeakObjectPtr<UObject>> ObjectsBeingCustomized;
 		DetailBuilder.GetObjectsBeingCustomized(ObjectsBeingCustomized);
 
@@ -45,142 +40,28 @@ protected:
 			return;
 		}
 
-		// Get the displayed job
-		UMoviePipelineExecutorJob* CustomizedJob = Cast<UMoviePipelineExecutorJob>(ObjectsBeingCustomized[0].Get());
-		if (!CustomizedJob)
-		{
-			return;
-		}
-
-		UMovieGraphConfig* GraphConfig = GetGraphConfig(CustomizedJob);
-		if (!GraphConfig)
-		{
-			return;
-		}
-
-		for (const UMovieGraphVariable* Variable : GraphConfig->GetVariables())
-		{
-			AddVariableAssignmentRow(GraphVariablesCategory, DetailBuilder, Variable, CustomizedJob);
-		}
-
-		// Refresh the variable listing if anything changes with the variables
-		GraphConfig->OnGraphVariablesChangedDelegate.AddSP(this, &FJobDetailsCustomization::Refresh);
-	}
-	//~ End IDetailCustomization interface
-
-	void Refresh() const
-	{
-		// Don't keep the details builder alive just for the sake of this Refresh() call
-		if (IDetailLayoutBuilder* DetailLayoutBuilder = WeakDetailBuilder.Pin().Get())
+		// Hide the original assignments property (since it presents an asset picker by default), and add the Value
+		// property (which is a property bag) instead -- that will present each of the variable overrides
+		if (UMoviePipelineExecutorJob* Job = CastChecked<UMoviePipelineExecutorJob>(ObjectsBeingCustomized[0]))
 		{
-			DetailLayoutBuilder->ForceRefreshDetails();
-		}
-	}
-
-	void AddVariableAssignmentRow(IDetailCategoryBuilder& DetailCategoryBuilder, const IDetailLayoutBuilder& DetailBuilder,
-		const UMovieGraphVariable* Variable, UMoviePipelineExecutorJob* CustomizedJob)
-	{
-		FDetailWidgetRow& NewRow = DetailCategoryBuilder.AddCustomRow(FText::FromString(Variable->Name));
-		NewRow.WholeRowWidget
-		[
-			SNew(SHorizontalBox)
+			const TSharedRef<IPropertyHandle> AssignmentsProperty =
+				DetailBuilder.GetProperty(GET_MEMBER_NAME_CHECKED(UMoviePipelineExecutorJob, VariableAssignments));
 
-			+ SHorizontalBox::Slot()
-			.AutoWidth()
-			.VAlign(VAlign_Center)
-			.Padding(0, 0, 5, 0)
-			[
-				SNew(SCheckBox)
-				.IsChecked(this, &FJobDetailsCustomization::IsVariableOverrideChecked, Variable, CustomizedJob)
-				.OnCheckStateChanged(this, &FJobDetailsCustomization::OnVariableOverrideEnableStateChanged, Variable, CustomizedJob)
-			]
-				
-			+ SHorizontalBox::Slot()
-			.FillWidth(0.5f)
-			.VAlign(VAlign_Center)
-			.Padding(0, 0, 5, 0)
-			[
-				SNew(STextBlock)
-				.Text(FText::FromString(Variable->Name))
-				.IsEnabled(this, &FJobDetailsCustomization::IsVariableOverrideEnabled, Variable, CustomizedJob)
-				.Font(DetailBuilder.GetDetailFont())
-			]
+			// Always hide the assignments property
+			DetailBuilder.HideProperty(AssignmentsProperty);
 
-			+ SHorizontalBox::Slot()
-			.FillWidth(0.5f)
-			.VAlign(VAlign_Center)
-			.Padding(0, 0, 1, 0)
-			[
-				SNew(SNumericEntryBox<float>)
-				.Value(this, &FJobDetailsCustomization::GetVariableValue, Variable, CustomizedJob)
-				.IsEnabled(this, &FJobDetailsCustomization::IsVariableOverrideEnabled, Variable, CustomizedJob)
-				.Font(DetailBuilder.GetDetailFont())
-				.OnValueCommitted(this, &FJobDetailsCustomization::OnVariableValueCommitted, Variable, CustomizedJob)
-			]
-		];
-	}
-
-	UMovieGraphConfig* GetGraphConfig(const UMoviePipelineExecutorJob* InJob) const
-	{
-		// TODO: This should get the config OR preset
-		return InJob->GetGraphPreset();
-	}
-	
-	ECheckBoxState IsVariableOverrideChecked(const UMovieGraphVariable* InVariable, UMoviePipelineExecutorJob* InJob) const
-	{
-		if (InJob && InVariable)
-		{
-			bool bIsEnabled;
-			if (InJob->GetVariableAssignmentEnableState(InVariable, bIsEnabled))
+			if (Job->VariableAssignments->GetNumAssignments() > 0)
 			{
-				return bIsEnabled ? ECheckBoxState::Checked : ECheckBoxState::Unchecked;
+				// Add a new "Graph Variables" category if there is at least one assignment. Set as "Uncommon" priority to
+				// push variables down below the other properties.
+				IDetailCategoryBuilder& GraphVariablesCategory = DetailBuilder.EditCategory(
+					"GraphVariables", LOCTEXT("GraphVariablesCategory", "Graph Variables"), ECategoryPriority::Uncommon);
+			
+				GraphVariablesCategory.AddExternalObjectProperty({Job->VariableAssignments.Get()}, FName("Value"));
 			}
 		}
-
-		return ECheckBoxState::Unchecked;
-	}
-
-	bool IsVariableOverrideEnabled(const UMovieGraphVariable* InVariable, UMoviePipelineExecutorJob* InJob) const
-	{
-		return IsVariableOverrideChecked(InVariable, InJob) == ECheckBoxState::Checked;
 	}
-
-	void OnVariableOverrideEnableStateChanged(ECheckBoxState CheckBoxState, const UMovieGraphVariable* InVariable, UMoviePipelineExecutorJob* InJob) const
-	{
-		if (InJob && InVariable)
-		{
-			InJob->SetVariableAssignmentEnableState(InVariable, CheckBoxState == ECheckBoxState::Checked);
-		}
-	}
-
-	TOptional<float> GetVariableValue(const UMovieGraphVariable* InVariable, UMoviePipelineExecutorJob* InJob) const
-	{
-		if (InJob && InVariable)
-		{
-			// Get the value of the variable from the job, if any
-			float VariableValue;
-			if (InJob->GetVariableAssignmentValue(InVariable, VariableValue))
-			{
-				return VariableValue;
-			}
-
-			// Fall back to the default value of the variable
-			return InVariable->Default;
-		}
-
-		return TOptional<float>();
-	}
-
-	void OnVariableValueCommitted(float NewValue, ETextCommit::Type Arg, const UMovieGraphVariable* Variable, UMoviePipelineExecutorJob* InJob) const
-	{
-		if (InJob)
-		{
-			InJob->SetVariableAssignmentValue(Variable, NewValue);
-		}
-	}
-
-private:
-	TWeakPtr<IDetailLayoutBuilder> WeakDetailBuilder;
+	//~ End IDetailCustomization interface
 };
 
 #undef LOCTEXT_NAMESPACE
\ No newline at end of file
diff --git a/Engine/Plugins/MovieScene/MovieRenderPipeline/Source/MovieRenderPipelineEditor/Private/Graph/MovieEdGraphNode.cpp b/Engine/Plugins/MovieScene/MovieRenderPipeline/Source/MovieRenderPipelineEditor/Private/Graph/MovieEdGraphNode.cpp
--- a/Engine/Plugins/MovieScene/MovieRenderPipeline/Source/MovieRenderPipelineEditor/Private/Graph/MovieEdGraphNode.cpp
+++ b/Engine/Plugins/MovieScene/MovieRenderPipeline/Source/MovieRenderPipelineEditor/Private/Graph/MovieEdGraphNode.cpp
@@ -53,26 +53,24 @@ FEdGraphPinType UMoviePipelineEdGraphNodeBase::GetPinType(const UMovieGraphPin*
 	EdPinType.PinCategory = NAME_None;
 	EdPinType.PinSubCategory = NAME_None;
 
+	// TODO: This needs to properly fill in the category/subcategory for all types. Also, container type?
 	switch (InPin->Properties.Type)
 	{
-	case EMovieGraphMemberType::Branch:
+	case EMovieGraphValueType::Branch:
 		EdPinType.PinCategory = UMovieGraphSchema::PC_Branch;
 		break;
-	case EMovieGraphMemberType::Float:
+	case EMovieGraphValueType::Float:
 		EdPinType.PinCategory = UMovieGraphSchema::PC_Float;
 		break;
-	case EMovieGraphMemberType::Int:
+	case EMovieGraphValueType::Int32:
 		EdPinType.PinCategory = UMovieGraphSchema::PC_Integer;
 		break;
-	case EMovieGraphMemberType::Bool:
+	case EMovieGraphValueType::Bool:
 		EdPinType.PinCategory = UMovieGraphSchema::PC_Boolean;
 		break;
-	case EMovieGraphMemberType::String:
+	case EMovieGraphValueType::String:
 		EdPinType.PinCategory = UMovieGraphSchema::PC_String;
 		break;
-	case EMovieGraphMemberType::IntPoint:
-		EdPinType.PinCategory = UMovieGraphSchema::PC_IntPoint;
-		break;
 	default:
 		EdPinType.PinCategory = UMovieGraphSchema::PC_Float;
 		break;
