MoviePipeline: Initial graph traversal implementation which traverses the graph and builds a flattened list of settings which have had their final values calculated via traversal, respecting override flags, etc.

diff --git a/Engine/Plugins/MovieScene/MovieRenderPipeline/Source/MovieRenderPipelineCore/Private/Graph/MovieGraphConfig.cpp b/Engine/Plugins/MovieScene/MovieRenderPipeline/Source/MovieRenderPipelineCore/Private/Graph/MovieGraphConfig.cpp
--- a/Engine/Plugins/MovieScene/MovieRenderPipeline/Source/MovieRenderPipelineCore/Private/Graph/MovieGraphConfig.cpp
+++ b/Engine/Plugins/MovieScene/MovieRenderPipeline/Source/MovieRenderPipelineCore/Private/Graph/MovieGraphConfig.cpp
@@ -664,4 +664,274 @@ bool UMovieGraphConfig::DeleteOutputMember(UMovieGraphOutput* OutputMemberToDele
 	return false;
 }
 
+FBoolProperty* UMovieGraphConfig::FindOverridePropertyForRealProperty(UClass* InClass, const FProperty* InRealProperty) const
+{
+	static const FName EditConditionKey = FName("EditCondition");
+	if (!ensure(InClass && InRealProperty))
+	{
+		return nullptr;
+	}
+
+	// We can't get access to metadata in shipping builds, so we need to just rely on a naming pattern of bOverride_<PropertyName>
+	const FString DesiredPropertyName = FString::Printf(TEXT("bOverride_%s"), *InRealProperty->GetName());
+
+	for (TFieldIterator<FProperty> PropertyIterator(InClass); PropertyIterator; ++PropertyIterator)
+	{
+		FProperty* CheckProperty = *PropertyIterator;
+		if (CheckProperty && CheckProperty->IsA<FBoolProperty>())
+		{
+			FBoolProperty* PropertyAsBool = CastFieldChecked<FBoolProperty>(CheckProperty);
+			if (PropertyAsBool->GetName() == DesiredPropertyName)
+			{
+				return PropertyAsBool;
+			}
+		}
+	}
+
+	return nullptr;
+}
+
+void UMovieGraphConfig::InitializeFlattenedNode(UMovieGraphNode* InNode)
+{
+	static const FName EditConditionKey = FName("EditCondition");
+	
+	// We go through each of the bOverride_ properties on this new instance and set
+	// them all as "not overridden", so that as we traverse the graph, we can only
+	// update the node the first time we hit another property that is overridden.
+	// If they never override the values anywhere in the chain then it's okay, because
+	// it will use the values from the CDO.
+	for (TFieldIterator<FProperty> PropertyIterator(InNode->GetClass()); PropertyIterator; ++PropertyIterator)
+	{
+		FProperty* CheckProperty = *PropertyIterator;
+		FBoolProperty* EditConditionProperty = FindOverridePropertyForRealProperty(InNode->GetClass(), CheckProperty);
+		if (EditConditionProperty)
+		{
+			EditConditionProperty->SetPropertyValue_InContainer(InNode, false);
+		}
+	}
+}
+
+void UMovieGraphConfig::CopyOverriddenProperties(UMovieGraphNode* FromNode, UMovieGraphNode* ToNode)
+{
+	if (!ensure(FromNode && ToNode))
+	{
+		return;
+	}
+
+	if (!ensureMsgf(FromNode->GetClass() == ToNode->GetClass(), TEXT("Cross-Class Property copying is not supported at this time.")))
+	{
+		return;
+	}
+
+	for (TFieldIterator<FProperty> PropertyIterator(ToNode->GetClass()); PropertyIterator; ++PropertyIterator)
+	{
+		// For each property on the destination node, decide if we need to try to update it (we don't update bOverride_ properties)
+		FProperty* PropertyToCheckOnDestNode = *PropertyIterator;
+
+		// We use the existence of a matching edit condition node to signal that this property is overrideable, ie:
+		// float MyFoo would have a bOverride_MyFoo, so FindOverridePropertyForRealProperty would match it. However when
+		// looking at the bOverride_MyFoo property it would look for bOverride_bOverride_MyFoo, which wouldn't exist, successfully
+		// filtering us to only the "real" overrideable properties.
+		FBoolProperty* EditConditionProperty = FindOverridePropertyForRealProperty(ToNode->GetClass(), PropertyToCheckOnDestNode);
+		if (PropertyToCheckOnDestNode && EditConditionProperty)
+		{
+			// The user has indicated that this property should be overrideable in the Graph network.
+			// The first thing we do is check to see if it is exposed, and if it's exposed, we need to get the value from the graph.
+			bool bGetValueFromPin = false;
+			const TArray<FName>& ExposedPins = FromNode->GetExposedDynamicProperties();
+			for (const FName& Name : ExposedPins)
+			{
+				if (Name == PropertyToCheckOnDestNode->GetName())
+				{
+					bGetValueFromPin = true;
+					break;
+				}
+			}
+
+			if (bGetValueFromPin)
+			{
+				UE_LOG(LogMovieRenderPipeline, Error, TEXT("not yet implemented."));
+				continue;
+			}
+
+			// If our destination node already has this property marked as overridden, then some other node in the graph has
+			// taken priority and set the value to something, so we don't want to override it.
+			bool bAlreadyOverridenOnDestNode = EditConditionProperty->GetPropertyValue_InContainer(ToNode);
+			if (bAlreadyOverridenOnDestNode)
+			{
+				continue;
+			}
+
+			// We know it's not already overridden, so now we should check to see if the incoming node wants to override it.
+			bool bSourceNodeOverwrites = EditConditionProperty->GetPropertyValue_InContainer(FromNode);
+			if (!bSourceNodeOverwrites)
+			{
+				// The source node didn't have the override flag checked, so we don't copy the value from it.
+				continue;
+			}
+
+			// Okay at this point we know that on our target node, no one has overridden it yet, and our source node wants to override this property.
+			// First, we update the booleans to say that yes, this property has been overridden on the target node.
+			EditConditionProperty->SetPropertyValue_InContainer(ToNode, true);
+
+			// Now we need to copy the value from the source to the destination
+			PropertyToCheckOnDestNode->CopyCompleteValue_InContainer(ToNode, FromNode);
+		}
+	}
+}
+
+void UMovieGraphConfig::CreateFlattenedGraph_Recursive(UMovieGraphEvaluatedConfig* InOwningConfig, FMovieGraphEvaluatedBranchConfig& OutBranchConfig,
+	FMovieGraphEvaluationContext& InEvaluationContext, UMovieGraphPin* InPinToFollow)
+{
+	if (!InPinToFollow)
+	{
+		return;
+	}
+
+	UMovieGraphNode* Node = InPinToFollow->Node;
+	if (!Node)
+	{
+		return;
+	}
+
+	// We only follow execution pins during traversal.
+	if (!ensureMsgf(InPinToFollow->Properties.Type == EMovieGraphValueType::Branch, TEXT("Only Branch pins should be contained by InPinToFollow!")))
+	{
+		return;
+	}
+
+	// Check to see if our flattened evaluation graph already has a copy of this node.
+	InEvaluationContext.VisitedNodes.Add(Node);
+	const bool bShouldIncludeNode = Node->IsA<UMovieGraphSettingNode>();
+
+	if(bShouldIncludeNode)
+	{
+		UMovieGraphNode* ExistingNode = OutBranchConfig.GetNodeByClassExactMatch(Node->GetClass());
+		if (!ExistingNode)
+		{
+			// Create a new instance of this node inside our flattened eval graph
+			ExistingNode = NewObject<UMovieGraphNode>(InOwningConfig, Node->GetClass());
+			OutBranchConfig.NodeInstances.Add(ExistingNode);
+
+			// Set all of the boolean edit condition values to false, so we can use "true" to indicate
+			// that the value was overridden already during traversal.
+			InitializeFlattenedNode(ExistingNode);
+		}
+
+		// Now do a property-copy from this node onto our flattened one. We don't use the generic property
+		// copy routines in the engine because we have special handling (we want to check if the property
+		// is actually marked for override, and also skip if this has already been overridden.
+
+		// ToDo: Handle "Disable" ndoes that can disable upstream types of nodes, etc. Push disable types into
+		// the currently evaluating context
+		CopyOverriddenProperties(Node, ExistingNode);
+	}
+	
+	// Now that we've potentially resolved the values on this node, continue to travel up-stream along any execution pins,
+	// potentially following re-route nodes, sub-graph nodes, through branches, etc.
+	TArray<UMovieGraphPin*> NewPinsToFollow = Node->EvaluatePinsToFollow(InEvaluationContext.UserContext);
+	
+	for(UMovieGraphPin* Pin : NewPinsToFollow)
+	{
+		for (UMovieGraphEdge* Edge : Pin->Edges)
+		{
+			if (UMovieGraphPin* OtherPin = Edge->GetOtherPin(Pin))
+			{
+				UMovieGraphNode* OtherNode = OtherPin->Node;
+
+				if (InEvaluationContext.VisitedNodes.Contains(OtherNode))
+				{
+					// ToDo: This won't work long term if you have two different branches visiting the same node
+					// also we need to reset this every time we go start from the root.
+					UE_LOG(LogMovieRenderPipeline, Error, TEXT("Circular graph?"));
+					continue;
+				}
+
+				CreateFlattenedGraph_Recursive(InOwningConfig, OutBranchConfig, InEvaluationContext, OtherPin);
+			}
+		}
+	}
+}
+
+UMovieGraphEvaluatedConfig* UMovieGraphConfig::CreateFlattenedGraph(const FMovieGraphTraversalContext& InContext)
+{
+	TRACE_CPUPROFILER_EVENT_SCOPE(MRQ_CreateFlattenedGraph);
+	LLM_SCOPE_BYNAME(TEXT("MovieGraph/CreateFlattenedGraph"));
+
+	UMovieGraphEvaluatedConfig* NewContext = NewObject<UMovieGraphEvaluatedConfig>(this);
+
+	if (OutputNode)
+	{
+		TArray<UMovieGraphPin*> InputPinsToFollow = OutputNode->GetInputPins();
+		UMovieGraphPin* GlobalsPin = OutputNode->GetInputPin(UMovieGraphNode::GlobalsPinName);
+
+		// For each input pin, we create an instance of the config (including Globals, since some queries are made with no context)
+		// but we also want each named branch to have a complete, resolved set of configs.
+		for (UMovieGraphPin* InputPin : OutputNode->GetInputPins())
+		{
+			const FName BranchName = InputPin->Properties.Label;
+			FMovieGraphEvaluatedBranchConfig& BranchConfig = NewContext->BranchConfigMapping.Add(BranchName);
+			
+			// The stack evaluation context is per-branch
+			FMovieGraphEvaluationContext StackContext;
+			StackContext.UserContext = InContext;
+
+			// Follow the branch connected to this pin
+			for (UMovieGraphEdge* Edge : InputPin->Edges)
+			{
+				UMovieGraphPin* OtherPin = Edge->GetOtherPin(InputPin);
+				if (OtherPin)
+				{
+					CreateFlattenedGraph_Recursive(NewContext, /*InOut*/ BranchConfig, StackContext, OtherPin);
+				}
+			}
+
+			// Now, if this isn't the Globals branch, we apply the Globals branch after the branch settings.
+			// This allows users to override things on a per-branch basis, and if they don't set it, then the Globals
+			// branch has a chance to set "defaults" (which then fall back to CDO values if the Globals branch doesn't 
+			// set it). We skip doing this for the Globals branch because the above loop just did that.
+			if (BranchName != UMovieGraphNode::GlobalsPinName && GlobalsPin)
+			{
+				// We use a new context here as we consider every branch independent.
+				FMovieGraphEvaluationContext GlobalStackContext;
+				GlobalStackContext.UserContext = InContext;
+				for (UMovieGraphEdge* Edge : GlobalsPin->Edges)
+				{
+					UMovieGraphPin* OtherPin = Edge->GetOtherPin(GlobalsPin);
+					if (OtherPin)
+					{
+						CreateFlattenedGraph_Recursive(NewContext, /*InOut*/ BranchConfig, GlobalStackContext, OtherPin);
+					}
+				}
+			}
+		}		
+	}
+
+	UE_LOG(LogTemp, Warning, TEXT("Traversed Graph:"));
+	for (const TPair<FName, FMovieGraphEvaluatedBranchConfig>& Pair : NewContext->BranchConfigMapping)
+	{
+		UE_LOG(LogTemp, Warning, TEXT("\t Branch: %s"), *Pair.Key.ToString());
+
+		for (UMovieGraphNode* Node : Pair.Value.GetNodes())
+		{
+			UE_LOG(LogTemp, Warning, TEXT("\t\t %s Class:"), *Node->GetClass()->GetName());
+			for (TFieldIterator<FProperty> PropertyIterator(Node->GetClass()); PropertyIterator; ++PropertyIterator)
+			{
+				FProperty* CheckProperty = *PropertyIterator;
+				FBoolProperty* EditConditionProperty = FindOverridePropertyForRealProperty(Node->GetClass(), CheckProperty);
+				if (EditConditionProperty)
+				{
+					FString ExportText;
+					CheckProperty->ExportText_InContainer(0, ExportText, Node, Node, Node, 0);
+					UE_LOG(LogTemp, Warning, TEXT("\t\t\t %s : %s"), *CheckProperty->GetName(), *ExportText);
+				}
+			}
+		}
+
+	}
+
+	return NewContext;
+}
+
+
 #undef LOCTEXT_NAMESPACE
diff --git a/Engine/Plugins/MovieScene/MovieRenderPipeline/Source/MovieRenderPipelineCore/Private/Graph/MovieGraphNode.cpp b/Engine/Plugins/MovieScene/MovieRenderPipeline/Source/MovieRenderPipelineCore/Private/Graph/MovieGraphNode.cpp
--- a/Engine/Plugins/MovieScene/MovieRenderPipeline/Source/MovieRenderPipelineCore/Private/Graph/MovieGraphNode.cpp
+++ b/Engine/Plugins/MovieScene/MovieRenderPipeline/Source/MovieRenderPipelineCore/Private/Graph/MovieGraphNode.cpp
@@ -4,6 +4,8 @@
 
 #include "Graph/MovieGraphConfig.h"
 
+FName UMovieGraphNode::GlobalsPinName = FName("Globals");
+
 UMovieGraphNode::UMovieGraphNode()
 {
 }
@@ -204,4 +206,35 @@ void UMovieGraphNode::PostLoad()
 	Super::PostLoad();
 
 	RegisterDelegates();
+}
+
+TArray<UMovieGraphPin*> UMovieGraphNode::EvaluatePinsToFollow(const FMovieGraphTraversalContext& InContext) const
+{
+	TArray<UMovieGraphPin*> PinsToFollow;
+
+	// By default we provide all Input Pins to this node that are the Branch Type.
+	// You should override this in downstream nodes that need custom logic, such
+	// as branch or switch nodes.
+	for (const TObjectPtr<UMovieGraphPin>& InputPin : GetInputPins())
+	{
+		if (InputPin->Properties.Type == EMovieGraphValueType::Branch)
+		{
+			PinsToFollow.Add(InputPin);
+		}
+	}
+	return PinsToFollow;
+}
+
+TArray<FMovieGraphPinProperties> UMovieGraphSettingNode::GetInputPinProperties() const
+{
+	TArray<FMovieGraphPinProperties> Properties;
+	Properties.Add(FMovieGraphPinProperties(NAME_None, EMovieGraphValueType::Branch, false));
+	return Properties;
+}
+
+TArray<FMovieGraphPinProperties> UMovieGraphSettingNode::GetOutputPinProperties() const
+{
+	TArray<FMovieGraphPinProperties> Properties;
+	Properties.Add(FMovieGraphPinProperties(NAME_None, EMovieGraphValueType::Branch, false));
+	return Properties;
 }
\ No newline at end of file
diff --git a/Engine/Plugins/MovieScene/MovieRenderPipeline/Source/MovieRenderPipelineCore/Private/Graph/Nodes/MovieGraphBranchNode.cpp b/Engine/Plugins/MovieScene/MovieRenderPipeline/Source/MovieRenderPipelineCore/Private/Graph/Nodes/MovieGraphBranchNode.cpp
new file mode 100644
--- /dev/null
+++ b/Engine/Plugins/MovieScene/MovieRenderPipeline/Source/MovieRenderPipelineCore/Private/Graph/Nodes/MovieGraphBranchNode.cpp
@@ -0,0 +1,130 @@
+// Copyright Epic Games, Inc. All Rights Reserved.
+
+#include "Graph/Nodes/MovieGraphBranchNode.h"
+
+#include "Graph/MovieGraphConfig.h"
+#include "Graph/MovieGraphEdge.h"
+#include "MovieRenderPipelineCoreModule.h"
+
+#include "Styling/AppStyle.h"
+
+#define LOCTEXT_NAMESPACE "MovieGraphNode"
+
+namespace UE::MovieGraph::BranchNode
+{
+	static const FName TrueBranch("True");
+	static const FName FalseBranch("False");
+	static const FName Condition("Condition");
+}
+
+TArray<UMovieGraphPin*> UMovieGraphBranchNode::EvaluatePinsToFollow(const FMovieGraphTraversalContext& InContext) const
+{
+	TArray<UMovieGraphPin*> PinsToFollow;
+
+	// The branch node has two branches that could be followed, True or False. To figure out which one we're actually going
+	// to follow, we need to evaluate the Conditional pin. 
+	UMovieGraphPin* ConditionalPin = GetInputPin(UE::MovieGraph::BranchNode::Condition);
+	if (!ensure(ConditionalPin))
+	{
+		return PinsToFollow;
+	}
+
+	UMovieGraphPin* OtherPin = nullptr;
+	for (UMovieGraphEdge* Edge : ConditionalPin->Edges)
+	{
+		OtherPin = Edge->GetOtherPin(ConditionalPin);
+
+		// We only support a single connection to this pin type anyways.
+		break;
+	}
+
+	// There may not be a node actually connected. We don't know what to do in this case (as our nodes don't have default values)
+	// so for now we choose to follow neither branch.
+	if (!OtherPin)
+	{
+		UE_LOG(LogMovieRenderPipeline, Warning, TEXT("Unconnected conditional pin has no default value. Node: %s"), *GetName());
+		return PinsToFollow;
+	}
+
+	UMovieGraphNode* ConnectedNode = OtherPin->Node;
+	if (!ensure(ConnectedNode))
+	{
+		return PinsToFollow;
+	}
+
+	// Connected Node is the node that the pin is actually connected to.
+	const UMovieGraphValueContainer* ValueContainer = ConnectedNode->GetPropertyValueContainerForPin(OtherPin->GetName());
+
+	if (!ValueContainer)
+	{
+		// Somehow we got connected to node that can't provide a value.
+		UE_LOG(LogMovieRenderPipeline, Warning, TEXT("ConnectedNode: %s provided no ValueContainer to fetch value."), *ConnectedNode->GetName());
+		return PinsToFollow;
+	}
+
+	bool bEvaluatedValue;
+	bool bSuccessfullyFetched = ValueContainer->GetValueBool(bEvaluatedValue);
+	if (!bSuccessfullyFetched)
+	{
+		// There was a node connected but it didn't know how to return a boolean value for us.
+		UE_LOG(LogMovieRenderPipeline, Warning, TEXT("ValueContainer had no bool value."));
+		return PinsToFollow;
+	}
+
+	if (bEvaluatedValue)
+	{
+		PinsToFollow.Add(GetInputPin(UE::MovieGraph::BranchNode::TrueBranch));
+	}
+	else
+	{
+		PinsToFollow.Add(GetInputPin(UE::MovieGraph::BranchNode::FalseBranch));
+	}
+
+	return PinsToFollow;
+}
+
+TArray<FMovieGraphPinProperties> UMovieGraphBranchNode::GetInputPinProperties() const
+{
+	TArray<FMovieGraphPinProperties> Properties;
+	
+	Properties.Add(FMovieGraphPinProperties(UE::MovieGraph::BranchNode::TrueBranch, EMovieGraphValueType::Branch, false));
+	Properties.Add(FMovieGraphPinProperties(UE::MovieGraph::BranchNode::FalseBranch, EMovieGraphValueType::Branch, false));
+	Properties.Add(FMovieGraphPinProperties(UE::MovieGraph::BranchNode::Condition, EMovieGraphValueType::Bool, false));
+	return Properties;
+}
+
+TArray<FMovieGraphPinProperties> UMovieGraphBranchNode::GetOutputPinProperties() const
+{
+	TArray<FMovieGraphPinProperties> Properties;
+	Properties.Add(FMovieGraphPinProperties(NAME_None, EMovieGraphValueType::Branch, false));
+	return Properties;
+}
+
+#if WITH_EDITOR
+FText UMovieGraphBranchNode::GetNodeTitle(const bool bGetDescriptive) const
+{
+	static const FText BranchNodeName = LOCTEXT("NodeName_Branch", "Branch");
+	return BranchNodeName;
+}
+
+FText UMovieGraphBranchNode::GetMenuCategory() const
+{
+	return LOCTEXT("NodeCategory_Conditionals", "Conditionals");
+}
+
+FLinearColor UMovieGraphBranchNode::GetNodeTitleColor() const
+{
+	static const FLinearColor BranchNodeColor = FLinearColor(0.266f, 0.266f, 0.266f);
+	return BranchNodeColor;
+}
+
+FSlateIcon UMovieGraphBranchNode::GetIconAndTint(FLinearColor& OutColor) const
+{
+	static const FSlateIcon BranchIcon = FSlateIcon(FAppStyle::GetAppStyleSetName(), "Icons.Merge");
+
+	OutColor = FLinearColor::White;
+	return BranchIcon;
+}
+#endif // WITH_EDITOR
+
+#undef LOCTEXT_NAMESPACE
\ No newline at end of file
diff --git a/Engine/Plugins/MovieScene/MovieRenderPipeline/Source/MovieRenderPipelineCore/Public/Graph/MovieGraphConfig.h b/Engine/Plugins/MovieScene/MovieRenderPipeline/Source/MovieRenderPipelineCore/Public/Graph/MovieGraphConfig.h
--- a/Engine/Plugins/MovieScene/MovieRenderPipeline/Source/MovieRenderPipelineCore/Public/Graph/MovieGraphConfig.h
+++ b/Engine/Plugins/MovieScene/MovieRenderPipeline/Source/MovieRenderPipelineCore/Public/Graph/MovieGraphConfig.h
@@ -152,6 +152,86 @@ public:
 	DECLARE_MULTICAST_DELEGATE_OneParam(FOnMovieGraphNodesDeleted, TArray<UMovieGraphNode*>);
 #endif // WITH_EDITOR
 
+/**
+* A flattened list of configuration values for a given Graph Branch. For named branches, this incldues the "Globals"
+* branch (for any value not also overridden by the named branch).
+* 
+*/
+USTRUCT()
+struct MOVIERENDERPIPELINECORE_API FMovieGraphEvaluatedBranchConfig
+{
+	GENERATED_BODY()
+public:
+	UMovieGraphNode* GetNodeByClassExactMatch(const TSubclassOf<UMovieGraphNode>& InClass)
+	{
+		for (const TObjectPtr<UMovieGraphNode>& Instance : NodeInstances)
+		{
+			if (Instance && Instance->GetClass() == InClass)
+			{
+				return Instance;
+			}
+		}
+
+		return nullptr;
+	}
+
+	const TArray<TObjectPtr<UMovieGraphNode>>& GetNodes() const
+	{
+		return NodeInstances;
+	}
+
+private:
+	// Allow the config to add nodes to this, but otherwise we don't want the public adding nodes to them
+	// without going through the graph resolving.
+	friend class UMovieGraphConfig;
+	
+	UPROPERTY(Transient)
+	TArray<TObjectPtr<UMovieGraphNode>> NodeInstances;
+};
+
+/**
+* An evaluated config for the current frame. Each named branch (including Globals) has its own
+* copy of the config, fully resolved (so there is no need to check the Globals branch when
+* looking at a named branch). You can use the functions to fetch a node by type from a given
+* branch and it will return the right object (or the CDO if the node is NOT in the config).
+*/
+UCLASS()
+class MOVIERENDERPIPELINECORE_API UMovieGraphEvaluatedConfig : public UObject
+{
+	GENERATED_BODY()
+
+public:
+	/** Mapping between named branches (at the root of the config) and their evaluated values. */
+	UPROPERTY(Transient)
+	TMap<FName, FMovieGraphEvaluatedBranchConfig> BranchConfigMapping;
+};
+
+/**
+* This stores short-term information needed during traversal of the graph
+* such as disabled nodes, already visited nodes, etc. This information is
+* discarded after traversal.
+*/
+USTRUCT()
+struct FMovieGraphEvaluationContext
+{
+	GENERATED_BODY()
+public:
+	
+	/** 
+	* This is the user provided traversal context which specifies high level user decisions. This is the calling
+	* context such as what frame you're on, or what the shot name is, stuff generally driven by global variables.
+	*/
+	UPROPERTY()
+	FMovieGraphTraversalContext UserContext;
+
+	/**
+	* A list of nodes that have been visited. Used for cycle detection right now.
+	*/
+	UPROPERTY()
+	TSet<TObjectPtr<UMovieGraphNode>> VisitedNodes;
+};
+
+
 /**
 * This is the runtime representation of the UMoviePipelineEdGraph which contains the actual strongly
 * typed graph network that is read by the MoviePipeline. There is an editor-only representation of
@@ -227,6 +307,8 @@ public:
 	void SetEditorOnlyNodes(const TArray<TObjectPtr<const UObject>>& InNodes);
 #endif
 
+	UMovieGraphEvaluatedConfig* CreateFlattenedGraph(const FMovieGraphTraversalContext& InContext);
+
 	template<typename NodeType>
 	static NodeType* IterateGraphForClass(const FMovieGraphTraversalContext& InContext)
 	{
@@ -260,6 +342,15 @@ public:
 
 protected:
 	void TraverseGraphRecursive(UMovieGraphNode* InNode, TSubclassOf<UMovieGraphNode> InClassType, const FMovieGraphTraversalContext& InContext, TArray<UMovieGraphNode*>& OutNodes) const;
+	/** Copies properties in FromNode that are marked for override into ToNode, but only if ToNode doesn't already override that value. */
+	void CopyOverriddenProperties(UMovieGraphNode* FromNode, UMovieGraphNode* ToNode);
+	/** Given a property you want to override (passes IsPropertyOverrideable) look for a matching bOverride_ named bool property. */
+	FBoolProperty* FindOverridePropertyForRealProperty(UClass* InClass, const FProperty* InRealProperty) const;
+	/** Find all "Overrideable" marked properties, then find their edit condition properties, then set those to false. */
+	void InitializeFlattenedNode(UMovieGraphNode* InNode);
+
+	/** Traverse the graph, generating a combined "flatten" graph as it goes. */
+	void CreateFlattenedGraph_Recursive(UMovieGraphEvaluatedConfig* InOwningConfig, FMovieGraphEvaluatedBranchConfig& OutBranchConfig, FMovieGraphEvaluationContext& InEvaluationContext, UMovieGraphPin* InPinToFollow);
 
 public:
 	// Names of global variables that are provided by the graph
diff --git a/Engine/Plugins/MovieScene/MovieRenderPipeline/Source/MovieRenderPipelineCore/Public/Graph/MovieGraphNode.h b/Engine/Plugins/MovieScene/MovieRenderPipeline/Source/MovieRenderPipelineCore/Public/Graph/MovieGraphNode.h
--- a/Engine/Plugins/MovieScene/MovieRenderPipeline/Source/MovieRenderPipelineCore/Public/Graph/MovieGraphNode.h
+++ b/Engine/Plugins/MovieScene/MovieRenderPipeline/Source/MovieRenderPipelineCore/Public/Graph/MovieGraphNode.h
@@ -6,6 +6,7 @@
 #include "MovieGraphPin.h"
 #include "InstancedStruct.h"
 #include "PropertyBag.h"
+#include "Graph/MovieGraphValueContainer.h"
 
 #if WITH_EDITOR
 #include "Textures/SlateIcon.h"
@@ -20,6 +21,7 @@ class UMovieGraphMember;
 class UMovieGraphOutput;
 class UMovieGraphPin;
 class UMovieGraphVariable;
+struct FMovieGraphTraversalContext;
 
 #if WITH_EDITOR
 class UEdGraphNode;
@@ -37,6 +39,8 @@ class MOVIERENDERPIPELINECORE_API UMovieGraphNode : public UObject
 {
 	GENERATED_BODY()
 
+	static FName GlobalsPinName;
+
 	friend class UMovieGraphConfig;
 	friend class UMovieGraphEdge;
 	
@@ -65,6 +69,25 @@ public:
 		return ExposedDynamicPropertyNames;
 	}
 
+	/** 
+	* Used to determine which Branch type pins we should follow when trying to traverse the graph.
+	* By default we will follow any input pin (with Branch type) on the node, but override this in
+	* inherited classes and change that if you need custom logic, such as boolean nodes that want 
+	* to choose one or the other based on the results of a conditional property.
+	*/
+	virtual TArray<UMovieGraphPin*> EvaluatePinsToFollow(const FMovieGraphTraversalContext& InContext) const;
+
+	/**
+	* When a non-branch pin type is being evaluated on a node, the calling node will ask this node
+	* for the value connected to the given pin name. For example, a Branch node will call this
+	* function on whatever node is connected to the Conditional pin, and then will try to get a 
+	* Boolean value out of the returned UMovieGraphValueContainer.
+	*/
+	virtual UMovieGraphValueContainer* GetPropertyValueContainerForPin(const FString& InPinName) const
+	{
+		return nullptr;
+	}
+
 	/** Promotes the property with the given name to a pin on the node via a dynamic property. */
 	virtual void PromoteDynamicPropertyToPin(const FName& PropertyName);
 
@@ -165,4 +188,18 @@ protected:
 	/** A GUID which uniquely identifies this node. */
 	UPROPERTY()
 	FGuid Guid;
+};
+
+/**
+* Nodes representing user settings should derive from this. This is the only node type copied into flattened eval.
+*/
+UCLASS(Abstract)
+class MOVIERENDERPIPELINECORE_API UMovieGraphSettingNode : public UMovieGraphNode
+{
+	GENERATED_BODY()
+public:
+	// UMovieGraphNode Interface
+	virtual TArray<FMovieGraphPinProperties> GetInputPinProperties() const override;
+	virtual TArray<FMovieGraphPinProperties> GetOutputPinProperties() const override;
+	// ~UMovieGraphNode Interface
 };
\ No newline at end of file
diff --git a/Engine/Plugins/MovieScene/MovieRenderPipeline/Source/MovieRenderPipelineCore/Public/Graph/MovieGraphPlaceholderNodes.cpp b/Engine/Plugins/MovieScene/MovieRenderPipeline/Source/MovieRenderPipelineCore/Public/Graph/MovieGraphPlaceholderNodes.cpp
--- a/Engine/Plugins/MovieScene/MovieRenderPipeline/Source/MovieRenderPipelineCore/Public/Graph/MovieGraphPlaceholderNodes.cpp
+++ b/Engine/Plugins/MovieScene/MovieRenderPipeline/Source/MovieRenderPipelineCore/Public/Graph/MovieGraphPlaceholderNodes.cpp
@@ -254,53 +254,6 @@ FSlateIcon UMovieGraphEXRSequenceNode::GetIconAndTint(FLinearColor& OutColor) co
 }
 #endif // WITH_EDITOR
 
-TArray<FMovieGraphPinProperties> UMovieGraphBranchNode::GetInputPinProperties() const
-{
-	static const FName TrueBranch("True");
-	static const FName FalseBranch("False");
-	static const FName Condition("Condition");
-	
-	TArray<FMovieGraphPinProperties> Properties;
-	Properties.Add(FMovieGraphPinProperties(TrueBranch, EMovieGraphValueType::Branch, false));
-	Properties.Add(FMovieGraphPinProperties(FalseBranch, EMovieGraphValueType::Branch, false));
-	Properties.Add(FMovieGraphPinProperties(Condition, EMovieGraphValueType::Bool, false));
-	return Properties;
-}
-
-TArray<FMovieGraphPinProperties> UMovieGraphBranchNode::GetOutputPinProperties() const
-{
-	TArray<FMovieGraphPinProperties> Properties;
-	Properties.Add(FMovieGraphPinProperties(NAME_None, EMovieGraphValueType::Branch, false));
-	return Properties;
-}
-
-#if WITH_EDITOR
-FText UMovieGraphBranchNode::GetNodeTitle(const bool bGetDescriptive) const
-{
-	static const FText BranchNodeName = LOCTEXT("NodeName_Branch", "Branch");
-	return BranchNodeName;
-}
-
-FText UMovieGraphBranchNode::GetMenuCategory() const
-{
-	return NodeCategory_Conditionals;
-}
-
-FLinearColor UMovieGraphBranchNode::GetNodeTitleColor() const
-{
-	static const FLinearColor BranchNodeColor = FLinearColor(0.266f, 0.266f, 0.266f);
-	return BranchNodeColor;
-}
-
-FSlateIcon UMovieGraphBranchNode::GetIconAndTint(FLinearColor& OutColor) const
-{
-	static const FSlateIcon BranchIcon = FSlateIcon(FAppStyle::GetAppStyleSetName(), "Icons.Merge");
-
-	OutColor = FLinearColor::White;
-	return BranchIcon;
-}
-#endif // WITH_EDITOR
-
 TArray<FMovieGraphPinProperties> UMovieGraphSelectNode::GetInputPinProperties() const
 {
 	TArray<FMovieGraphPinProperties> Properties;
@@ -372,47 +325,6 @@ void UMovieGraphSelectNode::PostEditChangeProperty(FPropertyChangedEvent& Proper
 }
 #endif // WITH_EDITOR
 
-TArray<FMovieGraphPinProperties> UMovieGraphOutputSettingsNode::GetInputPinProperties() const
-{
-	TArray<FMovieGraphPinProperties> Properties;
-	Properties.Add(FMovieGraphPinProperties(NAME_None, EMovieGraphValueType::Branch, false));
-	return Properties;
-}
-
-TArray<FMovieGraphPinProperties> UMovieGraphOutputSettingsNode::GetOutputPinProperties() const
-{
-	TArray<FMovieGraphPinProperties> Properties;
-	Properties.Add(FMovieGraphPinProperties(NAME_None, EMovieGraphValueType::Branch, false));
-	return Properties;
-}
-
-#if WITH_EDITOR
-FText UMovieGraphOutputSettingsNode::GetNodeTitle(const bool bGetDescriptive) const
-{
-	static const FText OutputSettingsNodeName = LOCTEXT("NodeName_OutputSettings", "Output Settings");
-	return OutputSettingsNodeName;
-}
-
-FText UMovieGraphOutputSettingsNode::GetMenuCategory() const
-{
-	return NodeCategory_Settings;
-}
-
-FLinearColor UMovieGraphOutputSettingsNode::GetNodeTitleColor() const
-{
-	static const FLinearColor OutputSettingsColor = FLinearColor(0.854f, 0.509f, 0.039f);
-	return OutputSettingsColor;
-}
-
-FSlateIcon UMovieGraphOutputSettingsNode::GetIconAndTint(FLinearColor& OutColor) const
-{
-	static const FSlateIcon SettingsIcon = FSlateIcon(FAppStyle::GetAppStyleSetName(), "Icons.Settings");
-
-	OutColor = FLinearColor::White;
-	return SettingsIcon;
-}
-#endif // WITH_EDITOR
-
 TArray<FMovieGraphPinProperties> UMovieGraphAntiAliasingNode::GetInputPinProperties() const
 {
 	TArray<FMovieGraphPinProperties> Properties;
diff --git a/Engine/Plugins/MovieScene/MovieRenderPipeline/Source/MovieRenderPipelineCore/Public/Graph/MovieGraphPlaceholderNodes.h b/Engine/Plugins/MovieScene/MovieRenderPipeline/Source/MovieRenderPipelineCore/Public/Graph/MovieGraphPlaceholderNodes.h
--- a/Engine/Plugins/MovieScene/MovieRenderPipeline/Source/MovieRenderPipelineCore/Public/Graph/MovieGraphPlaceholderNodes.h
+++ b/Engine/Plugins/MovieScene/MovieRenderPipeline/Source/MovieRenderPipelineCore/Public/Graph/MovieGraphPlaceholderNodes.h
@@ -143,26 +143,6 @@ public:
 #endif
 };
 
-/** A node which creates a branching condition. */
-UCLASS()
-class MOVIERENDERPIPELINECORE_API UMovieGraphBranchNode : public UMovieGraphNode
-{
-	GENERATED_BODY()
-
-public:
-	UMovieGraphBranchNode() = default;
-
-	virtual TArray<FMovieGraphPinProperties> GetInputPinProperties() const override;
-	virtual TArray<FMovieGraphPinProperties> GetOutputPinProperties() const override;
-
-#if WITH_EDITOR
-	virtual FText GetNodeTitle(const bool bGetDescriptive = false) const override;
-	virtual FText GetMenuCategory() const override;
-	virtual FLinearColor GetNodeTitleColor() const override;
-	virtual FSlateIcon GetIconAndTint(FLinearColor& OutColor) const override;
-#endif
-};
-
 // TODO: Currently this node is restricted to just accepting string-based options.
 /** A node which creates a condition that selects from a set of input branches. */
 UCLASS()
@@ -200,26 +180,6 @@ public:
 	FString Description;
 };
 
-/** A node which configures output settings. */
-UCLASS()
-class MOVIERENDERPIPELINECORE_API UMovieGraphOutputSettingsNode : public UMovieGraphNode
-{
-	GENERATED_BODY()
-
-public:
-	UMovieGraphOutputSettingsNode() = default;
-
-	virtual TArray<FMovieGraphPinProperties> GetInputPinProperties() const override;
-	virtual TArray<FMovieGraphPinProperties> GetOutputPinProperties() const override;
-
-#if WITH_EDITOR
-	virtual FText GetNodeTitle(const bool bGetDescriptive = false) const override;
-	virtual FText GetMenuCategory() const override;
-	virtual FLinearColor GetNodeTitleColor() const override;
-	virtual FSlateIcon GetIconAndTint(FLinearColor& OutColor) const override;
-#endif
-};
-
 /** A node which configures anti-aliasing settings. */
 UCLASS()
 class MOVIERENDERPIPELINECORE_API UMovieGraphAntiAliasingNode : public UMovieGraphNode
diff --git a/Engine/Plugins/MovieScene/MovieRenderPipeline/Source/MovieRenderPipelineCore/Public/Graph/Nodes/MovieGraphBranchNode.h b/Engine/Plugins/MovieScene/MovieRenderPipeline/Source/MovieRenderPipelineCore/Public/Graph/Nodes/MovieGraphBranchNode.h
new file mode 100644
--- /dev/null
+++ b/Engine/Plugins/MovieScene/MovieRenderPipeline/Source/MovieRenderPipelineCore/Public/Graph/Nodes/MovieGraphBranchNode.h
@@ -0,0 +1,32 @@
+// Copyright Epic Games, Inc. All Rights Reserved.
+
+#pragma once
+
+#include "Graph/MovieGraphNode.h"
+
+#include "MovieGraphBranchNode.generated.h"
+
+/** 
+* A node which creates a True/False branching condition. A user Graph Variable can be plugged
+* into the conditional pin and this will be evaluated when flattening the graph, choosing which
+* branch path to follow.
+*/
+UCLASS()
+class MOVIERENDERPIPELINECORE_API UMovieGraphBranchNode : public UMovieGraphNode
+{
+	GENERATED_BODY()
+
+public:
+	UMovieGraphBranchNode() = default;
+
+	virtual TArray<FMovieGraphPinProperties> GetInputPinProperties() const override;
+	virtual TArray<FMovieGraphPinProperties> GetOutputPinProperties() const override;
+	TArray<UMovieGraphPin*> EvaluatePinsToFollow(const FMovieGraphTraversalContext& InContext) const override;
+
+#if WITH_EDITOR
+	virtual FText GetNodeTitle(const bool bGetDescriptive = false) const override;
+	virtual FText GetMenuCategory() const override;
+	virtual FLinearColor GetNodeTitleColor() const override;
+	virtual FSlateIcon GetIconAndTint(FLinearColor& OutColor) const override;
+#endif
+};
\ No newline at end of file
