MoviePipeline: Added a new subgraph node which can reference existing graph assets. Also integrated this node into the traversal process so the nodes contained within the subgraph are followed.

Subgraphs can also expose non-branch input/output pins. Because the "Branch" type cannot be selected in the property bags type selector (it explicitly forbids choosing the "Exec" type), the "Branch" type was moved from the EMovieGraphValueType enum to being a bool in FMovieGraphPinProperties. This may not be an ideal long-term solution, but works for now.

diff --git a/Engine/Plugins/MovieScene/MovieRenderPipeline/Source/MovieRenderPipelineCore/Private/Graph/MovieGraphConfig.cpp b/Engine/Plugins/MovieScene/MovieRenderPipeline/Source/MovieRenderPipelineCore/Private/Graph/MovieGraphConfig.cpp
--- a/Engine/Plugins/MovieScene/MovieRenderPipeline/Source/MovieRenderPipelineCore/Private/Graph/MovieGraphConfig.cpp
+++ b/Engine/Plugins/MovieScene/MovieRenderPipeline/Source/MovieRenderPipelineCore/Private/Graph/MovieGraphConfig.cpp
@@ -795,11 +795,13 @@ void UMovieGraphConfig::CreateFlattenedGraph_Recursive(UMovieGraphEvaluatedConfi
 	}
 
 	// We only follow execution pins during traversal.
-	if (!ensureMsgf(InPinToFollow->Properties.Type == EMovieGraphValueType::Branch, TEXT("Only Branch pins should be contained by InPinToFollow!")))
+	if (!ensureMsgf(InPinToFollow->Properties.bIsBranch, TEXT("Only Branch pins should be contained by InPinToFollow!")))
 	{
 		return;
 	}
 
+	InEvaluationContext.PinBeingFollowed = InPinToFollow;
+
 	// Check to see if our flattened evaluation graph already has a copy of this node.
 	InEvaluationContext.VisitedNodes.Add(Node);
 	const bool bShouldIncludeNode = Node->IsA<UMovieGraphSettingNode>();
@@ -829,7 +831,7 @@ void UMovieGraphConfig::CreateFlattenedGraph_Recursive(UMovieGraphEvaluatedConfi
 	
 	// Now that we've potentially resolved the values on this node, continue to travel up-stream along any execution pins,
 	// potentially following re-route nodes, sub-graph nodes, through branches, etc.
-	TArray<UMovieGraphPin*> NewPinsToFollow = Node->EvaluatePinsToFollow(InEvaluationContext.UserContext);
+	TArray<UMovieGraphPin*> NewPinsToFollow = Node->EvaluatePinsToFollow(InEvaluationContext);
 	
 	for(UMovieGraphPin* Pin : NewPinsToFollow)
 	{
diff --git a/Engine/Plugins/MovieScene/MovieRenderPipeline/Source/MovieRenderPipelineCore/Private/Graph/MovieGraphNode.cpp b/Engine/Plugins/MovieScene/MovieRenderPipeline/Source/MovieRenderPipelineCore/Private/Graph/MovieGraphNode.cpp
--- a/Engine/Plugins/MovieScene/MovieRenderPipeline/Source/MovieRenderPipelineCore/Private/Graph/MovieGraphNode.cpp
+++ b/Engine/Plugins/MovieScene/MovieRenderPipeline/Source/MovieRenderPipelineCore/Private/Graph/MovieGraphNode.cpp
@@ -208,7 +208,7 @@ void UMovieGraphNode::PostLoad()
 	RegisterDelegates();
 }
 
-TArray<UMovieGraphPin*> UMovieGraphNode::EvaluatePinsToFollow(const FMovieGraphTraversalContext& InContext) const
+TArray<UMovieGraphPin*> UMovieGraphNode::EvaluatePinsToFollow(FMovieGraphEvaluationContext& InContext) const
 {
 	TArray<UMovieGraphPin*> PinsToFollow;
 
@@ -217,7 +217,7 @@ TArray<UMovieGraphPin*> UMovieGraphNode::EvaluatePinsToFollow(const FMovieGraphT
 	// as branch or switch nodes.
 	for (const TObjectPtr<UMovieGraphPin>& InputPin : GetInputPins())
 	{
-		if (InputPin->Properties.Type == EMovieGraphValueType::Branch)
+		if (InputPin->Properties.bIsBranch)
 		{
 			PinsToFollow.Add(InputPin);
 		}
@@ -228,13 +228,13 @@ TArray<UMovieGraphPin*> UMovieGraphNode::EvaluatePinsToFollow(const FMovieGraphT
 TArray<FMovieGraphPinProperties> UMovieGraphSettingNode::GetInputPinProperties() const
 {
 	TArray<FMovieGraphPinProperties> Properties;
-	Properties.Add(FMovieGraphPinProperties(NAME_None, EMovieGraphValueType::Branch, false));
+	Properties.Add(FMovieGraphPinProperties::MakeBranchProperties());
 	return Properties;
 }
 
 TArray<FMovieGraphPinProperties> UMovieGraphSettingNode::GetOutputPinProperties() const
 {
 	TArray<FMovieGraphPinProperties> Properties;
-	Properties.Add(FMovieGraphPinProperties(NAME_None, EMovieGraphValueType::Branch, false));
+	Properties.Add(FMovieGraphPinProperties::MakeBranchProperties());
 	return Properties;
 }
\ No newline at end of file
diff --git a/Engine/Plugins/MovieScene/MovieRenderPipeline/Source/MovieRenderPipelineCore/Private/Graph/Nodes/MovieGraphBranchNode.cpp b/Engine/Plugins/MovieScene/MovieRenderPipeline/Source/MovieRenderPipelineCore/Private/Graph/Nodes/MovieGraphBranchNode.cpp
--- a/Engine/Plugins/MovieScene/MovieRenderPipeline/Source/MovieRenderPipelineCore/Private/Graph/Nodes/MovieGraphBranchNode.cpp
+++ b/Engine/Plugins/MovieScene/MovieRenderPipeline/Source/MovieRenderPipelineCore/Private/Graph/Nodes/MovieGraphBranchNode.cpp
@@ -17,7 +17,7 @@ namespace UE::MovieGraph::BranchNode
 	static const FName Condition("Condition");
 }
 
-TArray<UMovieGraphPin*> UMovieGraphBranchNode::EvaluatePinsToFollow(const FMovieGraphTraversalContext& InContext) const
+TArray<UMovieGraphPin*> UMovieGraphBranchNode::EvaluatePinsToFollow(FMovieGraphEvaluationContext& InContext) const
 {
 	TArray<UMovieGraphPin*> PinsToFollow;
 
@@ -87,8 +87,8 @@ TArray<FMovieGraphPinProperties> UMovieGraphBranchNode::GetInputPinProperties()
 {
 	TArray<FMovieGraphPinProperties> Properties;
 	
-	Properties.Add(FMovieGraphPinProperties(UE::MovieGraph::BranchNode::TrueBranch, EMovieGraphValueType::Branch, false));
-	Properties.Add(FMovieGraphPinProperties(UE::MovieGraph::BranchNode::FalseBranch, EMovieGraphValueType::Branch, false));
+	Properties.Add(FMovieGraphPinProperties::MakeBranchProperties(UE::MovieGraph::BranchNode::TrueBranch));
+	Properties.Add(FMovieGraphPinProperties::MakeBranchProperties(UE::MovieGraph::BranchNode::FalseBranch));
 	Properties.Add(FMovieGraphPinProperties(UE::MovieGraph::BranchNode::Condition, EMovieGraphValueType::Bool, false));
 	return Properties;
 }
@@ -96,7 +96,7 @@ TArray<FMovieGraphPinProperties> UMovieGraphBranchNode::GetInputPinProperties()
 TArray<FMovieGraphPinProperties> UMovieGraphBranchNode::GetOutputPinProperties() const
 {
 	TArray<FMovieGraphPinProperties> Properties;
-	Properties.Add(FMovieGraphPinProperties(NAME_None, EMovieGraphValueType::Branch, false));
+	Properties.Add(FMovieGraphPinProperties::MakeBranchProperties());
 	return Properties;
 }
 
diff --git a/Engine/Plugins/MovieScene/MovieRenderPipeline/Source/MovieRenderPipelineCore/Private/Graph/Nodes/MovieGraphInputNode.cpp b/Engine/Plugins/MovieScene/MovieRenderPipeline/Source/MovieRenderPipelineCore/Private/Graph/Nodes/MovieGraphInputNode.cpp
--- a/Engine/Plugins/MovieScene/MovieRenderPipeline/Source/MovieRenderPipelineCore/Private/Graph/Nodes/MovieGraphInputNode.cpp
+++ b/Engine/Plugins/MovieScene/MovieRenderPipeline/Source/MovieRenderPipelineCore/Private/Graph/Nodes/MovieGraphInputNode.cpp
@@ -3,6 +3,7 @@
 #include "Graph/Nodes/MovieGraphInputNode.h"
 
 #include "Graph/MovieGraphConfig.h"
+#include "Graph/Nodes/MovieGraphSubgraphNode.h"
 #include "Styling/AppStyle.h"
 
 UMovieGraphInputNode::UMovieGraphInputNode()
@@ -27,13 +28,58 @@ TArray<FMovieGraphPinProperties> UMovieGraphInputNode::GetOutputPinProperties()
 	{
 		for (const UMovieGraphInput* Input : ParentGraph->GetInputs())
 		{
-			Properties.Add(FMovieGraphPinProperties(FName(Input->Name), EMovieGraphValueType::Branch, false));
+			FMovieGraphPinProperties PinProperties(FName(Input->Name), Input->GetValueType(), false);
+			PinProperties.bIsBranch = Input->bIsBranch;
+			Properties.Add(MoveTemp(PinProperties));
 		}
 	}
 	
 	return Properties;
 }
 
+TArray<UMovieGraphPin*> UMovieGraphInputNode::EvaluatePinsToFollow(FMovieGraphEvaluationContext& InContext) const
+{
+	TArray<UMovieGraphPin*> PinsToFollow;
+	
+	if (!ensure(InContext.PinBeingFollowed))
+	{
+		return PinsToFollow;
+	}
+
+	// If this input node occurs in a subgraph (ie, the subgraph stack is not empty), follow the pin to the input pin on
+	// the subgraph node in the parent graph.
+	//
+	//    Subgraph
+	//    +------------------
+	//    | Inputs node
+	//    | +--------+ 
+	// In1| |    In1 | <--- Pin being followed ---
+	// In2| |    In2 | 
+	// In3| |    In3 | 
+	//    | +--------+   
+	//    +------------------
+
+	if (InContext.SubgraphStack.IsEmpty())
+	{
+		// This was not an input node in a subgraph; nothing to continue following
+		return PinsToFollow;
+	}
+
+	// Pop the subgraph node off the stack and continue traversal in the parent graph
+	const TObjectPtr<const UMovieGraphSubgraphNode> ParentSubgraphNode = InContext.SubgraphStack.Pop();
+	if (!ensure(ParentSubgraphNode.Get()))
+	{
+		return PinsToFollow;
+	}
+
+	if (UMovieGraphPin* InputPin = ParentSubgraphNode->GetInputPin(InContext.PinBeingFollowed->Properties.Label))
+	{
+		PinsToFollow.Add(InputPin);
+	}
+	
+	return PinsToFollow;
+}
+
 #if WITH_EDITOR
 FText UMovieGraphInputNode::GetNodeTitle(const bool bGetDescriptive) const
 {
@@ -92,6 +138,8 @@ void UMovieGraphInputNode::UpdateExistingPins(UMovieGraphMember* ChangedInput) c
 			for (int32 Index = 0; Index < InputMembers.Num(); ++Index)
 			{
 				OutputPins[Index]->Properties.Label = FName(InputMembers[Index]->Name);
+				OutputPins[Index]->Properties.Type = InputMembers[Index]->GetValueType();
+				OutputPins[Index]->Properties.bIsBranch = InputMembers[Index]->bIsBranch;
 			}
 		}
 		
diff --git a/Engine/Plugins/MovieScene/MovieRenderPipeline/Source/MovieRenderPipelineCore/Private/Graph/Nodes/MovieGraphOutputNode.cpp b/Engine/Plugins/MovieScene/MovieRenderPipeline/Source/MovieRenderPipelineCore/Private/Graph/Nodes/MovieGraphOutputNode.cpp
--- a/Engine/Plugins/MovieScene/MovieRenderPipeline/Source/MovieRenderPipelineCore/Private/Graph/Nodes/MovieGraphOutputNode.cpp
+++ b/Engine/Plugins/MovieScene/MovieRenderPipeline/Source/MovieRenderPipelineCore/Private/Graph/Nodes/MovieGraphOutputNode.cpp
@@ -27,7 +27,9 @@ TArray<FMovieGraphPinProperties> UMovieGraphOutputNode::GetInputPinProperties()
 	{
 		for (const UMovieGraphOutput* Output : ParentGraph->GetOutputs())
 		{
-			Properties.Add(FMovieGraphPinProperties(FName(Output->Name), EMovieGraphValueType::Branch, false));
+			FMovieGraphPinProperties PinProperties(FName(Output->Name), Output->GetValueType(), false);
+			PinProperties.bIsBranch = Output->bIsBranch;
+			Properties.Add(MoveTemp(PinProperties));
 		}
 	}
 	
@@ -92,6 +94,8 @@ void UMovieGraphOutputNode::UpdateExistingPins(UMovieGraphMember* ChangedOutput)
 			for (int32 Index = 0; Index < OutputMembers.Num(); ++Index)
 			{
 				InputPins[Index]->Properties.Label = FName(OutputMembers[Index]->Name);
+				InputPins[Index]->Properties.Type = OutputMembers[Index]->GetValueType();
+				InputPins[Index]->Properties.bIsBranch = OutputMembers[Index]->bIsBranch;
 			}
 		}
 
diff --git a/Engine/Plugins/MovieScene/MovieRenderPipeline/Source/MovieRenderPipelineCore/Private/Graph/Nodes/MovieGraphRenderLayerNode.cpp b/Engine/Plugins/MovieScene/MovieRenderPipeline/Source/MovieRenderPipelineCore/Private/Graph/Nodes/MovieGraphRenderLayerNode.cpp
--- a/Engine/Plugins/MovieScene/MovieRenderPipeline/Source/MovieRenderPipelineCore/Private/Graph/Nodes/MovieGraphRenderLayerNode.cpp
+++ b/Engine/Plugins/MovieScene/MovieRenderPipeline/Source/MovieRenderPipelineCore/Private/Graph/Nodes/MovieGraphRenderLayerNode.cpp
@@ -13,14 +13,14 @@ UMovieGraphRenderLayerNode::UMovieGraphRenderLayerNode()
 TArray<FMovieGraphPinProperties> UMovieGraphRenderLayerNode::GetInputPinProperties() const
 {
 	TArray<FMovieGraphPinProperties> Properties;
-	Properties.Add(FMovieGraphPinProperties(NAME_None, EMovieGraphValueType::Branch, false));
+	Properties.Add(FMovieGraphPinProperties::MakeBranchProperties());
 	return Properties;
 }
 
 TArray<FMovieGraphPinProperties> UMovieGraphRenderLayerNode::GetOutputPinProperties() const
 {
 	TArray<FMovieGraphPinProperties> Properties;
-	Properties.Add(FMovieGraphPinProperties(NAME_None, EMovieGraphValueType::Branch, false));
+	Properties.Add(FMovieGraphPinProperties::MakeBranchProperties());
 	return Properties;
 }
 
diff --git a/Engine/Plugins/MovieScene/MovieRenderPipeline/Source/MovieRenderPipelineCore/Private/Graph/Nodes/MovieGraphSubgraphNode.cpp b/Engine/Plugins/MovieScene/MovieRenderPipeline/Source/MovieRenderPipelineCore/Private/Graph/Nodes/MovieGraphSubgraphNode.cpp
new file mode 100644
--- /dev/null
+++ b/Engine/Plugins/MovieScene/MovieRenderPipeline/Source/MovieRenderPipelineCore/Private/Graph/Nodes/MovieGraphSubgraphNode.cpp
@@ -0,0 +1,181 @@
+// Copyright Epic Games, Inc. All Rights Reserved.
+
+#include "Graph/Nodes/MovieGraphSubgraphNode.h"
+
+#include "UObject/ObjectSaveContext.h"
+#include "UObject/Package.h"
+
+#define LOCTEXT_NAMESPACE "MovieGraphNode"
+
+namespace UE::MovieGraph::Private
+{
+	/** Converts graph input/output members to pin properties. */
+	TArray<FMovieGraphPinProperties> MembersToPinProperties(const TArray<UMovieGraphInterfaceBase*>& InMembers)
+	{
+		TArray<FMovieGraphPinProperties> PinProperties;
+
+		Algo::TransformIf(InMembers, PinProperties,
+			[](const UMovieGraphInterfaceBase* Member) -> bool
+			{
+				return Member != nullptr;
+			},
+			[](const UMovieGraphInterfaceBase* Member) -> FMovieGraphPinProperties
+			{
+				constexpr bool bAllowMultipleConnections = false;
+				FMovieGraphPinProperties Properties(FName(Member->Name), Member->GetValueType(), bAllowMultipleConnections);
+				Properties.bIsBranch = Member->bIsBranch;
+				return Properties;
+			});
+
+		return PinProperties;
+	}
+}
+
+UMovieGraphSubgraphNode::UMovieGraphSubgraphNode()
+{
+	if (!HasAnyFlags(RF_ClassDefaultObject))
+	{
+		RefreshOnSubgraphAssetSaved();
+	}
+}
+
+TArray<FMovieGraphPinProperties> UMovieGraphSubgraphNode::GetInputPinProperties() const
+{
+	if (const UMovieGraphConfig* Config = GetSubgraphAsset())
+	{
+		TArray<UMovieGraphInterfaceBase*> InputMembers;
+		Algo::Transform(Config->GetInputs(), InputMembers, [](UMovieGraphInput* Input) { return Input; });
+		
+		return UE::MovieGraph::Private::MembersToPinProperties(InputMembers);
+	}
+
+	return {};
+}
+
+TArray<FMovieGraphPinProperties> UMovieGraphSubgraphNode::GetOutputPinProperties() const
+{
+	if (const UMovieGraphConfig* Config = GetSubgraphAsset())
+	{
+		TArray<UMovieGraphInterfaceBase*> OutputMembers;
+		Algo::Transform(Config->GetOutputs(), OutputMembers, [](UMovieGraphOutput* Output) { return Output; });
+		
+		return UE::MovieGraph::Private::MembersToPinProperties(OutputMembers);
+	}
+
+	return {};
+}
+
+TArray<UMovieGraphPin*> UMovieGraphSubgraphNode::EvaluatePinsToFollow(FMovieGraphEvaluationContext& InContext) const
+{
+	TArray<UMovieGraphPin*> PinsToFollow;
+	
+	if (!ensure(InContext.PinBeingFollowed))
+	{
+		return PinsToFollow;
+	}
+
+	const UMovieGraphConfig* SubgraphPtr = SubgraphAsset.LoadSynchronous();
+	if (!SubgraphPtr)
+	{
+		return PinsToFollow;
+	}
+
+	// Find the input pin on the subgraph asset's outputs node that corresponds to the pin being followed
+	//
+	// Subgraph
+	// -----------------+
+	//     Outputs node |
+	//     +----------+ |
+	// --->| Out1     | |Out1 <--- Pin being followed ---
+	//     | Out2     | |Out2
+	//     | Out3     | |Out3
+	//     +----------+ |
+	// -----------------+
+	
+	const UMovieGraphNode* OutputNode = SubgraphPtr->GetOutputNode();
+	if (!ensure(OutputNode))
+	{
+		return PinsToFollow;
+	}
+
+	if (UMovieGraphPin* OutputNodePin = OutputNode->GetInputPin(InContext.PinBeingFollowed->Properties.Label))
+	{
+		PinsToFollow.Add(OutputNodePin);
+	}
+
+	// Update the subgraph traversal stack so traversal can resume in parent graphs when traversal reaches the input
+	// node of this subgraph
+	InContext.SubgraphStack.Add(this);
+	
+	return PinsToFollow;
+}
+
+#if WITH_EDITOR
+void UMovieGraphSubgraphNode::PostEditChangeProperty(FPropertyChangedEvent& PropertyChangedEvent)
+{
+	Super::PostEditChangeProperty(PropertyChangedEvent);
+
+	// Update pins when the subgraph asset is changed
+	if (PropertyChangedEvent.GetPropertyName() == GET_MEMBER_NAME_CHECKED(UMovieGraphSubgraphNode, SubgraphAsset))
+	{
+		UpdatePins();
+	}
+}
+#endif
+
+#if WITH_EDITOR
+FText UMovieGraphSubgraphNode::GetNodeTitle(const bool bGetDescriptive) const
+{
+	// Note: Only the descriptive title indicates invalidity because we don't want the node creation menu to say the
+	// subgraph node is invalid
+	static const FText SubgraphTitle = LOCTEXT("SubgraphTitle", "Subgraph");
+	static const FText DescriptiveSubgraphTitle = LOCTEXT("DescriptiveSubgraphTitle", "Invalid Subgraph\nNo subgraph asset specified");
+
+	if (const UMovieGraphConfig* Subgraph = GetSubgraphAsset())
+	{
+		return FText::Format(LOCTEXT("SubgraphTitle", "Subgraph: {0}"), FText::FromString(Subgraph->GetName()));
+	}
+
+	return bGetDescriptive ? DescriptiveSubgraphTitle : SubgraphTitle;
+}
+
+FText UMovieGraphSubgraphNode::GetMenuCategory() const
+{
+	static const FText NodeCategory = LOCTEXT("NodeCategory_Utility", "Utility");
+
+	return NodeCategory;
+}
+#endif
+
+void UMovieGraphSubgraphNode::SetSubGraphAsset(const TSoftObjectPtr<UMovieGraphConfig>& InSubgraphAsset)
+{
+	SubgraphAsset = InSubgraphAsset;
+}
+
+UMovieGraphConfig* UMovieGraphSubgraphNode::GetSubgraphAsset() const
+{
+	return SubgraphAsset.LoadSynchronous();
+}
+
+void UMovieGraphSubgraphNode::RefreshOnSubgraphAssetSaved()
+{
+	UPackage::PackageSavedWithContextEvent.AddWeakLambda(this, [this](const FString& PackageFilename, const UPackage* Package, FObjectPostSaveContext Context)
+	{
+		// Ignore procedural saves (ie, saves not triggered by a user)
+		if (Context.IsProceduralSave())
+		{
+			return;
+		}
+
+		// Refresh the node if the subgraph asset was saved
+		if (const UMovieGraphConfig* Subgraph = GetSubgraphAsset())
+		{
+			if (Subgraph->GetPackage() == Package)
+			{
+				UpdatePins();
+			}
+		}
+	});
+}
+
+#undef LOCTEXT_NAMESPACE
diff --git a/Engine/Plugins/MovieScene/MovieRenderPipeline/Source/MovieRenderPipelineCore/Private/Graph/Nodes/MovieGraphVariableNode.cpp b/Engine/Plugins/MovieScene/MovieRenderPipeline/Source/MovieRenderPipelineCore/Private/Graph/Nodes/MovieGraphVariableNode.cpp
--- a/Engine/Plugins/MovieScene/MovieRenderPipeline/Source/MovieRenderPipelineCore/Private/Graph/Nodes/MovieGraphVariableNode.cpp
+++ b/Engine/Plugins/MovieScene/MovieRenderPipeline/Source/MovieRenderPipelineCore/Private/Graph/Nodes/MovieGraphVariableNode.cpp
@@ -23,7 +23,7 @@ TArray<FMovieGraphPinProperties> UMovieGraphVariableNode::GetOutputPinProperties
 	}
 	else
 	{
-		Properties.Add(FMovieGraphPinProperties(TEXT("Unknown"), EMovieGraphValueType::Unknown, false));
+		Properties.Add(FMovieGraphPinProperties(TEXT("Unknown"), EMovieGraphValueType::None, false));
 	}
 	
 	return Properties;
diff --git a/Engine/Plugins/MovieScene/MovieRenderPipeline/Source/MovieRenderPipelineCore/Public/Graph/MovieGraphCommon.h b/Engine/Plugins/MovieScene/MovieRenderPipeline/Source/MovieRenderPipelineCore/Public/Graph/MovieGraphCommon.h
--- a/Engine/Plugins/MovieScene/MovieRenderPipeline/Source/MovieRenderPipelineCore/Public/Graph/MovieGraphCommon.h
+++ b/Engine/Plugins/MovieScene/MovieRenderPipeline/Source/MovieRenderPipelineCore/Public/Graph/MovieGraphCommon.h
@@ -29,15 +29,11 @@ enum class EMovieGraphValueType : uint8
 	Class UMETA(Hidden),
 	SoftClass UMETA(Hidden),
 
-	Branch UMETA(Hidden),	// Added for MRQ purposes
-	Unknown UMETA(Hidden),	// Added for MRQ purposes
-
 	Count UMETA(Hidden)
 };
 
-// Note: The +2 is to account for the MRQ-added items in the enum
 // TODO: We may want a method which converts between these enum types instead
-static_assert((uint8)EMovieGraphValueType::Count == (uint8)EPropertyBagPropertyType::Count + 2);
+static_assert((uint8)EMovieGraphValueType::Count == (uint8)EPropertyBagPropertyType::Count);
 
 // Note: Also a copy of the property bag's container types for the same reason as EMovieGraphValueType.
 /** The container type of a graph member's value. */
diff --git a/Engine/Plugins/MovieScene/MovieRenderPipeline/Source/MovieRenderPipelineCore/Public/Graph/MovieGraphConfig.h b/Engine/Plugins/MovieScene/MovieRenderPipeline/Source/MovieRenderPipelineCore/Public/Graph/MovieGraphConfig.h
--- a/Engine/Plugins/MovieScene/MovieRenderPipeline/Source/MovieRenderPipelineCore/Public/Graph/MovieGraphConfig.h
+++ b/Engine/Plugins/MovieScene/MovieRenderPipeline/Source/MovieRenderPipelineCore/Public/Graph/MovieGraphConfig.h
@@ -11,6 +11,7 @@
 
 // Forward Declare
 class UMovieGraphNode;
+class UMovieGraphSubgraphNode;
 
 #if WITH_EDITOR
 	DECLARE_MULTICAST_DELEGATE_OneParam(FOnMovieGraphVariableChanged, class UMovieGraphMember*);
@@ -38,6 +39,9 @@ public:
 	/** Determines if this member can be deleted. */
 	virtual bool IsDeletable() const { return true; }
 
+	/** Gets whether this member is editable via the UI. */
+	virtual bool IsEditable() const { return bIsEditable; }
+
 public:
 	// TODO: Need a details customization that validates whether or not the name is valid/unique
 	/** The name of this member, which is user-facing. */
@@ -98,11 +102,25 @@ private:
 	bool bIsGlobal = false;
 };
 
+/**
+ * Common base class for input/output members on the graph.
+ */
+UCLASS(Abstract)
+class MOVIERENDERPIPELINECORE_API UMovieGraphInterfaceBase : public UMovieGraphMember
+{
+	GENERATED_BODY()
+
+public:
+	/** Whether this interface member represents a branch. If not a branch, then a value is associated with it. */
+	UPROPERTY(EditAnywhere, Category = "Value", meta=(EditCondition="bIsEditable", HideEditConditionToggle))
+	bool bIsBranch = true;
+};
+
 /**
  * An input exposed on the graph that will be available for nodes to connect to.
  */
 UCLASS(BlueprintType)
-class MOVIERENDERPIPELINECORE_API UMovieGraphInput : public UMovieGraphMember
+class MOVIERENDERPIPELINECORE_API UMovieGraphInput : public UMovieGraphInterfaceBase
 {
 	GENERATED_BODY()
 
@@ -125,7 +143,7 @@ public:
  * An output exposed on the graph that will be available for nodes to connect to.
  */
 UCLASS(BlueprintType)
-class MOVIERENDERPIPELINECORE_API UMovieGraphOutput : public UMovieGraphMember
+class MOVIERENDERPIPELINECORE_API UMovieGraphOutput : public UMovieGraphInterfaceBase
 {
 	GENERATED_BODY()
 
@@ -229,6 +247,19 @@ public:
 	*/
 	UPROPERTY()
 	TSet<TObjectPtr<UMovieGraphNode>> VisitedNodes;
+
+	/**
+	* The pin that is currently being followed in the traversal process.
+	*/
+	UPROPERTY()
+	TObjectPtr<UMovieGraphPin> PinBeingFollowed;
+
+	/**
+	* The current stack of subgraphs that are being visited. The last subgraph in the stack is the one currently being
+	* visited. If no subgraphs are in this stack, then the parent-most graph is being traversed currently.
+	*/
+	UPROPERTY()
+	TArray<TObjectPtr<const UMovieGraphSubgraphNode>> SubgraphStack;
 };
 
 
@@ -307,6 +338,7 @@ public:
 	void SetEditorOnlyNodes(const TArray<TObjectPtr<const UObject>>& InNodes);
 #endif
 
+	UFUNCTION(BlueprintCallable, Category="Experimental")
 	UMovieGraphEvaluatedConfig* CreateFlattenedGraph(const FMovieGraphTraversalContext& InContext);
 
 	template<typename NodeType>
diff --git a/Engine/Plugins/MovieScene/MovieRenderPipeline/Source/MovieRenderPipelineCore/Public/Graph/MovieGraphNode.h b/Engine/Plugins/MovieScene/MovieRenderPipeline/Source/MovieRenderPipelineCore/Public/Graph/MovieGraphNode.h
--- a/Engine/Plugins/MovieScene/MovieRenderPipeline/Source/MovieRenderPipelineCore/Public/Graph/MovieGraphNode.h
+++ b/Engine/Plugins/MovieScene/MovieRenderPipeline/Source/MovieRenderPipelineCore/Public/Graph/MovieGraphNode.h
@@ -21,7 +21,7 @@ class UMovieGraphMember;
 class UMovieGraphOutput;
 class UMovieGraphPin;
 class UMovieGraphVariable;
-struct FMovieGraphTraversalContext;
+struct FMovieGraphEvaluationContext;
 
 #if WITH_EDITOR
 class UEdGraphNode;
@@ -75,7 +75,7 @@ public:
 	* inherited classes and change that if you need custom logic, such as boolean nodes that want 
 	* to choose one or the other based on the results of a conditional property.
 	*/
-	virtual TArray<UMovieGraphPin*> EvaluatePinsToFollow(const FMovieGraphTraversalContext& InContext) const;
+	virtual TArray<UMovieGraphPin*> EvaluatePinsToFollow(FMovieGraphEvaluationContext& InContext) const;
 
 	/**
 	* When a non-branch pin type is being evaluated on a node, the calling node will ask this node
diff --git a/Engine/Plugins/MovieScene/MovieRenderPipeline/Source/MovieRenderPipelineCore/Public/Graph/MovieGraphPin.h b/Engine/Plugins/MovieScene/MovieRenderPipeline/Source/MovieRenderPipelineCore/Public/Graph/MovieGraphPin.h
--- a/Engine/Plugins/MovieScene/MovieRenderPipeline/Source/MovieRenderPipelineCore/Public/Graph/MovieGraphPin.h
+++ b/Engine/Plugins/MovieScene/MovieRenderPipeline/Source/MovieRenderPipelineCore/Public/Graph/MovieGraphPin.h
@@ -21,20 +21,35 @@ struct FMovieGraphPinProperties
 		, bAllowMultipleConnections(bInAllowMultipleConnections)
 	{}
 
+	static FMovieGraphPinProperties MakeBranchProperties(const FName& InLabel = NAME_None)
+	{
+		FMovieGraphPinProperties Properties(InLabel, EMovieGraphValueType::None, false);
+		Properties.bIsBranch = true;
+		return MoveTemp(Properties);
+	}
+
+	/** The name assigned to the pin. */
 	UPROPERTY(BlueprintReadWrite, EditAnywhere, Category = "Settings")
 	FName Label = NAME_None;
 
+	/** The type of the pin. If the pin represents a branch, this type is ignored. */
 	UPROPERTY(BlueprintReadWrite, EditAnywhere, Category = "Settings")
 	EMovieGraphValueType Type = EMovieGraphValueType::Float;
 
+	/** Whether this pin can accept multiple connections. */
 	UPROPERTY(BlueprintReadWrite, EditAnywhere, Category = "Settings")
 	bool bAllowMultipleConnections = true;
 
+	/** Whether this pin represents a branch. If it does not represent a branch, then it is a value-providing pin. */
+	UPROPERTY(BlueprintReadWrite, EditAnywhere, Category = "Settings")
+	bool bIsBranch = false;
+
 	bool operator==(const FMovieGraphPinProperties& Other) const
 	{
 		return Label == Other.Label
 			&& Type == Other.Type
-			&& bAllowMultipleConnections == Other.bAllowMultipleConnections;
+			&& bAllowMultipleConnections == Other.bAllowMultipleConnections
+			&& bIsBranch == Other.bIsBranch;
 	}
 
 	bool operator !=(const FMovieGraphPinProperties& Other) const
diff --git a/Engine/Plugins/MovieScene/MovieRenderPipeline/Source/MovieRenderPipelineCore/Public/Graph/MovieGraphPlaceholderNodes.cpp b/Engine/Plugins/MovieScene/MovieRenderPipeline/Source/MovieRenderPipelineCore/Public/Graph/MovieGraphPlaceholderNodes.cpp
--- a/Engine/Plugins/MovieScene/MovieRenderPipeline/Source/MovieRenderPipelineCore/Public/Graph/MovieGraphPlaceholderNodes.cpp
+++ b/Engine/Plugins/MovieScene/MovieRenderPipeline/Source/MovieRenderPipelineCore/Public/Graph/MovieGraphPlaceholderNodes.cpp
@@ -17,14 +17,14 @@ static const FText NodeCategory_Settings = LOCTEXT("NodeCategory_Settings", "Set
 TArray<FMovieGraphPinProperties> UMoviePipelineCollectionNode::GetInputPinProperties() const
 {
 	TArray<FMovieGraphPinProperties> Properties;
-	Properties.Add(FMovieGraphPinProperties(NAME_None, EMovieGraphValueType::Branch, false));
+	Properties.Add(FMovieGraphPinProperties::MakeBranchProperties());
 	return Properties;
 }
 
 TArray<FMovieGraphPinProperties> UMoviePipelineCollectionNode::GetOutputPinProperties() const
 {
 	TArray<FMovieGraphPinProperties> Properties;
-	Properties.Add(FMovieGraphPinProperties(NAME_None, EMovieGraphValueType::Branch, false));
+	Properties.Add(FMovieGraphPinProperties::MakeBranchProperties());
 	return Properties;
 }
 
@@ -75,14 +75,14 @@ void UMoviePipelineCollectionNode::PostEditChangeProperty(FPropertyChangedEvent&
 TArray<FMovieGraphPinProperties> UMovieGraphModifierNode::GetInputPinProperties() const
 {
 	TArray<FMovieGraphPinProperties> Properties;
-	Properties.Add(FMovieGraphPinProperties(NAME_None, EMovieGraphValueType::Branch, false));
+	Properties.Add(FMovieGraphPinProperties::MakeBranchProperties());
 	return Properties;
 }
 
 TArray<FMovieGraphPinProperties> UMovieGraphModifierNode::GetOutputPinProperties() const
 {
 	TArray<FMovieGraphPinProperties> Properties;
-	Properties.Add(FMovieGraphPinProperties(NAME_None, EMovieGraphValueType::Branch, false));
+	Properties.Add(FMovieGraphPinProperties::MakeBranchProperties());
 	return Properties;
 }
 
@@ -134,14 +134,14 @@ void UMovieGraphModifierNode::PostEditChangeProperty(FPropertyChangedEvent& Prop
 TArray<FMovieGraphPinProperties> UMovieGraphGlobalGameOverridesNode::GetInputPinProperties() const
 {
 	TArray<FMovieGraphPinProperties> Properties;
-	Properties.Add(FMovieGraphPinProperties(NAME_None, EMovieGraphValueType::Branch, false));
+	Properties.Add(FMovieGraphPinProperties::MakeBranchProperties());
 	return Properties;
 }
 
 TArray<FMovieGraphPinProperties> UMovieGraphGlobalGameOverridesNode::GetOutputPinProperties() const
 {
 	TArray<FMovieGraphPinProperties> Properties;
-	Properties.Add(FMovieGraphPinProperties(NAME_None, EMovieGraphValueType::Branch, false));
+	Properties.Add(FMovieGraphPinProperties::MakeBranchProperties());
 	return Properties;
 }
 
@@ -175,14 +175,14 @@ FSlateIcon UMovieGraphGlobalGameOverridesNode::GetIconAndTint(FLinearColor& OutC
 TArray<FMovieGraphPinProperties> UMovieGraphPathTracedRendererNode::GetInputPinProperties() const
 {
 	TArray<FMovieGraphPinProperties> Properties;
-	Properties.Add(FMovieGraphPinProperties(NAME_None, EMovieGraphValueType::Branch, false));
+	Properties.Add(FMovieGraphPinProperties::MakeBranchProperties());
 	return Properties;
 }
 
 TArray<FMovieGraphPinProperties> UMovieGraphPathTracedRendererNode::GetOutputPinProperties() const
 {
 	TArray<FMovieGraphPinProperties> Properties;
-	Properties.Add(FMovieGraphPinProperties(NAME_None, EMovieGraphValueType::Branch, false));
+	Properties.Add(FMovieGraphPinProperties::MakeBranchProperties());
 	return Properties;
 }
 
@@ -216,14 +216,14 @@ FSlateIcon UMovieGraphPathTracedRendererNode::GetIconAndTint(FLinearColor& OutCo
 TArray<FMovieGraphPinProperties> UMovieGraphEXRSequenceNode::GetInputPinProperties() const
 {
 	TArray<FMovieGraphPinProperties> Properties;
-	Properties.Add(FMovieGraphPinProperties(NAME_None, EMovieGraphValueType::Branch, false));
+	Properties.Add(FMovieGraphPinProperties::MakeBranchProperties());
 	return Properties;
 }
 
 TArray<FMovieGraphPinProperties> UMovieGraphEXRSequenceNode::GetOutputPinProperties() const
 {
 	TArray<FMovieGraphPinProperties> Properties;
-	Properties.Add(FMovieGraphPinProperties(NAME_None, EMovieGraphValueType::Branch, false));
+	Properties.Add(FMovieGraphPinProperties::MakeBranchProperties());
 	return Properties;
 }
 
@@ -261,7 +261,7 @@ TArray<FMovieGraphPinProperties> UMovieGraphSelectNode::GetInputPinProperties()
 	// Generate branch pins for each option
 	for (const FString& SelectOption : SelectOptions)
 	{
-		Properties.Add(FMovieGraphPinProperties(FName(SelectOption), EMovieGraphValueType::Branch, false));
+		Properties.Add(FMovieGraphPinProperties::MakeBranchProperties(FName(SelectOption)));
 	}
 
 	Properties.Add(FMovieGraphPinProperties(NAME_None, EMovieGraphValueType::String, false));
@@ -272,7 +272,7 @@ TArray<FMovieGraphPinProperties> UMovieGraphSelectNode::GetInputPinProperties()
 TArray<FMovieGraphPinProperties> UMovieGraphSelectNode::GetOutputPinProperties() const
 {
 	TArray<FMovieGraphPinProperties> Properties;
-	Properties.Add(FMovieGraphPinProperties(NAME_None, EMovieGraphValueType::Branch, false));
+	Properties.Add(FMovieGraphPinProperties::MakeBranchProperties());
 	return Properties;
 }
 
@@ -328,14 +328,14 @@ void UMovieGraphSelectNode::PostEditChangeProperty(FPropertyChangedEvent& Proper
 TArray<FMovieGraphPinProperties> UMovieGraphAntiAliasingNode::GetInputPinProperties() const
 {
 	TArray<FMovieGraphPinProperties> Properties;
-	Properties.Add(FMovieGraphPinProperties(NAME_None, EMovieGraphValueType::Branch, false));
+	Properties.Add(FMovieGraphPinProperties::MakeBranchProperties());
 	return Properties;
 }
 
 TArray<FMovieGraphPinProperties> UMovieGraphAntiAliasingNode::GetOutputPinProperties() const
 {
 	TArray<FMovieGraphPinProperties> Properties;
-	Properties.Add(FMovieGraphPinProperties(NAME_None, EMovieGraphValueType::Branch, false));
+	Properties.Add(FMovieGraphPinProperties::MakeBranchProperties());
 	return Properties;
 }
 
diff --git a/Engine/Plugins/MovieScene/MovieRenderPipeline/Source/MovieRenderPipelineCore/Public/Graph/MovieGraphValueContainer.h b/Engine/Plugins/MovieScene/MovieRenderPipeline/Source/MovieRenderPipelineCore/Public/Graph/MovieGraphValueContainer.h
--- a/Engine/Plugins/MovieScene/MovieRenderPipeline/Source/MovieRenderPipelineCore/Public/Graph/MovieGraphValueContainer.h
+++ b/Engine/Plugins/MovieScene/MovieRenderPipeline/Source/MovieRenderPipelineCore/Public/Graph/MovieGraphValueContainer.h
@@ -226,6 +226,6 @@ private:
 	// of a property which can have its type changed, 3) the ability to set the value of the property from both Python
 	// and C++, and 4) the ability to change the property at runtime.
 	/** The value held by this object. */
-	UPROPERTY(EditAnywhere, meta=(ShowOnlyInnerProperties), Category = "General")
+	UPROPERTY(EditAnywhere, meta=(ShowOnlyInnerProperties), Category = "Value")
 	FInstancedPropertyBag Value;
 };
\ No newline at end of file
diff --git a/Engine/Plugins/MovieScene/MovieRenderPipeline/Source/MovieRenderPipelineCore/Public/Graph/Nodes/MovieGraphBranchNode.h b/Engine/Plugins/MovieScene/MovieRenderPipeline/Source/MovieRenderPipelineCore/Public/Graph/Nodes/MovieGraphBranchNode.h
--- a/Engine/Plugins/MovieScene/MovieRenderPipeline/Source/MovieRenderPipelineCore/Public/Graph/Nodes/MovieGraphBranchNode.h
+++ b/Engine/Plugins/MovieScene/MovieRenderPipeline/Source/MovieRenderPipelineCore/Public/Graph/Nodes/MovieGraphBranchNode.h
@@ -21,7 +21,7 @@ public:
 
 	virtual TArray<FMovieGraphPinProperties> GetInputPinProperties() const override;
 	virtual TArray<FMovieGraphPinProperties> GetOutputPinProperties() const override;
-	TArray<UMovieGraphPin*> EvaluatePinsToFollow(const FMovieGraphTraversalContext& InContext) const override;
+	virtual TArray<UMovieGraphPin*> EvaluatePinsToFollow(FMovieGraphEvaluationContext& InContext) const override;
 
 #if WITH_EDITOR
 	virtual FText GetNodeTitle(const bool bGetDescriptive = false) const override;
diff --git a/Engine/Plugins/MovieScene/MovieRenderPipeline/Source/MovieRenderPipelineCore/Public/Graph/Nodes/MovieGraphFileOutputNode.h b/Engine/Plugins/MovieScene/MovieRenderPipeline/Source/MovieRenderPipelineCore/Public/Graph/Nodes/MovieGraphFileOutputNode.h
--- a/Engine/Plugins/MovieScene/MovieRenderPipeline/Source/MovieRenderPipelineCore/Public/Graph/Nodes/MovieGraphFileOutputNode.h
+++ b/Engine/Plugins/MovieScene/MovieRenderPipeline/Source/MovieRenderPipelineCore/Public/Graph/Nodes/MovieGraphFileOutputNode.h
@@ -36,14 +36,14 @@ public:
 	virtual TArray<FMovieGraphPinProperties> GetInputPinProperties() const override
 	{
 		TArray<FMovieGraphPinProperties> Properties;
-		Properties.Add(FMovieGraphPinProperties(NAME_None, EMovieGraphValueType::Branch, false));
+		Properties.Add(FMovieGraphPinProperties::MakeBranchProperties());
 		return Properties;
 	}
 
 	virtual TArray<FMovieGraphPinProperties> GetOutputPinProperties() const override
 	{
 		TArray<FMovieGraphPinProperties> Properties;
-		Properties.Add(FMovieGraphPinProperties(NAME_None, EMovieGraphValueType::Branch, false));
+		Properties.Add(FMovieGraphPinProperties::MakeBranchProperties());
 		return Properties;
 	}
 
diff --git a/Engine/Plugins/MovieScene/MovieRenderPipeline/Source/MovieRenderPipelineCore/Public/Graph/Nodes/MovieGraphInputNode.h b/Engine/Plugins/MovieScene/MovieRenderPipeline/Source/MovieRenderPipelineCore/Public/Graph/Nodes/MovieGraphInputNode.h
--- a/Engine/Plugins/MovieScene/MovieRenderPipeline/Source/MovieRenderPipelineCore/Public/Graph/Nodes/MovieGraphInputNode.h
+++ b/Engine/Plugins/MovieScene/MovieRenderPipeline/Source/MovieRenderPipelineCore/Public/Graph/Nodes/MovieGraphInputNode.h
@@ -14,8 +14,10 @@ class MOVIERENDERPIPELINECORE_API UMovieGraphInputNode : public UMovieGraphNode
 
 public:
 	UMovieGraphInputNode();
-	
+
+	//~ Begin UMovieGraphNode interface
 	virtual TArray<FMovieGraphPinProperties> GetOutputPinProperties() const override;
+	virtual TArray<UMovieGraphPin*> EvaluatePinsToFollow(FMovieGraphEvaluationContext& InContext) const override;
 
 #if WITH_EDITOR
 	virtual FText GetNodeTitle(const bool bGetDescriptive = false) const override;
@@ -23,6 +25,7 @@ public:
 	virtual FLinearColor GetNodeTitleColor() const override;
 	virtual FSlateIcon GetIconAndTint(FLinearColor& OutColor) const override;
 #endif
+	//~ End UMovieGraphNode interface
 
 private:
 	virtual void RegisterDelegates() const override;
diff --git a/Engine/Plugins/MovieScene/MovieRenderPipeline/Source/MovieRenderPipelineCore/Public/Graph/Nodes/MovieGraphRenderPassNode.h b/Engine/Plugins/MovieScene/MovieRenderPipeline/Source/MovieRenderPipelineCore/Public/Graph/Nodes/MovieGraphRenderPassNode.h
--- a/Engine/Plugins/MovieScene/MovieRenderPipeline/Source/MovieRenderPipelineCore/Public/Graph/Nodes/MovieGraphRenderPassNode.h
+++ b/Engine/Plugins/MovieScene/MovieRenderPipeline/Source/MovieRenderPipelineCore/Public/Graph/Nodes/MovieGraphRenderPassNode.h
@@ -46,14 +46,14 @@ public:
 	virtual TArray<FMovieGraphPinProperties> GetInputPinProperties() const override
 	{
 		TArray<FMovieGraphPinProperties> Properties;
-		Properties.Add(FMovieGraphPinProperties(NAME_None, EMovieGraphValueType::Branch, false));
+		Properties.Add(FMovieGraphPinProperties::MakeBranchProperties());
 		return Properties;
 	}
 
 	virtual TArray<FMovieGraphPinProperties> GetOutputPinProperties() const override
 	{
 		TArray<FMovieGraphPinProperties> Properties;
-		Properties.Add(FMovieGraphPinProperties(NAME_None, EMovieGraphValueType::Branch, false));
+		Properties.Add(FMovieGraphPinProperties::MakeBranchProperties());
 		return Properties;
 	}
 
diff --git a/Engine/Plugins/MovieScene/MovieRenderPipeline/Source/MovieRenderPipelineCore/Public/Graph/Nodes/MovieGraphSubgraphNode.h b/Engine/Plugins/MovieScene/MovieRenderPipeline/Source/MovieRenderPipelineCore/Public/Graph/Nodes/MovieGraphSubgraphNode.h
new file mode 100644
--- /dev/null
+++ b/Engine/Plugins/MovieScene/MovieRenderPipeline/Source/MovieRenderPipelineCore/Public/Graph/Nodes/MovieGraphSubgraphNode.h
@@ -0,0 +1,55 @@
+// Copyright Epic Games, Inc. All Rights Reserved.
+
+#pragma once
+
+#include "Graph/MovieGraphConfig.h"
+#include "Graph/MovieGraphNode.h"
+#include "Graph/MovieGraphPin.h"
+
+#include "MovieGraphSubgraphNode.generated.h"
+
+/**
+ * A node which represents another graph asset. Inputs/outputs on this subgraph will update if the underlying graph
+ * asset's inputs/outputs change.
+ */
+UCLASS(BlueprintType)
+class MOVIERENDERPIPELINECORE_API UMovieGraphSubgraphNode : public UMovieGraphNode
+{
+	GENERATED_BODY()
+
+public:
+	UMovieGraphSubgraphNode();
+
+	//~ Begin UMovieGraphNode interface
+	virtual TArray<FMovieGraphPinProperties> GetInputPinProperties() const override;
+	virtual TArray<FMovieGraphPinProperties> GetOutputPinProperties() const override;
+	virtual TArray<UMovieGraphPin*> EvaluatePinsToFollow(FMovieGraphEvaluationContext& InContext) const override;
+
+#if WITH_EDITOR
+	virtual FText GetNodeTitle(const bool bGetDescriptive = false) const override;
+	virtual FText GetMenuCategory() const override;
+#endif
+	//~ End UMovieGraphNode interface
+
+	//~ Begin UObject interface
+#if WITH_EDITOR
+	virtual void PostEditChangeProperty(FPropertyChangedEvent& PropertyChangedEvent) override;
+#endif
+	//~ End UObject interface
+
+	/** Sets the graph asset this subgraph points to. */
+	UFUNCTION(BlueprintCallable, Category="Graph")
+	void SetSubGraphAsset(const TSoftObjectPtr<UMovieGraphConfig>& InSubgraphAsset);
+
+	/** Gets the graph asset this subgraph points to. */
+	UFUNCTION(BlueprintCallable, Category="Graph")
+	UMovieGraphConfig* GetSubgraphAsset() const;
+
+private:
+	/** Update the subgraph to reflect the subgraph asset when the subgraph asset is saved. */
+	void RefreshOnSubgraphAssetSaved();
+
+private:
+	UPROPERTY(EditAnywhere, Category="Graph")
+	TSoftObjectPtr<UMovieGraphConfig> SubgraphAsset;
+};
\ No newline at end of file
diff --git a/Engine/Plugins/MovieScene/MovieRenderPipeline/Source/MovieRenderPipelineEditor/Private/Customizations/Graph/GraphMemberCustomization.h b/Engine/Plugins/MovieScene/MovieRenderPipeline/Source/MovieRenderPipelineEditor/Private/Customizations/Graph/GraphMemberCustomization.h
new file mode 100644
--- /dev/null
+++ b/Engine/Plugins/MovieScene/MovieRenderPipeline/Source/MovieRenderPipelineEditor/Private/Customizations/Graph/GraphMemberCustomization.h
@@ -0,0 +1,62 @@
+// Copyright Epic Games, Inc. All Rights Reserved.
+
+#pragma once
+
+#include "DetailLayoutBuilder.h"
+#include "IDetailCustomization.h"
+#include "Graph/MovieGraphConfig.h"
+#include "PropertyHandle.h"
+
+#define LOCTEXT_NAMESPACE "MoviePipelineEditor"
+
+/** Customize how members for a graph appear in the details panel. */
+class FGraphMemberCustomization : public IDetailCustomization
+{
+public:
+	static TSharedRef<IDetailCustomization> MakeInstance()
+	{
+		return MakeShared<FGraphMemberCustomization>();
+	}
+
+protected:
+	//~ Begin IDetailCustomization interface
+	virtual void CustomizeDetails(const TSharedPtr<IDetailLayoutBuilder>& DetailBuilder) override
+	{
+		CustomizeDetails(*DetailBuilder);
+	}
+
+	virtual void CustomizeDetails(IDetailLayoutBuilder& DetailBuilder) override
+	{
+		TArray<TWeakObjectPtr<UObject>> ObjectsBeingCustomized;
+		DetailBuilder.GetObjectsBeingCustomized(ObjectsBeingCustomized);
+
+		for (const TWeakObjectPtr<UObject>& CustomizedObject : ObjectsBeingCustomized)
+		{
+			// Note: The graph members inherit the Value property from a base class, so the enable/disable state cannot
+			// be driven by UPROPERTY metadata. Hence why this needs to be done w/ a details customization.
+			
+			// Enable/disable the value property for inputs/outputs based on whether it is specified as a branch or not
+			if (const UMovieGraphInterfaceBase* InterfaceBase = Cast<UMovieGraphInterfaceBase>(CustomizedObject))
+			{
+				const TSharedRef<IPropertyHandle> ValueProperty = DetailBuilder.GetProperty("Value", UMovieGraphValueContainer::StaticClass());
+				if (ValueProperty->IsValidHandle())
+				{
+					DetailBuilder.EditDefaultProperty(ValueProperty)->IsEnabled(!InterfaceBase->bIsBranch);
+				}
+			}
+
+			// Enable/disable the value property for variables based on the editable state
+			if (const UMovieGraphVariable* Variable = Cast<UMovieGraphVariable>(CustomizedObject))
+			{
+				const TSharedRef<IPropertyHandle> ValueProperty = DetailBuilder.GetProperty("Value", UMovieGraphValueContainer::StaticClass());
+				if (ValueProperty->IsValidHandle())
+				{
+					DetailBuilder.EditDefaultProperty(ValueProperty)->IsEnabled(Variable->IsEditable());
+				}
+			}
+		}
+	}
+	//~ End IDetailCustomization interface
+};
+
+#undef LOCTEXT_NAMESPACE
\ No newline at end of file
diff --git a/Engine/Plugins/MovieScene/MovieRenderPipeline/Source/MovieRenderPipelineEditor/Private/Graph/MovieEdGraphNode.cpp b/Engine/Plugins/MovieScene/MovieRenderPipeline/Source/MovieRenderPipelineEditor/Private/Graph/MovieEdGraphNode.cpp
--- a/Engine/Plugins/MovieScene/MovieRenderPipeline/Source/MovieRenderPipelineEditor/Private/Graph/MovieEdGraphNode.cpp
+++ b/Engine/Plugins/MovieScene/MovieRenderPipeline/Source/MovieRenderPipelineEditor/Private/Graph/MovieEdGraphNode.cpp
@@ -53,12 +53,16 @@ FEdGraphPinType UMoviePipelineEdGraphNodeBase::GetPinType(const UMovieGraphPin*
 	EdPinType.PinCategory = NAME_None;
 	EdPinType.PinSubCategory = NAME_None;
 
+	// Special case for branch pins
+	if (InPin->Properties.bIsBranch)
+	{
+		EdPinType.PinCategory = UMovieGraphSchema::PC_Branch;
+		return EdPinType;
+	}
+
 	// TODO: This needs to properly fill in the category/subcategory for all types. Also, container type?
 	switch (InPin->Properties.Type)
 	{
-	case EMovieGraphValueType::Branch:
-		EdPinType.PinCategory = UMovieGraphSchema::PC_Branch;
-		break;
 	case EMovieGraphValueType::Float:
 		EdPinType.PinCategory = UMovieGraphSchema::PC_Float;
 		break;
diff --git a/Engine/Plugins/MovieScene/MovieRenderPipeline/Source/MovieRenderPipelineEditor/Private/Graph/MovieGraphAssetToolkit.cpp b/Engine/Plugins/MovieScene/MovieRenderPipeline/Source/MovieRenderPipelineEditor/Private/Graph/MovieGraphAssetToolkit.cpp
--- a/Engine/Plugins/MovieScene/MovieRenderPipeline/Source/MovieRenderPipelineEditor/Private/Graph/MovieGraphAssetToolkit.cpp
+++ b/Engine/Plugins/MovieScene/MovieRenderPipeline/Source/MovieRenderPipelineEditor/Private/Graph/MovieGraphAssetToolkit.cpp
@@ -2,6 +2,7 @@
 
 #include "MovieGraphAssetToolkit.h"
 
+#include "Customizations/Graph/GraphMemberCustomization.h"
 #include "Framework/Commands/GenericCommands.h"
 #include "Graph/MovieGraphConfig.h"
 #include "MovieEdGraphNode.h"
@@ -200,6 +201,10 @@ TSharedRef<SDockTab> FMovieGraphAssetToolkit::SpawnTab_RenderGraphDetails(const
 	DetailsViewArgs.ViewIdentifier = "MovieGraphSettings";
 
 	SelectedGraphObjectsDetailsWidget = PropertyEditorModule.CreateDetailView(DetailsViewArgs);
+
+	SelectedGraphObjectsDetailsWidget->RegisterInstancedCustomPropertyLayout(
+		UMovieGraphMember::StaticClass(),
+		FOnGetDetailCustomizationInstance::CreateStatic(&FGraphMemberCustomization::MakeInstance));
 	
 	return SNew(SDockTab)
 		.TabColorScale(GetTabColorScale())
