MoviePipeline: Fixed localization warning.

diff --git a/Engine/Plugins/MovieScene/MovieRenderPipeline/Source/MovieRenderPipelineCore/Private/Graph/Nodes/MovieGraphSubgraphNode.cpp b/Engine/Plugins/MovieScene/MovieRenderPipeline/Source/MovieRenderPipelineCore/Private/Graph/Nodes/MovieGraphSubgraphNode.cpp
--- a/Engine/Plugins/MovieScene/MovieRenderPipeline/Source/MovieRenderPipelineCore/Private/Graph/Nodes/MovieGraphSubgraphNode.cpp
+++ b/Engine/Plugins/MovieScene/MovieRenderPipeline/Source/MovieRenderPipelineCore/Private/Graph/Nodes/MovieGraphSubgraphNode.cpp
@@ -133,7 +133,7 @@ FText UMovieGraphSubgraphNode::GetNodeTitle(const bool bGetDescriptive) const
 
 	if (const UMovieGraphConfig* Subgraph = GetSubgraphAsset())
 	{
-		return FText::Format(LOCTEXT("SubgraphTitle", "Subgraph: {0}"), FText::FromString(Subgraph->GetName()));
+		return FText::Format(LOCTEXT("SubgraphTitle_IncludesAsset", "Subgraph: {0}"), FText::FromString(Subgraph->GetName()));
 	}
 
 	return bGetDescriptive ? DescriptiveSubgraphTitle : SubgraphTitle;
