MoviePipeline: Fleshed out what the pipeline looks like to use flattened graph configs generated by traversing them. Traversal generates a "flat" list that needs to travel through the system with the shot that generated it so we can query the values from it later to read settings and resolve values.

diff --git a/Engine/Plugins/MovieScene/MovieRenderPipeline/Source/MovieRenderPipelineCore/Private/Graph/MovieGraphBlueprintLibrary.cpp b/Engine/Plugins/MovieScene/MovieRenderPipeline/Source/MovieRenderPipelineCore/Private/Graph/MovieGraphBlueprintLibrary.cpp
new file mode 100644
--- /dev/null
+++ b/Engine/Plugins/MovieScene/MovieRenderPipeline/Source/MovieRenderPipelineCore/Private/Graph/MovieGraphBlueprintLibrary.cpp
@@ -0,0 +1,14 @@
+// Copyright Epic Games, Inc. All Rights Reserved.
+
+#include "Graph/MovieGraphBlueprintLibrary.h"
+#include "Graph/Nodes/MovieGraphOutputSettingNode.h"
+
+FFrameRate UMovieGraphBlueprintLibrary::GetEffectiveFrameRate(UMovieGraphOutputSettingNode* InNode, const FFrameRate& InDefaultRate)
+{
+	if (InNode && InNode->bOverride_OutputFrameRate)
+	{
+		return InNode->OutputFrameRate;
+	}
+
+	return InDefaultRate;
+}
\ No newline at end of file
diff --git a/Engine/Plugins/MovieScene/MovieRenderPipeline/Source/MovieRenderPipelineCore/Private/Graph/MovieGraphConfig.cpp b/Engine/Plugins/MovieScene/MovieRenderPipeline/Source/MovieRenderPipelineCore/Private/Graph/MovieGraphConfig.cpp
--- a/Engine/Plugins/MovieScene/MovieRenderPipeline/Source/MovieRenderPipelineCore/Private/Graph/MovieGraphConfig.cpp
+++ b/Engine/Plugins/MovieScene/MovieRenderPipeline/Source/MovieRenderPipelineCore/Private/Graph/MovieGraphConfig.cpp
@@ -99,50 +99,6 @@ void UMovieGraphConfig::PostLoad()
 	}
 }
 
-void UMovieGraphConfig::TraverseGraphRecursive(UMovieGraphNode* InNode, TSubclassOf<UMovieGraphNode> InClassType, const FMovieGraphTraversalContext& InContext, TArray<UMovieGraphNode*>& OutNodes) const
-{
-	check(InNode);
-	
-	if (InNode->IsA(InClassType))
-	{
-		OutNodes.Add(InNode);
-	}
-
-	// Loop through the input pins on each node
-	for (UMovieGraphPin* Pin : InNode->GetInputPins())
-	{
-		// ToDo: Pins should have a type (property, branch line, etc.) and we can mask following them.
-
-		bool bFollowPath = true;
-		// If this is an output node we mask their pins against the traversal context
-		if (UMovieGraphOutputNode* NodeAsOutputNode = Cast<UMovieGraphOutputNode>(InNode))
-		{
-			if (InContext.RootBranch.BranchName != NAME_None)
-			{
-				bFollowPath = Pin->Properties.Label == InContext.RootBranch.BranchName;
-				UE_LOG(LogMovieRenderPipeline, Log, TEXT("Traversing Graph with Root Branch: %s (bMatch: %d)"), *InContext.RootBranch.BranchName.ToString(), bFollowPath);
-			}
-		}
-
-		if (bFollowPath)
-		{
-			for (UMovieGraphEdge* Edge : Pin->Edges)
-			{
-				if (UMovieGraphPin* OtherPin = Edge->GetOtherPin(Pin))
-				{
-					// ToDo: This needs to be upgraded to a full cyclic detection so we don't get stuck in an infinite
-					// loop through some badly malformed graph created via scripting/etc.
-					if (OtherPin->Node != InNode)
-					{
-						TraverseGraphRecursive(OtherPin->Node, InClassType, InContext, OutNodes);
-					}
-				}
-			}
-		}
-	}
-
-}
-
 UMovieGraphVariable* UMovieGraphConfig::AddGlobalVariable(const FName& InName)
 {
 	// Don't add duplicate global variables
@@ -206,17 +162,6 @@ void UMovieGraphConfig::AddDefaultMembers()
 	}
 }
 
-TArray<UMovieGraphNode*> UMovieGraphConfig::TraverseGraph(TSubclassOf<UMovieGraphNode> InClassType, const FMovieGraphTraversalContext& InContext) const
-{
-	TArray<UMovieGraphNode*> OutNodes;
-
-	if (OutputNode)
-	{
-		TraverseGraphRecursive(OutputNode, InClassType, InContext, /*Out*/OutNodes);
-	}
-	return OutNodes;
-}
-
 bool UMovieGraphConfig::AddLabeledEdge(UMovieGraphNode* FromNode, const FName& FromPinLabel, UMovieGraphNode* ToNode, const FName& ToPinLabel)
 {
 	if (!FromNode || !ToNode)
@@ -605,21 +550,6 @@ bool UMovieGraphConfig::DeleteVariableMember(UMovieGraphVariable* VariableMember
 	return true;
 }
 
-TArray<FMovieGraphBranch> UMovieGraphConfig::GetOutputBranches() const
-{
-	TArray<FMovieGraphBranch> Branches;
-	if (OutputNode)
-	{
-		for (UMovieGraphPin* Pin : OutputNode->GetInputPins())
-		{
-			FMovieGraphBranch& NewBranch = Branches.AddDefaulted_GetRef();
-			NewBranch.BranchName = Pin->Properties.Label;
-		}
-	}
-
-	return Branches;
-}
-
 #if WITH_EDITOR
 void UMovieGraphConfig::SetEditorOnlyNodes(const TArray<TObjectPtr<const UObject>>& InNodes)
 {
diff --git a/Engine/Plugins/MovieScene/MovieRenderPipeline/Source/MovieRenderPipelineCore/Private/Graph/MovieGraphDefaultRenderer.cpp b/Engine/Plugins/MovieScene/MovieRenderPipeline/Source/MovieRenderPipelineCore/Private/Graph/MovieGraphDefaultRenderer.cpp
--- a/Engine/Plugins/MovieScene/MovieRenderPipeline/Source/MovieRenderPipelineCore/Private/Graph/MovieGraphDefaultRenderer.cpp
+++ b/Engine/Plugins/MovieScene/MovieRenderPipeline/Source/MovieRenderPipelineCore/Private/Graph/MovieGraphDefaultRenderer.cpp
@@ -19,36 +19,40 @@ void UMovieGraphDefaultRenderer::SetupRenderingPipelineForShot(UMoviePipelineExe
 	struct FMovieGraphPass
 	{
 		TSubclassOf<UMovieGraphRenderPassNode> ClassType;
-		TArray<FString> LayerNames;
+		TArray<FName> LayerBranchNames;
 	};
 
 	TArray<FMovieGraphPass> OutputPasses;
 
 	// Start by getting our root set of branches we should follow
-	TArray<FMovieGraphBranch> GraphBranches = RootGraph->GetOutputBranches();
-	for (const FMovieGraphBranch& Branch : GraphBranches)
-	{
-		// We follow each branch looking for Render Layer nodes to figure out what render layer this should be.
-		FMovieGraphTraversalContext TraversalContext = GetOwningGraph()->GetCurrentTraversalContext();
-		TraversalContext.RootBranch = Branch;
+	const FMovieGraphTimeStepData& TimeStepData = GetOwningGraph()->GetTimeStepInstance()->GetCalculatedTimeData();
+	UMovieGraphEvaluatedConfig* EvaluatedConfig = TimeStepData.EvaluatedConfig;
+	TArray<FName> GraphBranches = EvaluatedConfig->GetBranchNames();
 
-		FString RenderLayerName = TEXT("Unnamed_Layer");
-		UMovieGraphRenderLayerNode* FirstRenderLayerNode = RootGraph->IterateGraphForClass<UMovieGraphRenderLayerNode>(TraversalContext);
-		if (FirstRenderLayerNode)
+	for (const FName& Branch : GraphBranches)
+	{
+		// We follow each branch looking for Render Layer nodes to figure out what render layer this should be. We assume a render layer is named
+		// after the branch it is on, unless they specifically add a UMovieGraphRenderLayerNode to rename it.
+		FString RenderLayerName = Branch.ToString();
+		const bool bIncludeCDOs = false;
+		UMovieGraphRenderLayerNode* RenderLayerNode = EvaluatedConfig->GetSettingForBranch<UMovieGraphRenderLayerNode>(Branch, bIncludeCDOs);
+			
+		// A RenderLayerNode is required for now to indicate you wish to actually render something.
+		if (!RenderLayerNode)
 		{
-			RenderLayerName = FirstRenderLayerNode->GetRenderLayerName();;
+			continue;
 		}
 
 		// Now we need to figure out which renderers are on this branch.
-		TArray<UMovieGraphRenderPassNode*> Renderers = RootGraph->IterateGraphForClassAll<UMovieGraphRenderPassNode>(TraversalContext);
-		if (FirstRenderLayerNode && Renderers.Num() == 0)
+		const bool bExactMatch = false;
+		TArray<UMovieGraphRenderPassNode*> Renderers = EvaluatedConfig->GetSettingsForBranch<UMovieGraphRenderPassNode>(Branch, bIncludeCDOs, bExactMatch);
+		if (RenderLayerNode && Renderers.Num() == 0)
 		{
 			UE_LOG(LogMovieRenderPipeline, Warning, TEXT("Found RenderLayer: \"%s\" but no Renderers defined."), *RenderLayerName);
 		}
 
 		for (UMovieGraphRenderPassNode* RenderPassNode : Renderers)
 		{
-			FString RendererName = RenderPassNode->GetRendererName();
 
 			FMovieGraphPass* ExistingPass = OutputPasses.FindByPredicate([RenderPassNode](const FMovieGraphPass& Pass)
 				{ return Pass.ClassType == RenderPassNode->GetClass(); });
@@ -58,7 +62,7 @@ void UMovieGraphDefaultRenderer::SetupRenderingPipelineForShot(UMoviePipelineExe
 				ExistingPass = &OutputPasses.AddDefaulted_GetRef();
 				ExistingPass->ClassType = RenderPassNode->GetClass();
 			}
-			ExistingPass->LayerNames.AddUnique(RenderLayerName);
+			ExistingPass->LayerBranchNames.AddUnique(Branch);
 		}
 	}
 
@@ -68,17 +72,17 @@ void UMovieGraphDefaultRenderer::SetupRenderingPipelineForShot(UMoviePipelineExe
 	{
 		// ToDo: This should probably come from the Renderers themselves, as they can internally produce multiple
 		// renders (such as ObjectID passes).
-		TotalLayerCount += Pass.LayerNames.Num();
+		TotalLayerCount += Pass.LayerBranchNames.Num();
 		
 		FMovieGraphRenderPassSetupData SetupData;
 		SetupData.Renderer = this;
 		UE_LOG(LogMovieRenderPipeline, Warning, TEXT("\tRenderer Class: %s"), *Pass.ClassType->GetName());
-		for (const FString& LayerName : Pass.LayerNames)
+		for (const FName& LayerBranchName : Pass.LayerBranchNames)
 		{
 			FMovieGraphRenderPassLayerData& LayerData = SetupData.Layers.AddDefaulted_GetRef();
-			LayerData.LayerName = LayerName;
+			LayerData.BranchName = LayerBranchName;
 
-			UE_LOG(LogMovieRenderPipeline, Warning, TEXT("\t\tLayer Name: %s"), *LayerName);
+			UE_LOG(LogMovieRenderPipeline, Warning, TEXT("\t\tBranch Name: %s"), *LayerBranchName.ToString());
 		}
 
 		UMovieGraphRenderPassNode* RenderPassCDO = Pass.ClassType->GetDefaultObject<UMovieGraphRenderPassNode>();
@@ -188,6 +192,7 @@ void UMovieGraphDefaultRenderer::Render(const FMovieGraphTimeStepData& InTimeSte
 		// Get the Traversal Context (not specific to any render pass) at the first sample. This is so
 		// we can easily fetch things that are shared between all render layers later.
 		NewOutputFrame.TraversalContext = GetOwningGraph()->GetCurrentTraversalContext();
+		NewOutputFrame.EvaluatedConfig = TStrongObjectPtr<UMovieGraphEvaluatedConfig>(InTimeStepData.EvaluatedConfig);
 
 		for (const TObjectPtr<UMovieGraphRenderPassNode>& RenderPass : RenderPassesInUse)
 		{
@@ -330,4 +335,4 @@ FSurfaceAccumulatorPool::FInstancePtr FSurfaceAccumulatorPool::BlockAndGetAccumu
 
 	return Accumulators[AvailableIndex];
 }
-}
\ No newline at end of file
+}
diff --git a/Engine/Plugins/MovieScene/MovieRenderPipeline/Source/MovieRenderPipelineCore/Private/Graph/MovieGraphLinearTimeStep.cpp b/Engine/Plugins/MovieScene/MovieRenderPipeline/Source/MovieRenderPipelineCore/Private/Graph/MovieGraphLinearTimeStep.cpp
--- a/Engine/Plugins/MovieScene/MovieRenderPipeline/Source/MovieRenderPipelineCore/Private/Graph/MovieGraphLinearTimeStep.cpp
+++ b/Engine/Plugins/MovieScene/MovieRenderPipeline/Source/MovieRenderPipelineCore/Private/Graph/MovieGraphLinearTimeStep.cpp
@@ -2,6 +2,8 @@
 
 #include "Graph/MovieGraphLinearTimeStep.h"
 #include "Graph/MovieGraphPipeline.h"
+#include "Graph/MovieGraphBlueprintLibrary.h"
+#include "Graph/Nodes/MovieGraphOutputSettingNode.h"
 #include "MoviePipelineQueue.h"
 #include "MovieRenderPipelineCoreModule.h"
 #include "LevelSequence.h"
@@ -66,17 +68,30 @@ void UMovieGraphLinearTimeStep::TickProducingFrames()
 		UE_LOG(LogMovieRenderPipeline, Log, TEXT("MovieGraph Initializing Camera Cut [%d/%d] in [%s] %s."),
 			CurrentShotIndex + 1, ActiveShotList.Num(), *CurrentCameraCut->OuterName, *CurrentCameraCut->InnerName);
 
+		// Evaluate the graph so we can fetch values for this shot.
+		UMovieGraphConfig* Config = GetOwningGraph()->GetRootGraphForShot(CurrentCameraCut);
+		FMovieGraphTraversalContext Context = GetOwningGraph()->GetCurrentTraversalContext();
+		CurrentFrameData.EvaluatedConfig = TStrongObjectPtr<UMovieGraphEvaluatedConfig>(Config->CreateFlattenedGraph(Context));
+		CurrentFrameData.TemporalSampleIndex = 0;
+
+		// Ensure we've set it in the CurrentTimeStepData so things can fetch from it below.
+		CurrentTimeStepData.EvaluatedConfig = TObjectPtr<UMovieGraphEvaluatedConfig>(CurrentFrameData.EvaluatedConfig.Get());
+
+
 		// Sets up the render state, etc.
 		GetOwningGraph()->SetupShot(CurrentCameraCut);
 
 		// Seeks the external data source to match
 		GetOwningGraph()->GetDataSourceInstance()->InitializeShot(CurrentCameraCut);
 
-		CurrentFrameData.TemporalSampleIndex = 0;
+		const bool bIncludeCDO = false;
+		UMovieGraphOutputSettingNode* OutputNode = CurrentFrameData.EvaluatedConfig->GetSettingForBranch<UMovieGraphOutputSettingNode>(UMovieGraphNode::GlobalsPinName, bIncludeCDO);
 
-		FFrameRate TickResolution = GetOwningGraph()->GetDataSourceInstance()->GetTickResolution();
-		FFrameRate FrameRate = GetOwningGraph()->GetDataSourceInstance()->GetDisplayRate(); // ToDo, needs to come from config (config can override)
-		FFrameTime FrameTimePerOutputFrame = FFrameRate::TransformTime(FFrameTime(FFrameNumber(1)), FrameRate, TickResolution);
+		const FFrameRate TickResolution = GetOwningGraph()->GetDataSourceInstance()->GetTickResolution();
+		const FFrameRate SourceFrameRate = GetOwningGraph()->GetDataSourceInstance()->GetDisplayRate(); // ToDo, needs to come from config (config can override)
+		const FFrameRate FinalFrameRate = UMovieGraphBlueprintLibrary::GetEffectiveFrameRate(OutputNode, SourceFrameRate);
+
+		FFrameTime FrameTimePerOutputFrame = FFrameRate::TransformTime(FFrameTime(FFrameNumber(1)), FinalFrameRate, TickResolution);
 
 		// Dummy range so that the first rendering frame knows to start at the right time.
 		FFrameTime UpperBound = CurrentCameraCut->ShotInfo.CurrentTimeInRoot;
@@ -97,6 +112,10 @@ void UMovieGraphLinearTimeStep::TickProducingFrames()
 	{
 		if (IsFirstTemporalSample())
 		{
+			FMovieGraphTraversalContext Context = GetOwningGraph()->GetCurrentTraversalContext();
+			UMovieGraphConfig* Config = GetOwningGraph()->GetRootGraphForShot(CurrentCameraCut);
+			CurrentFrameData.EvaluatedConfig = TStrongObjectPtr<UMovieGraphEvaluatedConfig>(Config->CreateFlattenedGraph(Context));
+
 			// The temporal sample count can change every frame due to graph evaluations, so when we're on our first temporal
 			// sub-sample of the new frame, we need to re-fetch the value.
 			UpdateTemporalSampleCount();
@@ -211,6 +230,8 @@ void UMovieGraphLinearTimeStep::TickProducingFrames()
 		CurrentTimeStepData.bIsLastTemporalSampleForFrame = IsLastTemporalSample();
 		CurrentTimeStepData.bRequiresAccumulator = CurrentFrameData.TemporalSampleCount > 1;
 		CurrentTimeStepData.OutputFrameNumber = CurrentFrameData.OutputFrameNumber;
+		
+		CurrentTimeStepData.EvaluatedConfig = TObjectPtr<UMovieGraphEvaluatedConfig>(CurrentFrameData.EvaluatedConfig.Get());
 
 		UE_LOG(LogTemp, Warning, TEXT("F# %d bFirst: %d bLast: %d bReqAc: %d"),
 			CurrentTimeStepData.OutputFrameNumber, CurrentTimeStepData.bIsFirstTemporalSampleForFrame,
@@ -263,7 +284,6 @@ void UMovieGraphLinearTimeStep::TickProducingFrames()
 
 }
 
-
 bool UMovieGraphLinearTimeStep::IsFirstTemporalSample() const
 {
 	return CurrentFrameData.TemporalSampleIndex == 0;
diff --git a/Engine/Plugins/MovieScene/MovieRenderPipeline/Source/MovieRenderPipelineCore/Private/Graph/MovieGraphPipeline.cpp b/Engine/Plugins/MovieScene/MovieRenderPipeline/Source/MovieRenderPipelineCore/Private/Graph/MovieGraphPipeline.cpp
--- a/Engine/Plugins/MovieScene/MovieRenderPipeline/Source/MovieRenderPipelineCore/Private/Graph/MovieGraphPipeline.cpp
+++ b/Engine/Plugins/MovieScene/MovieRenderPipeline/Source/MovieRenderPipelineCore/Private/Graph/MovieGraphPipeline.cpp
@@ -688,25 +688,57 @@ void UMovieGraphPipeline::ProcessOutstandingFinishedFrames()
 			// UE_LOG(LogTemp, Log, TEXT("Frame: %d Duration: %s"), FrameNumber, *DurationTimeStr);
 		}
 
+		// How we choose which image gets sent to which output container is a little tricky. We use the CDO of each node type on purpose
+		// as the graph can change node instances every frame. If a output type node is in the "Globals" graph, we assume that we should
+		// send all rendered layers to the image type. However we should also allow placing an output of a given type in a render layer,
+		// so you could choose to send the layer to only some output types (ie: only send the "beauty" layer to the .jpeg container).
+		// Because our data is always treated as a block of all image data (and is non-copyable) we instead generate both a list of
+		// output containers for the data in this output frame, and a mask which lets that container know if it should skip a layer.
+		TMap<UMovieGraphFileOutputNode*, TSet<FMovieGraphRenderDataIdentifier>> MaskData;
+
+		for (TPair<FMovieGraphRenderDataIdentifier, TUniquePtr<FImagePixelData>>& RenderData : OutputFrame.ImageOutputData)
+		{
+			UE::MovieGraph::FMovieGraphSampleState* Payload = RenderData.Value->GetPayload<UE::MovieGraph::FMovieGraphSampleState>();
+			if (!ensure(Payload))
+			{
+				continue;
+			}
 
-		// We need to get the list of output nodes for this frame's shot. We pass all of the data at once to each relevant one
-		// and individual output nodes can decide if they should skip writing that layer's data, or combine them into one file
-		// or multiple files, etc. To get the right evaluation context though we need to use common data from one node for now.
-		TSet<UMovieGraphFileOutputNode*> ContainerCDOs;
+			if (!ensure(OutputFrame.EvaluatedConfig))
+			{
+				continue;
+			}
 
-		TArray<UMovieGraphFileOutputNode*> AllNodeInstances = UMovieGraphConfig::IterateGraphForClassAll<UMovieGraphFileOutputNode>(OutputFrame.TraversalContext);
-		for (UMovieGraphFileOutputNode* Node : AllNodeInstances)
-		{
-			ContainerCDOs.Add(Node->GetClass()->GetDefaultObject<UMovieGraphFileOutputNode>());
+			// Get a list of all output nodes for this particular render layer. We specifically skip CDOs here because we're just trying
+			// to find out which output nodes the user _wanted_ to place items onto, we later collect only the CDOs so that we have
+			// a central point for actually handling the file writing.
+			const bool bIncludeCDOs = false;
+			const FName BranchName = RenderData.Key.RootBranchName;
+			TArray<TObjectPtr<UMovieGraphFileOutputNode>> OutputNodeInstances = OutputFrame.EvaluatedConfig->GetSettingsForBranch<UMovieGraphFileOutputNode>(BranchName, bIncludeCDOs);
+			for (const TObjectPtr<UMovieGraphFileOutputNode>& Instance : OutputNodeInstances)
+			{
+				UMovieGraphFileOutputNode* CDO = Instance->GetClass()->GetDefaultObject<UMovieGraphFileOutputNode>();
+				MaskData.FindOrAdd(CDO).Add(RenderData.Key);
+			}
 		}
 
-		for (UMovieGraphFileOutputNode* CDOInstance : ContainerCDOs)
+		// Now that we've looped through the above, we have the total list of which output formats are being used by the graph for
+		// all of the render layers given. We also have a list of which identifiers should go into each one. So we can loop through
+		// the CDO instances and pass the data to them.
+		UE_LOG(LogMovieRenderPipeline, Warning, TEXT("File Outputs:"));
+		for (const TPair<UMovieGraphFileOutputNode*, TSet<FMovieGraphRenderDataIdentifier>>& Pair : MaskData)
 		{
-			CDOInstance->OnReceiveImageData(this, &OutputFrame);
+			UE_LOG(LogMovieRenderPipeline, Warning, TEXT("\tNode: %s"), *Pair.Key->GetClass()->GetName());
+			for (const FMovieGraphRenderDataIdentifier& ID : Pair.Value)
+			{
+				UE_LOG(LogMovieRenderPipeline, Warning, TEXT("\t\tBranch: %s:"), *ID.RootBranchName.ToString());
+			}
+
+			Pair.Key->OnReceiveImageData(this, &OutputFrame, Pair.Value);
 
 			// Ensure we keep track of which nodes we actually sent data to during this render
 			// so that we can call BeginFinalize/IsFinishedWritingToDisk on them when shutting down.
-			OutputNodesDataSentTo.Add(CDOInstance);
+			OutputNodesDataSentTo.Add(Pair.Key);
 		}
 	}
 }
diff --git a/Engine/Plugins/MovieScene/MovieRenderPipeline/Source/MovieRenderPipelineCore/Private/Graph/Nodes/MovieGraphDeferredRenderPassNode.cpp b/Engine/Plugins/MovieScene/MovieRenderPipeline/Source/MovieRenderPipelineCore/Private/Graph/Nodes/MovieGraphDeferredRenderPassNode.cpp
--- a/Engine/Plugins/MovieScene/MovieRenderPipeline/Source/MovieRenderPipelineCore/Private/Graph/Nodes/MovieGraphDeferredRenderPassNode.cpp
+++ b/Engine/Plugins/MovieScene/MovieRenderPipeline/Source/MovieRenderPipelineCore/Private/Graph/Nodes/MovieGraphDeferredRenderPassNode.cpp
@@ -1,6 +1,7 @@
 // Copyright Epic Games, Inc. All Rights Reserved.
 
 #include "Graph/Nodes/MovieGraphDeferredRenderPassNode.h"
+#include "Graph/Nodes/MovieGraphOutputSettingNode.h"
 #include "Graph/MovieGraphDataTypes.h"
 #include "Graph/MovieGraphDefaultRenderer.h"
 #include "Graph/MovieGraphPipeline.h"
@@ -81,7 +82,7 @@ void UMovieGraphDeferredRenderPassNode::FMovieGraphDeferredRenderPass::Setup(TWe
 	Renderer = InRenderer;
 	RenderPassNode = InRenderPassNode;
 
-	RenderDataIdentifier.RenderLayerName = LayerData.LayerName;
+	RenderDataIdentifier.RootBranchName = LayerData.BranchName;
 	RenderDataIdentifier.RendererName = RenderPassNode->GetRendererName();
 	RenderDataIdentifier.SubResourceName = TEXT("beauty");
 	
@@ -130,9 +131,16 @@ void UMovieGraphDeferredRenderPassNode::FMovieGraphDeferredRenderPass::AddRefere
 
 void UMovieGraphDeferredRenderPassNode::FMovieGraphDeferredRenderPass::Render(const FMovieGraphTraversalContext& InFrameTraversalContext, const FMovieGraphTimeStepData& InTimeData)
 {
+	const bool bIncludeCDOs = true;
+	UMovieGraphOutputSettingNode* OutputSetting = InTimeData.EvaluatedConfig->GetSettingForBranch<UMovieGraphOutputSettingNode>(LayerData.BranchName, bIncludeCDOs);
+	if (!ensure(OutputSetting))
+	{
+		return;
+	}
+
 	// This is the size we actually render at.
 	UE::MovieGraph::DefaultRenderer::FRenderTargetInitParams RenderTargetInitParams;
-	RenderTargetInitParams.Size = FIntPoint(1920, 1080);
+	RenderTargetInitParams.Size = OutputSetting->OutputResolution;
 	
 	// OCIO: Since this is a manually created Render target we don't need Gamma to be applied.
 	// We use this render target to render to via a display extension that utilizes Display Gamma
@@ -232,23 +240,34 @@ void UMovieGraphDeferredRenderPassNode::FMovieGraphDeferredRenderPass::PostRende
 
 	auto OnSurfaceReadbackFinished = [this, InSampleState, AccumulationArgs, AccumulatorInstance](TUniquePtr<FImagePixelData>&& InPixelData)
 	{
-			UE::Tasks::TTask<void> Task = UE::Tasks::Launch(UE_SOURCE_LOCATION, [PixelData = MoveTemp(InPixelData), InSampleState, AccumulationArgs, AccumulatorInstance]() mutable
-			{
-				// Enqueue a encode for this frame onto our worker thread.
-				UE::MovieGraph::AccumulateSample_TaskThread(MoveTemp(PixelData), InSampleState, AccumulationArgs);
+		UE::Tasks::TTask<void> Task = UE::Tasks::Launch(UE_SOURCE_LOCATION, [PixelData = MoveTemp(InPixelData), InSampleState, AccumulationArgs, AccumulatorInstance]() mutable
+		{
+			// Enqueue a encode for this frame onto our worker thread.
+			UE::MovieGraph::AccumulateSample_TaskThread(MoveTemp(PixelData), InSampleState, AccumulationArgs);
 
-				// We have to defer clearing the accumulator until after sample accumulation has finished
-				if (AccumulationArgs.bIsLastSample)
-				{
-					// Final sample has now been executed, free the accumulator for reuse.
-					AccumulatorInstance->bIsActive = false;
-				}
-			}, AccumulatorInstance->TaskPrereq);
+			// We have to defer clearing the accumulator until after sample accumulation has finished
+			if (AccumulationArgs.bIsLastSample)
+			{
+				// Final sample has now been executed, free the accumulator for reuse.
+				AccumulatorInstance->bIsActive = false;
+			}
+		}, AccumulatorInstance->TaskPrereq);
 
 		// Make the next accumulation task that uses this accumulator use the task we just created as a pre-req.
 		AccumulatorInstance->TaskPrereq = Task;
 
-		this->Renderer->AddOutstandingRenderTask_AnyThread(Task);
+		// Because we're run on a separate thread, we need to check validity differently. The standard
+		// TWeakObjectPtr will report non-valid during GC (even if the object it's pointing to isn't being
+		// GC'd).
+		const bool bEvenIfPendingKill = false;
+		const bool bThreadsafeTest = true;
+		const bool bValid = this->Renderer.IsValid(bEvenIfPendingKill, bThreadsafeTest);
+		if (ensureMsgf(bValid, TEXT("Renderer was garbage collected while outstanding tasks existed, outstanding tasks were not flushed properly during shutdown!")))
+		{
+			// The regular Get() will fail during GC so we use the above check to see if it's valid
+			// before ignoring it and directly getting the object.
+			this->Renderer.GetEvenIfUnreachable()->AddOutstandingRenderTask_AnyThread(Task);
+		}
 	};
 
 	FRenderTarget* RenderTarget = InCanvas.GetRenderTarget();
@@ -593,4 +612,4 @@ namespace UE::MovieGraph
 			SamplePixelData.Reset();
 		}
 	}
-}
\ No newline at end of file
+}
diff --git a/Engine/Plugins/MovieScene/MovieRenderPipeline/Source/MovieRenderPipelineCore/Private/Graph/Nodes/MovieGraphImageSequenceOutputNode.cpp b/Engine/Plugins/MovieScene/MovieRenderPipeline/Source/MovieRenderPipelineCore/Private/Graph/Nodes/MovieGraphImageSequenceOutputNode.cpp
--- a/Engine/Plugins/MovieScene/MovieRenderPipeline/Source/MovieRenderPipelineCore/Private/Graph/Nodes/MovieGraphImageSequenceOutputNode.cpp
+++ b/Engine/Plugins/MovieScene/MovieRenderPipeline/Source/MovieRenderPipelineCore/Private/Graph/Nodes/MovieGraphImageSequenceOutputNode.cpp
@@ -1,8 +1,11 @@
 // Copyright Epic Games, Inc. All Rights Reserved.
 
 #include "Graph/Nodes/MovieGraphImageSequenceOutputNode.h"
+#include "Graph/Nodes/MovieGraphOutputSettingNode.h"
+#include "Graph/Nodes/MovieGraphRenderLayerNode.h"
 #include "Graph/MovieGraphDataTypes.h"
 #include "Graph/MovieGraphPipeline.h"
+#include "Graph/MovieGraphConfig.h"
 #include "Modules/ModuleManager.h"
 #include "ImageWriteQueue.h"
 #include "Misc/Paths.h"
@@ -24,9 +27,10 @@ bool UMovieGraphImageSequenceOutputNode::IsFinishedWritingToDiskImpl() const
 	return Super::IsFinishedWritingToDiskImpl() && (!FinalizeFence.IsValid() || FinalizeFence.WaitFor(0));
 }
 
-void UMovieGraphImageSequenceOutputNode::OnReceiveImageDataImpl(UMovieGraphPipeline* InPipeline, UE::MovieGraph::FMovieGraphOutputMergerFrame* InRawFrameData)
+void UMovieGraphImageSequenceOutputNode::OnReceiveImageDataImpl(UMovieGraphPipeline* InPipeline, UE::MovieGraph::FMovieGraphOutputMergerFrame* InRawFrameData, const TSet<FMovieGraphRenderDataIdentifier>& InMask)
 {
 	check(InRawFrameData);
+
 	// ToDo:
 	// The ImageWriteQueue is set up in a fire-and-forget manner. This means that the data needs to be placed in the WriteQueue
 	// as a TUniquePtr (so it can free the data when its done). Unfortunately we can have multiple output formats at once,
@@ -42,11 +46,55 @@ void UMovieGraphImageSequenceOutputNode::OnReceiveImageDataImpl(UMovieGraphPipel
 	// The base ImageSequenceOutputNode doesn't support any multilayer formats, so we write out each render pass separately.
 	for (TPair<FMovieGraphRenderDataIdentifier, TUniquePtr<FImagePixelData>>& RenderData : InRawFrameData->ImageOutputData)
 	{
+		// A layer within this output data may have chosen to not be written to disk by this CDO node
+		if (!InMask.Contains(RenderData.Key))
+		{
+			continue;
+		}
+
 		UE::MovieGraph::FMovieGraphSampleState* Payload = RenderData.Value->GetPayload<UE::MovieGraph::FMovieGraphSampleState>();
 
-		// ToDo: Real file paths.
+		const bool bIncludeCDOs = true;
+		UMovieGraphOutputSettingNode* OutputSettingNode = InRawFrameData->EvaluatedConfig->GetSettingForBranch<UMovieGraphOutputSettingNode>(RenderData.Key.RootBranchName, bIncludeCDOs);
+		if (!ensure(OutputSettingNode))
+		{
+			continue;
+		}
+
+		FString RenderLayerName = RenderData.Key.RootBranchName.ToString();
+		UMovieGraphRenderLayerNode* RenderLayerNode = InRawFrameData->EvaluatedConfig->GetSettingForBranch<UMovieGraphRenderLayerNode>(RenderData.Key.RootBranchName, bIncludeCDOs);
+		if (RenderLayerNode)
+		{
+			RenderLayerName = RenderLayerNode->GetRenderLayerName();
+		}
+		// ToDo: Certain images may require transparency, at which point
+		// we write out a .png instead of a .jpeg.
+		EImageFormat PreferredOutputFormat = OutputFormat;
+
+		const TCHAR* Extension = TEXT("");
+		switch (PreferredOutputFormat)
+		{
+		case EImageFormat::PNG: Extension = TEXT("png"); break;
+		case EImageFormat::JPEG: Extension = TEXT("jpeg"); break;
+		case EImageFormat::BMP: Extension = TEXT("bmp"); break;
+		case EImageFormat::EXR: Extension = TEXT("exr"); break;
+		}
+
+		// Generate one string that puts the directory combined with the filename format.
+		FString FileNameFormatString = OutputSettingNode->OutputDirectory.Path / OutputSettingNode->FileNameFormat;
+
+		// ToDo: Validate the string, ie: ensure it has {render_pass} in there somewhere there are multiple render passes
+		// in the output data, include {camera_name} if there are multiple cameras for that render pass, etc.
+		FileNameFormatString += TEXT(".{ext}");
+
+		// Map the .ext to be specific to our output data.
+		TMap<FString, FString> AdditionalFormatArgs;
+		AdditionalFormatArgs.Add(TEXT("ext"), Extension);
+
+
+
 		const FString ResolvedProjectDir = FPaths::ConvertRelativePathToFull(FPaths::ProjectDir());
-		FString FileName = FString::Printf(TEXT("%s/Saved/MovieRenders/%d.jpg"), *ResolvedProjectDir, Payload->TraversalContext.Time.OutputFrameNumber);
+		FString FileName = FString::Printf(TEXT("%s/Saved/MovieRenders/%s_%d.jpg"), *ResolvedProjectDir, *RenderLayerName, Payload->TraversalContext.Time.OutputFrameNumber);
 
 		TUniquePtr<FImageWriteTask> TileImageTask = MakeUnique<FImageWriteTask>();
 		TileImageTask->Format = OutputFormat;
diff --git a/Engine/Plugins/MovieScene/MovieRenderPipeline/Source/MovieRenderPipelineCore/Private/Graph/Nodes/MovieGraphOutputSettingNode.cpp b/Engine/Plugins/MovieScene/MovieRenderPipeline/Source/MovieRenderPipelineCore/Private/Graph/Nodes/MovieGraphOutputSettingNode.cpp
new file mode 100644
--- /dev/null
+++ b/Engine/Plugins/MovieScene/MovieRenderPipeline/Source/MovieRenderPipelineCore/Private/Graph/Nodes/MovieGraphOutputSettingNode.cpp
@@ -0,0 +1,13 @@
+// Copyright Epic Games, Inc. All Rights Reserved.
+
+#include "Graph/Nodes/MovieGraphOutputSettingNode.h"
+
+UMovieGraphOutputSettingNode::UMovieGraphOutputSettingNode()
+	: OutputResolution(FIntPoint(1920, 1080))
+	, OutputFrameRate(FFrameRate(24, 1))
+	, bOverwriteExistingOutput(true)
+{
+
+	FileNameFormat = TEXT("{sequence_name}.{frame_number}");
+	OutputDirectory.Path = TEXT("{project_dir}/Saved/MovieRenders/");
+}
diff --git a/Engine/Plugins/MovieScene/MovieRenderPipeline/Source/MovieRenderPipelineCore/Private/Graph/Nodes/MovieGraphRenderLayerNode.cpp b/Engine/Plugins/MovieScene/MovieRenderPipeline/Source/MovieRenderPipelineCore/Private/Graph/Nodes/MovieGraphRenderLayerNode.cpp
--- a/Engine/Plugins/MovieScene/MovieRenderPipeline/Source/MovieRenderPipelineCore/Private/Graph/Nodes/MovieGraphRenderLayerNode.cpp
+++ b/Engine/Plugins/MovieScene/MovieRenderPipeline/Source/MovieRenderPipelineCore/Private/Graph/Nodes/MovieGraphRenderLayerNode.cpp
@@ -1,4 +1,4 @@
-﻿// Copyright Epic Games, Inc. All Rights Reserved.
+// Copyright Epic Games, Inc. All Rights Reserved.
 
 #include "Graph/Nodes/MovieGraphRenderLayerNode.h"
 
@@ -10,20 +10,6 @@ UMovieGraphRenderLayerNode::UMovieGraphRenderLayerNode()
 	LayerName = TEXT("beauty");
 }
 
-TArray<FMovieGraphPinProperties> UMovieGraphRenderLayerNode::GetInputPinProperties() const
-{
-	TArray<FMovieGraphPinProperties> Properties;
-	Properties.Add(FMovieGraphPinProperties::MakeBranchProperties());
-	return Properties;
-}
-
-TArray<FMovieGraphPinProperties> UMovieGraphRenderLayerNode::GetOutputPinProperties() const
-{
-	TArray<FMovieGraphPinProperties> Properties;
-	Properties.Add(FMovieGraphPinProperties::MakeBranchProperties());
-	return Properties;
-}
-
 TArray<FPropertyBagPropertyDesc> UMovieGraphRenderLayerNode::GetDynamicPropertyDescriptions() const
 {
 	TArray<FPropertyBagPropertyDesc> Properties;
diff --git a/Engine/Plugins/MovieScene/MovieRenderPipeline/Source/MovieRenderPipelineCore/Public/Graph/MovieGraphBlueprintLibrary.h b/Engine/Plugins/MovieScene/MovieRenderPipeline/Source/MovieRenderPipelineCore/Public/Graph/MovieGraphBlueprintLibrary.h
new file mode 100644
--- /dev/null
+++ b/Engine/Plugins/MovieScene/MovieRenderPipeline/Source/MovieRenderPipelineCore/Public/Graph/MovieGraphBlueprintLibrary.h
@@ -0,0 +1,29 @@
+// Copyright Epic Games, Inc. All Rights Reserved.
+
+#pragma once
+
+#include "Kismet/BlueprintFunctionLibrary.h"
+#include "Misc/FrameRate.h"
+
+#include "MovieGraphBlueprintLibrary.generated.h"
+
+// Forward Declare
+class UMovieGraphOutputSettingNode;
+
+UCLASS(meta = (ScriptName = "MovieGraphLibrary"))
+class MOVIERENDERPIPELINECORE_API UMovieGraphBlueprintLibrary : public UBlueprintFunctionLibrary
+{
+	GENERATED_BODY()
+
+public:
+	/**
+	* If InNode is valid, inspects the provided OutputsettingNode to determine if it wants to override the
+	* Frame Rate, and if so, returns the overwritten frame rate. If nullptr, or it does not have the
+	* bOverride_bUseCustomFrameRate flag set, then InDefaultrate is returned.
+	* @param	InNode			- Optional, setting to inspect for a custom framerate.
+	* @param	InDefaultRate	- The frame rate to use if the node is nullptr or doesn't want to override the rate.
+	* @return					- The effective frame rate (taking into account the node's desire to override it). 
+	*/
+	UFUNCTION(BlueprintCallable, Category = "Movie Graph")
+	static FFrameRate GetEffectiveFrameRate(UMovieGraphOutputSettingNode* InNode, const FFrameRate& InDefaultRate);
+};
\ No newline at end of file
diff --git a/Engine/Plugins/MovieScene/MovieRenderPipeline/Source/MovieRenderPipelineCore/Public/Graph/MovieGraphConfig.h b/Engine/Plugins/MovieScene/MovieRenderPipeline/Source/MovieRenderPipelineCore/Public/Graph/MovieGraphConfig.h
--- a/Engine/Plugins/MovieScene/MovieRenderPipeline/Source/MovieRenderPipelineCore/Public/Graph/MovieGraphConfig.h
+++ b/Engine/Plugins/MovieScene/MovieRenderPipeline/Source/MovieRenderPipelineCore/Public/Graph/MovieGraphConfig.h
@@ -207,23 +207,6 @@ private:
 	TArray<TObjectPtr<UMovieGraphNode>> NodeInstances;
 };
 
-/**
-* An evaluated config for the current frame. Each named branch (including Globals) has its own
-* copy of the config, fully resolved (so there is no need to check the Globals branch when
-* looking at a named branch). You can use the functions to fetch a node by type from a given
-* branch and it will return the right object (or the CDO if the node is NOT in the config).
-*/
-UCLASS()
-class MOVIERENDERPIPELINECORE_API UMovieGraphEvaluatedConfig : public UObject
-{
-	GENERATED_BODY()
-
-public:
-	/** Mapping between named branches (at the root of the config) and their evaluated values. */
-	UPROPERTY(Transient)
-	TMap<FName, FMovieGraphEvaluatedBranchConfig> BranchConfigMapping;
-};
-
 /**
 * This stores short-term information needed during traversal of the graph
 * such as disabled nodes, already visited nodes, etc. This information is
@@ -262,6 +245,73 @@ public:
 	TArray<TObjectPtr<const UMovieGraphSubgraphNode>> SubgraphStack;
 };
 
+/**
+* An evaluated config for the current frame. Each named branch (including Globals) has its own
+* copy of the config, fully resolved (so there is no need to check the Globals branch when
+* looking at a named branch). You can use the functions to fetch a node by type from a given
+* branch and it will return the right object (or the CDO if the node is NOT in the config).
+*/
+UCLASS()
+class MOVIERENDERPIPELINECORE_API UMovieGraphEvaluatedConfig : public UObject
+{
+	GENERATED_BODY()
+public:
+
+	const TArray<FName> GetBranchNames() const
+	{
+		TArray<FName> OutKeys;
+		BranchConfigMapping.GenerateKeyArray(OutKeys);
+ 		return OutKeys;
+	}
+
+	template<typename NodeType>
+	TObjectPtr<NodeType> GetSettingForBranch(const FName InBranchName, bool bIncludeCDOs = true, bool bExactMatch = false)
+	{
+		TArray<TObjectPtr<NodeType>> AllSettings = GetSettingsForBranch<NodeType>(InBranchName, bIncludeCDOs, bExactMatch);
+		if (AllSettings.Num() > 0)
+		{
+			return AllSettings[0];
+		}
+
+		return nullptr;
+	}
+
+	template<typename NodeType>
+	TArray<TObjectPtr<NodeType>> GetSettingsForBranch(const FName InBranchName, bool bIncludeCDOs = true, bool bExactMatch = false)
+	{
+		TArray<TObjectPtr<NodeType>> ResultNodes;
+		FMovieGraphEvaluatedBranchConfig* BranchConfig = BranchConfigMapping.Find(InBranchName);
+		ensureMsgf(BranchConfig, TEXT("Failed to find branch mapping for Branch: %s"), *InBranchName.ToString());
+
+		if (BranchConfig)
+		{
+			// Check to see if the branch has an instance of this.
+			for (const TObjectPtr<UMovieGraphNode>& Node : BranchConfig->GetNodes())
+			{
+				bool bMatches = bExactMatch ? Node->GetClass() == NodeType::StaticClass() : Node->IsA<NodeType>();
+				if (bMatches)
+				{
+					ResultNodes.Add(Cast<NodeType>(Node.Get()));
+				}
+			}
+		}
+
+		// If they didn't return above, then either they specified an invalid branch (for which the ensure tripped)
+		// or the config simply didn't override that setting class, at which point we might try to return a default
+		if (bIncludeCDOs)
+		{
+			ResultNodes.Add(GetMutableDefault<NodeType>());
+		}
+
+		return ResultNodes;
+	}
+	
+
+public:
+	/** Mapping between named branches (at the root of the config) and their evaluated values. */
+	UPROPERTY(Transient)
+	TMap<FName, FMovieGraphEvaluatedBranchConfig> BranchConfigMapping;
+};
 
 /**
 * This is the runtime representation of the UMoviePipelineEdGraph which contains the actual strongly
@@ -327,9 +377,6 @@ public:
 	/** Remove the specified member (input, output, variable) from the graph. */
 	bool DeleteMember(UMovieGraphMember* MemberToDelete);
 
-	/** Returns only the names of the root branches in the Output Node, with no depth information. */
-	TArray<FMovieGraphBranch> GetOutputBranches() const;
-
 #if WITH_EDITOR
 	/** Gets the editor-only nodes in this graph. Editor-only nodes do not have an equivalent runtime node. */
 	const TArray<TObjectPtr<UObject>>& GetEditorOnlyNodes() const { return EditorOnlyNodes; }
@@ -338,42 +385,11 @@ public:
 	void SetEditorOnlyNodes(const TArray<TObjectPtr<const UObject>>& InNodes);
 #endif
 
+	/** Given a user-defined evaluation context, evaluate the graph and build a "flattened" list of settings for each branch discovered. */
 	UFUNCTION(BlueprintCallable, Category="Experimental")
 	UMovieGraphEvaluatedConfig* CreateFlattenedGraph(const FMovieGraphTraversalContext& InContext);
 
-	template<typename NodeType>
-	static NodeType* IterateGraphForClass(const FMovieGraphTraversalContext& InContext)
-	{
-		TArray<NodeType*> AllSettings = IterateGraphForClassAll<NodeType>(InContext);
-		if (AllSettings.Num() > 0)
-		{
-			return AllSettings[0];
-		}
-		return nullptr;
-	}
-	
-	template<typename NodeType>
-	static TArray<NodeType*> IterateGraphForClassAll(const FMovieGraphTraversalContext& InContext)
-	{
-		TArray<NodeType*> TypedNodes;
-		if (!ensureMsgf(InContext.RootGraph, TEXT("You must specify a RootGraph to traverse with")))
-		{
-			return TypedNodes;
-		}
-
-		const TArray<UMovieGraphNode*> FoundNodes = InContext.RootGraph->TraverseGraph(NodeType::StaticClass(), InContext);
-		for (UMovieGraphNode* Node : FoundNodes)
-		{
-			TypedNodes.Add(CastChecked<NodeType>(Node));
-		}
-
-		return TypedNodes;
-	}
-
-	TArray<UMovieGraphNode*> TraverseGraph(TSubclassOf<UMovieGraphNode> InClassType, const FMovieGraphTraversalContext& InContext) const;
-
 protected:
-	void TraverseGraphRecursive(UMovieGraphNode* InNode, TSubclassOf<UMovieGraphNode> InClassType, const FMovieGraphTraversalContext& InContext, TArray<UMovieGraphNode*>& OutNodes) const;
 	/** Copies properties in FromNode that are marked for override into ToNode, but only if ToNode doesn't already override that value. */
 	void CopyOverriddenProperties(UMovieGraphNode* FromNode, UMovieGraphNode* ToNode);
 	/** Given a property you want to override (passes IsPropertyOverrideable) look for a matching bOverride_ named bool property. */
diff --git a/Engine/Plugins/MovieScene/MovieRenderPipeline/Source/MovieRenderPipelineCore/Public/Graph/MovieGraphDataTypes.h b/Engine/Plugins/MovieScene/MovieRenderPipeline/Source/MovieRenderPipelineCore/Public/Graph/MovieGraphDataTypes.h
--- a/Engine/Plugins/MovieScene/MovieRenderPipeline/Source/MovieRenderPipelineCore/Public/Graph/MovieGraphDataTypes.h
+++ b/Engine/Plugins/MovieScene/MovieRenderPipeline/Source/MovieRenderPipelineCore/Public/Graph/MovieGraphDataTypes.h
@@ -4,6 +4,7 @@
 #include "Graph/MovieGraphTraversalContext.h"
 #include "Graph/MovieGraphRenderDataIdentifier.h"
 #include "UObject/Object.h"
+#include "UObject/StrongObjectPtr.h"
 #include "Templates/SubclassOf.h"
 #include "ImagePixelData.h"
 #include "Containers/Queue.h"
@@ -12,6 +13,7 @@
 
 // Forward Declares
 class UMovieGraphPipeline;
+class UMovieGraphEvaluatedConfig;
 class UMovieGraphTimeStepBase;
 class UMovieGraphRendererBase;
 class UMovieGraphTimeRangeBuilderBase;
@@ -162,7 +164,7 @@ public:
 // ToDo: Both of these can probably go into the Default Renderer implementation.
 struct FMovieGraphRenderPassLayerData
 {
-	FString LayerName;
+	FName BranchName;
 	FGuid CameraIdentifier;
 };
 // ToDo: Both of these can probably go into the Default Renderer implementation.
@@ -251,6 +253,11 @@ namespace UE::MovieGraph
 	{
 		FMovieGraphOutputMergerFrame() = default;
 		
+		/**
+		* The configuraton file that was being used by this frame during submission. We store it here so that it doesn't have to go through the pipeline
+		* (which is difficult becuase it's a Garbage Collected object), and the Output Merger is the "join" step before we go back onto the Game Thread.
+		*/
+		TStrongObjectPtr<UMovieGraphEvaluatedConfig> EvaluatedConfig;
 
 		/** 
 		* Traversal context for the Movie Pipeline that isn't layer specific. Each image layer will have
diff --git a/Engine/Plugins/MovieScene/MovieRenderPipeline/Source/MovieRenderPipelineCore/Public/Graph/MovieGraphLinearTimeStep.h b/Engine/Plugins/MovieScene/MovieRenderPipeline/Source/MovieRenderPipelineCore/Public/Graph/MovieGraphLinearTimeStep.h
--- a/Engine/Plugins/MovieScene/MovieRenderPipeline/Source/MovieRenderPipelineCore/Public/Graph/MovieGraphLinearTimeStep.h
+++ b/Engine/Plugins/MovieScene/MovieRenderPipeline/Source/MovieRenderPipelineCore/Public/Graph/MovieGraphLinearTimeStep.h
@@ -2,9 +2,11 @@
 #pragma once
 
 #include "Graph/MovieGraphDataTypes.h"
+#include "Graph/MovieGraphConfig.h"
 #include "Misc/FrameRate.h"
 #include "Misc/FrameTime.h"
 #include "Engine/EngineCustomTimeStep.h"
+#include "UObject/StrongObjectPtr.h"
 #include "MovieGraphLinearTimeStep.generated.h"
 
 // Forward Declares
@@ -43,6 +45,7 @@ protected:
 
 protected:
 	/** This is the output data needed by the rest of MRQ to produce a frame. */
+	UPROPERTY(Transient)
 	FMovieGraphTimeStepData CurrentTimeStepData;
 
 	struct FOutputFrameMetrics
@@ -121,6 +124,9 @@ protected:
 
 		/** A pre-calculated (when the output frame starts) range of time, breaking RangeShutterOpen into smaller chunks, one per temporal sub-sample. */
 		TArray<TRange<FFrameTime>> TemporalRanges;
+
+		/** Updated at the start of each Temporal Sample. */
+		TStrongObjectPtr<UMovieGraphEvaluatedConfig> EvaluatedConfig;
 	};
 
 	/**
diff --git a/Engine/Plugins/MovieScene/MovieRenderPipeline/Source/MovieRenderPipelineCore/Public/Graph/MovieGraphNode.h b/Engine/Plugins/MovieScene/MovieRenderPipeline/Source/MovieRenderPipelineCore/Public/Graph/MovieGraphNode.h
--- a/Engine/Plugins/MovieScene/MovieRenderPipeline/Source/MovieRenderPipelineCore/Public/Graph/MovieGraphNode.h
+++ b/Engine/Plugins/MovieScene/MovieRenderPipeline/Source/MovieRenderPipelineCore/Public/Graph/MovieGraphNode.h
@@ -39,12 +39,12 @@ class MOVIERENDERPIPELINECORE_API UMovieGraphNode : public UObject
 {
 	GENERATED_BODY()
 
-	static FName GlobalsPinName;
 
 	friend class UMovieGraphConfig;
 	friend class UMovieGraphEdge;
 	
 public:
+	static FName GlobalsPinName;
 	UMovieGraphNode();
 
 	const TArray<TObjectPtr<UMovieGraphPin>>& GetInputPins() const { return InputPins; }
@@ -193,7 +193,7 @@ protected:
 /**
 * Nodes representing user settings should derive from this. This is the only node type copied into flattened eval.
 */
-UCLASS(Abstract)
+UCLASS(Abstract, BlueprintType)
 class MOVIERENDERPIPELINECORE_API UMovieGraphSettingNode : public UMovieGraphNode
 {
 	GENERATED_BODY()
diff --git a/Engine/Plugins/MovieScene/MovieRenderPipeline/Source/MovieRenderPipelineCore/Public/Graph/MovieGraphRenderDataIdentifier.h b/Engine/Plugins/MovieScene/MovieRenderPipeline/Source/MovieRenderPipelineCore/Public/Graph/MovieGraphRenderDataIdentifier.h
--- a/Engine/Plugins/MovieScene/MovieRenderPipeline/Source/MovieRenderPipelineCore/Public/Graph/MovieGraphRenderDataIdentifier.h
+++ b/Engine/Plugins/MovieScene/MovieRenderPipeline/Source/MovieRenderPipelineCore/Public/Graph/MovieGraphRenderDataIdentifier.h
@@ -17,9 +17,9 @@ struct FMovieGraphRenderDataIdentifier
 	FMovieGraphRenderDataIdentifier()
 	{}
 
-	FMovieGraphRenderDataIdentifier(const FString& InRenderLayerName, const FString& InRendererName,
+	FMovieGraphRenderDataIdentifier(const FName& InRootBranchName, const FString& InRendererName,
 		const FString& InSubRenderResourceName, const FString& InCameraName)
-		: RenderLayerName(InRenderLayerName)
+		: RootBranchName(InRootBranchName)
 		, RendererName(InRendererName)
 		, SubResourceName(InSubRenderResourceName)
 		, CameraName(InCameraName)
@@ -28,7 +28,7 @@ struct FMovieGraphRenderDataIdentifier
 
 	bool operator == (const FMovieGraphRenderDataIdentifier& InRHS) const
 	{
-		return RenderLayerName == InRHS.RenderLayerName && 
+		return RootBranchName == InRHS.RootBranchName && 
 			RendererName == InRHS.RendererName &&
 			SubResourceName == InRHS.SubResourceName &&
 			CameraName == InRHS.CameraName;
@@ -41,7 +41,7 @@ struct FMovieGraphRenderDataIdentifier
 
 	friend uint32 GetTypeHash(FMovieGraphRenderDataIdentifier InIdentifier)
 	{
-		return HashCombineFast(GetTypeHash(InIdentifier.RenderLayerName),
+		return HashCombineFast(GetTypeHash(InIdentifier.RootBranchName),
 			HashCombineFast(GetTypeHash(InIdentifier.RendererName),
 				HashCombineFast(GetTypeHash(InIdentifier.SubResourceName),
 					GetTypeHash(InIdentifier.CameraName))));
@@ -49,13 +49,19 @@ struct FMovieGraphRenderDataIdentifier
 
 	friend FString LexToString(const FMovieGraphRenderDataIdentifier InIdentifier)
 	{
-		return FString::Printf(TEXT("RenderLayer: %s Renderer:%s SubResource: %s Camera: %s"), *InIdentifier.RenderLayerName, *InIdentifier.RendererName, *InIdentifier.SubResourceName, *InIdentifier.CameraName);
+		return FString::Printf(TEXT("RootBranch: %s Renderer:%s SubResource: %s Camera: %s"), *InIdentifier.RootBranchName.ToString(), *InIdentifier.RendererName, *InIdentifier.SubResourceName, *InIdentifier.CameraName);
 	}
 
 public:
-	/** The user provided name for the whole Render Layer ("character", "background", etc.) */
+	/** 
+	* The root branch name that this render layer exists on. Actual display name comes from a UMovieGraphRenderLayerNode (if found in the branch)
+	* otherwise it falls back to just displaying RootBranchName. All of our internal lookups for branches are done based on the path we followed 
+	* out from the root so that we can handle overriding a render layer display name via regular node overriding.
+	* 
+	* Could have a value like "Globals" or a user-provided one "character", "background", etc.
+	*/
 	UPROPERTY(EditAnywhere, BlueprintReadWrite, Category = "Movie Pipeline")
-	FString RenderLayerName;
+	FName RootBranchName;
 
 	/** Which renderer was used to produce this image ("panoramic" "deferred" "path tracer", etc.) */
 	UPROPERTY(EditAnywhere, BlueprintReadWrite, Category = "Movie Pipeline")
diff --git a/Engine/Plugins/MovieScene/MovieRenderPipeline/Source/MovieRenderPipelineCore/Public/Graph/MovieGraphTimeStepData.h b/Engine/Plugins/MovieScene/MovieRenderPipeline/Source/MovieRenderPipelineCore/Public/Graph/MovieGraphTimeStepData.h
--- a/Engine/Plugins/MovieScene/MovieRenderPipeline/Source/MovieRenderPipelineCore/Public/Graph/MovieGraphTimeStepData.h
+++ b/Engine/Plugins/MovieScene/MovieRenderPipeline/Source/MovieRenderPipelineCore/Public/Graph/MovieGraphTimeStepData.h
@@ -4,6 +4,8 @@
 #include "CoreMinimal.h"
 #include "MovieGraphTimeStepData.generated.h"
 
+class UMovieGraphEvaluatedConfig;
+
 /** 
 * This data structure needs to be filled out each frame by the UMovieGraphTimeStepBase,
 * which will eventually be passed to the renderer. It controls per-sample behavior such
@@ -63,4 +65,14 @@ public:
 	*/
 	UPROPERTY(BlueprintReadWrite, EditAnywhere, Category = "Movie Graph")
 	bool bRequiresAccumulator;
+
+	/**
+	* The evaluated config holds the configuration used for this given frame. This pointer
+	* can potentially change each frame (if the graph for that frame is different) but
+	* users can rely on the EvaluatedConfig being correct for a given frame, thus all
+	* resolves (such as filenames) should use the config for that frame, not the latest
+	* one available.
+	*/
+	UPROPERTY(BlueprintReadWrite, EditAnywhere, Category = "Movie Graph")
+	TObjectPtr<UMovieGraphEvaluatedConfig> EvaluatedConfig;
 };
diff --git a/Engine/Plugins/MovieScene/MovieRenderPipeline/Source/MovieRenderPipelineCore/Public/Graph/MovieGraphTraversalContext.h b/Engine/Plugins/MovieScene/MovieRenderPipeline/Source/MovieRenderPipelineCore/Public/Graph/MovieGraphTraversalContext.h
--- a/Engine/Plugins/MovieScene/MovieRenderPipeline/Source/MovieRenderPipelineCore/Public/Graph/MovieGraphTraversalContext.h
+++ b/Engine/Plugins/MovieScene/MovieRenderPipeline/Source/MovieRenderPipelineCore/Public/Graph/MovieGraphTraversalContext.h
@@ -34,7 +34,7 @@ public:
 	{}
 
 	UPROPERTY(BlueprintReadWrite, EditAnywhere, Category = "MovieGraph")
-	FMovieGraphBranch RootBranch;
+	FName RootBranch;
 
 	/** Which shot (out of ShotCount) is this time step for? */
 	UPROPERTY(BlueprintReadWrite, EditAnywhere, Category = "Movie Graph")
diff --git a/Engine/Plugins/MovieScene/MovieRenderPipeline/Source/MovieRenderPipelineCore/Public/Graph/Nodes/MovieGraphFileOutputNode.h b/Engine/Plugins/MovieScene/MovieRenderPipeline/Source/MovieRenderPipelineCore/Public/Graph/Nodes/MovieGraphFileOutputNode.h
--- a/Engine/Plugins/MovieScene/MovieRenderPipeline/Source/MovieRenderPipelineCore/Public/Graph/Nodes/MovieGraphFileOutputNode.h
+++ b/Engine/Plugins/MovieScene/MovieRenderPipeline/Source/MovieRenderPipelineCore/Public/Graph/Nodes/MovieGraphFileOutputNode.h
@@ -14,7 +14,7 @@
 * multiple output formats at once (such as exr and jpeg).
 */
 UCLASS(Abstract)
-class UMovieGraphFileOutputNode : public UMovieGraphNode
+class UMovieGraphFileOutputNode : public UMovieGraphSettingNode
 {
 	GENERATED_BODY()
 public:
@@ -22,13 +22,11 @@ public:
 	
 	// UMovieGraphFileOutputNode Interface
 	/** 
-	* This is called by the Game Thread containing the final output data for a given output frame. If you need to hold 
-	* onto this data for longer than the duration of the single function call (such as async/task graph operations) then
-	* you need to increment the reference count on the data in the UMovieGraphPipeline, and decrement when you are done.
-	* Failing to increment will lead to the data being deleted out from underneath you, and failing to decrement will leak 
-	* memory until the end of a render.
+	* This is called by the Game Thread containing the final output data for a given output frame. InRawFrameData contains
+	* a list of all of the data generated for this output frame, and InMask specifies which render layers within InRawFrameData
+	* should actually be written to disk by this output node.
 	*/
-	void OnReceiveImageData(UMovieGraphPipeline* InPipeline, UE::MovieGraph::FMovieGraphOutputMergerFrame* InRawFrameData) { OnReceiveImageDataImpl(InPipeline, InRawFrameData); }
+	void OnReceiveImageData(UMovieGraphPipeline* InPipeline, UE::MovieGraph::FMovieGraphOutputMergerFrame* InRawFrameData, const TSet<FMovieGraphRenderDataIdentifier>& InMask) { OnReceiveImageDataImpl(InPipeline, InRawFrameData, InMask); }
 	void OnAllFramesSubmitted() { OnAllFramesSubmittedImpl(); }
 	bool IsFinishedWritingToDisk() const { return IsFinishedWritingToDiskImpl(); }
 	// ~UMovieGraphFileOutputNode Interface
@@ -55,7 +53,7 @@ public:
 #endif
 
 protected:
-	virtual void OnReceiveImageDataImpl(UMovieGraphPipeline* InPipeline, UE::MovieGraph::FMovieGraphOutputMergerFrame* InRawFrameData) {}
+	virtual void OnReceiveImageDataImpl(UMovieGraphPipeline* InPipeline, UE::MovieGraph::FMovieGraphOutputMergerFrame* InRawFrameData, const TSet<FMovieGraphRenderDataIdentifier>& InMask) {}
 	virtual void OnAllFramesSubmittedImpl() {}
 	virtual bool IsFinishedWritingToDiskImpl() const { return true; }
 };
\ No newline at end of file
diff --git a/Engine/Plugins/MovieScene/MovieRenderPipeline/Source/MovieRenderPipelineCore/Public/Graph/Nodes/MovieGraphImageSequenceOutputNode.h b/Engine/Plugins/MovieScene/MovieRenderPipeline/Source/MovieRenderPipelineCore/Public/Graph/Nodes/MovieGraphImageSequenceOutputNode.h
--- a/Engine/Plugins/MovieScene/MovieRenderPipeline/Source/MovieRenderPipelineCore/Public/Graph/Nodes/MovieGraphImageSequenceOutputNode.h
+++ b/Engine/Plugins/MovieScene/MovieRenderPipeline/Source/MovieRenderPipelineCore/Public/Graph/Nodes/MovieGraphImageSequenceOutputNode.h
@@ -23,7 +23,7 @@ public:
 	UMovieGraphImageSequenceOutputNode();
 	
 	// UMovieGraphFileOutputNode Interface
-	virtual void OnReceiveImageDataImpl(UMovieGraphPipeline* InPipeline, UE::MovieGraph::FMovieGraphOutputMergerFrame* InRawFrameData) override;
+	virtual void OnReceiveImageDataImpl(UMovieGraphPipeline* InPipeline, UE::MovieGraph::FMovieGraphOutputMergerFrame* InRawFrameData, const TSet<FMovieGraphRenderDataIdentifier>& InMask) override;
 	virtual void OnAllFramesSubmittedImpl() override;
 	virtual bool IsFinishedWritingToDiskImpl() const override;
 	// ~UMovieGraphFileOutputNode Interface
diff --git a/Engine/Plugins/MovieScene/MovieRenderPipeline/Source/MovieRenderPipelineCore/Public/Graph/Nodes/MovieGraphOutputSettingNode.h b/Engine/Plugins/MovieScene/MovieRenderPipeline/Source/MovieRenderPipelineCore/Public/Graph/Nodes/MovieGraphOutputSettingNode.h
new file mode 100644
--- /dev/null
+++ b/Engine/Plugins/MovieScene/MovieRenderPipeline/Source/MovieRenderPipelineCore/Public/Graph/Nodes/MovieGraphOutputSettingNode.h
@@ -0,0 +1,79 @@
+// Copyright Epic Games, Inc. All Rights Reserved.
+
+#pragma once
+#include "Graph/MovieGraphNode.h"
+#include "Misc/FrameRate.h"
+#include "Styling/AppStyle.h"
+#include "MovieGraphOutputSettingNode.generated.h"
+
+
+UCLASS()
+class MOVIERENDERPIPELINECORE_API UMovieGraphOutputSettingNode : public UMovieGraphSettingNode
+{
+	GENERATED_BODY()
+public:
+	UMovieGraphOutputSettingNode();
+
+#if WITH_EDITOR
+	virtual FText GetNodeTitle(const bool bGetDescriptive) const override
+	{
+		static const FText OutputSettingsNodeName = NSLOCTEXT("MoviePipelineGraph", "NodeName_OutputSettings", "Output Settings");
+		return OutputSettingsNodeName;
+	}
+
+	virtual FText GetMenuCategory() const override
+	{
+		return NSLOCTEXT("MoviePipelineGraph", "Settings_Category", "Settings");
+	}
+
+	virtual FLinearColor GetNodeTitleColor() const override
+	{
+		static const FLinearColor OutputSettingsColor = FLinearColor(0.854f, 0.509f, 0.039f);
+		return OutputSettingsColor;
+	}
+
+	virtual FSlateIcon GetIconAndTint(FLinearColor& OutColor) const override
+	{
+		static const FSlateIcon SettingsIcon = FSlateIcon(FAppStyle::GetAppStyleSetName(), "Icons.Settings");
+
+		OutColor = FLinearColor::White;
+		return SettingsIcon;
+	}
+#endif
+
+public:
+	UPROPERTY(EditAnywhere, BlueprintReadWrite, Category = Overrides, meta = (InlineEditConditionToggle))
+	uint8 bOverride_OutputDirectory : 1;
+
+	UPROPERTY(EditAnywhere, BlueprintReadWrite, Category = Overrides, meta = (InlineEditConditionToggle))
+	uint8 bOverride_FileNameFormat : 1;
+
+	UPROPERTY(EditAnywhere, BlueprintReadWrite, Category = Overrides, meta = (InlineEditConditionToggle))
+	uint8 bOverride_OutputResolution : 1;
+	
+	UPROPERTY(EditAnywhere, BlueprintReadWrite, Category = Overrides, meta = (InlineEditConditionToggle))
+	uint8 bOverride_OutputFrameRate : 1;
+	
+	UPROPERTY(EditAnywhere, BlueprintReadWrite, Category = Overrides, meta = (InlineEditConditionToggle))
+	uint8 bOverride_bOverwriteExistingOutput : 1;
+
+	/** What directory should all of our output files be relative to. */
+	UPROPERTY(EditAnywhere, BlueprintReadWrite, Category = "File Output", meta=(EditCondition="bOverride_OutputDirectory"))
+	FDirectoryPath OutputDirectory;
+
+	/** What format string should the final files use? Can include folder prefixes, and format string ({shot_name}, etc.) */
+	UPROPERTY(EditAnywhere, BlueprintReadWrite, Category = "File Output", meta = (EditCondition = "bOverride_FileNameFormat"))
+	FString FileNameFormat;
+
+	/** What resolution should our output files be exported at? */
+	UPROPERTY(EditAnywhere, BlueprintReadWrite, Category = "File Output", meta = (EditCondition = "bOverride_OutputResolution"))
+	FIntPoint OutputResolution;
+	
+	/** What frame rate should the output files be exported at? This overrides the Display Rate of the target sequence. If not overwritten, uses the default Sequence Display Rate. */
+	UPROPERTY(EditAnywhere, BlueprintReadWrite, Category = "File Output", meta=(EditCondition="bOverride_OutputFrameRate"))
+	FFrameRate OutputFrameRate;
+	
+	/** If true, output containers should attempt to override any existing files with the same name. */
+	UPROPERTY(EditAnywhere, BlueprintReadWrite, Category = "File Output", meta = (EditCondition = "bOverride_bOverwriteExistingOutput"))
+	bool bOverwriteExistingOutput;
+};
diff --git a/Engine/Plugins/MovieScene/MovieRenderPipeline/Source/MovieRenderPipelineCore/Public/Graph/Nodes/MovieGraphRenderLayerNode.h b/Engine/Plugins/MovieScene/MovieRenderPipeline/Source/MovieRenderPipelineCore/Public/Graph/Nodes/MovieGraphRenderLayerNode.h
--- a/Engine/Plugins/MovieScene/MovieRenderPipeline/Source/MovieRenderPipelineCore/Public/Graph/Nodes/MovieGraphRenderLayerNode.h
+++ b/Engine/Plugins/MovieScene/MovieRenderPipeline/Source/MovieRenderPipelineCore/Public/Graph/Nodes/MovieGraphRenderLayerNode.h
@@ -6,7 +6,7 @@
 #include "MovieGraphRenderLayerNode.generated.h"
 
 UCLASS()
-class MOVIERENDERPIPELINECORE_API UMovieGraphRenderLayerNode : public UMovieGraphNode
+class MOVIERENDERPIPELINECORE_API UMovieGraphRenderLayerNode : public UMovieGraphSettingNode
 {
 	GENERATED_BODY()
 public:
@@ -14,8 +14,6 @@ public:
 
 	FString GetRenderLayerName() const { return LayerName; }
 
-	virtual TArray<FMovieGraphPinProperties> GetInputPinProperties() const override;
-	virtual TArray<FMovieGraphPinProperties> GetOutputPinProperties() const override;
 	virtual TArray<FPropertyBagPropertyDesc> GetDynamicPropertyDescriptions() const override;
 
 #if WITH_EDITOR
@@ -29,6 +27,10 @@ public:
 	//~ End UObject interface
 #endif
 
-	UPROPERTY(BlueprintReadWrite, EditAnywhere, Category = "Settings")
+public:
+	UPROPERTY(EditAnywhere, BlueprintReadWrite, Category = Overrides, meta = (InlineEditConditionToggle))
+	uint8 bOverride_LayerName : 1;
+
+	UPROPERTY(BlueprintReadWrite, EditAnywhere, Category = "Settings", meta = (EditCondition = "bOverride_LayerName"))
 	FString LayerName;
 };
\ No newline at end of file
diff --git a/Engine/Plugins/MovieScene/MovieRenderPipeline/Source/MovieRenderPipelineCore/Public/Graph/Nodes/MovieGraphRenderPassNode.h b/Engine/Plugins/MovieScene/MovieRenderPipeline/Source/MovieRenderPipelineCore/Public/Graph/Nodes/MovieGraphRenderPassNode.h
--- a/Engine/Plugins/MovieScene/MovieRenderPipeline/Source/MovieRenderPipelineCore/Public/Graph/Nodes/MovieGraphRenderPassNode.h
+++ b/Engine/Plugins/MovieScene/MovieRenderPipeline/Source/MovieRenderPipelineCore/Public/Graph/Nodes/MovieGraphRenderPassNode.h
@@ -17,7 +17,7 @@ struct FMovieGraphTimeStepData;
 * number of instances (decoupled from the number of times the node is used in the graph).
 */
 UCLASS(Abstract)
-class UMovieGraphRenderPassNode : public UMovieGraphNode
+class UMovieGraphRenderPassNode : public UMovieGraphSettingNode
 {
 	GENERATED_BODY()
 public:
@@ -43,20 +43,6 @@ public:
 	void GatherOutputPasses(TArray<FMovieGraphRenderDataIdentifier>& OutExpectedPasses) const { GatherOutputPassesImpl(OutExpectedPasses); }
 	// ~UMovieGraphRenderPassNode Interface
 
-	virtual TArray<FMovieGraphPinProperties> GetInputPinProperties() const override
-	{
-		TArray<FMovieGraphPinProperties> Properties;
-		Properties.Add(FMovieGraphPinProperties::MakeBranchProperties());
-		return Properties;
-	}
-
-	virtual TArray<FMovieGraphPinProperties> GetOutputPinProperties() const override
-	{
-		TArray<FMovieGraphPinProperties> Properties;
-		Properties.Add(FMovieGraphPinProperties::MakeBranchProperties());
-		return Properties;
-	}
-
 #if WITH_EDITOR
 	virtual FText GetMenuCategory() const override
 	{
