MoviePipeline: Resolve the output filenames from the configuration file instead of hard-coding it.

diff --git a/Engine/Plugins/MovieScene/MovieRenderPipeline/Source/MovieRenderPipelineCore/Private/Graph/MovieGraphBlueprintLibrary.cpp b/Engine/Plugins/MovieScene/MovieRenderPipeline/Source/MovieRenderPipelineCore/Private/Graph/MovieGraphBlueprintLibrary.cpp
--- a/Engine/Plugins/MovieScene/MovieRenderPipeline/Source/MovieRenderPipelineCore/Private/Graph/MovieGraphBlueprintLibrary.cpp
+++ b/Engine/Plugins/MovieScene/MovieRenderPipeline/Source/MovieRenderPipelineCore/Private/Graph/MovieGraphBlueprintLibrary.cpp
@@ -2,6 +2,8 @@
 
 #include "Graph/MovieGraphBlueprintLibrary.h"
 #include "Graph/Nodes/MovieGraphOutputSettingNode.h"
+#include "Graph/Nodes/MovieGraphRenderLayerNode.h"
+#include "MoviePipelineUtils.h"
 
 FFrameRate UMovieGraphBlueprintLibrary::GetEffectiveFrameRate(UMovieGraphOutputSettingNode* InNode, const FFrameRate& InDefaultRate)
 {
@@ -11,4 +13,188 @@ FFrameRate UMovieGraphBlueprintLibrary::GetEffectiveFrameRate(UMovieGraphOutputS
 	}
 
 	return InDefaultRate;
+}
+
+FString UMovieGraphBlueprintLibrary::ResolveFilenameFormatArguments(const FString& InFormatString, const FMovieGraphFilenameResolveParams& InParams, FMovieGraphResolveArgs& OutMergedFormatArgs)
+{
+	// There are a number of different sources for file/metadata KVPs. 
+	// UMovieGraphPipeline - these are handled as resolve params to allow resolving during edit time
+	// UMoviePipelineExecutorJob - author, sequence, level, etc.
+	// UMovieGraphSettingNode - every node can contribute it's own file/metadata KVP based on their own settings
+	
+	bool bOverwriteExisting = true;
+
+	// Data from the UMovieGraphPipeline comes from the resolve params themselves to allow resolving without a runtime instance.
+	{
+		// Zero-pad our frame numbers when we format the strings. Some programs struggle when ingesting frames that 
+		// go 1,2,3,...,10,11. To work around this issue we allow the user to specify how many zeros they want to
+		// pad the numbers with, 0001, 0002, etc. We also allow offsetting the output frame numbers, this is useful
+		// when your sequence starts at zero and you use handle frames (which would cause negative output frame 
+		// numbers), so we allow the user to add a fixed amount to all output to ensure they are positive.
+		int32 FrameNumberOffset = InParams.FrameNumberOffset;
+		FString FrameNumber = UE::MoviePipeline::GetPaddingFormatString(InParams.ZeroPadFrameNumberCount, InParams.RootFrameNumber + FrameNumberOffset); // Sequence Frame #
+		FString FrameNumberShot = UE::MoviePipeline::GetPaddingFormatString(InParams.ZeroPadFrameNumberCount, InParams.ShotFrameNumber + FrameNumberOffset); // Shot Frame #
+		FString FrameNumberRel = UE::MoviePipeline::GetPaddingFormatString(InParams.ZeroPadFrameNumberCount, InParams.RootFrameNumberRel + FrameNumberOffset); // Relative to 0
+		FString FrameNumberShotRel = UE::MoviePipeline::GetPaddingFormatString(InParams.ZeroPadFrameNumberCount, InParams.ShotFrameNumberRel + FrameNumberOffset); // Relative to 0 within the shot.
+
+		// Ensure they used relative frame numbers in the output so they get the right number of output frames.
+		if (InParams.bForceRelativeFrameNumbers)
+		{
+			FrameNumber = FrameNumberRel;
+			FrameNumberShot = FrameNumberShotRel;
+		}
+
+		FString ShotName = InParams.Shot ? InParams.Shot->OuterName : FString();
+		ShotName = ShotName.Len() > 0 ? ShotName : TEXT("NoShot");
+
+		// Use the legacy system to add filename and metadata kvp
+		MoviePipeline::GetOutputStateFormatArgs(OutMergedFormatArgs.FilenameArguments, OutMergedFormatArgs.FileMetadata, FrameNumber, FrameNumberShot, FrameNumberRel, FrameNumberShotRel, TEXT("DummyCameraToken"), ShotName);
+
+		// Look up the Render Layer display name from the evaluated config if possible.
+		FString RenderLayerName = InParams.RenderDataIdentifier.RootBranchName.ToString();
+		if (InParams.EvaluatedConfig)
+		{
+			const bool bIncludeCDOs = false;
+			TObjectPtr<UMovieGraphRenderLayerNode> RenderLayerNode = InParams.EvaluatedConfig->GetSettingForBranch<UMovieGraphRenderLayerNode>(InParams.RenderDataIdentifier.RootBranchName, bIncludeCDOs);
+			if (RenderLayerNode)
+			{
+				RenderLayerName = RenderLayerNode->GetRenderLayerName();
+			}
+
+			TObjectPtr<UMovieGraphOutputSettingNode> OutputSettingNode = InParams.EvaluatedConfig->GetSettingForBranch<UMovieGraphOutputSettingNode>(InParams.RenderDataIdentifier.RootBranchName, bIncludeCDOs);
+			if(OutputSettingNode)
+			{
+				bOverwriteExisting = OutputSettingNode->bOverwriteExistingOutput;
+			}
+		}
+
+		// Add on Render Data Identifier, overwriting the dummy camera name above.
+		OutMergedFormatArgs.FilenameArguments.Add(TEXT("camera_name"), InParams.RenderDataIdentifier.CameraName);
+		OutMergedFormatArgs.FilenameArguments.Add(TEXT("renderer_name"), InParams.RenderDataIdentifier.RendererName);
+		OutMergedFormatArgs.FilenameArguments.Add(TEXT("renderer_sub_name"), InParams.RenderDataIdentifier.SubResourceName);
+		OutMergedFormatArgs.FilenameArguments.Add(TEXT("render_layer"), RenderLayerName);
+
+		OutMergedFormatArgs.FileMetadata.Add(TEXT("unreal/cameraName"), InParams.RenderDataIdentifier.CameraName);
+		OutMergedFormatArgs.FileMetadata.Add(TEXT("unreal/rendererName"), InParams.RenderDataIdentifier.RendererName);
+		OutMergedFormatArgs.FileMetadata.Add(TEXT("unreal/rendererSubName"), InParams.RenderDataIdentifier.SubResourceName);
+		OutMergedFormatArgs.FileMetadata.Add(TEXT("unreal/renderLayer"), RenderLayerName);
+	}
+
+
+	// Add data from the job. ToDo: This probably needs to come from the job itself, so you can plug in different data sources eventually
+	{
+		// Frame Rate we handle separately, because it needs to be resolved.
+		double FrameRate = InParams.DefaultFrameRate.AsDecimal();
+		if (InParams.EvaluatedConfig)
+		{
+			const bool bIncludeCDOs = false;
+			UMovieGraphOutputSettingNode* OutputSettingNode = InParams.EvaluatedConfig->GetSettingForBranch<UMovieGraphOutputSettingNode>(InParams.RenderDataIdentifier.RootBranchName, bIncludeCDOs);
+			FrameRate = GetEffectiveFrameRate(OutputSettingNode, InParams.DefaultFrameRate).AsDecimal();
+		}
+
+		OutMergedFormatArgs.FilenameArguments.Add(TEXT("frame_rate"), FString::SanitizeFloat(FrameRate));
+		OutMergedFormatArgs.FileMetadata.Add(TEXT("unreal/frameRate"), FString::SanitizeFloat(FrameRate));
+
+		if (InParams.Job)
+		{
+			FString SequenceName = InParams.Job->Sequence.GetAssetName();
+			OutMergedFormatArgs.FilenameArguments.Add(TEXT("sequence_name"), SequenceName);
+			OutMergedFormatArgs.FileMetadata.Add(TEXT("unreal/sequenceName"), SequenceName);
+		}
+
+
+		// Add KVP data from the job (date, time, job name, job author, job comment)
+		UE::MoviePipeline::GetSharedFormatArguments(OutMergedFormatArgs.FilenameArguments, OutMergedFormatArgs.FileMetadata, InParams.InitializationTime, InParams.Version, InParams.Job);
+	}
+	
+
+	//  Now get the settings from our config. We need to gather KVP data from all possible nodes, even if not expressed in your configuration. This is because you might want to
+	// always use the {ts_count} token even if you don't have a Temporal Sample Count node to add it. So we loop through all the possible class types, and call a function on the 
+	// CDO, but then we pass that class type from the evaluated config (if it exists), and we pass the CDO as an argument if it doesn't.
+	TArray<UClass*> AllSettingsNodeClasses = UE::MovieRenderPipeline::FindMoviePipelineSettingClasses(UMovieGraphSettingNode::StaticClass());
+
+	// ToDo: This loops through class iterators every frame, we should probably initialize a copy of everything into the flattened config, since we could cache the classes
+	// once per run there. We don't cache the returned results here because you could potentially add/remove classes (via Blueprints) which would invalidate our cache.
+	for (UClass* InClass : AllSettingsNodeClasses)
+	{
+		UMovieGraphSettingNode* SettingInstance = nullptr;
+		if (InParams.EvaluatedConfig)
+		{
+			const bool bIncludeCDOs = true;
+			const bool bExactMatch = true;
+			SettingInstance = InParams.EvaluatedConfig->GetSettingForBranch(InClass, InParams.RenderDataIdentifier.RootBranchName, bIncludeCDOs, bExactMatch);
+		}
+
+		if (SettingInstance)
+		{
+			SettingInstance->GetFormatResolveArgs(OutMergedFormatArgs);
+		}
+	}
+
+	// Copy the metadata in the parameters into the merged output version
+	OutMergedFormatArgs.FileMetadata = InParams.FileMetadata;
+
+	// We expect the incoming string to have a {file_dup} token where they want file duplication numbers to be handled
+	// but we don't actually want a value there (unless there's a collision). So we override it with an empty string by default.
+	OutMergedFormatArgs.FilenameArguments.Add(TEXT("file_dup"), FString());
+
+	// Overwrite any {tokens} with the user-supplied overrides if needed. This allows different requesters to share the same variables (ie: filename extension)
+	for (const TPair<FString, FString>& KVP : InParams.FileNameFormatOverrides)
+	{
+		OutMergedFormatArgs.FilenameArguments.Add(KVP.Key, KVP.Value);
+	}
+
+	// Convert from our Python/BP exposed <FString, FString> to a named argument array for the formatter.
+	FStringFormatNamedArguments NamedArgs;
+	for (const TPair<FString, FString>& Argument : OutMergedFormatArgs.FilenameArguments)
+	{
+		NamedArgs.Add(Argument.Key, Argument.Value);
+	}
+
+	// Apply all of our named args to the file and generate a path.
+	FString BaseFilename = FString::Format(*InFormatString, NamedArgs);
+	if (InParams.bEnsureAbsolutePath)
+	{
+		BaseFilename = FPaths::ConvertRelativePathToFull(BaseFilename);
+	}
+
+	// Fix-up slashes
+	FPaths::NormalizeFilename(BaseFilename);
+
+	
+	// Check that it's a valid path that we can write to.
+	if (UE::MoviePipeline::CanWriteToFile(*BaseFilename, bOverwriteExisting))
+	{
+		return BaseFilename;
+	}
+
+	// The base filename must contain {file_dup} at this point, otherwise we're going to
+	// be stuck in an infinite loop of never resolving names as we'd just keep checking the
+	// base name!
+	if (!BaseFilename.Contains(TEXT("{file_dup}")))
+	{
+		BaseFilename.Append(TEXT("{file_dup}"));
+	}
+	FString ThisTry = BaseFilename;
+
+	// If we got here it means that there was already a file there with that name and bOverwriteExisting was false.
+	// So we start by swapping in _(2) where the file_dup token is, and try again. If that fails, _(3), etc. We start
+	// at _(2) because the file they're trying to write (foo.png) is technically "_(1)", so we start the rename process
+	// at _(2) so it accurately reflects that this is the second verison of that file.
+	int32 DuplicateIndex = 2;
+	while (true)
+	{
+		NamedArgs.Add(TEXT("file_dup"), FString::Printf(TEXT("_(%d)"), DuplicateIndex));
+
+		// Re-resolve the format string now that we've reassigned frame_dup to a number.
+		ThisTry = FString::Format(*InFormatString, NamedArgs);
+
+		// If the file doesn't exist, we can use that, else, increment the index and try again
+		if (UE::MoviePipeline::CanWriteToFile(*ThisTry, bOverwriteExisting))
+		{
+			return ThisTry;
+		}
+
+		++DuplicateIndex;
+	}
 }
\ No newline at end of file
diff --git a/Engine/Plugins/MovieScene/MovieRenderPipeline/Source/MovieRenderPipelineCore/Private/Graph/MovieGraphLinearTimeStep.cpp b/Engine/Plugins/MovieScene/MovieRenderPipeline/Source/MovieRenderPipelineCore/Private/Graph/MovieGraphLinearTimeStep.cpp
--- a/Engine/Plugins/MovieScene/MovieRenderPipeline/Source/MovieRenderPipelineCore/Private/Graph/MovieGraphLinearTimeStep.cpp
+++ b/Engine/Plugins/MovieScene/MovieRenderPipeline/Source/MovieRenderPipelineCore/Private/Graph/MovieGraphLinearTimeStep.cpp
@@ -88,7 +88,7 @@ void UMovieGraphLinearTimeStep::TickProducingFrames()
 		UMovieGraphOutputSettingNode* OutputNode = CurrentFrameData.EvaluatedConfig->GetSettingForBranch<UMovieGraphOutputSettingNode>(UMovieGraphNode::GlobalsPinName, bIncludeCDO);
 
 		const FFrameRate TickResolution = GetOwningGraph()->GetDataSourceInstance()->GetTickResolution();
-		const FFrameRate SourceFrameRate = GetOwningGraph()->GetDataSourceInstance()->GetDisplayRate(); // ToDo, needs to come from config (config can override)
+		const FFrameRate SourceFrameRate = GetOwningGraph()->GetDataSourceInstance()->GetDisplayRate();
 		const FFrameRate FinalFrameRate = UMovieGraphBlueprintLibrary::GetEffectiveFrameRate(OutputNode, SourceFrameRate);
 
 		FFrameTime FrameTimePerOutputFrame = FFrameRate::TransformTime(FFrameTime(FFrameNumber(1)), FinalFrameRate, TickResolution);
diff --git a/Engine/Plugins/MovieScene/MovieRenderPipeline/Source/MovieRenderPipelineCore/Private/Graph/MovieGraphPipeline.cpp b/Engine/Plugins/MovieScene/MovieRenderPipeline/Source/MovieRenderPipelineCore/Private/Graph/MovieGraphPipeline.cpp
--- a/Engine/Plugins/MovieScene/MovieRenderPipeline/Source/MovieRenderPipelineCore/Private/Graph/MovieGraphPipeline.cpp
+++ b/Engine/Plugins/MovieScene/MovieRenderPipeline/Source/MovieRenderPipelineCore/Private/Graph/MovieGraphPipeline.cpp
@@ -714,8 +714,8 @@ void UMovieGraphPipeline::ProcessOutstandingFinishedFrames()
 			// a central point for actually handling the file writing.
 			const bool bIncludeCDOs = false;
 			const FName BranchName = RenderData.Key.RootBranchName;
-			TArray<TObjectPtr<UMovieGraphFileOutputNode>> OutputNodeInstances = OutputFrame.EvaluatedConfig->GetSettingsForBranch<UMovieGraphFileOutputNode>(BranchName, bIncludeCDOs);
-			for (const TObjectPtr<UMovieGraphFileOutputNode>& Instance : OutputNodeInstances)
+			TArray<UMovieGraphFileOutputNode*> OutputNodeInstances = OutputFrame.EvaluatedConfig->GetSettingsForBranch<UMovieGraphFileOutputNode>(BranchName, bIncludeCDOs);
+			for (const UMovieGraphFileOutputNode* Instance : OutputNodeInstances)
 			{
 				UMovieGraphFileOutputNode* CDO = Instance->GetClass()->GetDefaultObject<UMovieGraphFileOutputNode>();
 				MaskData.FindOrAdd(CDO).Add(RenderData.Key);
diff --git a/Engine/Plugins/MovieScene/MovieRenderPipeline/Source/MovieRenderPipelineCore/Private/Graph/Nodes/MovieGraphImageSequenceOutputNode.cpp b/Engine/Plugins/MovieScene/MovieRenderPipeline/Source/MovieRenderPipelineCore/Private/Graph/Nodes/MovieGraphImageSequenceOutputNode.cpp
--- a/Engine/Plugins/MovieScene/MovieRenderPipeline/Source/MovieRenderPipelineCore/Private/Graph/Nodes/MovieGraphImageSequenceOutputNode.cpp
+++ b/Engine/Plugins/MovieScene/MovieRenderPipeline/Source/MovieRenderPipelineCore/Private/Graph/Nodes/MovieGraphImageSequenceOutputNode.cpp
@@ -6,6 +6,8 @@
 #include "Graph/MovieGraphDataTypes.h"
 #include "Graph/MovieGraphPipeline.h"
 #include "Graph/MovieGraphConfig.h"
+#include "Graph/MovieGraphFilenameResolveParams.h"
+#include "Graph/MovieGraphBlueprintLibrary.h"
 #include "Modules/ModuleManager.h"
 #include "ImageWriteQueue.h"
 #include "Misc/Paths.h"
@@ -84,17 +86,43 @@ void UMovieGraphImageSequenceOutputNode::OnReceiveImageDataImpl(UMovieGraphPipel
 		FString FileNameFormatString = OutputSettingNode->OutputDirectory.Path / OutputSettingNode->FileNameFormat;
 
 		// ToDo: Validate the string, ie: ensure it has {render_pass} in there somewhere there are multiple render passes
-		// in the output data, include {camera_name} if there are multiple cameras for that render pass, etc.
+		// in the output data, include {camera_name} if there are multiple cameras for that render pass, etc. Validation
+		// should insert {file_dup} tokens so it can put them at a logical place (ie: before frame numbers?)
 		FileNameFormatString += TEXT(".{ext}");
 
 		// Map the .ext to be specific to our output data.
 		TMap<FString, FString> AdditionalFormatArgs;
 		AdditionalFormatArgs.Add(TEXT("ext"), Extension);
 
-
-
-		const FString ResolvedProjectDir = FPaths::ConvertRelativePathToFull(FPaths::ProjectDir());
-		FString FileName = FString::Printf(TEXT("%s/Saved/MovieRenders/%s_%d.jpg"), *ResolvedProjectDir, *RenderLayerName, Payload->TraversalContext.Time.OutputFrameNumber);
+		FMovieGraphFilenameResolveParams Params = FMovieGraphFilenameResolveParams();
+		Params.RenderDataIdentifier = RenderData.Key;
+		//Params.RootFrameNumber = Payload->TraversalContext.Time.RootFrameNumber;
+		//Params.ShotFrameNumber = Payload->TraversalContext.Time.ShotFrameNumber;
+		Params.RootFrameNumberRel = Payload->TraversalContext.Time.OutputFrameNumber;
+		//Params.ShotFrameNumberRel = Payload->TraversalCOntext.Time.ShotFrameNumberRel
+		//Params.FileMetadata = ToDo: Track File Metadata
+		Params.Version = 1; // ToDo: Track versions
+		Params.ZeroPadFrameNumberCount = OutputSettingNode->ZeroPadFrameNumbers;
+		Params.FrameNumberOffset = OutputSettingNode->FrameNumberOffset;
+
+		// If time dilation is in effect, RootFrameNumber and ShotFrameNumber will contain duplicates and the files will overwrite each other, 
+		// so we force them into relative mode and then warn users we did that (as their numbers will jump from say 1001 -> 0000).
+		bool bForceRelativeFrameNumbers = true; // TODO: Use relative frame numbers until we track Root vs. Shot frame numbers. (Previously false);
+		//if (FileNameFormatString.Contains(TEXT("{frame")) Payload->TraversalContext.Time.IsTimeDilated() && !FileNameFormatString.Contains(TEXT("_rel}")))
+		//{
+		//	UE_LOG(LogMovieRenderPipeline, Warning, TEXT("Time Dilation was used but output format does not use relative time, forcing relative numbers. Change {frame_number} to {frame_number_rel} (or shot version) to remove this message."));
+		//	bForceRelativeFrameNumbers = true;
+		//}
+		Params.bForceRelativeFrameNumbers = bForceRelativeFrameNumbers;
+		Params.bEnsureAbsolutePath = true;
+		Params.FileNameFormatOverrides = AdditionalFormatArgs;
+		Params.InitializationTime = InPipeline->GetInitializationTime();
+		Params.Job = InPipeline->GetCurrentJob();
+		Params.EvaluatedConfig = InRawFrameData->EvaluatedConfig.Get();
+
+		// Take our string path from the Output Setting and resolve it.
+		FMovieGraphResolveArgs FinalResolvedKVPs;
+		const FString FileName = UMovieGraphBlueprintLibrary::ResolveFilenameFormatArguments(FileNameFormatString, Params, FinalResolvedKVPs);
 
 		TUniquePtr<FImageWriteTask> TileImageTask = MakeUnique<FImageWriteTask>();
 		TileImageTask->Format = OutputFormat;
diff --git a/Engine/Plugins/MovieScene/MovieRenderPipeline/Source/MovieRenderPipelineCore/Private/Graph/Nodes/MovieGraphOutputSettingNode.cpp b/Engine/Plugins/MovieScene/MovieRenderPipeline/Source/MovieRenderPipelineCore/Private/Graph/Nodes/MovieGraphOutputSettingNode.cpp
--- a/Engine/Plugins/MovieScene/MovieRenderPipeline/Source/MovieRenderPipelineCore/Private/Graph/Nodes/MovieGraphOutputSettingNode.cpp
+++ b/Engine/Plugins/MovieScene/MovieRenderPipeline/Source/MovieRenderPipelineCore/Private/Graph/Nodes/MovieGraphOutputSettingNode.cpp
@@ -1,13 +1,60 @@
 // Copyright Epic Games, Inc. All Rights Reserved.
 
 #include "Graph/Nodes/MovieGraphOutputSettingNode.h"
+#include "Styling/AppStyle.h"
 
 UMovieGraphOutputSettingNode::UMovieGraphOutputSettingNode()
 	: OutputResolution(FIntPoint(1920, 1080))
 	, OutputFrameRate(FFrameRate(24, 1))
 	, bOverwriteExistingOutput(true)
+	, ZeroPadFrameNumbers(4)
+	, FrameNumberOffset(0)
 {
 
 	FileNameFormat = TEXT("{sequence_name}.{frame_number}");
 	OutputDirectory.Path = TEXT("{project_dir}/Saved/MovieRenders/");
 }
+
+void UMovieGraphOutputSettingNode::GetFormatResolveArgs(FMovieGraphResolveArgs& OutMergedFormatArgs) const
+{
+	const FString ResolvedProjectDir = FPaths::ConvertRelativePathToFull(FPaths::ProjectDir());
+	OutMergedFormatArgs.FilenameArguments.Add(TEXT("project_dir"), ResolvedProjectDir);
+	OutMergedFormatArgs.FileMetadata.Add(TEXT("unreal/project_dir"), ResolvedProjectDir);
+
+	// Resolution Arguments
+	{
+		FString Resolution = FString::Printf(TEXT("%d_%d"), OutputResolution.X, OutputResolution.Y);
+		OutMergedFormatArgs.FilenameArguments.Add(TEXT("output_resolution"), Resolution);
+		OutMergedFormatArgs.FilenameArguments.Add(TEXT("output_width"), FString::FromInt(OutputResolution.X));
+		OutMergedFormatArgs.FilenameArguments.Add(TEXT("output_height"), FString::FromInt(OutputResolution.Y));
+	}
+
+	// We don't resolve the version here because that's handled on a per-file/shot basis
+}
+
+#if WITH_EDITOR
+FText UMovieGraphOutputSettingNode::GetNodeTitle(const bool bGetDescriptive) const
+{
+	static const FText OutputSettingsNodeName = NSLOCTEXT("MoviePipelineGraph", "NodeName_OutputSettings", "Output Settings");
+	return OutputSettingsNodeName;
+}
+
+FText UMovieGraphOutputSettingNode::GetMenuCategory() const 
+{
+	return NSLOCTEXT("MoviePipelineGraph", "Settings_Category", "Settings");
+}
+
+FLinearColor UMovieGraphOutputSettingNode::GetNodeTitleColor() const 
+{
+	static const FLinearColor OutputSettingsColor = FLinearColor(0.854f, 0.509f, 0.039f);
+	return OutputSettingsColor;
+}
+
+FSlateIcon UMovieGraphOutputSettingNode::GetIconAndTint(FLinearColor& OutColor) const 
+{
+	static const FSlateIcon SettingsIcon = FSlateIcon(FAppStyle::GetAppStyleSetName(), "Icons.Settings");
+
+	OutColor = FLinearColor::White;
+	return SettingsIcon;
+}
+#endif // WITH_EDITOR
\ No newline at end of file
diff --git a/Engine/Plugins/MovieScene/MovieRenderPipeline/Source/MovieRenderPipelineCore/Private/MoviePipelineBlueprintLibrary.cpp b/Engine/Plugins/MovieScene/MovieRenderPipeline/Source/MovieRenderPipelineCore/Private/MoviePipelineBlueprintLibrary.cpp
--- a/Engine/Plugins/MovieScene/MovieRenderPipeline/Source/MovieRenderPipelineCore/Private/MoviePipelineBlueprintLibrary.cpp
+++ b/Engine/Plugins/MovieScene/MovieRenderPipeline/Source/MovieRenderPipelineCore/Private/MoviePipelineBlueprintLibrary.cpp
@@ -915,31 +915,6 @@ UMoviePipelineSetting* UMoviePipelineBlueprintLibrary::FindOrGetDefaultSettingFo
 	return InSettingType->GetDefaultObject<UMoviePipelineSetting>();
 }
 
-static bool CanWriteToFile(const TCHAR* InFilename, bool bOverwriteExisting)
-{
-	// Check if there is space on the output disk.
-	bool bIsFreeSpace = true;
-
-	uint64 TotalNumberOfBytes, NumberOfFreeBytes;
-	if (FPlatformMisc::GetDiskTotalAndFreeSpace(InFilename, TotalNumberOfBytes, NumberOfFreeBytes))
-	{
-		bIsFreeSpace = NumberOfFreeBytes > 64 * 1024 * 1024; // 64mb minimum
-	}
-	// ToDO: Infinite loop possible.
-	return bIsFreeSpace && (bOverwriteExisting || IFileManager::Get().FileSize(InFilename) == -1);
-}
-
-static FString GetPaddingFormatString(int32 InZeroPadCount, const int32 InFrameNumber)
-{
-	// Printf takes the - sign into account when you specify the number of digits to pad to
-	// so we bump it by one to make the negative sign come first (ie: -0001 instead of -001)
-	if (InFrameNumber < 0)
-	{
-		InZeroPadCount++;
-	}
-
-	return FString::Printf(TEXT("%0*d"), InZeroPadCount, InFrameNumber);
-}
 
 void UMoviePipelineBlueprintLibrary::ResolveFilenameFormatArguments(const FString& InFormatString, const FMoviePipelineFilenameResolveParams& InParams, FString& OutFinalPath, FMoviePipelineFormatArgs& OutMergedFormatArgs)
 {
@@ -969,10 +944,10 @@ void UMoviePipelineBlueprintLibrary::ResolveFilenameFormatArguments(const FStrin
 		// when your sequence starts at zero and you use handle frames (which would cause negative output frame 
 		// numbers), so we allow the user to add a fixed amount to all output to ensure they are positive.
 		int32 FrameNumberOffset = OutputSetting->FrameNumberOffset + InParams.AdditionalFrameNumberOffset;
-		FString FrameNumber = GetPaddingFormatString(OutputSetting->ZeroPadFrameNumbers, InParams.FrameNumber + FrameNumberOffset); // Sequence Frame #
-		FString FrameNumberShot = GetPaddingFormatString(OutputSetting->ZeroPadFrameNumbers, InParams.FrameNumberShot + FrameNumberOffset); // Shot Frame #
-		FString FrameNumberRel = GetPaddingFormatString(OutputSetting->ZeroPadFrameNumbers, InParams.FrameNumberRel + FrameNumberOffset); // Relative to 0
-		FString FrameNumberShotRel = GetPaddingFormatString(OutputSetting->ZeroPadFrameNumbers, InParams.FrameNumberShotRel + FrameNumberOffset); // Relative to 0 within the shot.
+		FString FrameNumber = UE::MoviePipeline::GetPaddingFormatString(OutputSetting->ZeroPadFrameNumbers, InParams.FrameNumber + FrameNumberOffset); // Sequence Frame #
+		FString FrameNumberShot = UE::MoviePipeline::GetPaddingFormatString(OutputSetting->ZeroPadFrameNumbers, InParams.FrameNumberShot + FrameNumberOffset); // Shot Frame #
+		FString FrameNumberRel = UE::MoviePipeline::GetPaddingFormatString(OutputSetting->ZeroPadFrameNumbers, InParams.FrameNumberRel + FrameNumberOffset); // Relative to 0
+		FString FrameNumberShotRel = UE::MoviePipeline::GetPaddingFormatString(OutputSetting->ZeroPadFrameNumbers, InParams.FrameNumberShotRel + FrameNumberOffset); // Relative to 0 within the shot.
 
 		// Ensure they used relative frame numbers in the output so they get the right number of output frames.
 		if (InParams.bForceRelativeFrameNumbers)
@@ -994,7 +969,7 @@ void UMoviePipelineBlueprintLibrary::ResolveFilenameFormatArguments(const FStrin
 			ShotName = InParams.ShotNameOverride;
 		}
 
-		MoviePipeline::GetOutputStateFormatArgs(OutMergedFormatArgs, FrameNumber, FrameNumberShot, FrameNumberRel, FrameNumberShotRel, CameraName, ShotName);
+		MoviePipeline::GetOutputStateFormatArgs(OutMergedFormatArgs.FilenameArguments, OutMergedFormatArgs.FileMetadata, FrameNumber, FrameNumberShot, FrameNumberRel, FrameNumberShotRel, CameraName, ShotName);
 	}
 
 
@@ -1056,7 +1031,7 @@ void UMoviePipelineBlueprintLibrary::ResolveFilenameFormatArguments(const FStrin
 
 	FString ThisTry = BaseFilename + Extension;
 
-	if (CanWriteToFile(*ThisTry, OutputSetting->bOverrideExistingOutput))
+	if (UE::MoviePipeline::CanWriteToFile(*ThisTry, OutputSetting->bOverrideExistingOutput))
 	{
 		OutFinalPath = ThisTry;
 		return;
@@ -1071,7 +1046,7 @@ void UMoviePipelineBlueprintLibrary::ResolveFilenameFormatArguments(const FStrin
 		ThisTry = FString::Format(*FileNameFormatString, NamedArgs) + Extension;
 
 		// If the file doesn't exist, we can use that, else, increment the index and try again
-		if (CanWriteToFile(*ThisTry, OutputSetting->bOverrideExistingOutput))
+		if (UE::MoviePipeline::CanWriteToFile(*ThisTry, OutputSetting->bOverrideExistingOutput))
 		{
 			OutFinalPath = ThisTry;
 			return;
diff --git a/Engine/Plugins/MovieScene/MovieRenderPipeline/Source/MovieRenderPipelineCore/Private/MoviePipelinePrimaryConfig.cpp b/Engine/Plugins/MovieScene/MovieRenderPipeline/Source/MovieRenderPipelineCore/Private/MoviePipelinePrimaryConfig.cpp
--- a/Engine/Plugins/MovieScene/MovieRenderPipeline/Source/MovieRenderPipelineCore/Private/MoviePipelinePrimaryConfig.cpp
+++ b/Engine/Plugins/MovieScene/MovieRenderPipeline/Source/MovieRenderPipelineCore/Private/MoviePipelinePrimaryConfig.cpp
@@ -38,7 +38,7 @@ void UMoviePipelinePrimaryConfig::InitializeTransientSettings()
 	}
 
 	// Reflect over all possible settings that could be added to this class.
-	TArray<UClass*> AllSettingClasses = UE::MovieRenderPipeline::FindMoviePipelineSettingClasses();
+	TArray<UClass*> AllSettingClasses = UE::MovieRenderPipeline::FindMoviePipelineSettingClasses(UMoviePipelineSetting::StaticClass());
 
 	// Now remove any from the list that we already have a user setting for
 	for (const UMoviePipelineSetting* ExistingSetting : GetUserSettings())
@@ -140,7 +140,7 @@ void UMoviePipelinePrimaryConfig::GetFormatArguments(FMoviePipelineFormatArgs& I
 
 		// We use the FrameNumberOffset as the number for all of these so they can see it changing in the UI.
 		FString FramePlaceholderNumber = FString::Printf(TEXT("%0*d"), OutputSettings->ZeroPadFrameNumbers, OutputSetting->FrameNumberOffset);
-		MoviePipeline::GetOutputStateFormatArgs(InOutFormatArgs, FramePlaceholderNumber, FramePlaceholderNumber, FramePlaceholderNumber, FramePlaceholderNumber, TEXT("CameraName"), TEXT("ShotName"));
+		MoviePipeline::GetOutputStateFormatArgs(InOutFormatArgs.FilenameArguments, InOutFormatArgs.FileMetadata, FramePlaceholderNumber, FramePlaceholderNumber, FramePlaceholderNumber, FramePlaceholderNumber, TEXT("CameraName"), TEXT("ShotName"));
 	}
 
 	// Let each setting provide its own set of key/value pairs.
diff --git a/Engine/Plugins/MovieScene/MovieRenderPipeline/Source/MovieRenderPipelineCore/Private/MoviePipelineUtils.cpp b/Engine/Plugins/MovieScene/MovieRenderPipeline/Source/MovieRenderPipelineCore/Private/MoviePipelineUtils.cpp
--- a/Engine/Plugins/MovieScene/MovieRenderPipeline/Source/MovieRenderPipelineCore/Private/MoviePipelineUtils.cpp
+++ b/Engine/Plugins/MovieScene/MovieRenderPipeline/Source/MovieRenderPipelineCore/Private/MoviePipelineUtils.cpp
@@ -28,6 +28,7 @@
 #include "MoviePipelineQueue.h"
 #include "HAL/PlatformMisc.h"
 #include "HAL/PlatformMemory.h"
+#include "HAL/FileManager.h"
 #include "GenericPlatform/GenericPlatformDriver.h"
 #include "MovieRenderPipelineCoreModule.h"
 #include "RHI.h"
@@ -37,14 +38,14 @@ namespace UE
 {
 	namespace MovieRenderPipeline
 	{
-		TArray<UClass*> FindMoviePipelineSettingClasses()
+		TArray<UClass*> FindMoviePipelineSettingClasses(UClass* InBaseClass)
 		{
 			FAssetRegistryModule& AssetRegistryModule = FModuleManager::LoadModuleChecked<FAssetRegistryModule>(TEXT("AssetRegistry"));
 
 			TArray<FAssetData> ClassList;
 
 			FARFilter Filter;
-			Filter.ClassPaths.Add(UMoviePipelineSetting::StaticClass()->GetClassPathName());
+			Filter.ClassPaths.Add(InBaseClass->GetClassPathName());
 			Filter.ClassPaths.Add(UBlueprint::StaticClass()->GetClassPathName());
 
 			// Include any Blueprint based objects as well, this includes things like Blutilities, UMG, and GameplayAbility objects
@@ -69,7 +70,7 @@ namespace UE
 				{
 					ParentClass = UClass::TryFindTypeSlow<UClass>(FPackageName::ExportTextPathToObjectPath(ParentClassName));
 
-					if (ParentClass == nullptr || !ParentClass->IsChildOf(UMoviePipelineSetting::StaticClass()))
+					if (ParentClass == nullptr || !ParentClass->IsChildOf(InBaseClass))
 					{
 						continue;
 					}
@@ -83,7 +84,7 @@ namespace UE
 			// Now iterate through the loaded classes.
 			for (TObjectIterator<UClass> ClassIterator; ClassIterator; ++ClassIterator)
 			{
-				if (ClassIterator->IsChildOf(UMoviePipelineSetting::StaticClass()) && !ClassIterator->HasAnyClassFlags(CLASS_Abstract | CLASS_Deprecated | CLASS_NewerVersionExists))
+				if (ClassIterator->IsChildOf(InBaseClass) && !ClassIterator->HasAnyClassFlags(CLASS_Abstract | CLASS_Deprecated | CLASS_NewerVersionExists))
 				{
 					// While a blueprint is being compiled, there can be some transient instances of the class. We don't want these to show up in the list
 					// so just manually skip over them as they will go away next time GC is run anyways.
@@ -149,21 +150,21 @@ namespace MoviePipeline
 		}
 	}
 
-	void GetOutputStateFormatArgs(FMoviePipelineFormatArgs& InOutFinalFormatArgs, const FString FrameNumber, const FString FrameNumberShot, const FString FrameNumberRel, const FString FrameNumberShotRel, const FString CameraName, const FString ShotName)
+	void GetOutputStateFormatArgs(TMap<FString, FString>& InOutFilenameArguments, TMap<FString, FString>& InOutFileMetadata, const FString FrameNumber, const FString FrameNumberShot, const FString FrameNumberRel, const FString FrameNumberShotRel, const FString CameraName, const FString ShotName)
 	{
-		InOutFinalFormatArgs.FilenameArguments.Add(TEXT("frame_number"), FrameNumber);
-		InOutFinalFormatArgs.FilenameArguments.Add(TEXT("frame_number_shot"), FrameNumberShot);
-		InOutFinalFormatArgs.FilenameArguments.Add(TEXT("frame_number_rel"), FrameNumberRel);
-		InOutFinalFormatArgs.FilenameArguments.Add(TEXT("frame_number_shot_rel"), FrameNumberShotRel);
-		InOutFinalFormatArgs.FilenameArguments.Add(TEXT("camera_name"), CameraName);
-		InOutFinalFormatArgs.FilenameArguments.Add(TEXT("shot_name"), ShotName);
-
-		InOutFinalFormatArgs.FileMetadata.Add(TEXT("unreal/sequenceFrameNumber"), FrameNumber);
-		InOutFinalFormatArgs.FileMetadata.Add(TEXT("unreal/shotFrameNumber"), FrameNumberShot);
-		InOutFinalFormatArgs.FileMetadata.Add(TEXT("unreal/sequenceFrameNumberRelative"), FrameNumberRel);
-		InOutFinalFormatArgs.FileMetadata.Add(TEXT("unreal/shotFrameNumberRelative"), FrameNumberShotRel);
-		InOutFinalFormatArgs.FileMetadata.Add(TEXT("unreal/cameraName"), CameraName);
-		InOutFinalFormatArgs.FileMetadata.Add(TEXT("unreal/shotName"), ShotName);
+		InOutFilenameArguments.Add(TEXT("frame_number"), FrameNumber);
+		InOutFilenameArguments.Add(TEXT("frame_number_shot"), FrameNumberShot);
+		InOutFilenameArguments.Add(TEXT("frame_number_rel"), FrameNumberRel);
+		InOutFilenameArguments.Add(TEXT("frame_number_shot_rel"), FrameNumberShotRel);
+		InOutFilenameArguments.Add(TEXT("camera_name"), CameraName);
+		InOutFilenameArguments.Add(TEXT("shot_name"), ShotName);
+
+		InOutFileMetadata.Add(TEXT("unreal/sequenceFrameNumber"), FrameNumber);
+		InOutFileMetadata.Add(TEXT("unreal/shotFrameNumber"), FrameNumberShot);
+		InOutFileMetadata.Add(TEXT("unreal/sequenceFrameNumberRelative"), FrameNumberRel);
+		InOutFileMetadata.Add(TEXT("unreal/shotFrameNumberRelative"), FrameNumberShotRel);
+		InOutFileMetadata.Add(TEXT("unreal/cameraName"), CameraName);
+		InOutFileMetadata.Add(TEXT("unreal/shotName"), ShotName);
 	}
 
 }
@@ -826,6 +827,7 @@ namespace UE
 				InFilenameArguments.Add(TEXT("job_name"), InJob->JobName);
 				InFileMetadata.Add(TEXT("unreal/jobName"), InJob->JobName);
 				InFileMetadata.Add(TEXT("unreal/jobAuthor"), GetJobAuthor(InJob));
+				InFileMetadata.Add(TEXT("unreal/jobComment"), InJob->Comment);
 			}
 		}
 
@@ -930,6 +932,34 @@ namespace UE
 			InOutMetadata.Add(FString::Printf(TEXT("unreal/%s/prevRot/roll"), *CamName), FString::SanitizeFloat(InPrevRot.Roll));
 		}
 
+
+		bool CanWriteToFile(const TCHAR* InFilename, bool bOverwriteExisting)
+		{
+			// Check if there is space on the output disk.
+			bool bIsFreeSpace = true;
+
+			uint64 TotalNumberOfBytes, NumberOfFreeBytes;
+			if (FPlatformMisc::GetDiskTotalAndFreeSpace(InFilename, TotalNumberOfBytes, NumberOfFreeBytes))
+			{
+				bIsFreeSpace = NumberOfFreeBytes > 64 * 1024 * 1024; // 64mb minimum
+			}
+			// ToDO: Infinite loop possible.
+			return bIsFreeSpace && (bOverwriteExisting || IFileManager::Get().FileSize(InFilename) == -1);
+		}
+
+		FString GetPaddingFormatString(int32 InZeroPadCount, const int32 InFrameNumber)
+		{
+			// Printf takes the - sign into account when you specify the number of digits to pad. This
+			// means padding "3" to 4 digits becomes "0003", while "-3" becomes "-003". We combat this by
+			// incrementing the pad count for negative numbers by 1 so that you end up with "0003" and "-0003".
+			if (InFrameNumber < 0)
+			{
+				InZeroPadCount++;
+			}
+
+			return FString::Printf(TEXT("%0*d"), InZeroPadCount, InFrameNumber);
+		}
+
 	}
 }
 namespace UE
diff --git a/Engine/Plugins/MovieScene/MovieRenderPipeline/Source/MovieRenderPipelineCore/Public/Graph/MovieGraphBlueprintLibrary.h b/Engine/Plugins/MovieScene/MovieRenderPipeline/Source/MovieRenderPipelineCore/Public/Graph/MovieGraphBlueprintLibrary.h
--- a/Engine/Plugins/MovieScene/MovieRenderPipeline/Source/MovieRenderPipelineCore/Public/Graph/MovieGraphBlueprintLibrary.h
+++ b/Engine/Plugins/MovieScene/MovieRenderPipeline/Source/MovieRenderPipelineCore/Public/Graph/MovieGraphBlueprintLibrary.h
@@ -3,6 +3,7 @@
 #pragma once
 
 #include "Kismet/BlueprintFunctionLibrary.h"
+#include "Graph/MovieGraphFilenameResolveParams.h"
 #include "Misc/FrameRate.h"
 
 #include "MovieGraphBlueprintLibrary.generated.h"
@@ -26,4 +27,19 @@ public:
 	*/
 	UFUNCTION(BlueprintCallable, Category = "Movie Graph")
 	static FFrameRate GetEffectiveFrameRate(UMovieGraphOutputSettingNode* InNode, const FFrameRate& InDefaultRate);
+
+	/**
+	* Takes a Movie Graph format string (in the form of {token}), a list of parameters (which normally come from the running UMovieGraphPipeline) and
+	* resolves them into a string. Unknown tokens are ignored. Which tokens can be resolved depends on the contents of InParams, tokens from settings
+	* rely on a evaluated config being provided, etc.
+	* @param	InFormatString		- Format string to attempt to resolve.
+	* @param	InParams			- A list of parameters to use as source data for the resolve step. Normally comes from the UMovieGraphPipeline instance,
+	*								- but takes (mostly) POD here to allow using this function outside of the render runtime.
+	* @param	OutMergedFormatArgs - The set of KVP for both filename formats and file metadata that is generated as a result of this. Provided in case you
+	* 								- needed to do your own string resolving with the final dataset.
+	* @return						- The resolved format string.
+	*/
+	UFUNCTION(BlueprintCallable, Category = "Movie Graph")
+	static FString ResolveFilenameFormatArguments(const FString& InFormatString, const FMovieGraphFilenameResolveParams& InParams, FMovieGraphResolveArgs& OutMergedFormatArgs);
+
 };
\ No newline at end of file
diff --git a/Engine/Plugins/MovieScene/MovieRenderPipeline/Source/MovieRenderPipelineCore/Public/Graph/MovieGraphConfig.h b/Engine/Plugins/MovieScene/MovieRenderPipeline/Source/MovieRenderPipelineCore/Public/Graph/MovieGraphConfig.h
--- a/Engine/Plugins/MovieScene/MovieRenderPipeline/Source/MovieRenderPipelineCore/Public/Graph/MovieGraphConfig.h
+++ b/Engine/Plugins/MovieScene/MovieRenderPipeline/Source/MovieRenderPipelineCore/Public/Graph/MovieGraphConfig.h
@@ -264,10 +264,9 @@ public:
  		return OutKeys;
 	}
 
-	template<typename NodeType>
-	TObjectPtr<NodeType> GetSettingForBranch(const FName InBranchName, bool bIncludeCDOs = true, bool bExactMatch = false)
+	UMovieGraphSettingNode* GetSettingForBranch(UClass* InClass, const FName InBranchName, bool bIncludeCDOs = true, bool bExactMatch = false)
 	{
-		TArray<TObjectPtr<NodeType>> AllSettings = GetSettingsForBranch<NodeType>(InBranchName, bIncludeCDOs, bExactMatch);
+		TArray<UMovieGraphSettingNode*> AllSettings = GetSettingsForBranch(InClass, InBranchName, bIncludeCDOs, bExactMatch);
 		if (AllSettings.Num() > 0)
 		{
 			return AllSettings[0];
@@ -276,22 +275,21 @@ public:
 		return nullptr;
 	}
 
-	template<typename NodeType>
-	TArray<TObjectPtr<NodeType>> GetSettingsForBranch(const FName InBranchName, bool bIncludeCDOs = true, bool bExactMatch = false)
+	TArray<UMovieGraphSettingNode*> GetSettingsForBranch(UClass* InClass, const FName InBranchName, bool bIncludeCDOs = true, bool bExactMatch = false)
 	{
-		TArray<TObjectPtr<NodeType>> ResultNodes;
 		FMovieGraphEvaluatedBranchConfig* BranchConfig = BranchConfigMapping.Find(InBranchName);
 		ensureMsgf(BranchConfig, TEXT("Failed to find branch mapping for Branch: %s"), *InBranchName.ToString());
 
+		TArray<UMovieGraphSettingNode*> ResultNodes;
 		if (BranchConfig)
 		{
 			// Check to see if the branch has an instance of this.
 			for (const TObjectPtr<UMovieGraphNode>& Node : BranchConfig->GetNodes())
 			{
-				bool bMatches = bExactMatch ? Node->GetClass() == NodeType::StaticClass() : Node->IsA<NodeType>();
+				bool bMatches = bExactMatch ? Node->GetClass() == InClass : Node->IsA(InClass);
 				if (bMatches)
 				{
-					ResultNodes.Add(Cast<NodeType>(Node.Get()));
+					ResultNodes.Add(Cast<UMovieGraphSettingNode>(Node.Get()));
 				}
 			}
 		}
@@ -300,13 +298,33 @@ public:
 		// or the config simply didn't override that setting class, at which point we might try to return a default
 		if (bIncludeCDOs)
 		{
-			ResultNodes.Add(GetMutableDefault<NodeType>());
+			ResultNodes.Add(Cast<UMovieGraphSettingNode>(InClass->GetDefaultObject()));
 		}
 
 		return ResultNodes;
 	}
-	
 
+	template<typename NodeType>
+	NodeType* GetSettingForBranch(const FName InBranchName, bool bIncludeCDOs = true, bool bExactMatch = false)
+	{
+		return Cast<NodeType>(GetSettingForBranch(NodeType::StaticClass(), InBranchName, bIncludeCDOs, bExactMatch));
+	}
+
+	template<typename NodeType>
+	TArray<NodeType*> GetSettingsForBranch(const FName InBranchName, bool bIncludeCDOs = true, bool bExactMatch = false)
+	{
+		TArray<UMovieGraphSettingNode*> UntypedResults = GetSettingsForBranch(NodeType::StaticClass(), InBranchName, bIncludeCDOs, bExactMatch);
+
+		TArray<TObjectPtr<NodeType>> ResultNodes;
+		ResultNodes.Reserve(UntypedResults.Num());
+		for (UMovieGraphSettingNode* UntypedNode : UntypedResults)
+		{
+			ResultNodes.Add(Cast<NodeType>(UntypedNode));
+		}
+
+		return ResultNodes;
+	}
+	
 public:
 	/** Mapping between named branches (at the root of the config) and their evaluated values. */
 	UPROPERTY(Transient)
diff --git a/Engine/Plugins/MovieScene/MovieRenderPipeline/Source/MovieRenderPipelineCore/Public/Graph/MovieGraphFilenameResolveParams.h b/Engine/Plugins/MovieScene/MovieRenderPipeline/Source/MovieRenderPipelineCore/Public/Graph/MovieGraphFilenameResolveParams.h
new file mode 100644
--- /dev/null
+++ b/Engine/Plugins/MovieScene/MovieRenderPipeline/Source/MovieRenderPipelineCore/Public/Graph/MovieGraphFilenameResolveParams.h
@@ -0,0 +1,154 @@
+// Copyright Epic Games, Inc. All Rights Reserved.
+#pragma once
+
+#include "CoreMinimal.h"
+#include "Misc/FrameRate.h"
+#include "Graph/MovieGraphRenderDataIdentifier.h"
+#include "MovieGraphFilenameResolveParams.generated.h"
+
+// Forward Declare
+class UMoviePipelineExecutorJob;
+class UMovieGraphEvaluatedConfig;
+class UMoviePipelineExecutorShot;
+
+/**
+* This data structure contains a list of key-value pairs (as strings) for both filename resolving, and file metadata.
+* They are stored as two separate arrays (as metadata is often a bit more verbose/human friendly) and it is not required
+* that the two arrays be in sync. When resolving final filenames, any {tokens} will be replaced with a matching key-value
+* if found in the FilenameArguments variable. Because of this, FilenameArguments should be short and avoid illegal filename chars.
+*/
+USTRUCT(BlueprintType)
+struct FMovieGraphResolveArgs
+{
+	GENERATED_BODY()
+
+	FMovieGraphResolveArgs() = default;
+
+	/** A set of Key/Value pairs for output filename format strings (without {}) and their values. */
+	UPROPERTY(BlueprintReadWrite, EditAnywhere, Category = "Movie Render Pipeline")
+	TMap<FString, FString> FilenameArguments;
+
+	/** A set of Key/Value pairs for file metadata for file formats that support metadata. */
+	UPROPERTY(BlueprintReadWrite, EditAnywhere, Category = "Movie Render Pipeline")
+	TMap<FString, FString> FileMetadata;
+};
+
+/*
+* This data structure contains all of the information needed to resolve a filename format string from the Movie Graph.
+* During runtime the values are generated by the actual rendering and then placed into this data structure before resolving
+* an actual {format_string} using the UMovieGraphBlueprintLibrary. You can use this to resolve filenames when editing. If
+* you wish to resolve a partial format string (such as {sequence_name}/{frame_number} but have it only resolve {sequence_name}
+* you can add ("frame_number", "{frame_number}") to the FileNameFormatOverrides and that portion of the format string will
+* be preserved so the resulting string back would say "my_seq_name/{frame_number}".
+*/
+USTRUCT(BlueprintType)
+struct FMovieGraphFilenameResolveParams
+{
+	GENERATED_BODY()
+
+	FMovieGraphFilenameResolveParams()
+		: RootFrameNumber(0)
+		, ShotFrameNumber(0)
+		, RootFrameNumberRel(0)
+		, ShotFrameNumberRel(0)
+		, Version(0)
+		, bForceRelativeFrameNumbers(false)
+		, bEnsureAbsolutePath(true)
+		, InitializationTime(0)
+		, ZeroPadFrameNumberCount(0)
+		, Job(nullptr)
+		, Shot(nullptr)
+		, EvaluatedConfig(nullptr)
+		, FrameNumberOffset(0)
+	{
+	}
+	
+	/**
+	* This is used to fill out tokens related to the render data that the file represents.
+	* Used to fill out the {camera_name}, {render_layer}, {renderer_name}, {renderer_sub_name} tokens.
+	*/
+	UPROPERTY(BlueprintReadWrite, EditAnywhere, Category = "Movie Graph")
+	FMovieGraphRenderDataIdentifier RenderDataIdentifier;
+	
+	/**
+	* Frame Number for the Root (matching what you see in the Sequencer timeline. ie: If the Sequence PlaybackRange starts on 50, this value would be 50 on the first frame.
+	* Used to fill out the {frame_number} token.
+	*/
+	UPROPERTY(BlueprintReadWrite, EditAnywhere, Category = "Movie Graph")
+	int32 RootFrameNumber;
+	
+	/** 
+	* Frame Number for the Shot (matching what you would see in Sequencer at the sub-sequence level. 
+	* Used to fill out the {frame_number_shot} token.
+	*/
+	UPROPERTY(BlueprintReadWrite, EditAnywhere, Category = "Movie Graph")
+	int32 ShotFrameNumber;
+	
+	/** 
+	* Frame Number for the Root (relative to 0, not what you would see in the Sequencer timeline. ie: If sequence PlaybackRange starts on 50, this value would be 0 on the first frame.)
+	* Used to fill out the {frame_number_rel} token.
+	*/
+	UPROPERTY(BlueprintReadWrite, EditAnywhere, Category = "Movie Graph")
+	int32 RootFrameNumberRel;
+	
+	/** 
+	* Frame Number for the Shot (relative to 0, not what you would see in the Sequencer timeline.) 
+	* Used to fill out the {frame_number_shot_rel} token.
+	*/
+	UPROPERTY(BlueprintReadWrite, EditAnywhere, Category = "Movie Graph")
+	int32 ShotFrameNumberRel;
+	
+	/** A key/value pair that maps metadata names to their values. Output is only supported in exr formats at the moment. */
+	UPROPERTY(BlueprintReadWrite, EditAnywhere, Category = "Movie Graph")
+	TMap<FString, FString> FileMetadata;
+	
+	/** Used to resolve {version} format arguments. */
+	UPROPERTY(BlueprintReadWrite, EditAnywhere, Category = "Movie Graph")
+	int32 Version;
+	
+	/** If true, force format strings (like {frame_number}) to resolve using the relative version. Used when slow-mo is detected as frame numbers would overlap. */
+	UPROPERTY(BlueprintReadWrite, EditAnywhere, Category = "Movie Graph")
+	bool bForceRelativeFrameNumbers;
+
+	/** 
+	* If true, If the format string is a relative path, then the resulting resolved path will be converted to an absolute path.
+	* Assumes that the relative path should be resolved relative to the engine/editor executable.
+	*/
+	UPROPERTY(BlueprintReadWrite, EditAnywhere, Category = "Movie Graph")
+	bool bEnsureAbsolutePath;
+	
+	/** The initialization time for this job. Used to resolve time-based format arguments. */
+	UPROPERTY(BlueprintReadWrite, EditAnywhere, Category = "Movie Graph")
+	FDateTime InitializationTime;
+	
+	/** When converitng frame numbers to strings, how many digits should we pad them up to? ie: 5 => 0005 with a count of 4. */
+	UPROPERTY(BlueprintReadWrite, EditAnywhere, Category = "Movie Graph")
+	int32 ZeroPadFrameNumberCount;
+
+	/** 
+	* A map between "{format}" tokens and their values. These are applied after the auto-generated ones from the system,
+	* which allows the caller to override things like {.ext}
+	*/
+	UPROPERTY(BlueprintReadWrite, EditAnywhere, Category = "Movie Graph")
+	TMap<FString, FString> FileNameFormatOverrides;
+
+	/** Optional. If specified, used to resolve some tokens related to the job.*/
+	UPROPERTY(BlueprintReadWrite, EditAnywhere, Category = "Movie Graph")
+	TObjectPtr<UMoviePipelineExecutorJob> Job;
+
+	/** Optional. If specified, used to resolve the shot name. No settings from this shot are ever used (they must come form the EvaluatedConfig instead). */
+	UPROPERTY(BlueprintReadWrite, EditAnywhere, Category = "Movie Graph")
+	TObjectPtr<UMoviePipelineExecutorShot> Shot;
+
+	/** Optional. If specified, used to resolve some tokens related to settings in the shot (such as {ts_count}). */
+	UPROPERTY(BlueprintReadWrite, EditAnywhere, Category = "Movie Graph")
+	TObjectPtr<UMovieGraphEvaluatedConfig> EvaluatedConfig;
+
+	/** Added onto the various FrameNumber settings automatically. Can be used to offset values (such as starting your Level Sequence on frame 0, but outputting it to disk at 1001). */
+	UPROPERTY(BlueprintReadWrite, EditAnywhere, Category = "Movie Graph")
+	int32 FrameNumberOffset;
+
+	/** What is the default Frame Rate? This will be overwritten by the Output Setting from the evaluated config (if specified). */
+	UPROPERTY(BlueprintReadWrite, EditAnywhere, Category = "Movie Graph")
+	FFrameRate DefaultFrameRate;
+};
\ No newline at end of file
diff --git a/Engine/Plugins/MovieScene/MovieRenderPipeline/Source/MovieRenderPipelineCore/Public/Graph/MovieGraphNode.h b/Engine/Plugins/MovieScene/MovieRenderPipeline/Source/MovieRenderPipelineCore/Public/Graph/MovieGraphNode.h
--- a/Engine/Plugins/MovieScene/MovieRenderPipeline/Source/MovieRenderPipelineCore/Public/Graph/MovieGraphNode.h
+++ b/Engine/Plugins/MovieScene/MovieRenderPipeline/Source/MovieRenderPipelineCore/Public/Graph/MovieGraphNode.h
@@ -7,6 +7,7 @@
 #include "InstancedStruct.h"
 #include "PropertyBag.h"
 #include "Graph/MovieGraphValueContainer.h"
+#include "Graph/MovieGraphFilenameResolveParams.h"
 
 #if WITH_EDITOR
 #include "Textures/SlateIcon.h"
@@ -202,4 +203,15 @@ public:
 	virtual TArray<FMovieGraphPinProperties> GetInputPinProperties() const override;
 	virtual TArray<FMovieGraphPinProperties> GetOutputPinProperties() const override;
 	// ~UMovieGraphNode Interface
+	
+	/*
+	* This is called either on the CDO, or on a "flattened" instance of the node every frame when
+	* generating filename/file metadata, allowing the node to add custom key-value pairs (FString, FString)
+	* to be used as {format_tokens} in filenames, or to be included in File metadata. Nodes can read
+	* their own settings (such as temporal sub-sample count) and add it to the available list of tokens.
+	*
+	* Because this is called either on the CDO or on a flattened instance, there is no need to worry about
+	* resolving the settings of the graph, the node only needs to read its own values.
+	*/
+	virtual void GetFormatResolveArgs(FMovieGraphResolveArgs& OutMergedFormatArgs) const {}
 };
\ No newline at end of file
diff --git a/Engine/Plugins/MovieScene/MovieRenderPipeline/Source/MovieRenderPipelineCore/Public/Graph/MovieGraphPipeline.h b/Engine/Plugins/MovieScene/MovieRenderPipeline/Source/MovieRenderPipelineCore/Public/Graph/MovieGraphPipeline.h
--- a/Engine/Plugins/MovieScene/MovieRenderPipeline/Source/MovieRenderPipelineCore/Public/Graph/MovieGraphPipeline.h
+++ b/Engine/Plugins/MovieScene/MovieRenderPipeline/Source/MovieRenderPipelineCore/Public/Graph/MovieGraphPipeline.h
@@ -33,6 +33,22 @@ public:
 	UFUNCTION(BlueprintCallable, Category = "Movie Graph")
 	UMoviePipelineExecutorJob* GetCurrentJob() const { return CurrentJob; }
 
+	/**
+	* Returns the time this movie pipeline was initialized at.
+	*/
+	UFUNCTION(BlueprintPure, Category = "Movie Graph")
+	FDateTime GetInitializationTime() const { return GraphInitializationTime; }
+
+	/**
+	* Override the time this movie pipeline was initialized at. This can be used for render farms
+	* to ensure that jobs on all machines use the same date/time instead of each calculating it locally.
+	*
+	* Needs to be called after ::Initialize(...)
+	*
+	* @param InDateTime - The DateTime object to return for GetInitializationTime.
+	*/
+	UFUNCTION(BlueprintCallable, Category = "Movie Render Pipeline")
+	void SetInitializationTime(const FDateTime& InDateTime) { GraphInitializationTime = InDateTime; }
 
 public:
 	UMovieGraphConfig* GetRootGraphForShot(UMoviePipelineExecutorShot* InShot) const;
diff --git a/Engine/Plugins/MovieScene/MovieRenderPipeline/Source/MovieRenderPipelineCore/Public/Graph/Nodes/MovieGraphOutputSettingNode.h b/Engine/Plugins/MovieScene/MovieRenderPipeline/Source/MovieRenderPipelineCore/Public/Graph/Nodes/MovieGraphOutputSettingNode.h
--- a/Engine/Plugins/MovieScene/MovieRenderPipeline/Source/MovieRenderPipelineCore/Public/Graph/Nodes/MovieGraphOutputSettingNode.h
+++ b/Engine/Plugins/MovieScene/MovieRenderPipeline/Source/MovieRenderPipelineCore/Public/Graph/Nodes/MovieGraphOutputSettingNode.h
@@ -3,7 +3,6 @@
 #pragma once
 #include "Graph/MovieGraphNode.h"
 #include "Misc/FrameRate.h"
-#include "Styling/AppStyle.h"
 #include "MovieGraphOutputSettingNode.generated.h"
 
 
@@ -14,32 +13,16 @@ class MOVIERENDERPIPELINECORE_API UMovieGraphOutputSettingNode : public UMovieGr
 public:
 	UMovieGraphOutputSettingNode();
 
+	// UMovieGraphSettingNode Interface
+	virtual void GetFormatResolveArgs(FMovieGraphResolveArgs& OutMergedFormatArgs) const override;
+
 #if WITH_EDITOR
-	virtual FText GetNodeTitle(const bool bGetDescriptive) const override
-	{
-		static const FText OutputSettingsNodeName = NSLOCTEXT("MoviePipelineGraph", "NodeName_OutputSettings", "Output Settings");
-		return OutputSettingsNodeName;
-	}
-
-	virtual FText GetMenuCategory() const override
-	{
-		return NSLOCTEXT("MoviePipelineGraph", "Settings_Category", "Settings");
-	}
-
-	virtual FLinearColor GetNodeTitleColor() const override
-	{
-		static const FLinearColor OutputSettingsColor = FLinearColor(0.854f, 0.509f, 0.039f);
-		return OutputSettingsColor;
-	}
-
-	virtual FSlateIcon GetIconAndTint(FLinearColor& OutColor) const override
-	{
-		static const FSlateIcon SettingsIcon = FSlateIcon(FAppStyle::GetAppStyleSetName(), "Icons.Settings");
-
-		OutColor = FLinearColor::White;
-		return SettingsIcon;
-	}
+	virtual FText GetNodeTitle(const bool bGetDescriptive) const override;
+	virtual FText GetMenuCategory() const override;
+	virtual FLinearColor GetNodeTitleColor() const override;
+	virtual FSlateIcon GetIconAndTint(FLinearColor& OutColor) const override;
 #endif
+	// ~UMovieGraphSettingNode Interface
 
 public:
 	UPROPERTY(EditAnywhere, BlueprintReadWrite, Category = Overrides, meta = (InlineEditConditionToggle))
@@ -57,6 +40,12 @@ public:
 	UPROPERTY(EditAnywhere, BlueprintReadWrite, Category = Overrides, meta = (InlineEditConditionToggle))
 	uint8 bOverride_bOverwriteExistingOutput : 1;
 
+	UPROPERTY(EditAnywhere, BlueprintReadWrite, Category = Overrides, meta = (InlineEditConditionToggle))
+	uint8 bOverride_ZeroPadFrameNumbers : 1;
+
+	UPROPERTY(EditAnywhere, BlueprintReadWrite, Category = Overrides, meta = (InlineEditConditionToggle))
+	uint8 bOverride_FrameNumberOffset : 1;
+
 	/** What directory should all of our output files be relative to. */
 	UPROPERTY(EditAnywhere, BlueprintReadWrite, Category = "File Output", meta=(EditCondition="bOverride_OutputDirectory"))
 	FDirectoryPath OutputDirectory;
@@ -76,4 +65,16 @@ public:
 	/** If true, output containers should attempt to override any existing files with the same name. */
 	UPROPERTY(EditAnywhere, BlueprintReadWrite, Category = "File Output", meta = (EditCondition = "bOverride_bOverwriteExistingOutput"))
 	bool bOverwriteExistingOutput;
+
+	/** How many digits should all output frame numbers be padded to? MySequence_1.png -> MySequence_0001.png. Useful for software that struggles to recognize frame ranges when non-padded. */
+	UPROPERTY(EditAnywhere, BlueprintReadWrite, Category = "File Output", meta = (EditCondition = "bOverride_ZeroPadFrameNumbers", UIMin = "1", UIMax = "5", ClampMin = "1"))
+	int32 ZeroPadFrameNumbers;
+
+	/**
+	* How many frames should we offset the output frame number by? This is useful when using handle frames on Sequences that start at frame 0,
+	* as the output would start in negative numbers. This can be used to offset by a fixed amount to ensure there's no negative numbers.
+	*/
+	UPROPERTY(EditAnywhere, BlueprintReadWrite,  Category = "File Output", meta = (EditCondition = "bOverride_FrameNumberOffset"))
+	int32 FrameNumberOffset;
+
 };
diff --git a/Engine/Plugins/MovieScene/MovieRenderPipeline/Source/MovieRenderPipelineCore/Public/MoviePipelineUtils.h b/Engine/Plugins/MovieScene/MovieRenderPipeline/Source/MovieRenderPipelineCore/Public/MoviePipelineUtils.h
--- a/Engine/Plugins/MovieScene/MovieRenderPipeline/Source/MovieRenderPipelineCore/Public/MoviePipelineUtils.h
+++ b/Engine/Plugins/MovieScene/MovieRenderPipeline/Source/MovieRenderPipelineCore/Public/MoviePipelineUtils.h
@@ -122,7 +122,7 @@ namespace UE
 {
 	namespace MovieRenderPipeline
 	{
-		MOVIERENDERPIPELINECORE_API TArray<UClass*> FindMoviePipelineSettingClasses();
+		MOVIERENDERPIPELINECORE_API TArray<UClass*> FindMoviePipelineSettingClasses(UClass* InBaseClass);
 		MOVIERENDERPIPELINECORE_API EAntiAliasingMethod GetEffectiveAntiAliasingMethod(const UMoviePipelineAntiAliasingSetting* InSetting);
 	}
 
@@ -139,12 +139,14 @@ namespace UE
 		MOVIERENDERPIPELINECORE_API void GetMetadataFromCineCamera(class UCineCameraComponent* InComponent, const FString& InCameraName, const FString& InRenderPassName, TMap<FString, FString>& InOutMetadata);
 		MOVIERENDERPIPELINECORE_API void GetMetadataFromCameraLocRot(const FString& InCameraName, const FString& InRenderPassName, const FVector& InCurLoc, const FRotator& InCurRot, const FVector& InPrevLoc, const FRotator& InPrevRot, TMap<FString, FString>& InOutMetadata);
 		MOVIERENDERPIPELINECORE_API FMoviePipelineRenderPassMetrics GetRenderPassMetrics(UMoviePipelinePrimaryConfig* InPrimaryConfig, UMoviePipelineExecutorShot* InPipelineExecutorShot, const FMoviePipelineRenderPassMetrics& InRenderPassMetrics, const FIntPoint& InEffectiveOutputResolution);
+		MOVIERENDERPIPELINECORE_API bool CanWriteToFile(const TCHAR* InFilename, bool bOverwriteExisting);
+		MOVIERENDERPIPELINECORE_API FString GetPaddingFormatString(int32 InZeroPadCount, const int32 InFrameNumber);
 	}
 }
 
 namespace MoviePipeline
 {
-	MOVIERENDERPIPELINECORE_API void GetOutputStateFormatArgs(FMoviePipelineFormatArgs& InOutFinalFormatArgs, const FString FrameNumber, const FString FrameNumberShot, const FString FrameNumberRel, const FString FrameNumberShotRel, const FString CameraName, const FString ShotName);
+	MOVIERENDERPIPELINECORE_API void GetOutputStateFormatArgs(TMap<FString, FString>& InFilenameArguments, TMap<FString, FString>& InFileMetadata, const FString FrameNumber, const FString FrameNumberShot, const FString FrameNumberRel, const FString FrameNumberShotRel, const FString CameraName, const FString ShotName);
 	
 	/** Iterate root-to-tails and generate a HierarchyNode for each level. Caches the sub-section range, playback range, camera cut range, etc. */
 	void CacheCompleteSequenceHierarchy(UMovieSceneSequence* InSequence, TSharedPtr<FCameraCutSubSectionHierarchyNode> InRootNode);
diff --git a/Engine/Plugins/MovieScene/MovieRenderPipeline/Source/MovieRenderPipelineEditor/Private/Widgets/SMoviePipelineConfigEditor.cpp b/Engine/Plugins/MovieScene/MovieRenderPipeline/Source/MovieRenderPipelineEditor/Private/Widgets/SMoviePipelineConfigEditor.cpp
--- a/Engine/Plugins/MovieScene/MovieRenderPipeline/Source/MovieRenderPipelineEditor/Private/Widgets/SMoviePipelineConfigEditor.cpp
+++ b/Engine/Plugins/MovieScene/MovieRenderPipeline/Source/MovieRenderPipelineEditor/Private/Widgets/SMoviePipelineConfigEditor.cpp
@@ -621,7 +621,7 @@ TSharedRef<SWidget> SMoviePipelineConfigEditor::OnGenerateSettingsMenu()
 	FMenuBuilder MenuBuilder(true, nullptr, Extender);
 
 	// Put the different categories into different sections
-	TArray<UClass*> SourceClasses = UE::MovieRenderPipeline::FindMoviePipelineSettingClasses();
+	TArray<UClass*> SourceClasses = UE::MovieRenderPipeline::FindMoviePipelineSettingClasses(UMoviePipelineSetting::StaticClass());
 
 	TMap<FString, TArray<UClass*>> CategorizedClasses;
 
