MoviePipeline: Minor improvements to Graph UX. Pins that are don't support multiple connections now "pop" the existing connection off when connected (mimicking Blueprint Execution Pin behavior), and cleaned up some scenarios where you could make invalid connections.

diff --git a/Engine/Plugins/MovieScene/MovieRenderPipeline/Source/MovieRenderPipelineCore/Private/Graph/MovieGraphConfig.cpp b/Engine/Plugins/MovieScene/MovieRenderPipeline/Source/MovieRenderPipelineCore/Private/Graph/MovieGraphConfig.cpp
--- a/Engine/Plugins/MovieScene/MovieRenderPipeline/Source/MovieRenderPipelineCore/Private/Graph/MovieGraphConfig.cpp
+++ b/Engine/Plugins/MovieScene/MovieRenderPipeline/Source/MovieRenderPipelineCore/Private/Graph/MovieGraphConfig.cpp
@@ -184,13 +184,19 @@ bool UMovieGraphConfig::AddLabeledEdge(UMovieGraphNode* FromNode, const FName& F
 		return false;
 	}
 
-	// Add the edge
-	FromPin->AddEdgeTo(ToPin);
 	bool bConnectionBrokeOtherEdges = false;
-//	if (!ToPin->AllowMultipleConnections())
-//	{
-//		bConnectionBrokeOtherEdges = ToPin->BreakAllIncompatibleEdges();
-//	}
+
+	// Input pins can only have one edge connected to them at once, so if the pin we're connecting to already
+	// has a connection, then we break the existing connection.
+	if (!ToPin->AllowsMultipleConnections() && ToPin->EdgeCount() > 0)
+	{
+		ToPin->BreakAllEdges();
+		bConnectionBrokeOtherEdges = true;
+	}
+
+	// Add the edge. We do this after the above
+	// since that will break all edges first if there's already one.
+	FromPin->AddEdgeTo(ToPin);
 //
 //#if WITH_EDITOR
 //	NotifyGraphChanged(EMoviePipelineGraphChangeType::Structural);
diff --git a/Engine/Plugins/MovieScene/MovieRenderPipeline/Source/MovieRenderPipelineCore/Private/Graph/MovieGraphPin.cpp b/Engine/Plugins/MovieScene/MovieRenderPipeline/Source/MovieRenderPipelineCore/Private/Graph/MovieGraphPin.cpp
--- a/Engine/Plugins/MovieScene/MovieRenderPipeline/Source/MovieRenderPipelineCore/Private/Graph/MovieGraphPin.cpp
+++ b/Engine/Plugins/MovieScene/MovieRenderPipeline/Source/MovieRenderPipelineCore/Private/Graph/MovieGraphPin.cpp
@@ -120,3 +120,8 @@ int32 UMovieGraphPin::EdgeCount() const
 	return EdgeCount;
 }
 
+bool UMovieGraphPin::AllowsMultipleConnections() const
+{
+	// Always allow multiple connection on output pin
+	return IsOutputPin() || Properties.bAllowMultipleConnections;
+}
\ No newline at end of file
diff --git a/Engine/Plugins/MovieScene/MovieRenderPipeline/Source/MovieRenderPipelineCore/Public/Graph/MovieGraphPin.h b/Engine/Plugins/MovieScene/MovieRenderPipeline/Source/MovieRenderPipelineCore/Public/Graph/MovieGraphPin.h
--- a/Engine/Plugins/MovieScene/MovieRenderPipeline/Source/MovieRenderPipelineCore/Public/Graph/MovieGraphPin.h
+++ b/Engine/Plugins/MovieScene/MovieRenderPipeline/Source/MovieRenderPipelineCore/Public/Graph/MovieGraphPin.h
@@ -71,6 +71,8 @@ public:
 	bool IsConnected() const;
 	bool IsOutputPin() const;
 	int32 EdgeCount() const;
+	bool AllowsMultipleConnections() const;
+
 public:
 	// The node that this pin belongs to.
 	UPROPERTY(BlueprintReadOnly, Category = "Properties")
diff --git a/Engine/Plugins/MovieScene/MovieRenderPipeline/Source/MovieRenderPipelineEditor/Private/Graph/MovieGraphSchema.cpp b/Engine/Plugins/MovieScene/MovieRenderPipeline/Source/MovieRenderPipelineEditor/Private/Graph/MovieGraphSchema.cpp
--- a/Engine/Plugins/MovieScene/MovieRenderPipeline/Source/MovieRenderPipelineEditor/Private/Graph/MovieGraphSchema.cpp
+++ b/Engine/Plugins/MovieScene/MovieRenderPipeline/Source/MovieRenderPipelineEditor/Private/Graph/MovieGraphSchema.cpp
@@ -150,6 +150,23 @@ const FPinConnectionResponse UMovieGraphSchema::CanCreateConnection(const UEdGra
 	{
 		return FPinConnectionResponse(CONNECT_RESPONSE_DISALLOW, NSLOCTEXT("MoviePipeline", "PinTypeMismatchError", "Pin types don't match!"));
 	}
+
+	// Re-organize PinA/PinB to Input/Output by comparing internal directions to avoid having to check both cases depending on which
+	// direction the conneciton was made.
+	const UEdGraphPin* InputPin = nullptr;
+	const UEdGraphPin* OutputPin = nullptr;
+
+	if (!CategorizePinsByDirection(PinA, PinB, /*out*/ InputPin, /*out*/ OutputPin))
+	{
+		return FPinConnectionResponse(CONNECT_RESPONSE_DISALLOW, NSLOCTEXT("MoviePipeline", "PinDirectionMismatchError", "Directions are not compatible!"));
+	}
+
+	// We don't allow multiple things to be connected to an Input Pin
+	if(InputPin->HasAnyConnections())
+	{
+		const ECanCreateConnectionResponse ReplyBreakInputs = (PinA == InputPin) ? CONNECT_RESPONSE_BREAK_OTHERS_A : CONNECT_RESPONSE_BREAK_OTHERS_B;
+		return FPinConnectionResponse(ReplyBreakInputs, NSLOCTEXT("MoviePipeline", "PinInputReplaceExisting","Replace existing input connections"));
+	}
 	
 	// Make sure the pins are not on the same node
 	return FPinConnectionResponse(CONNECT_RESPONSE_MAKE, NSLOCTEXT("MoviePipeline", "PinConnect", "Connect nodes"));
