MoviePipeline: Added a new "Active Render Settings" tab (a debug view) to the graph's asset editor for previewing the graph's output. Currently it is restricted to viewing the output node, but eventually it will be capable of previewing the output at any arbitrary node.

diff --git a/Engine/Plugins/MovieScene/MovieRenderPipeline/Source/MovieRenderPipelineCore/Private/Graph/MovieGraphConfig.cpp b/Engine/Plugins/MovieScene/MovieRenderPipeline/Source/MovieRenderPipelineCore/Private/Graph/MovieGraphConfig.cpp
--- a/Engine/Plugins/MovieScene/MovieRenderPipeline/Source/MovieRenderPipelineCore/Private/Graph/MovieGraphConfig.cpp
+++ b/Engine/Plugins/MovieScene/MovieRenderPipeline/Source/MovieRenderPipelineCore/Private/Graph/MovieGraphConfig.cpp
@@ -600,7 +600,7 @@ bool UMovieGraphConfig::DeleteOutputMember(UMovieGraphOutput* OutputMemberToDele
 	return false;
 }
 
-FBoolProperty* UMovieGraphConfig::FindOverridePropertyForRealProperty(UClass* InClass, const FProperty* InRealProperty) const
+FBoolProperty* UMovieGraphConfig::FindOverridePropertyForRealProperty(UClass* InClass, const FProperty* InRealProperty)
 {
 	static const FName EditConditionKey = FName("EditCondition");
 	if (!ensure(InClass && InRealProperty))
@@ -807,6 +807,12 @@ UMovieGraphEvaluatedConfig* UMovieGraphConfig::CreateFlattenedGraph(const FMovie
 		// but we also want each named branch to have a complete, resolved set of configs.
 		for (UMovieGraphPin* InputPin : OutputNode->GetInputPins())
 		{
+			if (!InputPin->Properties.bIsBranch)
+			{
+				// TODO: For now, only branches can be evaluated
+				continue;
+			}
+			
 			const FName BranchName = InputPin->Properties.Label;
 			FMovieGraphEvaluatedBranchConfig& BranchConfig = NewContext->BranchConfigMapping.Add(BranchName);
 			
diff --git a/Engine/Plugins/MovieScene/MovieRenderPipeline/Source/MovieRenderPipelineCore/Public/Graph/MovieGraphConfig.h b/Engine/Plugins/MovieScene/MovieRenderPipeline/Source/MovieRenderPipelineCore/Public/Graph/MovieGraphConfig.h
--- a/Engine/Plugins/MovieScene/MovieRenderPipeline/Source/MovieRenderPipelineCore/Public/Graph/MovieGraphConfig.h
+++ b/Engine/Plugins/MovieScene/MovieRenderPipeline/Source/MovieRenderPipelineCore/Public/Graph/MovieGraphConfig.h
@@ -407,11 +407,13 @@ public:
 	UFUNCTION(BlueprintCallable, Category="Experimental")
 	UMovieGraphEvaluatedConfig* CreateFlattenedGraph(const FMovieGraphTraversalContext& InContext);
 
+	/** Given a class and FProperty that belongs to that class, search for a FBoolProperty that matches the name "bOverride_<name of InRealProperty>. */
+	static FBoolProperty* FindOverridePropertyForRealProperty(UClass* InClass, const FProperty* InRealProperty);
+
 protected:
 	/** Copies properties in FromNode that are marked for override into ToNode, but only if ToNode doesn't already override that value. */
 	void CopyOverriddenProperties(UMovieGraphNode* FromNode, UMovieGraphNode* ToNode);
-	/** Given a property you want to override (passes IsPropertyOverrideable) look for a matching bOverride_ named bool property. */
-	FBoolProperty* FindOverridePropertyForRealProperty(UClass* InClass, const FProperty* InRealProperty) const;
+	
 	/** Find all "Overrideable" marked properties, then find their edit condition properties, then set those to false. */
 	void InitializeFlattenedNode(UMovieGraphNode* InNode);
 
diff --git a/Engine/Plugins/MovieScene/MovieRenderPipeline/Source/MovieRenderPipelineCore/Public/Graph/MovieGraphPlaceholderNodes.h b/Engine/Plugins/MovieScene/MovieRenderPipeline/Source/MovieRenderPipelineCore/Public/Graph/MovieGraphPlaceholderNodes.h
--- a/Engine/Plugins/MovieScene/MovieRenderPipeline/Source/MovieRenderPipelineCore/Public/Graph/MovieGraphPlaceholderNodes.h
+++ b/Engine/Plugins/MovieScene/MovieRenderPipeline/Source/MovieRenderPipelineCore/Public/Graph/MovieGraphPlaceholderNodes.h
@@ -85,7 +85,7 @@ public:
 
 /** A node which configures the global game overrides. */
 UCLASS()
-class MOVIERENDERPIPELINECORE_API UMovieGraphGlobalGameOverridesNode : public UMovieGraphNode
+class MOVIERENDERPIPELINECORE_API UMovieGraphGlobalGameOverridesNode : public UMovieGraphSettingNode
 {
 	GENERATED_BODY()
 
diff --git a/Engine/Plugins/MovieScene/MovieRenderPipeline/Source/MovieRenderPipelineCore/Public/Graph/Nodes/MovieGraphDeferredRenderPassNode.h b/Engine/Plugins/MovieScene/MovieRenderPipeline/Source/MovieRenderPipelineCore/Public/Graph/Nodes/MovieGraphDeferredRenderPassNode.h
--- a/Engine/Plugins/MovieScene/MovieRenderPipeline/Source/MovieRenderPipelineCore/Public/Graph/Nodes/MovieGraphDeferredRenderPassNode.h
+++ b/Engine/Plugins/MovieScene/MovieRenderPipeline/Source/MovieRenderPipelineCore/Public/Graph/Nodes/MovieGraphDeferredRenderPassNode.h
@@ -50,7 +50,7 @@ namespace UE::MovieGraph
 }
 
 UCLASS()
-class UMovieGraphDeferredRenderPassNode : public UMovieGraphRenderPassNode
+class MOVIERENDERPIPELINECORE_API UMovieGraphDeferredRenderPassNode : public UMovieGraphRenderPassNode
 {
 	GENERATED_BODY()
 public:
diff --git a/Engine/Plugins/MovieScene/MovieRenderPipeline/Source/MovieRenderPipelineCore/Public/Graph/Nodes/MovieGraphRenderPassNode.h b/Engine/Plugins/MovieScene/MovieRenderPipeline/Source/MovieRenderPipelineCore/Public/Graph/Nodes/MovieGraphRenderPassNode.h
--- a/Engine/Plugins/MovieScene/MovieRenderPipeline/Source/MovieRenderPipelineCore/Public/Graph/Nodes/MovieGraphRenderPassNode.h
+++ b/Engine/Plugins/MovieScene/MovieRenderPipeline/Source/MovieRenderPipelineCore/Public/Graph/Nodes/MovieGraphRenderPassNode.h
@@ -17,7 +17,7 @@ struct FMovieGraphTimeStepData;
 * number of instances (decoupled from the number of times the node is used in the graph).
 */
 UCLASS(Abstract)
-class UMovieGraphRenderPassNode : public UMovieGraphSettingNode
+class MOVIERENDERPIPELINECORE_API UMovieGraphRenderPassNode : public UMovieGraphSettingNode
 {
 	GENERATED_BODY()
 public:
diff --git a/Engine/Plugins/MovieScene/MovieRenderPipeline/Source/MovieRenderPipelineEditor/Private/Graph/MovieGraphAssetToolkit.cpp b/Engine/Plugins/MovieScene/MovieRenderPipeline/Source/MovieRenderPipelineEditor/Private/Graph/MovieGraphAssetToolkit.cpp
--- a/Engine/Plugins/MovieScene/MovieRenderPipeline/Source/MovieRenderPipelineEditor/Private/Graph/MovieGraphAssetToolkit.cpp
+++ b/Engine/Plugins/MovieScene/MovieRenderPipeline/Source/MovieRenderPipelineEditor/Private/Graph/MovieGraphAssetToolkit.cpp
@@ -8,6 +8,7 @@
 #include "MovieEdGraphNode.h"
 #include "MovieGraphSchema.h"
 #include "PropertyEditorModule.h"
+#include "SMovieGraphActiveRenderSettingsTabContent.h"
 #include "SMovieGraphMembersTabContent.h"
 #include "ToolMenus.h"
 #include "Widgets/Docking/SDockTab.h"
@@ -19,6 +20,7 @@ const FName FMovieGraphAssetToolkit::AppIdentifier(TEXT("MovieGraphAssetEditorAp
 const FName FMovieGraphAssetToolkit::GraphTabId(TEXT("MovieGraphAssetToolkit"));
 const FName FMovieGraphAssetToolkit::DetailsTabId(TEXT("MovieGraphAssetToolkitDetails"));
 const FName FMovieGraphAssetToolkit::MembersTabId(TEXT("MovieGraphAssetToolkitMembers"));
+const FName FMovieGraphAssetToolkit::ActiveRenderSettingsTabId(TEXT("MovieGraphAssetToolkitActiveRenderSettings"));
 
 // Temporary cvar to enable/disable upgrading to a graph-based configuration
 static TAutoConsoleVariable<bool> CVarMoviePipelineEnableRenderGraph(
@@ -53,6 +55,11 @@ void FMovieGraphAssetToolkit::RegisterTabSpawners(const TSharedRef<FTabManager>&
 		.SetDisplayName(LOCTEXT("MembersTab", "Members"))
 		.SetGroup(WorkspaceMenuCategory.ToSharedRef())
 		.SetIcon(FSlateIcon(FAppStyle::GetAppStyleSetName(), "LevelEditor.Tabs.Outliner"));
+
+	InTabManager->RegisterTabSpawner(ActiveRenderSettingsTabId, FOnSpawnTab::CreateSP(this, &FMovieGraphAssetToolkit::SpawnTab_RenderGraphActiveRenderSettings))
+		.SetDisplayName(LOCTEXT("ActiveRenderSettingsTab", "Active Render Settings"))
+		.SetGroup(WorkspaceMenuCategory.ToSharedRef())
+		.SetIcon(FSlateIcon(FAppStyle::GetAppStyleSetName(), "LevelEditor.Tabs.Debug"));
 }
 
 void FMovieGraphAssetToolkit::UnregisterTabSpawners(const TSharedRef<FTabManager>& InTabManager)
@@ -62,6 +69,7 @@ void FMovieGraphAssetToolkit::UnregisterTabSpawners(const TSharedRef<FTabManager
 	InTabManager->UnregisterTabSpawner(GraphTabId);
 	InTabManager->UnregisterTabSpawner(DetailsTabId);
 	InTabManager->UnregisterTabSpawner(MembersTabId);
+	InTabManager->UnregisterTabSpawner(ActiveRenderSettingsTabId);
 }
 
 void FMovieGraphAssetToolkit::InitMovieGraphAssetToolkit(const EToolkitMode::Type Mode, const TSharedPtr<class IToolkitHost>& InitToolkitHost, UMovieGraphConfig* InitGraph)
@@ -71,7 +79,7 @@ void FMovieGraphAssetToolkit::InitMovieGraphAssetToolkit(const EToolkitMode::Typ
 	// Note: Changes to the layout should include a increment to the layout's ID, i.e.
 	// MoviePipelineRenderGraphEditor[X] -> MoviePipelineRenderGraphEditor[X+1]. Otherwise, layouts may be messed up
 	// without a full reset to layout defaults inside the editor.
-	const FName LayoutString = FName("MoviePipelineRenderGraphEditor1");
+	const FName LayoutString = FName("MoviePipelineRenderGraphEditor2");
 
 	// Override the Default Layout provided by FBaseAssetToolkit to hide the viewport and details panel tabs.
 	const TSharedRef<FTabManager::FLayout> Layout = FTabManager::NewLayout(FName(LayoutString))
@@ -98,9 +106,21 @@ void FMovieGraphAssetToolkit::InitMovieGraphAssetToolkit(const EToolkitMode::Typ
 				)
 				->Split
 				(
-					FTabManager::NewStack()
+					FTabManager::NewSplitter()
+					->SetOrientation(Orient_Vertical)
 					->SetSizeCoefficient(0.2f)
-					->AddTab(DetailsTabId, ETabState::OpenedTab)
+					->Split
+					(
+					FTabManager::NewStack()
+						->SetSizeCoefficient(0.5f)
+						->AddTab(DetailsTabId, ETabState::OpenedTab)
+					)
+					->Split
+					(
+					FTabManager::NewStack()
+						->SetSizeCoefficient(0.5f)
+						->AddTab(ActiveRenderSettingsTabId, ETabState::OpenedTab)
+					)
 				)
 			)
 		);
@@ -214,6 +234,17 @@ TSharedRef<SDockTab> FMovieGraphAssetToolkit::SpawnTab_RenderGraphDetails(const
 		];
 }
 
+TSharedRef<SDockTab> FMovieGraphAssetToolkit::SpawnTab_RenderGraphActiveRenderSettings(const FSpawnTabArgs& Args)
+{
+	return SNew(SDockTab)
+		.TabColorScale(GetTabColorScale())
+		.Label(LOCTEXT("ActiveRenderSettings_Title", "Active Render Settings"))
+		[
+			SAssignNew(ActiveRenderSettingsTabContent, SMovieGraphActiveRenderSettingsTabContent)
+			.Graph(InitialGraph)
+		];
+}
+
 void FMovieGraphAssetToolkit::BindGraphCommands()
 {
 	ToolkitCommands->MapAction(FGenericCommands::Get().Delete,
diff --git a/Engine/Plugins/MovieScene/MovieRenderPipeline/Source/MovieRenderPipelineEditor/Private/Graph/MovieGraphAssetToolkit.h b/Engine/Plugins/MovieScene/MovieRenderPipeline/Source/MovieRenderPipelineEditor/Private/Graph/MovieGraphAssetToolkit.h
--- a/Engine/Plugins/MovieScene/MovieRenderPipeline/Source/MovieRenderPipelineEditor/Private/Graph/MovieGraphAssetToolkit.h
+++ b/Engine/Plugins/MovieScene/MovieRenderPipeline/Source/MovieRenderPipelineEditor/Private/Graph/MovieGraphAssetToolkit.h
@@ -31,6 +31,7 @@ private:
 	TSharedRef<SDockTab> SpawnTab_RenderGraphEditor(const FSpawnTabArgs& Args);
 	TSharedRef<SDockTab> SpawnTab_RenderGraphMembers(const FSpawnTabArgs& Args);
 	TSharedRef<SDockTab> SpawnTab_RenderGraphDetails(const FSpawnTabArgs& Args);
+	TSharedRef<SDockTab> SpawnTab_RenderGraphActiveRenderSettings(const FSpawnTabArgs& Args);
 
 	/** Register commands used by the editor. */
 	void BindGraphCommands();
@@ -54,6 +55,9 @@ private:
 	/** The widget which encapsulates all of the content in the Members tab */
 	TSharedPtr<class SMovieGraphMembersTabContent> MembersTabContent;
 
+	/** The widget which encapsulates all of the content in the Active Render Settings tab */
+	TSharedPtr<class SMovieGraphActiveRenderSettingsTabContent> ActiveRenderSettingsTabContent;
+	
 	/** The graph that the editor was initialized with */
 	TObjectPtr<UMovieGraphConfig> InitialGraph;
 
@@ -64,4 +68,5 @@ private:
 	static const FName GraphTabId;
 	static const FName DetailsTabId;
 	static const FName MembersTabId;
+	static const FName ActiveRenderSettingsTabId;
 };
diff --git a/Engine/Plugins/MovieScene/MovieRenderPipeline/Source/MovieRenderPipelineEditor/Private/Graph/SMovieGraphActiveRenderSettingsTabContent.cpp b/Engine/Plugins/MovieScene/MovieRenderPipeline/Source/MovieRenderPipelineEditor/Private/Graph/SMovieGraphActiveRenderSettingsTabContent.cpp
new file mode 100644
--- /dev/null
+++ b/Engine/Plugins/MovieScene/MovieRenderPipeline/Source/MovieRenderPipelineEditor/Private/Graph/SMovieGraphActiveRenderSettingsTabContent.cpp
@@ -0,0 +1,374 @@
+// Copyright Epic Games, Inc. All Rights Reserved.
+
+#include "SMovieGraphActiveRenderSettingsTabContent.h"
+
+#include "Graph/MovieGraphConfig.h"
+#include "Graph/MovieGraphNode.h"
+#include "Graph/Nodes/MovieGraphRenderPassNode.h"
+#include "SPositiveActionButton.h"
+
+#define LOCTEXT_NAMESPACE "SMoviePipelineActiveRenderSettingsTabContent"
+
+const FName SMovieGraphActiveRenderSettingsTreeItem::ColumnID_Name(TEXT("Name"));
+const FName SMovieGraphActiveRenderSettingsTreeItem::ColumnID_Value(TEXT("Value"));
+
+const FName FActiveRenderSettingsTreeElement::RootName_Globals(TEXT("Globals"));
+const FName FActiveRenderSettingsTreeElement::RootName_Branches(TEXT("Branches"));
+
+FActiveRenderSettingsTreeElement::FActiveRenderSettingsTreeElement(const FName& InName, const EElementType InType)
+	: Name(InName)
+	, Type(InType)
+{
+}
+
+FString FActiveRenderSettingsTreeElement::GetValue() const
+{
+	if (SettingsProperty && SettingsNode)
+	{
+		FString ValueString;
+		SettingsProperty->ExportTextItem_InContainer(ValueString, SettingsNode, nullptr, nullptr, PPF_None);
+		return ValueString;
+	}
+
+	return FString();
+}
+
+bool FActiveRenderSettingsTreeElement::IsBranchRenderable() const
+{
+	if (Type != EElementType::NamedBranch)
+	{
+		return false;
+	}
+
+	// The branch is renderable if there's a render pass node in the branch's children
+	for (const TSharedPtr<FActiveRenderSettingsTreeElement>& Child : GetChildren())
+	{
+		if (Child.IsValid() && Child->SettingsNode && Child->SettingsNode->IsA<UMovieGraphRenderPassNode>())
+		{
+			return true;
+		}
+	}
+
+	return false;
+}
+
+const TArray<TSharedPtr<FActiveRenderSettingsTreeElement>>& FActiveRenderSettingsTreeElement::GetChildren() const
+{
+	auto MakeElementFromNode = [this](const TObjectPtr<UMovieGraphNode>& Node) -> TSharedPtr<FActiveRenderSettingsTreeElement>
+	{
+		TSharedPtr<FActiveRenderSettingsTreeElement> Element =
+			MakeShared<FActiveRenderSettingsTreeElement>(FName(*Node->GetNodeTitle().ToString()), EElementType::Node);
+		Element->SettingsNode = Node;
+		Element->FlattenedGraph = FlattenedGraph;
+
+		return MoveTemp(Element);
+	};
+
+	// Returned the cached children if they are available
+	if (!ChildrenCache.IsEmpty())
+	{
+		return ChildrenCache;
+	}
+	
+	if (!FlattenedGraph.IsValid())
+	{
+		// This should be an empty array
+		return ChildrenCache;
+	}
+
+	// For the root "Globals" element: all nodes which were found in the Globals branch should be children
+	if ((Name == RootName_Globals) && (Type == EElementType::Root))
+	{
+		if (const FMovieGraphEvaluatedBranchConfig* BranchConfig = FlattenedGraph->BranchConfigMapping.Find(RootName_Globals))
+		{
+			for (const TObjectPtr<UMovieGraphNode>& Node : BranchConfig->GetNodes())
+			{
+				ChildrenCache.Add(MakeElementFromNode(Node));
+			}
+		}
+
+		return ChildrenCache;
+	}
+
+	// For the root "Branches" element: all branches which were found should be children
+	if ((Name == RootName_Branches) && (Type == EElementType::Root))
+	{
+		for (const auto& Pair : FlattenedGraph->BranchConfigMapping)
+		{
+			if (Pair.Key != RootName_Globals)
+			{
+				TSharedPtr<FActiveRenderSettingsTreeElement> Element =
+					MakeShared<FActiveRenderSettingsTreeElement>(FName(Pair.Key), EElementType::NamedBranch);
+				Element->FlattenedGraph = FlattenedGraph;
+				
+				ChildrenCache.Add(MoveTemp(Element));
+			}
+		}
+
+		return ChildrenCache;
+	}
+
+	// For named branch elements: all the nodes under the branch should be children
+	if (Type == EElementType::NamedBranch)
+	{
+		const FMovieGraphEvaluatedBranchConfig* BranchConfig = FlattenedGraph->BranchConfigMapping.Find(Name);
+		if (!BranchConfig)
+		{
+			return ChildrenCache;
+		}
+
+		for (const TObjectPtr<UMovieGraphNode>& Node : BranchConfig->GetNodes())
+		{
+			ChildrenCache.Add(MakeElementFromNode(Node));
+		}
+
+		return ChildrenCache;
+	}
+
+	// For node elements: all overrideable properties on the node should be children
+	if (Type == EElementType::Node)
+	{
+		for (TFieldIterator<FProperty> PropertyIterator(SettingsNode->GetClass()); PropertyIterator; ++PropertyIterator)
+		{
+			FProperty* NodeProperty = *PropertyIterator;
+			if (UMovieGraphConfig::FindOverridePropertyForRealProperty(SettingsNode->GetClass(), NodeProperty))
+			{
+				TSharedPtr<FActiveRenderSettingsTreeElement> Element =
+					MakeShared<FActiveRenderSettingsTreeElement>(NodeProperty->GetFName(), EElementType::Property);
+				Element->SettingsNode = SettingsNode;
+				Element->SettingsProperty = NodeProperty;
+				Element->FlattenedGraph = FlattenedGraph;
+				
+				ChildrenCache.Add(MoveTemp(Element));
+			}
+		}
+	}
+	
+	return ChildrenCache;
+}
+
+void FActiveRenderSettingsTreeElement::ClearCachedChildren() const
+{
+	ChildrenCache.Empty();
+}
+
+void SMovieGraphActiveRenderSettingsTreeItem::Construct(
+	const FArguments& InArgs, const TSharedRef<STableViewBase>& InOwnerTable, const TSharedPtr<FActiveRenderSettingsTreeElement>& InTreeElement)
+{
+	WeakTreeElement = InTreeElement;
+
+	SMultiColumnTableRow<TSharedPtr<FActiveRenderSettingsTreeElement>>::Construct(
+		SMultiColumnTableRow<TSharedPtr<FActiveRenderSettingsTreeElement>>::FArguments()
+			.Padding(0)	
+		, InOwnerTable);
+}
+
+TSharedRef<SWidget> SMovieGraphActiveRenderSettingsTreeItem::GenerateWidgetForColumn(const FName& ColumnName)
+{
+	const TSharedPtr<FActiveRenderSettingsTreeElement> TreeElement = WeakTreeElement.Pin();
+	if (!TreeElement.IsValid())
+	{
+		return SNullWidget::NullWidget;
+	}
+	
+	if (ColumnName == ColumnID_Name)
+	{
+		return SNew(SHorizontalBox)
+			+ SHorizontalBox::Slot()
+			.AutoWidth()
+			[
+				SNew(SExpanderArrow, SharedThis(this))
+				.IndentAmount(16)
+				.ShouldDrawWires(false)
+			]
+
+			+ SHorizontalBox::Slot()
+			.AutoWidth()
+			.Padding(0, 0, 2.f, 0)
+			[
+				SNew(SImage)
+				.Image(FSlateIcon(FAppStyle::GetAppStyleSetName(), "ClassIcon.Texture2D").GetIcon())
+				.ToolTipText(LOCTEXT("RenderableBranchTooltip", "This branch outputs rendered images."))
+				.Visibility_Lambda([TreeElement]()
+				{
+					return TreeElement->IsBranchRenderable() ? EVisibility::Visible : EVisibility::Collapsed;
+				})
+			]
+
+			+ SHorizontalBox::Slot()
+			.AutoWidth()
+			.VAlign(VAlign_Center)
+			.HAlign(HAlign_Right)
+			[
+				SNew(STextBlock)
+				.Text(FText::FromName(TreeElement->Name))
+				.Font_Lambda([TreeElement]()
+				{
+					switch (TreeElement->Type)
+					{
+					case FActiveRenderSettingsTreeElement::EElementType::Root:
+						return FCoreStyle::GetDefaultFontStyle("Bold", 12);
+					case FActiveRenderSettingsTreeElement::EElementType::NamedBranch:
+						return FCoreStyle::GetDefaultFontStyle("Bold", 10);
+					case FActiveRenderSettingsTreeElement::EElementType::Property:
+						return FCoreStyle::GetDefaultFontStyle("Italic", 9);
+					default:
+						return FCoreStyle::GetDefaultFontStyle("Regular", 9);
+					}
+				})
+				.ColorAndOpacity_Lambda([TreeElement]()
+				{
+					return (TreeElement->Type == FActiveRenderSettingsTreeElement::EElementType::Property)
+						? FSlateColor::UseSubduedForeground()
+						: FSlateColor::UseForeground();
+				})
+			];
+	}
+
+	if (ColumnName == ColumnID_Value)
+	{
+		return SNew(SHorizontalBox)
+			+ SHorizontalBox::Slot()
+			.AutoWidth()
+			.VAlign(VAlign_Center)
+			.HAlign(HAlign_Right)
+			[
+				SNew(STextBlock)
+				.Text(FText::FromString(TreeElement->GetValue()))
+				.Font(FCoreStyle::GetDefaultFontStyle("Regular", 9))
+				.ColorAndOpacity(FSlateColor::UseSubduedForeground())
+			];
+	}
+
+	return SNullWidget::NullWidget;
+}
+
+void SMovieGraphActiveRenderSettingsTabContent::Construct(const FArguments& InArgs)
+{
+	CurrentGraph = InArgs._Graph;
+
+	// Add the two default root elements, "Globals" and "Branches", which are always visible. These elements will
+	// generate the children which should be displayed under them in the tree.
+	RootElements.Add(MakeShared<FActiveRenderSettingsTreeElement>(
+		FActiveRenderSettingsTreeElement::RootName_Globals, FActiveRenderSettingsTreeElement::EElementType::Root));
+	RootElements.Add(MakeShared<FActiveRenderSettingsTreeElement>(
+		FActiveRenderSettingsTreeElement::RootName_Branches, FActiveRenderSettingsTreeElement::EElementType::Root));
+	
+	ChildSlot
+	[
+		SNew(SVerticalBox)
+		+ SVerticalBox::Slot()
+		.AutoHeight()
+		.Padding(3.f)
+		[
+			SNew(SHorizontalBox)
+			+ SHorizontalBox::Slot()
+			.AutoWidth()
+			[
+				SNew(SPositiveActionButton)
+				.Text(FText::FromString("Evaluate Graph"))
+				.Icon(FAppStyle::GetBrush("Icons.Refresh"))
+				.OnClicked(this, &SMovieGraphActiveRenderSettingsTabContent::OnEvaluateGraphClicked)
+			]
+		]
+
+		+ SVerticalBox::Slot()
+		[
+			SAssignNew(TreeView, STreeView<TSharedPtr<FActiveRenderSettingsTreeElement>>)
+			.ItemHeight(28)
+			.TreeItemsSource(&RootElements)
+			.SelectionMode(ESelectionMode::Single)
+			.OnGenerateRow(this, &SMovieGraphActiveRenderSettingsTabContent::GenerateTreeRow)
+			.OnGetChildren(this, &SMovieGraphActiveRenderSettingsTabContent::GetChildrenForTree)
+			.HeaderRow
+			(
+				SNew(SHeaderRow)
+
+				+ SHeaderRow::Column(SMovieGraphActiveRenderSettingsTreeItem::ColumnID_Name)
+				.DefaultLabel(LOCTEXT("NameColumnLabel", "Name"))
+				.SortMode(EColumnSortMode::None)
+				.FillWidth(0.75f)
+
+				+ SHeaderRow::Column(SMovieGraphActiveRenderSettingsTreeItem::ColumnID_Value)
+				.DefaultLabel(LOCTEXT("ValueColumnLabel", "Value"))
+				.SortMode(EColumnSortMode::None)
+				.FillWidth(0.25f)
+			)
+			.Visibility_Lambda([this]()
+			{
+				return FlattenedGraph.IsValid() ? EVisibility::Visible : EVisibility::Collapsed;
+			})
+		]
+
+		+ SVerticalBox::Slot()
+		.Padding(5.f)
+		[
+			SNew(STextBlock)
+			.Text(LOCTEXT("NeedsEvaluationWarning", "Render settings have not been evaluated yet."))
+			.Font(FCoreStyle::GetDefaultFontStyle("Italic", 10))
+			.Visibility_Lambda([this]()
+			{
+				return FlattenedGraph.IsValid() ? EVisibility::Collapsed : EVisibility::Visible;
+			})
+		]
+	];
+}
+
+void SMovieGraphActiveRenderSettingsTabContent::TraverseGraph()
+{
+	if (!CurrentGraph.IsValid())
+	{
+		return;
+	}
+
+	// Clear out the resolved cached children under the root nodes
+	for (const TSharedPtr<FActiveRenderSettingsTreeElement>& RootElement : RootElements)
+	{
+		RootElement->ClearCachedChildren();
+	}
+
+	// TODO: Fill in the context via the UI (there are no widgets to provide this info yet)
+	FMovieGraphTraversalContext Context;
+
+	// Traverse the graph, and update the root elements
+	FlattenedGraph = TStrongObjectPtr(CurrentGraph->CreateFlattenedGraph(Context));
+	for (const TSharedPtr<FActiveRenderSettingsTreeElement>& RootElement : RootElements)
+	{
+		RootElement->FlattenedGraph = MakeWeakObjectPtr(FlattenedGraph.Get());
+	}
+	
+	TreeView->RequestTreeRefresh();
+}
+
+FReply SMovieGraphActiveRenderSettingsTabContent::OnEvaluateGraphClicked()
+{
+	TraverseGraph();
+	return FReply::Handled();
+}
+
+TSharedRef<ITableRow> SMovieGraphActiveRenderSettingsTabContent::GenerateTreeRow(
+	TSharedPtr<FActiveRenderSettingsTreeElement> InTreeElement, const TSharedRef<STableViewBase>& OwnerTable)
+{
+	return SNew(SMovieGraphActiveRenderSettingsTreeItem, OwnerTable, InTreeElement)
+		.Visibility_Lambda([InTreeElement]()
+		{
+			// If the root element has no children, hide it
+			if (InTreeElement->Type == FActiveRenderSettingsTreeElement::EElementType::Root)
+			{
+				if (InTreeElement->GetChildren().IsEmpty())
+				{
+					return EVisibility::Hidden;
+				}
+			}
+
+			return EVisibility::Visible;
+		});
+}
+
+void SMovieGraphActiveRenderSettingsTabContent::GetChildrenForTree(
+	TSharedPtr<FActiveRenderSettingsTreeElement> InItem, TArray<TSharedPtr<FActiveRenderSettingsTreeElement>>& OutChildren)
+{
+	OutChildren.Append(InItem->GetChildren());
+}
+
+#undef LOCTEXT_NAMESPACE
diff --git a/Engine/Plugins/MovieScene/MovieRenderPipeline/Source/MovieRenderPipelineEditor/Private/Graph/SMovieGraphActiveRenderSettingsTabContent.h b/Engine/Plugins/MovieScene/MovieRenderPipeline/Source/MovieRenderPipelineEditor/Private/Graph/SMovieGraphActiveRenderSettingsTabContent.h
new file mode 100644
--- /dev/null
+++ b/Engine/Plugins/MovieScene/MovieRenderPipeline/Source/MovieRenderPipelineEditor/Private/Graph/SMovieGraphActiveRenderSettingsTabContent.h
@@ -0,0 +1,140 @@
+// Copyright Epic Games, Inc. All Rights Reserved.
+
+#pragma once
+
+#include "UObject/StrongObjectPtr.h"
+#include "Widgets/SCompoundWidget.h"
+#include "Widgets/Views/STableRow.h"
+#include "Widgets/Views/STreeView.h"
+
+class UMovieGraphEvaluatedConfig;
+class UMovieGraphNode;
+
+/**
+ * Represents the data of one element displayed in the tree, which is one of the types in EElementType. Provides a way
+ * of getting the children under the element.
+ */
+class FActiveRenderSettingsTreeElement
+{
+public:
+	/** The type of this element. */
+	enum class EElementType : uint8
+	{
+		Root,			///< A root element; does not represent any actual data within the graph
+		NamedBranch,	///< Represents a named branch (output) within the graph
+		Node,			///< Represents a settings node within the graph
+		Property		///< Represents a property under an element with type EElementType::Node
+	};
+	
+	static const FName RootName_Globals;
+	static const FName RootName_Branches;
+	
+	FActiveRenderSettingsTreeElement(const FName& InName, const EElementType InType);
+
+	/** If this element represents a property, returns the string representation of the value. Otherwise returns an empty string. */
+	FString GetValue() const;
+
+	/** Determines if this element is a branch, and the branch has renderable output. */
+	bool IsBranchRenderable() const;
+
+	/**
+	 * Gets the child elements nested under this element. Returns a cached result if available, otherwise calculates the
+	 * children and returns the result. Call ClearCachedChildren() if the cache needs to be cleared (eg, when the tree
+	 * is being refreshed).
+	 */
+	const TArray<TSharedPtr<FActiveRenderSettingsTreeElement>>& GetChildren() const;
+
+	/** Clears the cached result of GetChildren(). */
+	void ClearCachedChildren() const;
+
+public:
+	/** The name of the element as it appears in the tree. */
+	FName Name;
+
+	/** The type of item in the tree this element represents. */
+	EElementType Type = EElementType::Root;
+
+	/**
+	 * If this element represents a node in the graph, this is the pointer to that node. If the element represents a
+	 * property, this is the pointer to the owning node.
+	 */
+	TObjectPtr<UMovieGraphNode> SettingsNode = nullptr;
+
+	/** If this element represents a property on a node, this is the pointer to that property. */
+	FProperty* SettingsProperty = nullptr;
+
+	/** The result of the graph traversal. */
+	TWeakObjectPtr<UMovieGraphEvaluatedConfig> FlattenedGraph = nullptr;
+
+private:
+	/** The (cached) child elements nested under this element. */
+	mutable TArray<TSharedPtr<FActiveRenderSettingsTreeElement>> ChildrenCache;
+};
+
+/**
+ * The widget that is responsible for generating the per-column content for each element in the tree.
+ */
+class SMovieGraphActiveRenderSettingsTreeItem : public SMultiColumnTableRow<TSharedPtr<FActiveRenderSettingsTreeElement>>
+{
+	SLATE_BEGIN_ARGS(SMovieGraphActiveRenderSettingsTreeItem) {}
+	SLATE_END_ARGS()
+
+	static const FName ColumnID_Name;
+	static const FName ColumnID_Value;
+
+	void Construct(const FArguments& InArgs, const TSharedRef<STableViewBase>& InOwnerTable,
+		const TSharedPtr<FActiveRenderSettingsTreeElement>& InTreeElement);
+
+	//~ Begin SMultiColumnTableRow overrides
+	virtual TSharedRef<SWidget> GenerateWidgetForColumn(const FName& ColumnName) override;
+	//~ End SMultiColumnTableRow overrides
+
+private:
+	/** The element that this tree item is displaying. */
+	TWeakPtr<FActiveRenderSettingsTreeElement> WeakTreeElement;
+};
+
+/**
+ * Contents of the "Active Render Settings" tab in the graph asset editor.
+ */
+class SMovieGraphActiveRenderSettingsTabContent : public SCompoundWidget
+{
+public:
+	SLATE_BEGIN_ARGS(SMovieGraphActiveRenderSettingsTabContent)
+		: _Graph(nullptr)
+	
+		{}
+		
+		/** The graph that is currently displayed. */
+		SLATE_ARGUMENT(class UMovieGraphConfig*, Graph)
+	
+	SLATE_END_ARGS();
+	
+	void Construct(const FArguments& InArgs);
+
+private:
+	/** Traverses the graph. Populates the FlattenedGraph member and refreshes the UI view. */
+	void TraverseGraph();
+
+	/** Handles the button click for evaluating the graph. */
+	FReply OnEvaluateGraphClicked();
+
+	/** Generates a row in the tree widget based on the provided element. */
+	TSharedRef<ITableRow> GenerateTreeRow(TSharedPtr<FActiveRenderSettingsTreeElement> InTreeElement, const TSharedRef<STableViewBase>& OwnerTable);
+
+	/** Gets the child elements in the tree for the provided element. */
+	void GetChildrenForTree(TSharedPtr<FActiveRenderSettingsTreeElement> InItem, TArray<TSharedPtr<FActiveRenderSettingsTreeElement>>& OutChildren);
+
+private:
+	/** The runtime graph that this UI gets data from. */
+	TWeakObjectPtr<UMovieGraphConfig> CurrentGraph = nullptr;
+
+	/** The main tree widget displayed in the tab. */
+	TSharedPtr<STreeView<TSharedPtr<FActiveRenderSettingsTreeElement>>> TreeView = nullptr;
+
+	/** The last result of a graph traversal. */
+	TStrongObjectPtr<UMovieGraphEvaluatedConfig> FlattenedGraph = nullptr;
+
+	/** The root-most elements in the tree, which are always present. */
+	TArray<TSharedPtr<FActiveRenderSettingsTreeElement>> RootElements;
+};
