MoviePipeline: Fixed a small memory leak with the backbuffer Render Target not getting released at the end of a render.

diff --git a/Engine/Plugins/MovieScene/MovieRenderPipeline/Source/MovieRenderPipelineRenderPasses/Private/MoviePipelineImagePassBase.cpp b/Engine/Plugins/MovieScene/MovieRenderPipeline/Source/MovieRenderPipelineRenderPasses/Private/MoviePipelineImagePassBase.cpp
--- a/Engine/Plugins/MovieScene/MovieRenderPipeline/Source/MovieRenderPipelineRenderPasses/Private/MoviePipelineImagePassBase.cpp
+++ b/Engine/Plugins/MovieScene/MovieRenderPipeline/Source/MovieRenderPipelineRenderPasses/Private/MoviePipelineImagePassBase.cpp
@@ -69,7 +69,7 @@ void UMoviePipelineImagePassBase::TeardownImpl()
 {
 	for (TPair<FIntPoint, TWeakObjectPtr<UTextureRenderTarget2D>>& TileRenderTargetIt : TileRenderTargets)
 	{
-		if (!TileRenderTargetIt.Value.IsValid())
+		if (TileRenderTargetIt.Value.IsValid())
 		{
 			TileRenderTargetIt.Value->RemoveFromRoot();
 		}
