MoviePipeline: Added copy/paste support to nodes in the Movie Graph editor.

diff --git a/Engine/Plugins/MovieScene/MovieRenderPipeline/Source/MovieRenderPipelineCore/Private/Graph/MovieGraphConfig.cpp b/Engine/Plugins/MovieScene/MovieRenderPipeline/Source/MovieRenderPipelineCore/Private/Graph/MovieGraphConfig.cpp
--- a/Engine/Plugins/MovieScene/MovieRenderPipeline/Source/MovieRenderPipelineCore/Private/Graph/MovieGraphConfig.cpp
+++ b/Engine/Plugins/MovieScene/MovieRenderPipeline/Source/MovieRenderPipelineCore/Private/Graph/MovieGraphConfig.cpp
@@ -330,6 +330,27 @@ bool UMovieGraphConfig::RemoveOutboundEdges(UMovieGraphNode* InNode, const FName
 	return bChanged;
 }
 
+void UMovieGraphConfig::AddNode(UMovieGraphNode* InNode)
+{
+	if (!InNode)
+	{
+		return;
+	}
+
+	InNode->SetFlags(RF_Transactional);
+
+	Modify();
+
+	// Reparent node to this graph
+	InNode->Rename(nullptr, this);
+
+	AllNodes.Add(InNode);
+
+#if WITH_EDITOR
+	OnGraphChangedDelegate.Broadcast();
+#endif
+}
+
 bool UMovieGraphConfig::RemoveNodes(TArray<UMovieGraphNode*> InNodes)
 {
 	bool bChanged = false;
diff --git a/Engine/Plugins/MovieScene/MovieRenderPipeline/Source/MovieRenderPipelineCore/Public/Graph/MovieGraphConfig.h b/Engine/Plugins/MovieScene/MovieRenderPipeline/Source/MovieRenderPipelineCore/Public/Graph/MovieGraphConfig.h
--- a/Engine/Plugins/MovieScene/MovieRenderPipeline/Source/MovieRenderPipelineCore/Public/Graph/MovieGraphConfig.h
+++ b/Engine/Plugins/MovieScene/MovieRenderPipeline/Source/MovieRenderPipelineCore/Public/Graph/MovieGraphConfig.h
@@ -355,8 +355,16 @@ public:
 	bool RemoveInboundEdges(UMovieGraphNode* InNode, const FName& InPinName);
 	bool RemoveOutboundEdges(UMovieGraphNode* InNode, const FName& InPinName);
 
+	/** 
+	* Add the specified node instance to the graph. This will rename the node to ensure the graph is the outer
+	* and then it will add it to the internal list of nodes used by the graph. See ConstructRuntimeNode if you
+	* want to construct a node by class and don't already have an instance.
+	*/
+	void AddNode(UMovieGraphNode* InNode);
+
 	/** Removes the specified node from the graph. */
 	bool RemoveNode(UMovieGraphNode* InNode);
+	/** Removes the specified nodes from the graph. */
 	bool RemoveNodes(TArray<UMovieGraphNode*> InNodes);
 
 	UMovieGraphNode* GetInputNode() const { return InputNode; }
@@ -435,7 +443,8 @@ public:
 	FOnMovieGraphOutputAdded OnGraphOutputAddedDelegate;
 	FOnMovieGraphNodesDeleted OnGraphNodesDeletedDelegate;
 #endif
-	
+
+protected:	
 	UPROPERTY()
 	TArray<TObjectPtr<UMovieGraphNode>> AllNodes;
 
@@ -444,7 +453,7 @@ public:
 
 	UPROPERTY()
 	TObjectPtr<UMovieGraphNode> OutputNode;
-
+public:
 #if WITH_EDITORONLY_DATA
 	// Not strongly typed to avoid a circular dependency between the editor only module
 	// and the runtime module, but it should be a UMoviePipelineEdGraph.
@@ -456,15 +465,12 @@ public:
 	T* ConstructRuntimeNode(TSubclassOf<UMovieGraphNode> PipelineGraphNodeClass = T::StaticClass())
 	{
 		// Construct a new object with ourselves as the outer, then keep track of it.
-		// ToDo: This is a runtime node, kept track in AllNodes, which is ultimately editor only. Probably
-		// because the system it's based on (SoundCues) have a root node and links nodes together later?
 		T* RuntimeNode = NewObject<T>(this, PipelineGraphNodeClass, NAME_None, RF_Transactional);
 		RuntimeNode->UpdateDynamicProperties();
 		RuntimeNode->UpdatePins();
 		RuntimeNode->Guid = FGuid::NewGuid();
-#if WITH_EDITOR
-		AllNodes.Add(RuntimeNode);
-#endif
+		
+		AddNode(RuntimeNode);
 		return RuntimeNode;
 	}
 
diff --git a/Engine/Plugins/MovieScene/MovieRenderPipeline/Source/MovieRenderPipelineCore/Public/Graph/MovieGraphPin.h b/Engine/Plugins/MovieScene/MovieRenderPipeline/Source/MovieRenderPipelineCore/Public/Graph/MovieGraphPin.h
--- a/Engine/Plugins/MovieScene/MovieRenderPipeline/Source/MovieRenderPipelineCore/Public/Graph/MovieGraphPin.h
+++ b/Engine/Plugins/MovieScene/MovieRenderPipeline/Source/MovieRenderPipelineCore/Public/Graph/MovieGraphPin.h
@@ -81,6 +81,10 @@ public:
 	UPROPERTY(BlueprintReadOnly, EditAnywhere, Category = "Properties", meta = (ShowOnlyInnerProperties))
 	FMovieGraphPinProperties Properties;
 
-	UPROPERTY(BlueprintReadOnly, Category = "Properties")
+	/** 
+	* A list of edges between pins. This is marked as TextExportTransient so that when we copy/paste nodes,
+	* we don't copy the edges, as they are rebuilt after paste based on the editor graph connections.
+	*/
+	UPROPERTY(BlueprintReadOnly, TextExportTransient, Category = "Properties")
 	TArray<TObjectPtr<UMovieGraphEdge>> Edges;
 };
\ No newline at end of file
diff --git a/Engine/Plugins/MovieScene/MovieRenderPipeline/Source/MovieRenderPipelineEditor/MovieRenderPipelineEditor.Build.cs b/Engine/Plugins/MovieScene/MovieRenderPipeline/Source/MovieRenderPipelineEditor/MovieRenderPipelineEditor.Build.cs
--- a/Engine/Plugins/MovieScene/MovieRenderPipeline/Source/MovieRenderPipelineEditor/MovieRenderPipelineEditor.Build.cs
+++ b/Engine/Plugins/MovieScene/MovieRenderPipeline/Source/MovieRenderPipelineEditor/MovieRenderPipelineEditor.Build.cs
@@ -69,6 +69,7 @@ public class MovieRenderPipelineEditor : ModuleRules
 				"StructUtils",
 				"GraphEditor",
 				"StructUtilsEditor",
+				"ApplicationCore",
 			}
         );
 
diff --git a/Engine/Plugins/MovieScene/MovieRenderPipeline/Source/MovieRenderPipelineEditor/Private/Graph/MovieEdGraphInputNode.h b/Engine/Plugins/MovieScene/MovieRenderPipeline/Source/MovieRenderPipelineEditor/Private/Graph/MovieEdGraphInputNode.h
--- a/Engine/Plugins/MovieScene/MovieRenderPipeline/Source/MovieRenderPipelineEditor/Private/Graph/MovieEdGraphInputNode.h
+++ b/Engine/Plugins/MovieScene/MovieRenderPipeline/Source/MovieRenderPipelineEditor/Private/Graph/MovieEdGraphInputNode.h
@@ -16,6 +16,9 @@ public:
 	virtual FText GetNodeTitle(ENodeTitleType::Type TitleType) const override;
 	virtual void AllocateDefaultPins() override;
 	virtual bool CanUserDeleteNode() const override { return false; }
-	virtual bool CanDuplicateNode() const override { return false; }
+	// This is set to true so that the node is deserialized on paste, required for CanPasteHere to be called.
+	virtual bool CanDuplicateNode() const override { return true; } 
+	// And this is set to false so the paste doesn't actually go through, and instead we get the editor prompt notifying user.
+	virtual bool CanPasteHere(const UEdGraph* TargetGraph) const { return false; }
 	// ~End UEdGraphNode interface
 };
diff --git a/Engine/Plugins/MovieScene/MovieRenderPipeline/Source/MovieRenderPipelineEditor/Private/Graph/MovieEdGraphNode.cpp b/Engine/Plugins/MovieScene/MovieRenderPipeline/Source/MovieRenderPipelineEditor/Private/Graph/MovieEdGraphNode.cpp
--- a/Engine/Plugins/MovieScene/MovieRenderPipeline/Source/MovieRenderPipelineEditor/Private/Graph/MovieEdGraphNode.cpp
+++ b/Engine/Plugins/MovieScene/MovieRenderPipeline/Source/MovieRenderPipelineEditor/Private/Graph/MovieEdGraphNode.cpp
@@ -4,6 +4,7 @@
 
 #include "Graph/MovieGraphPin.h"
 #include "Graph/MovieGraphNode.h"
+#include "Graph/MovieGraphConfig.h"
 #include "EdGraph/EdGraphPin.h"
 #include "Misc/TransactionObjectEvent.h"
 #include "MovieEdGraph.h"
@@ -11,6 +12,8 @@
 #include "PropertyBag.h"
 #include "ToolMenu.h"
 #include "EdGraph/EdGraphSchema.h"
+#include "GraphEditorActions.h"
+#include "Framework/Commands/GenericCommands.h"
 
 #define LOCTEXT_NAMESPACE "MoviePipelineEdGraphNodeBase"
 
@@ -131,7 +134,8 @@ FText UMoviePipelineEdGraphNode::GetNodeTitle(ENodeTitleType::Type TitleType) co
 
 FText UMoviePipelineEdGraphNode::GetTooltipText() const
 {
-	return LOCTEXT("GraphTestTooltip", "Test Tooltip");
+	// Return the UObject name for now for debugging purposes
+	return FText::FromString(GetName());
 }
 
 void UMoviePipelineEdGraphNode::GetNodeContextMenuActions(UToolMenu* Menu, UGraphNodeContextMenuContext* Context) const
@@ -342,6 +346,14 @@ void UMoviePipelineEdGraphNodeBase::OnRuntimeNodeChanged(const UMovieGraphNode*
 
 void UMoviePipelineEdGraphNodeBase::ReconstructNode()
 {
+	// Don't reconstruct the node during copy-paste. If we allow reconstruction,
+	// then the editor graph reconstructs connections to previous nodes that were
+	// not included in the copy/paste. This does not affect connections within the copy/pasted nodes.
+	if (bDisableReconstructNode)
+	{
+		return;
+	}
+
 	ReconstructPins();
 
 	UMoviePipelineEdGraph* Graph = CastChecked<UMoviePipelineEdGraph>(GetGraph());
@@ -388,4 +400,130 @@ void UMoviePipelineEdGraphNodeBase::ReconstructPins()
 	}
 }
 
+void UMoviePipelineEdGraphNodeBase::PrepareForCopying()
+{
+	if (RuntimeNode)
+	{
+		// Temporarily take ownership of the model's node, so that it is not deleted when copying.
+		// This is restored in PostCopy
+		RuntimeNode->Rename(nullptr, this, REN_DontCreateRedirectors | REN_DoNotDirty);
+	}
+}
+
+void UMoviePipelineEdGraphNodeBase::PostCopy()
+{
+	if (RuntimeNode)
+	{
+		// We briefly took ownership of the runtime node to create the copy/paste buffer,
+		// restore the ownership back to the owning graph.
+		UMoviePipelineEdGraph* MovieGraphEditorGraph = CastChecked<UMoviePipelineEdGraph>(GetGraph());
+		UMovieGraphConfig* RuntimeGraph = MovieGraphEditorGraph->GetPipelineGraph();
+		check(RuntimeGraph);
+		RuntimeNode->Rename(nullptr, RuntimeGraph, REN_DontCreateRedirectors | REN_DoNotDirty);
+	}
+}
+
+void UMoviePipelineEdGraphNodeBase::PostPasteNode()
+{
+	bDisableReconstructNode = true;
+}
+
+void UMoviePipelineEdGraphNodeBase::PostPaste()
+{
+	if (RuntimeNode)
+	{
+		// The editor nodes preserved the connections between nodes when copying/pasting
+		// but we intentionally don't preserve the edges of the runtime graph when copying
+		// (as the ownership isn't always clear given both input/output edges, which node owns
+		// the edge, the one inside the copied graph? Or the one outside it?), so instead 
+		// we just rebuild the runtime edge connections based on the editor graph connectivity. 
+		RebuildRuntimeEdgesFromPins();
+
+		// Ensure we're listening to the delegate for this pasted node, because we may have skipped ::Construct
+		RuntimeNode->OnNodeChangedDelegate.AddUObject(this, &UMoviePipelineEdGraphNodeBase::OnRuntimeNodeChanged);
+		RuntimeNode->SetNodePosX(NodePosX);
+		RuntimeNode->SetNodePosY(NodePosY);
+	}
+
+	bDisableReconstructNode = false;
+}
+
+void UMoviePipelineEdGraphNodeBase::RebuildRuntimeEdgesFromPins()
+{
+	check(RuntimeNode);
+	
+	for (UEdGraphPin* Pin : Pins)
+	{
+		// For each of our output pins, find the other editor node it's connected to, then
+		// translate that to runtime components and reconnect the runtime components. We only do
+		// the output side because it creates a two-way connection, and we're not worried about
+		// the nodes outside the copy/pasted nodes, as we won't have reconstructed the connection to them
+		// (so the resulting pasted nodes have no connection outside their "island" of copy/pasted nodes)
+		if (Pin->Direction == EEdGraphPinDirection::EGPD_Output)
+		{
+			for (UEdGraphPin* LinkedToPin : Pin->LinkedTo)
+			{
+				UEdGraphNode* ConnectedEdGraphNode = LinkedToPin->GetOwningNode();
+				UMoviePipelineEdGraphNodeBase* ConnectedMovieGraphNode = CastChecked<UMoviePipelineEdGraphNodeBase>(ConnectedEdGraphNode);
+				
+				if (UMovieGraphNode* ConnectedRuntimeNode = ConnectedMovieGraphNode->GetRuntimeNode())
+				{
+					UMovieGraphConfig* Graph = RuntimeNode->GetGraph();
+					check(Graph);
+
+					Graph->AddLabeledEdge(RuntimeNode, Pin->PinName, ConnectedRuntimeNode, LinkedToPin->PinName);
+				}
+			}
+		}
+	}
+
+}
+
+void UMoviePipelineEdGraphNodeBase::GetNodeContextMenuActions(UToolMenu* Menu, UGraphNodeContextMenuContext* Context) const
+{
+	if (!Context->Node)
+	{
+		return;
+	}
+
+	{
+		FToolMenuSection& Section = Menu->AddSection("EdGraphSchemaNodeActions", LOCTEXT("NodeActionsHeader", "Node Actions"));
+		Section.AddMenuEntry(FGraphEditorCommands::Get().BreakNodeLinks);
+		Section.AddMenuEntry(FGenericCommands::Get().Delete);
+		Section.AddMenuEntry(FGenericCommands::Get().Cut);
+		Section.AddMenuEntry(FGenericCommands::Get().Copy);
+		Section.AddMenuEntry(FGenericCommands::Get().Duplicate);
+	}
+
+	{
+		FToolMenuSection& Section = Menu->AddSection("EdGraphSchemaOrganization", LOCTEXT("OrganizationHeader", "Organization"));
+		Section.AddSubMenu("Alignment", LOCTEXT("AlignmentHeader", "Alignment"), FText(), FNewToolMenuDelegate::CreateLambda([](UToolMenu* AlignmentMenu)
+			{
+				{
+					FToolMenuSection& SubSection = AlignmentMenu->AddSection("EdGraphSchemaAlignment", LOCTEXT("AlignHeader", "Align"));
+					SubSection.AddMenuEntry(FGraphEditorCommands::Get().AlignNodesTop);
+					SubSection.AddMenuEntry(FGraphEditorCommands::Get().AlignNodesMiddle);
+					SubSection.AddMenuEntry(FGraphEditorCommands::Get().AlignNodesBottom);
+					SubSection.AddMenuEntry(FGraphEditorCommands::Get().AlignNodesLeft);
+					SubSection.AddMenuEntry(FGraphEditorCommands::Get().AlignNodesCenter);
+					SubSection.AddMenuEntry(FGraphEditorCommands::Get().AlignNodesRight);
+					SubSection.AddMenuEntry(FGraphEditorCommands::Get().StraightenConnections);
+				}
+
+				{
+					FToolMenuSection& SubSection = AlignmentMenu->AddSection("EdGraphSchemaDistribution", LOCTEXT("DistributionHeader", "Distribution"));
+					SubSection.AddMenuEntry(FGraphEditorCommands::Get().DistributeNodesHorizontally);
+					SubSection.AddMenuEntry(FGraphEditorCommands::Get().DistributeNodesVertically);
+				}
+			}));
+	}
+
+	{
+		FToolMenuSection& Section = Menu->AddSection("EdGraphSchemaCommentGroup", LOCTEXT("CommentGroupHeader", "Comment Group"));
+		Section.AddMenuEntry(FGraphEditorCommands::Get().CreateComment,
+			LOCTEXT("MultiCommentDesc", "Create Comment from Selection"),
+			LOCTEXT("CommentToolTip", "Create a resizable comment box around selection."));
+	}
+}
+
 #undef LOCTEXT_NAMESPACE
\ No newline at end of file
diff --git a/Engine/Plugins/MovieScene/MovieRenderPipeline/Source/MovieRenderPipelineEditor/Private/Graph/MovieEdGraphNode.h b/Engine/Plugins/MovieScene/MovieRenderPipeline/Source/MovieRenderPipelineEditor/Private/Graph/MovieEdGraphNode.h
--- a/Engine/Plugins/MovieScene/MovieRenderPipeline/Source/MovieRenderPipelineEditor/Private/Graph/MovieEdGraphNode.h
+++ b/Engine/Plugins/MovieScene/MovieRenderPipeline/Source/MovieRenderPipelineEditor/Private/Graph/MovieEdGraphNode.h
@@ -28,14 +28,22 @@ public:
 	virtual bool ShowPaletteIconOnNode() const override;
 	virtual void OnUpdateCommentText(const FString& NewComment) override;
 	virtual void OnCommentBubbleToggled(bool bInCommentBubbleVisible) override;
+	virtual void PrepareForCopying() override;
+	virtual void PostPasteNode() override;
+	virtual void GetNodeContextMenuActions(UToolMenu* Menu, UGraphNodeContextMenuContext* Context) const override;
 	//~ End UEdGraphNode Interface
-	
+
+	// Called after PrepareForCopying to restore changes made during preparation. 
+	virtual void PostCopy();
+	virtual void PostPaste();
+
 	//~ Begin UObject interface
 	virtual void PostTransacted(const FTransactionObjectEvent& TransactionEvent) override;
 	//~ End UObject interface
 
 protected:
 	virtual bool ShouldCreatePin(const UMovieGraphPin* InPin) const;
+	void RebuildRuntimeEdgesFromPins();
 
 	void CreatePins(const TArray<UMovieGraphPin*>& InInputPins, const TArray<UMovieGraphPin*>& InOutputPins);
 
@@ -57,6 +65,9 @@ protected:
 	/** The runtime node that this editor node represents. */
 	UPROPERTY()
 	TObjectPtr<UMovieGraphNode> RuntimeNode;
+
+	/** Should we early out in ReconstructNode(), skipping restoring the connections to other nodes. Set during Copy/Paste. */
+	bool bDisableReconstructNode = false;
 };
 
 UCLASS()
diff --git a/Engine/Plugins/MovieScene/MovieRenderPipeline/Source/MovieRenderPipelineEditor/Private/Graph/MovieEdGraphOutputNode.h b/Engine/Plugins/MovieScene/MovieRenderPipeline/Source/MovieRenderPipelineEditor/Private/Graph/MovieEdGraphOutputNode.h
--- a/Engine/Plugins/MovieScene/MovieRenderPipeline/Source/MovieRenderPipelineEditor/Private/Graph/MovieEdGraphOutputNode.h
+++ b/Engine/Plugins/MovieScene/MovieRenderPipeline/Source/MovieRenderPipelineEditor/Private/Graph/MovieEdGraphOutputNode.h
@@ -15,6 +15,9 @@ public:
 	virtual FText GetNodeTitle(ENodeTitleType::Type TitleType) const override;
 	virtual void AllocateDefaultPins() override;
 	virtual bool CanUserDeleteNode() const override { return false; }
-	virtual bool CanDuplicateNode() const override { return false; }
+	// This is set to true so that the node is deserialized on paste, required for CanPasteHere to be called
+	virtual bool CanDuplicateNode() const override { return true; }
+	// And this is set to false so the paste doesn't actually go through, and instead we get the editor prompt notifying user.
+	virtual bool CanPasteHere(const UEdGraph* TargetGraph) const { return false; }
 	// ~End UEdGraphNode interface
 };
\ No newline at end of file
diff --git a/Engine/Plugins/MovieScene/MovieRenderPipeline/Source/MovieRenderPipelineEditor/Private/Widgets/Graph/SMovieGraphConfigPanel.cpp b/Engine/Plugins/MovieScene/MovieRenderPipeline/Source/MovieRenderPipelineEditor/Private/Widgets/Graph/SMovieGraphConfigPanel.cpp
--- a/Engine/Plugins/MovieScene/MovieRenderPipeline/Source/MovieRenderPipelineEditor/Private/Widgets/Graph/SMovieGraphConfigPanel.cpp
+++ b/Engine/Plugins/MovieScene/MovieRenderPipeline/Source/MovieRenderPipelineEditor/Private/Widgets/Graph/SMovieGraphConfigPanel.cpp
@@ -31,6 +31,7 @@
 #include "Framework/MultiBox/MultiBoxBuilder.h"
 #include "SPrimaryButton.h"
 #include "Framework/Commands/GenericCommands.h"
+#include "SNodePanel.h"
 
 // ContentBrowser Includes
 #include "IContentBrowserSingleton.h"
@@ -49,6 +50,8 @@
 #include "FileHelpers.h"
 #include "AssetToolsModule.h"
 #include "Misc/FileHelper.h"
+#include "EdGraphUtilities.h"
+#include "HAL/PlatformApplicationMisc.h"
 
 #include "Graph/MovieGraphConfig.h"
 #include "Graph/Nodes/MovieGraphVariableNode.h"
@@ -89,6 +92,9 @@ void SMoviePipelineGraphPanel::Construct(const FArguments& InArgs)
 		PipelineQueueEditorWidget->SetSelectedJobs(Jobs);
 	}
 
+	FGraphAppearanceInfo AppearanceInfo;
+	AppearanceInfo.CornerText = LOCTEXT("EditorCornerText", "Movie Graph Config");
+
 	SGraphEditor::FGraphEditorEvents InEvents;
 	InEvents.OnSelectionChanged = SGraphEditor::FOnSelectionChanged::CreateSP(this, &SMoviePipelineGraphPanel::OnSelectedNodesChanged);
 	InEvents.OnNodeDoubleClicked = FSingleNodeEvent::CreateSP(this, &SMoviePipelineGraphPanel::OnNodeDoubleClicked);
@@ -115,9 +121,12 @@ void SMoviePipelineGraphPanel::Construct(const FArguments& InArgs)
 		.GraphToEdit(GraphToEdit)
 		.AdditionalCommands(GraphEditorCommands)
 		.GraphEvents(InEvents)
+		.Appearance(AppearanceInfo)
 	];
 }
 
+UE_ENABLE_OPTIMIZATION_SHIP
+
 void SMoviePipelineGraphPanel::MakeEditorCommands()
 {
 	GraphEditorCommands = MakeShareable(new FUICommandList);
@@ -136,6 +145,68 @@ void SMoviePipelineGraphPanel::MakeEditorCommands()
 			FExecuteAction::CreateSP(this, &SMoviePipelineGraphPanel::DeleteSelectedNodes),
 			FCanExecuteAction::CreateSP(this, &SMoviePipelineGraphPanel::CanDeleteSelectedNodes)
 		);
+
+		GraphEditorCommands->MapAction(FGenericCommands::Get().Copy,
+			FExecuteAction::CreateSP(this, &SMoviePipelineGraphPanel::CopySelectedNodes),
+			FCanExecuteAction::CreateSP(this, &SMoviePipelineGraphPanel::CanCopySelectedNodes)
+		);
+
+		GraphEditorCommands->MapAction(FGenericCommands::Get().Cut,
+			FExecuteAction::CreateSP(this, &SMoviePipelineGraphPanel::CutSelectedNodes),
+			FCanExecuteAction::CreateSP(this, &SMoviePipelineGraphPanel::CanCutSelectedNodes)
+		);
+
+		GraphEditorCommands->MapAction(FGenericCommands::Get().Paste,
+			FExecuteAction::CreateSP(this, &SMoviePipelineGraphPanel::PasteNodes),
+			FCanExecuteAction::CreateSP(this, &SMoviePipelineGraphPanel::CanPasteNodes)
+		);
+
+		GraphEditorCommands->MapAction(FGenericCommands::Get().Duplicate,
+			FExecuteAction::CreateSP(this, &SMoviePipelineGraphPanel::DuplicateNodes),
+			FCanExecuteAction::CreateSP(this, &SMoviePipelineGraphPanel::CanDuplicateNodes)
+		);
+
+		// Alignment Commands
+		GraphEditorCommands->MapAction(FGraphEditorCommands::Get().AlignNodesTop,
+			FExecuteAction::CreateSP(this, &SMoviePipelineGraphPanel::OnAlignTop)
+		);
+
+		GraphEditorCommands->MapAction(FGraphEditorCommands::Get().AlignNodesMiddle,
+			FExecuteAction::CreateSP(this, &SMoviePipelineGraphPanel::OnAlignMiddle)
+		);
+
+		GraphEditorCommands->MapAction(FGraphEditorCommands::Get().AlignNodesBottom,
+			FExecuteAction::CreateSP(this, &SMoviePipelineGraphPanel::OnAlignBottom)
+		);
+
+		GraphEditorCommands->MapAction(FGraphEditorCommands::Get().AlignNodesLeft,
+			FExecuteAction::CreateSP(this, &SMoviePipelineGraphPanel::OnAlignLeft)
+		);
+
+		GraphEditorCommands->MapAction(FGraphEditorCommands::Get().AlignNodesCenter,
+			FExecuteAction::CreateSP(this, &SMoviePipelineGraphPanel::OnAlignCenter)
+		);
+
+		GraphEditorCommands->MapAction(FGraphEditorCommands::Get().AlignNodesRight,
+			FExecuteAction::CreateSP(this, &SMoviePipelineGraphPanel::OnAlignRight)
+		);
+
+		GraphEditorCommands->MapAction(FGraphEditorCommands::Get().StraightenConnections,
+			FExecuteAction::CreateSP(this, &SMoviePipelineGraphPanel::OnStraightenConnections)
+		);
+
+		GraphEditorCommands->MapAction(FGraphEditorCommands::Get().CreateComment,
+			FExecuteAction::CreateSP(this, &SMoviePipelineGraphPanel::OnCreateComment)
+		);
+
+		// Distribution Commands
+		GraphEditorCommands->MapAction(FGraphEditorCommands::Get().DistributeNodesHorizontally,
+			FExecuteAction::CreateSP(this, &SMoviePipelineGraphPanel::OnDistributeNodesH)
+		);
+
+		GraphEditorCommands->MapAction(FGraphEditorCommands::Get().DistributeNodesVertically,
+			FExecuteAction::CreateSP(this, &SMoviePipelineGraphPanel::OnDistributeNodesV)
+		);
 	}
 }
 
@@ -212,12 +283,276 @@ void SMoviePipelineGraphPanel::DeleteSelectedNodes()
 	ClearGraphSelection();
 }
 
+void SMoviePipelineGraphPanel::CopySelectedNodes()
+{
+	if (GraphEditorWidget.IsValid())
+	{
+		const FGraphPanelSelectionSet SelectedNodes = GraphEditorWidget->GetSelectedNodes();
+
+		for (UObject* SelectedNode : SelectedNodes)
+		{
+			UEdGraphNode* GraphNode = CastChecked<UEdGraphNode>(SelectedNode);
+
+			// This causes the editor node to become the temporary owner of the underlying data model node
+			GraphNode->PrepareForCopying();
+		}
+
+		// Serialize it all to text, the linked data model nodes are copied too due to (temporarily) being owned
+		// by the editor node we're copying. 
+		FString ExportedText;
+		FEdGraphUtilities::ExportNodesToText(SelectedNodes, ExportedText);
+		FPlatformApplicationMisc::ClipboardCopy(*ExportedText);
+
+		// Now go back through our nodes and restore their proper ownership. There's no matching PostPrepareForCopying
+		// in UEdGraphNode, so we have to cast it to our type first.
+		for (UObject* SelectedNode : SelectedNodes)
+		{
+			if (UMoviePipelineEdGraphNodeBase* MovieGraphNode = Cast<UMoviePipelineEdGraphNodeBase>(SelectedNode))
+			{
+				MovieGraphNode->PostCopy();
+			}
+		}
+	}
+}
+
+bool SMoviePipelineGraphPanel::CanCopySelectedNodes() const
+{
+	// We handle non-copyable nodes on paste so the toast warning happens on paste
+	// to warn the user that some nodes (Input/Outputs) weren't copied.
+	bool bAtLeastOneCopyableNode = false;
+	if (GraphEditorWidget.IsValid())
+	{
+		for (UObject* Object : GraphEditorWidget->GetSelectedNodes())
+		{
+			if (UEdGraphNode* GraphNode = CastChecked<UEdGraphNode>(Object))
+			{
+				if (GraphNode->CanDuplicateNode())
+				{
+					bAtLeastOneCopyableNode = true;
+					break;
+				}
+			}
+		}
+	}
+
+	return bAtLeastOneCopyableNode;
+}
+
+void SMoviePipelineGraphPanel::CutSelectedNodes()
+{
+	CopySelectedNodes();
+	DeleteSelectedNodes();
+}
+
+bool SMoviePipelineGraphPanel::CanCutSelectedNodes() const
+{
+	return CanCopySelectedNodes() && CanDeleteSelectedNodes();
+}
+
+void SMoviePipelineGraphPanel::PasteNodes()
+{
+	if (GraphEditorWidget.IsValid())
+	{
+		PasteNodesHere(GraphEditorWidget->GetPasteLocation());
+	}
+}
+
+void SMoviePipelineGraphPanel::PasteNodesHere(const FVector2D& Location)
+{
+	if (!GraphEditorWidget.IsValid() || !CurrentGraph->PipelineEdGraph)
+	{
+		return;
+	}
+
+	const FScopedTransaction Transaction(LOCTEXT("PasteNodes_Transaction", "Movie Graph Paste Node(s)"));
+	UMoviePipelineEdGraph* EdGraph = Cast<UMoviePipelineEdGraph>(CurrentGraph->PipelineEdGraph);
+	EdGraph->Modify();
+
+	// Clear the selection set (newly pasted stuff will be selected)
+	GraphEditorWidget->ClearSelectionSet();
+
+	// Grab the text to paste from the clipboard.
+	FString TextToImport;
+	FPlatformApplicationMisc::ClipboardPaste(TextToImport);
+
+	// Import the nodes
+	TSet<UEdGraphNode*> PastedNodes;
+	FEdGraphUtilities::ImportNodesFromText(EdGraph, TextToImport, PastedNodes);
+
+	//Average position of nodes so that resulting nodes are centered around the paste location
+	FVector2D AvgNodePosition(0.0f, 0.0f);
+
+	int32 AvgCount = 0;
+
+	for (UEdGraphNode* PastedNode : PastedNodes)
+	{
+		if (PastedNode)
+		{
+			AvgNodePosition.X += PastedNode->NodePosX;
+			AvgNodePosition.Y += PastedNode->NodePosY;
+			++AvgCount;
+		}
+	}
+
+	if (AvgCount > 0)
+	{
+		float InvNumNodes = 1.0f / float(AvgCount);
+		AvgNodePosition.X *= InvNumNodes;
+		AvgNodePosition.Y *= InvNumNodes;
+	}
+
+	for (UEdGraphNode* PastedNode : PastedNodes)
+	{
+		GraphEditorWidget->SetNodeSelection(PastedNode, true);
+
+		// Offset to center the nodes on the paste location
+		PastedNode->NodePosX = (PastedNode->NodePosX - AvgNodePosition.X) + Location.X;
+		PastedNode->NodePosY = (PastedNode->NodePosY - AvgNodePosition.Y) + Location.Y;
+
+		PastedNode->SnapToGrid(SNodePanel::GetSnapGridSize());
+
+		PastedNode->CreateNewGuid();
+
+		if (UMoviePipelineEdGraphNodeBase* PastedEdGraphNode = Cast<UMoviePipelineEdGraphNodeBase>(PastedNode))
+		{
+			if (UMovieGraphNode* PastedRuntimeNode = PastedEdGraphNode->GetRuntimeNode())
+			{
+				// AddNode fixes up ownership to ensure the graph owns this node
+				CurrentGraph->AddNode(PastedRuntimeNode);
+
+				// Rebuild the runtime node connections based on the editor node pins
+				PastedEdGraphNode->PostPaste();
+			}
+		}
+	}
+
+	GraphEditorWidget->NotifyGraphChanged();
+}
+
+bool SMoviePipelineGraphPanel::CanPasteNodes() const
+{
+	FString ClipboardContent;
+	FPlatformApplicationMisc::ClipboardPaste(ClipboardContent);
+
+	return FEdGraphUtilities::CanImportNodesFromText(CurrentGraph->PipelineEdGraph, ClipboardContent);
+}
+
+void SMoviePipelineGraphPanel::DuplicateNodes()
+{
+	CopySelectedNodes();
+	PasteNodes();
+}
+
+bool SMoviePipelineGraphPanel::CanDuplicateNodes() const
+{
+	return CanCopySelectedNodes();
+}
+
+void SMoviePipelineGraphPanel::OnAlignTop()
+{
+	if (GraphEditorWidget.IsValid())
+	{
+		GraphEditorWidget->OnAlignTop();
+	}
+}
+
+void SMoviePipelineGraphPanel::OnAlignMiddle()
+{
+	if (GraphEditorWidget.IsValid())
+	{
+		GraphEditorWidget->OnAlignMiddle();
+	}
+}
+
+void SMoviePipelineGraphPanel::OnAlignBottom()
+{
+	if (GraphEditorWidget.IsValid())
+	{
+		GraphEditorWidget->OnAlignBottom();
+	}
+}
+
+void SMoviePipelineGraphPanel::OnAlignLeft()
+{
+	if (GraphEditorWidget.IsValid())
+	{
+		GraphEditorWidget->OnAlignLeft();
+	}
+}
+
+void SMoviePipelineGraphPanel::OnAlignCenter()
+{
+	if (GraphEditorWidget.IsValid())
+	{
+		GraphEditorWidget->OnAlignCenter();
+	}
+}
+
+void SMoviePipelineGraphPanel::OnAlignRight()
+{
+	if (GraphEditorWidget.IsValid())
+	{
+		GraphEditorWidget->OnAlignRight();
+	}
+}
+
+void SMoviePipelineGraphPanel::OnStraightenConnections()
+{
+	if (GraphEditorWidget.IsValid())
+	{
+		GraphEditorWidget->OnStraightenConnections();
+	}
+}
+
+void SMoviePipelineGraphPanel::OnDistributeNodesH()
+{
+	if (GraphEditorWidget.IsValid())
+	{
+		GraphEditorWidget->OnDistributeNodesH();
+	}
+}
+
+void SMoviePipelineGraphPanel::OnDistributeNodesV()
+{
+	if (GraphEditorWidget.IsValid())
+	{
+		GraphEditorWidget->OnDistributeNodesV();
+	}
+}
+
 void SMoviePipelineGraphPanel::ClearGraphSelection() const
 {
 	GraphEditorWidget->ClearSelectionSet();
 }
 
-UE_ENABLE_OPTIMIZATION_SHIP
+bool SMoviePipelineGraphPanel::MatchesContext(const FTransactionContext& InContext, const TArray<TPair<UObject*, FTransactionObjectEvent>>& TransactionObjectContexts) const
+{
+	// ToDo: If we wanted to rebuild the graph less often, we could make all of our transactions have a context
+	// and only handle PostUndo if the undo actually matches the context. For now we'll just always rebuild the widget.
+	return true;
+}
+
+void SMoviePipelineGraphPanel::PostUndo(bool bSuccess)
+{
+	if (bSuccess)
+	{
+		if (CurrentGraph)
+		{
+			CurrentGraph->OnGraphChangedDelegate.Broadcast();
+		}
+
+		if (GraphEditorWidget.IsValid())
+		{
+			// This unfortunately causes the graph to flicker for one frame, but all Blueprint Graphs do this.
+			// This is needed so that when a node is pasted, then undone, the widgets on the graph for the now
+			// deleted nodes also get removed, otherwise we try to draw them and they're pointed to invalid uobjects.
+			GraphEditorWidget->ClearSelectionSet();
+			GraphEditorWidget->NotifyGraphChanged();
+
+			FSlateApplication::Get().DismissAllMenus();
+		}
+	}
+}
 
 void SMoviePipelineGraphPanel::OnSelectedNodesChanged(const TSet<UObject*>& NewSelection)
 {
diff --git a/Engine/Plugins/MovieScene/MovieRenderPipeline/Source/MovieRenderPipelineEditor/Private/Widgets/Graph/SMovieGraphConfigPanel.h b/Engine/Plugins/MovieScene/MovieRenderPipeline/Source/MovieRenderPipelineEditor/Private/Widgets/Graph/SMovieGraphConfigPanel.h
--- a/Engine/Plugins/MovieScene/MovieRenderPipeline/Source/MovieRenderPipelineEditor/Private/Widgets/Graph/SMovieGraphConfigPanel.h
+++ b/Engine/Plugins/MovieScene/MovieRenderPipeline/Source/MovieRenderPipelineEditor/Private/Widgets/Graph/SMovieGraphConfigPanel.h
@@ -4,6 +4,7 @@
 
 #include "Widgets/DeclarativeSyntaxSupport.h"
 #include "Widgets/SCompoundWidget.h"
+#include "EditorUndoClient.h"
 
 class UMoviePipelineConfigBase;
 class SGraphEditor;
@@ -19,7 +20,7 @@ class UMovieGraphNode;
 /**
  * Outermost widget that is used for adding and removing jobs from the Movie Pipeline Queue Subsystem.
  */
-class SMoviePipelineGraphPanel : public SCompoundWidget
+class SMoviePipelineGraphPanel : public SCompoundWidget, public FSelfRegisteringEditorUndoClient
 {
 public:
 	DECLARE_DELEGATE_OneParam(FOnGraphSelectionChanged, TArray<UObject*>);
@@ -41,6 +42,12 @@ public:
 	/** Clears the selection in the graph widget being displayed. */
 	void ClearGraphSelection() const;
 
+	// FEditorUndoClient interface
+	virtual bool MatchesContext(const FTransactionContext& InContext, const TArray<TPair<UObject*, FTransactionObjectEvent>>& TransactionObjectContexts) const override;
+	virtual void PostUndo(bool bSuccess) override;
+	virtual void PostRedo(bool bSuccess) override { PostUndo(bSuccess); }
+	// ~FEditorUndoClient interface
+
 private:
 	FReply OnRenderLocalRequested();
 	bool IsRenderLocalEnabled() const;
@@ -81,6 +88,26 @@ private:
 	bool CanSelectAllNodes() const;
 	bool CanDeleteSelectedNodes() const;
 	void DeleteSelectedNodes();
+	void CopySelectedNodes();
+	bool CanCopySelectedNodes() const;
+	void CutSelectedNodes();
+	bool CanCutSelectedNodes() const;
+	void PasteNodes();
+	void PasteNodesHere(const FVector2D& Location);
+	bool CanPasteNodes() const;
+	void DuplicateNodes();
+	bool CanDuplicateNodes() const;
+
+	void OnAlignTop();
+	void OnAlignMiddle();
+	void OnAlignBottom();
+	void OnAlignLeft();
+	void OnAlignCenter();
+	void OnAlignRight();
+	void OnStraightenConnections();
+	void OnDistributeNodesH();
+	void OnDistributeNodesV();
+
 private:
 	/** The main movie pipeline queue editor widget */
 	TSharedPtr<SMoviePipelineQueueEditor> PipelineQueueEditorWidget;
