MoviePipeline: Updated the traversal to take variables into account, and made several other graph improvements.

* The context menu to expose properties as pins now works properly for both native and dynamic properties (and exposure is togglable as well).
* Pin colors are now available for all types that can be present in the graph.
* Active Render Settings now displays dynamic properties.

diff --git a/Engine/Plugins/MovieScene/MovieRenderPipeline/Source/MovieRenderPipelineCore/Private/Graph/MovieGraphConfig.cpp b/Engine/Plugins/MovieScene/MovieRenderPipeline/Source/MovieRenderPipelineCore/Private/Graph/MovieGraphConfig.cpp
--- a/Engine/Plugins/MovieScene/MovieRenderPipeline/Source/MovieRenderPipelineCore/Private/Graph/MovieGraphConfig.cpp
+++ b/Engine/Plugins/MovieScene/MovieRenderPipeline/Source/MovieRenderPipelineCore/Private/Graph/MovieGraphConfig.cpp
@@ -623,7 +623,6 @@ bool UMovieGraphConfig::DeleteOutputMember(UMovieGraphOutput* OutputMemberToDele
 
 FBoolProperty* UMovieGraphConfig::FindOverridePropertyForRealProperty(UClass* InClass, const FProperty* InRealProperty)
 {
-	static const FName EditConditionKey = FName("EditCondition");
 	if (!ensure(InClass && InRealProperty))
 	{
 		return nullptr;
@@ -650,8 +649,6 @@ FBoolProperty* UMovieGraphConfig::FindOverridePropertyForRealProperty(UClass* In
 
 void UMovieGraphConfig::InitializeFlattenedNode(UMovieGraphNode* InNode)
 {
-	static const FName EditConditionKey = FName("EditCondition");
-	
 	// We go through each of the bOverride_ properties on this new instance and set
 	// them all as "not overridden", so that as we traverse the graph, we can only
 	// update the node the first time we hit another property that is overridden.
@@ -666,6 +663,9 @@ void UMovieGraphConfig::InitializeFlattenedNode(UMovieGraphNode* InNode)
 			EditConditionProperty->SetPropertyValue_InContainer(InNode, false);
 		}
 	}
+
+	// Initialize the dynamic properties so they can be updated during traversal like UPROPERTY properties
+	InNode->UpdateDynamicProperties();
 }
 
 void UMovieGraphConfig::CopyOverriddenProperties(UMovieGraphNode* FromNode, UMovieGraphNode* ToNode)
@@ -680,59 +680,113 @@ void UMovieGraphConfig::CopyOverriddenProperties(UMovieGraphNode* FromNode, UMov
 		return;
 	}
 
-	for (TFieldIterator<FProperty> PropertyIterator(ToNode->GetClass()); PropertyIterator; ++PropertyIterator)
+	for (const FProperty* DestNodeProperty : ToNode->GetAllOverrideableProperties())
 	{
-		// For each property on the destination node, decide if we need to try to update it (we don't update bOverride_ properties)
-		FProperty* PropertyToCheckOnDestNode = *PropertyIterator;
+		if (!DestNodeProperty)
+		{
+			continue;
+		}
 
+		const FName PropertyName = DestNodeProperty->GetFName();
+		
+		// For each property on the destination node, decide if we need to try to update it (we don't update bOverride_ properties)
+		//
 		// We use the existence of a matching edit condition node to signal that this property is overrideable, ie:
 		// float MyFoo would have a bOverride_MyFoo, so FindOverridePropertyForRealProperty would match it. However when
 		// looking at the bOverride_MyFoo property it would look for bOverride_bOverride_MyFoo, which wouldn't exist, successfully
 		// filtering us to only the "real" overrideable properties.
-		FBoolProperty* EditConditionProperty = FindOverridePropertyForRealProperty(ToNode->GetClass(), PropertyToCheckOnDestNode);
-		if (PropertyToCheckOnDestNode && EditConditionProperty)
+
+		const bool bIsDynamic = ToNode->GetDynamicPropertyDescriptions().ContainsByPredicate([&PropertyName](const FPropertyBagPropertyDesc& PropDesc)
+		{
+			return PropDesc.Name == PropertyName;
+		});
+		
+		const bool bIsExposed = FromNode->GetExposedProperties().ContainsByPredicate([&PropertyName](const FMovieGraphPropertyInfo& ExposedPropertyInfo)
+		{
+			return ExposedPropertyInfo.Name == PropertyName;
+		});
+
+		// Ensure there's a property (the "bOverride_*" property) that tracks whether or not this property has already
+		// been set/overridden
+		const FBoolProperty* EditConditionProperty = bIsDynamic
+			? ToNode->FindOverridePropertyForDynamicProperty(PropertyName)
+			: FindOverridePropertyForRealProperty(ToNode->GetClass(), DestNodeProperty);
+		if (!EditConditionProperty)
 		{
-			// The user has indicated that this property should be overrideable in the Graph network.
-			// The first thing we do is check to see if it is exposed, and if it's exposed, we need to get the value from the graph.
-			bool bGetValueFromPin = false;
-			const TArray<FName>& ExposedPins = FromNode->GetExposedDynamicProperties();
-			for (const FName& Name : ExposedPins)
+			continue;
+		}
+
+		// If our destination node already has this property marked as overridden, then some other node in the graph has
+		// taken priority and set the value to something, so we don't want to override it.
+		const bool bAlreadyOverriddenOnDestNode = bIsDynamic
+			? ToNode->IsDynamicPropertyOverridden(PropertyName)
+			: EditConditionProperty->GetPropertyValue_InContainer(ToNode);
+		if (bAlreadyOverriddenOnDestNode)
+		{
+			continue;
+		}
+
+		// If this property (dynamic or not) has been exposed, attempt to get its value via the connection to it (if any)
+		if (bIsExposed)
+		{
+			if (const UMovieGraphPin* InputPin = FromNode->GetInputPin(PropertyName))
 			{
-				if (Name == PropertyToCheckOnDestNode->GetName())
+				const UMovieGraphPin* ConnectedPin = InputPin->GetFirstConnectedPin();
+				if (ConnectedPin && (ConnectedPin->Properties.Type == InputPin->Properties.Type) && ConnectedPin->Node)
 				{
-					bGetValueFromPin = true;
-					break;
+					// There was a valid connection to the input pin; resolve the value from the connected output and set the
+					// value on this property
+					const FString ResolvedValue = ConnectedPin->Node->GetResolvedValueForOutputPin(ConnectedPin->Properties.Label);
+					if (bIsDynamic)
+					{
+						ToNode->SetDynamicPropertyValue(PropertyName, ResolvedValue);
+						ToNode->SetDynamicPropertyOverridden(PropertyName, true);
+					}
+					else
+					{
+						DestNodeProperty->ImportText_Direct(*ResolvedValue, DestNodeProperty->ContainerPtrToValuePtr<uint8>(ToNode), nullptr, PPF_None);
+						EditConditionProperty->SetPropertyValue_InContainer(ToNode, true);
+					}
+
+					// The property value was set via a connected pin; move on to the next property
+					continue;
 				}
 			}
+		}
 
-			if (bGetValueFromPin)
-			{
-				UE_LOG(LogMovieRenderPipeline, Error, TEXT("not yet implemented."));
-				continue;
-			}
+		// We know it's not already overridden, so now we should check to see if the incoming node wants to override it.
+		const bool bSourceNodeOverwrites = bIsDynamic
+			? FromNode->IsDynamicPropertyOverridden(PropertyName)
+			: EditConditionProperty->GetPropertyValue_InContainer(FromNode);
+		if (!bSourceNodeOverwrites)
+		{
+			// The source node didn't have the override flag checked, so we don't copy the value from it.
+			continue;
+		}
 
-			// If our destination node already has this property marked as overridden, then some other node in the graph has
-			// taken priority and set the value to something, so we don't want to override it.
-			bool bAlreadyOverridenOnDestNode = EditConditionProperty->GetPropertyValue_InContainer(ToNode);
-			if (bAlreadyOverridenOnDestNode)
-			{
-				continue;
-			}
+		// Okay at this point we know that on our target node, no one has overridden it yet, and our source node wants to override this property.
+		// First, we update the booleans to say that yes, this property has been overridden on the target node.
+		if (bIsDynamic)
+		{
+			ToNode->SetDynamicPropertyOverridden(PropertyName, true);
+		}
+		else
+		{
+			EditConditionProperty->SetPropertyValue_InContainer(ToNode, true);
+		}
 
-			// We know it's not already overridden, so now we should check to see if the incoming node wants to override it.
-			bool bSourceNodeOverwrites = EditConditionProperty->GetPropertyValue_InContainer(FromNode);
-			if (!bSourceNodeOverwrites)
+		// Now we need to copy the value from the source to the destination
+		if (bIsDynamic)
+		{
+			FString SourceValue;
+			if (FromNode->GetDynamicPropertyValue(PropertyName, SourceValue))
 			{
-				// The source node didn't have the override flag checked, so we don't copy the value from it.
-				continue;
+				ToNode->SetDynamicPropertyValue(PropertyName, SourceValue);
 			}
-
-			// Okay at this point we know that on our target node, no one has overridden it yet, and our source node wants to override this property.
-			// First, we update the booleans to say that yes, this property has been overridden on the target node.
-			EditConditionProperty->SetPropertyValue_InContainer(ToNode, true);
-
-			// Now we need to copy the value from the source to the destination
-			PropertyToCheckOnDestNode->CopyCompleteValue_InContainer(ToNode, FromNode);
+		}
+		else
+		{
+			DestNodeProperty->CopyCompleteValue_InContainer(ToNode, FromNode);
 		}
 	}
 }
diff --git a/Engine/Plugins/MovieScene/MovieRenderPipeline/Source/MovieRenderPipelineCore/Private/Graph/MovieGraphNode.cpp b/Engine/Plugins/MovieScene/MovieRenderPipeline/Source/MovieRenderPipelineCore/Private/Graph/MovieGraphNode.cpp
--- a/Engine/Plugins/MovieScene/MovieRenderPipeline/Source/MovieRenderPipelineCore/Private/Graph/MovieGraphNode.cpp
+++ b/Engine/Plugins/MovieScene/MovieRenderPipeline/Source/MovieRenderPipelineCore/Private/Graph/MovieGraphNode.cpp
@@ -3,57 +3,180 @@
 #include "Graph/MovieGraphNode.h"
 
 #include "Graph/MovieGraphConfig.h"
+#include "UObject/TextProperty.h"
 
 FName UMovieGraphNode::GlobalsPinName = FName("Globals");
 
+namespace UE::MovieGraph::Private
+{
+	EMovieGraphValueType GetValueTypeFromProperty(const FProperty* InSourceProperty)
+	{
+		if (CastField<FBoolProperty>(InSourceProperty))
+		{
+			return EMovieGraphValueType::Bool;
+		}
+		if (const FByteProperty* ByteProperty = CastField<FByteProperty>(InSourceProperty))
+		{
+			return ByteProperty->IsEnum() ? EMovieGraphValueType::Enum : EMovieGraphValueType::Byte;
+		}
+		if (CastField<FIntProperty>(InSourceProperty))
+		{
+			return EMovieGraphValueType::Int32;
+		}
+		if (CastField<FInt64Property>(InSourceProperty))
+		{
+			return EMovieGraphValueType::Int64;
+		}
+		if (CastField<FFloatProperty>(InSourceProperty))
+		{
+			return EMovieGraphValueType::Float;
+		}
+		if (CastField<FDoubleProperty>(InSourceProperty))
+		{
+			return EMovieGraphValueType::Double;
+		}
+		if (CastField<FNameProperty>(InSourceProperty))
+		{
+			return EMovieGraphValueType::Name;
+		}
+		if (CastField<FStrProperty>(InSourceProperty))
+		{
+			return EMovieGraphValueType::String;
+		}
+		if (CastField<FTextProperty>(InSourceProperty))
+		{
+			return EMovieGraphValueType::Text;
+		}
+		if (CastField<FEnumProperty>(InSourceProperty))
+		{
+			return EMovieGraphValueType::Enum;
+		}
+		if (CastField<FStructProperty>(InSourceProperty))
+		{
+			return EMovieGraphValueType::Struct;
+		}
+		if (CastField<FObjectProperty>(InSourceProperty))
+		{
+			return EMovieGraphValueType::Object;
+		}
+		if (CastField<FSoftObjectProperty>(InSourceProperty))
+		{
+			return EMovieGraphValueType::SoftObject;
+		}
+		if (CastField<FClassProperty>(InSourceProperty))
+		{
+			return EMovieGraphValueType::Class;
+		}
+		if (CastField<FSoftClassProperty>(InSourceProperty))
+		{
+			return EMovieGraphValueType::SoftClass;
+		}
+
+		// Handle array property
+		if (const FArrayProperty* ArrayProperty = CastField<FArrayProperty>(InSourceProperty))
+		{
+			return GetValueTypeFromProperty(ArrayProperty->Inner);	
+		}
+
+		return EMovieGraphValueType::None;
+	}
+}
+
 UMovieGraphNode::UMovieGraphNode()
 {
 }
 
-TArray<FMovieGraphPinProperties> UMovieGraphNode::GetExposedDynamicPinProperties() const
+TArray<FMovieGraphPinProperties> UMovieGraphNode::GetExposedPinProperties() const
 {
 	TArray<FMovieGraphPinProperties> Properties;
 
-	// Exposed dynamic properties are tracked by name; convert to pin properties
-	for (const FName& PropertyName : GetExposedDynamicProperties())
+	// Exposed properties are tracked by name; convert to pin properties
+	for (const FMovieGraphPropertyInfo& PropertyInfo : GetExposedProperties())
 	{
-		// Note: Currently hardcoded to not allow multiple connections, but this may need to change in the future
-		const bool bInAllowMultipleConnections = false;
-		Properties.Add(FMovieGraphPinProperties(PropertyName, EMovieGraphValueType::Float, bInAllowMultipleConnections));
+		if (PropertyInfo.ValueType == EMovieGraphValueType::None)
+		{
+			continue;
+		}
+		
+		constexpr bool bAllowMultipleConnections = false;
+		Properties.Add(FMovieGraphPinProperties(PropertyInfo.Name, PropertyInfo.ValueType, bAllowMultipleConnections));
 	}
 
 	return Properties;
 }
 
-void UMovieGraphNode::PromoteDynamicPropertyToPin(const FName& PropertyName)
+void UMovieGraphNode::TogglePromotePropertyToPin(const FName& PropertyName)
 {
-	// Ensure this is a real dynamic property
-	bool bFoundMatchingProperty = false;
-	for (const FPropertyBagPropertyDesc& PropertyDescription : GetDynamicPropertyDescriptions())
+	// Ensure this is a property that can be promoted
+	bool bIsPromotableProperty = false;
+	FMovieGraphPropertyInfo OverrideablePropertyInfo;
+	for (const FMovieGraphPropertyInfo& ExposablePropertyInfo : GetOverrideablePropertyInfo())
 	{
-		if (PropertyDescription.Name == PropertyName)
+		if (ExposablePropertyInfo.Name == PropertyName)
 		{
-			bFoundMatchingProperty = true;
+			bIsPromotableProperty = true;
+			OverrideablePropertyInfo = ExposablePropertyInfo;
 			break;
 		}
 	}
 
-	// Add this property to promoted properties if it hasn't already been promoted
-	if (bFoundMatchingProperty && !ExposedDynamicPropertyNames.Contains(PropertyName))
+	// Add this property to promoted properties if it hasn't already been promoted, otherwise remove it
+	if (bIsPromotableProperty)
 	{
-		ExposedDynamicPropertyNames.Add(PropertyName);
+		const bool bFoundExposedPropertyInfo = ExposedPropertyInfo.ContainsByPredicate([&OverrideablePropertyInfo](const FMovieGraphPropertyInfo& ExposedInfo)
+		{
+			return OverrideablePropertyInfo.IsSamePropertyAs(ExposedInfo);
+		});
+
+		if (!bFoundExposedPropertyInfo)
+		{
+			ExposedPropertyInfo.Add(MoveTemp(OverrideablePropertyInfo));
+		}
+		else
+		{
+			ExposedPropertyInfo.RemoveAll([&OverrideablePropertyInfo](const FMovieGraphPropertyInfo& ExposedInfo)
+			{
+				return OverrideablePropertyInfo.IsSamePropertyAs(ExposedInfo);
+			});
+		}
 	}
 
 	UpdatePins();
 }
 
+TArray<const FProperty*> UMovieGraphNode::GetAllOverrideableProperties() const
+{
+	TArray<const FProperty*> AllProperties;
+	
+	for (TFieldIterator<FProperty> PropertyIterator(GetClass()); PropertyIterator; ++PropertyIterator)
+	{
+		if (UMovieGraphConfig::FindOverridePropertyForRealProperty(GetClass(), *PropertyIterator))
+		{
+			AllProperties.Add(*PropertyIterator);
+		}
+	}
+
+	if (const UPropertyBag* PropertyBag = DynamicProperties.GetPropertyBagStruct())
+	{
+		for (const FPropertyBagPropertyDesc& Desc : PropertyBag->GetPropertyDescs())
+		{
+			if (FindOverridePropertyForDynamicProperty(Desc.Name))
+			{
+				AllProperties.Add(Desc.CachedProperty);
+			}
+		}
+	}
+
+	return AllProperties;
+}
+
 void UMovieGraphNode::UpdatePins()
 {
 	TArray<FMovieGraphPinProperties> InputPinProperties = GetInputPinProperties();
 	TArray<FMovieGraphPinProperties> OutputPinProperties = GetOutputPinProperties();
 
 	// Include the exposed dynamic properties in the input pins
-	InputPinProperties.Append(GetExposedDynamicPinProperties());
+	InputPinProperties.Append(GetExposedPinProperties());
 
 	auto UpdatePins = [this](TArray<UMovieGraphPin*>& Pins, const TArray<FMovieGraphPinProperties>& PinProperties)
 	{
@@ -208,6 +331,92 @@ void UMovieGraphNode::PostLoad()
 	RegisterDelegates();
 }
 
+bool UMovieGraphNode::SetDynamicPropertyValue(const FName PropertyName, const FString& InNewValue)
+{
+	return DynamicProperties.SetValueSerializedString(PropertyName, InNewValue) == EPropertyBagResult::Success;
+}
+
+bool UMovieGraphNode::GetDynamicPropertyValue(const FName PropertyName, FString& OutValue)
+{
+	const TValueOrError<FString, EPropertyBagResult> Result = DynamicProperties.GetValueSerializedString(PropertyName);
+	if (Result.IsValid())
+	{
+		OutValue = Result.GetValue();
+	}
+
+	return Result.IsValid();
+}
+
+const FBoolProperty* UMovieGraphNode::FindOverridePropertyForDynamicProperty(const FName& InPropertyName) const
+{
+	const FString OverridePropertyName = FString::Printf(TEXT("bOverride_%s"), *InPropertyName.ToString());
+	const FPropertyBagPropertyDesc* Desc = DynamicProperties.FindPropertyDescByName(FName(OverridePropertyName));
+
+	if (!Desc)
+	{
+		return nullptr;
+	}
+
+	if (const FBoolProperty* DescBoolProperty = CastField<FBoolProperty>(Desc->CachedProperty))
+	{
+		return DescBoolProperty;
+	}
+	
+	return nullptr;
+}
+
+bool UMovieGraphNode::IsDynamicPropertyOverridden(const FName& InPropertyName) const
+{
+	if (const FBoolProperty* OverrideProperty = FindOverridePropertyForDynamicProperty(InPropertyName))
+	{
+		TValueOrError<bool, EPropertyBagResult> Result = DynamicProperties.GetValueBool(OverrideProperty->GetFName());
+		return Result.IsValid() && Result.GetValue();
+	}
+
+	return false;
+}
+
+void UMovieGraphNode::SetDynamicPropertyOverridden(const FName& InPropertyName, const bool bIsOverridden)
+{
+	if (const FBoolProperty* OverrideProperty = FindOverridePropertyForDynamicProperty(InPropertyName))
+	{
+		DynamicProperties.SetValueBool(OverrideProperty->GetFName(), bIsOverridden);
+	}
+}
+
+TArray<FMovieGraphPropertyInfo> UMovieGraphNode::GetOverrideablePropertyInfo() const
+{
+	TArray<FMovieGraphPropertyInfo> OverrideableProperties;
+	
+	for (TFieldIterator<FProperty> PropertyIterator(GetClass()); PropertyIterator; ++PropertyIterator)
+	{
+		if (UMovieGraphConfig::FindOverridePropertyForRealProperty(GetClass(), *PropertyIterator))
+		{
+			FMovieGraphPropertyInfo Info;
+			Info.Name = PropertyIterator->GetFName();
+			Info.bIsDynamicProperty = false;
+			Info.ValueType = UE::MovieGraph::Private::GetValueTypeFromProperty(*PropertyIterator);
+			
+			OverrideableProperties.Add(MoveTemp(Info));
+		}
+	}
+
+	for (const FPropertyBagPropertyDesc& Desc : GetDynamicPropertyDescriptions())
+	{
+		if (FindOverridePropertyForDynamicProperty(Desc.Name))
+		{
+			FMovieGraphPropertyInfo Info;
+			Info.Name = Desc.Name;
+			Info.bIsDynamicProperty = true;
+			Info.ValueType = static_cast<EMovieGraphValueType>(Desc.ValueType);
+
+			OverrideableProperties.Add(MoveTemp(Info));
+		}
+	}
+
+	return OverrideableProperties;
+}
+
 TArray<UMovieGraphPin*> UMovieGraphNode::EvaluatePinsToFollow(FMovieGraphEvaluationContext& InContext) const
 {
 	TArray<UMovieGraphPin*> PinsToFollow;
diff --git a/Engine/Plugins/MovieScene/MovieRenderPipeline/Source/MovieRenderPipelineCore/Private/Graph/MovieGraphPin.cpp b/Engine/Plugins/MovieScene/MovieRenderPipeline/Source/MovieRenderPipelineCore/Private/Graph/MovieGraphPin.cpp
--- a/Engine/Plugins/MovieScene/MovieRenderPipeline/Source/MovieRenderPipelineCore/Private/Graph/MovieGraphPin.cpp
+++ b/Engine/Plugins/MovieScene/MovieRenderPipeline/Source/MovieRenderPipelineCore/Private/Graph/MovieGraphPin.cpp
@@ -124,4 +124,19 @@ bool UMovieGraphPin::AllowsMultipleConnections() const
 {
 	// Always allow multiple connection on output pin
 	return IsOutputPin() || Properties.bAllowMultipleConnections;
+}
+
+UMovieGraphPin* UMovieGraphPin::GetFirstConnectedPin() const
+{
+	if (Edges.IsEmpty())
+	{
+		return nullptr;
+	}
+
+	if (ensureMsgf(Edges[0], TEXT("Null edge found when trying to get connected pin!")))
+	{
+		return Edges[0]->GetOtherPin(this);
+	}
+
+	return nullptr;
 }
\ No newline at end of file
diff --git a/Engine/Plugins/MovieScene/MovieRenderPipeline/Source/MovieRenderPipelineCore/Private/Graph/Nodes/MovieGraphRenderLayerNode.cpp b/Engine/Plugins/MovieScene/MovieRenderPipeline/Source/MovieRenderPipelineCore/Private/Graph/Nodes/MovieGraphRenderLayerNode.cpp
--- a/Engine/Plugins/MovieScene/MovieRenderPipeline/Source/MovieRenderPipelineCore/Private/Graph/Nodes/MovieGraphRenderLayerNode.cpp
+++ b/Engine/Plugins/MovieScene/MovieRenderPipeline/Source/MovieRenderPipelineCore/Private/Graph/Nodes/MovieGraphRenderLayerNode.cpp
@@ -1,4 +1,4 @@
-// Copyright Epic Games, Inc. All Rights Reserved.
+﻿// Copyright Epic Games, Inc. All Rights Reserved.
 
 #include "Graph/Nodes/MovieGraphRenderLayerNode.h"
 
@@ -14,7 +14,7 @@ TArray<FPropertyBagPropertyDesc> UMovieGraphRenderLayerNode::GetDynamicPropertyD
 {
 	TArray<FPropertyBagPropertyDesc> Properties;
 	FPropertyBagPropertyDesc FloatEditConProperty = FPropertyBagPropertyDesc("bOverride_TestPropName", EPropertyBagPropertyType::Bool);
-	FPropertyBagPropertyDesc FloatProperty = FPropertyBagPropertyDesc("TestPropName", EPropertyBagPropertyType::Float);
+	FPropertyBagPropertyDesc FloatProperty = FPropertyBagPropertyDesc("TestPropName", EPropertyBagPropertyType::String);
 #if WITH_EDITOR
 	FloatEditConProperty.MetaData.Add(FPropertyBagPropertyDescMetaData("InlineEditConditionToggle", "true"));
 	FloatProperty.MetaData.Add(FPropertyBagPropertyDescMetaData("EditCondition", "bOverride_TestPropName"));
diff --git a/Engine/Plugins/MovieScene/MovieRenderPipeline/Source/MovieRenderPipelineCore/Private/Graph/Nodes/MovieGraphVariableNode.cpp b/Engine/Plugins/MovieScene/MovieRenderPipeline/Source/MovieRenderPipelineCore/Private/Graph/Nodes/MovieGraphVariableNode.cpp
--- a/Engine/Plugins/MovieScene/MovieRenderPipeline/Source/MovieRenderPipelineCore/Private/Graph/Nodes/MovieGraphVariableNode.cpp
+++ b/Engine/Plugins/MovieScene/MovieRenderPipeline/Source/MovieRenderPipelineCore/Private/Graph/Nodes/MovieGraphVariableNode.cpp
@@ -29,6 +29,16 @@ TArray<FMovieGraphPinProperties> UMovieGraphVariableNode::GetOutputPinProperties
 	return Properties;
 }
 
+FString UMovieGraphVariableNode::GetResolvedValueForOutputPin(const FName& InPinName) const
+{
+	if (GraphVariable && (GraphVariable->Name == InPinName))
+	{
+		return GraphVariable->GetValueSerializedString();
+	}
+	
+	return FString();
+}
+
 void UMovieGraphVariableNode::SetVariable(UMovieGraphVariable* InVariable)
 {
 	if (InVariable)
diff --git a/Engine/Plugins/MovieScene/MovieRenderPipeline/Source/MovieRenderPipelineCore/Public/Graph/MovieGraphNode.h b/Engine/Plugins/MovieScene/MovieRenderPipeline/Source/MovieRenderPipelineCore/Public/Graph/MovieGraphNode.h
--- a/Engine/Plugins/MovieScene/MovieRenderPipeline/Source/MovieRenderPipelineCore/Public/Graph/MovieGraphNode.h
+++ b/Engine/Plugins/MovieScene/MovieRenderPipeline/Source/MovieRenderPipelineCore/Public/Graph/MovieGraphNode.h
@@ -30,6 +30,47 @@ class UEdGraphNode;
 
 DECLARE_MULTICAST_DELEGATE_OneParam(FOnMovieGraphNodeChanged, const UMovieGraphNode*);
 
+/**
+ * Information about a property that currently is (or can be) exposed on a node.
+ */
+USTRUCT(BlueprintType)
+struct FMovieGraphPropertyInfo
+{
+	GENERATED_BODY()
+	
+	/** The name of the property. */
+	UPROPERTY()
+	FName Name;
+
+	/** Whether this property is dynamic (ie, it does not correspond to a native UPROPERTY on the node). */
+	UPROPERTY()
+	bool bIsDynamicProperty = false;
+
+	/** The type of the value pointed to by the property. */
+	UPROPERTY()
+	EMovieGraphValueType ValueType = EMovieGraphValueType::None;
+
+	/**
+	 * Determines if this struct represents the same property as another instance of this struct.
+	 *
+	 * The equality operator ensures equality of all members of the struct. However, there may be cases where, for
+	 * example, a dynamic property changes its value type, and the ValueType member has not been changed yet. Two
+	 * structs that have the same Name and bIsDynamicProperty values should still be considered identical for the
+	 * purposes of determining if they identify the same property.
+	 */
+	bool IsSamePropertyAs(const FMovieGraphPropertyInfo& Other) const
+	{
+		return (Name == Other.Name) && (bIsDynamicProperty == Other.bIsDynamicProperty);
+	}
+
+	bool operator==(const FMovieGraphPropertyInfo& Other) const
+	{
+		return (Name == Other.Name)
+			&& (bIsDynamicProperty == Other.bIsDynamicProperty)
+			&& (ValueType == Other.ValueType);
+	}
+};
+
 /**
 * This is a base class for all nodes that can exist in the UMovieGraphConfig network.
 * In the editor, each node in the network will have an editor-only representation too 
@@ -55,19 +96,64 @@ public:
 	{
 		return TArray<FMovieGraphPinProperties>();
 	}
+	
 	virtual TArray<FMovieGraphPinProperties> GetOutputPinProperties() const
 	{
 		return TArray<FMovieGraphPinProperties>();
 	}
 
+	/**
+	 * Gets the resolved value of a named output pin (one that is returned in GetOutputPinProperties()). If the resolved
+	 * value could not be determined, an empty string is returned.
+	 */
+	virtual FString GetResolvedValueForOutputPin(const FName& InPinName) const
+	{
+		return FString();
+	}
+
+	/**
+	 * Gets the descriptions of properties which can be dynamically added to the node. These types of properties
+	 * do not correspond to a UPROPERTY defined on the node itself.
+	 */
 	virtual TArray<FPropertyBagPropertyDesc> GetDynamicPropertyDescriptions() const
 	{
 		return TArray<FPropertyBagPropertyDesc>();
 	}
 
-	virtual TArray<FName> GetExposedDynamicProperties() const
-	{ 
-		return ExposedDynamicPropertyNames;
+	/**
+	 * Sets the value of the dynamic property with the specified name. Note that the provided value must be the serialized
+	 * representation of the value. Returns true upon success, else false.
+	 */
+	bool SetDynamicPropertyValue(const FName PropertyName, const FString& InNewValue);
+
+	/**
+	 * Gets the value of the dynamic property with the specified name. Provides the serialized value of the property in
+	 * "OutValue". Returns true if "OutValue" was set and there were no errors, else returns false.
+	 */
+	bool GetDynamicPropertyValue(const FName PropertyName, FString& OutValue);
+
+	/** Gets the override property for the specified dynamic property. If one does not exist, returns nullptr. */
+	const FBoolProperty* FindOverridePropertyForDynamicProperty(const FName& InPropertyName) const;
+
+	/**
+	 * Returns true if the dynamic property with the provided name has been overridden, else false. Note that the name
+	 * provided should not be prefixed with "bOverride_".
+	 */
+	bool IsDynamicPropertyOverridden(const FName& InPropertyName) const;
+
+	/**
+	 * Sets the dynamic property with the provided name to the specified override state. Note that the name provided
+	 * should not be prefixed with "bOverride_".
+	 */
+	void SetDynamicPropertyOverridden(const FName& InPropertyName, const bool bIsOverridden);
+
+	/** Gets the information about properties which can be exposed as a pin on the node. */
+	virtual TArray<FMovieGraphPropertyInfo> GetOverrideablePropertyInfo() const;
+
+	/** Gets the information about properties which are currently exposed as pins on the node. */
+	virtual TArray<FMovieGraphPropertyInfo> GetExposedProperties() const
+	{
+		return ExposedPropertyInfo;
 	}
 
 	/** 
@@ -89,8 +175,15 @@ public:
 		return nullptr;
 	}
 
-	/** Promotes the property with the given name to a pin on the node via a dynamic property. */
-	virtual void PromoteDynamicPropertyToPin(const FName& PropertyName);
+	/** Toggles the promotion of the property with the given name to a pin on the node. */
+	virtual void TogglePromotePropertyToPin(const FName& PropertyName);
+
+	/**
+	 * Gets all overrideable properties that are defined on the node. This includes UPROPERTY-defined properties, as
+	 * well as dynamic properties. "Overrideable" means that the property has a corresponding property prefixed with
+	 * "bOverride_".
+	 */
+	TArray<const FProperty*> GetAllOverrideableProperties() const;
 
 	void UpdatePins();
 	void UpdateDynamicProperties();
@@ -151,7 +244,8 @@ public:
 #endif
 
 protected:
-	virtual TArray<FMovieGraphPinProperties> GetExposedDynamicPinProperties() const;
+	/** Gets the pin properties for all properties which have been exposed on the node. */
+	virtual TArray<FMovieGraphPinProperties> GetExposedPinProperties() const;
 
 	/** Register any delegates that need to be set up on the node. Called in PostLoad(). */
 	virtual void RegisterDelegates() const { }
@@ -163,11 +257,13 @@ protected:
 	UPROPERTY()
 	TArray<TObjectPtr<UMovieGraphPin>> OutputPins;
 
+	/** Properties which can be dynamically declared on the node (vs. native properties which are always present). */
 	UPROPERTY(EditAnywhere, meta=(FixedLayout, ShowOnlyInnerProperties), Category = "Node")
 	FInstancedPropertyBag DynamicProperties;
 
+	/** Tracks which properties have been exposed on the node as inputs. */
 	UPROPERTY()
-	TArray<FName> ExposedDynamicPropertyNames;
+	TArray<FMovieGraphPropertyInfo> ExposedPropertyInfo;
 
 #if WITH_EDITORONLY_DATA
 	UPROPERTY()
diff --git a/Engine/Plugins/MovieScene/MovieRenderPipeline/Source/MovieRenderPipelineCore/Public/Graph/MovieGraphPin.h b/Engine/Plugins/MovieScene/MovieRenderPipeline/Source/MovieRenderPipelineCore/Public/Graph/MovieGraphPin.h
--- a/Engine/Plugins/MovieScene/MovieRenderPipeline/Source/MovieRenderPipelineCore/Public/Graph/MovieGraphPin.h
+++ b/Engine/Plugins/MovieScene/MovieRenderPipeline/Source/MovieRenderPipelineCore/Public/Graph/MovieGraphPin.h
@@ -73,6 +73,9 @@ public:
 	int32 EdgeCount() const;
 	bool AllowsMultipleConnections() const;
 
+	/** Gets the first pin connected to this pin. Returns nullptr if no valid connection exists. */
+	UMovieGraphPin* GetFirstConnectedPin() const;
+
 public:
 	// The node that this pin belongs to.
 	UPROPERTY(BlueprintReadOnly, Category = "Properties")
diff --git a/Engine/Plugins/MovieScene/MovieRenderPipeline/Source/MovieRenderPipelineCore/Public/Graph/Nodes/MovieGraphVariableNode.h b/Engine/Plugins/MovieScene/MovieRenderPipeline/Source/MovieRenderPipelineCore/Public/Graph/Nodes/MovieGraphVariableNode.h
--- a/Engine/Plugins/MovieScene/MovieRenderPipeline/Source/MovieRenderPipelineCore/Public/Graph/Nodes/MovieGraphVariableNode.h
+++ b/Engine/Plugins/MovieScene/MovieRenderPipeline/Source/MovieRenderPipelineCore/Public/Graph/Nodes/MovieGraphVariableNode.h
@@ -17,6 +17,7 @@ public:
 	UMovieGraphVariableNode();
 
 	virtual TArray<FMovieGraphPinProperties> GetOutputPinProperties() const override;
+	virtual FString GetResolvedValueForOutputPin(const FName& InPinName) const override;
 
 	/** Gets the variable that this node represents. */
 	UMovieGraphVariable* GetVariable() const { return GraphVariable; }
diff --git a/Engine/Plugins/MovieScene/MovieRenderPipeline/Source/MovieRenderPipelineEditor/Private/Graph/MovieEdGraphNode.cpp b/Engine/Plugins/MovieScene/MovieRenderPipeline/Source/MovieRenderPipelineEditor/Private/Graph/MovieEdGraphNode.cpp
--- a/Engine/Plugins/MovieScene/MovieRenderPipeline/Source/MovieRenderPipelineEditor/Private/Graph/MovieEdGraphNode.cpp
+++ b/Engine/Plugins/MovieScene/MovieRenderPipeline/Source/MovieRenderPipelineEditor/Private/Graph/MovieEdGraphNode.cpp
@@ -63,21 +63,53 @@ FEdGraphPinType UMoviePipelineEdGraphNodeBase::GetPinType(const UMovieGraphPin*
 		return EdPinType;
 	}
 
-	// TODO: This needs to properly fill in the category/subcategory for all types. Also, container type?
 	switch (InPin->Properties.Type)
 	{
-	case EMovieGraphValueType::Float:
-		EdPinType.PinCategory = UMovieGraphSchema::PC_Float;
+	case EMovieGraphValueType::Bool:
+		EdPinType.PinCategory = UMovieGraphSchema::PC_Boolean;
+		break;
+	case EMovieGraphValueType::Byte:
+		EdPinType.PinCategory = UMovieGraphSchema::PC_Byte;
 		break;
 	case EMovieGraphValueType::Int32:
 		EdPinType.PinCategory = UMovieGraphSchema::PC_Integer;
 		break;
-	case EMovieGraphValueType::Bool:
-		EdPinType.PinCategory = UMovieGraphSchema::PC_Boolean;
+	case EMovieGraphValueType::Int64:
+		EdPinType.PinCategory = UMovieGraphSchema::PC_Int64;
+		break;
+	case EMovieGraphValueType::Float:
+		EdPinType.PinCategory = UMovieGraphSchema::PC_Float;
+		break;
+	case EMovieGraphValueType::Double:
+		EdPinType.PinCategory = UMovieGraphSchema::PC_Double;
+		break;
+	case EMovieGraphValueType::Name:
+		EdPinType.PinCategory = UMovieGraphSchema::PC_Name;
 		break;
 	case EMovieGraphValueType::String:
 		EdPinType.PinCategory = UMovieGraphSchema::PC_String;
 		break;
+	case EMovieGraphValueType::Text:
+		EdPinType.PinCategory = UMovieGraphSchema::PC_Text;
+		break;
+	case EMovieGraphValueType::Enum:
+		EdPinType.PinCategory = UMovieGraphSchema::PC_Enum;
+		break;
+	case EMovieGraphValueType::Struct:
+		EdPinType.PinCategory = UMovieGraphSchema::PC_Struct;
+		break;
+	case EMovieGraphValueType::Object:
+		EdPinType.PinCategory = UMovieGraphSchema::PC_Object;
+		break;
+	case EMovieGraphValueType::SoftObject:
+		EdPinType.PinCategory = UMovieGraphSchema::PC_SoftObject;
+		break;
+	case EMovieGraphValueType::Class:
+		EdPinType.PinCategory = UMovieGraphSchema::PC_Class;
+		break;
+	case EMovieGraphValueType::SoftClass:
+		EdPinType.PinCategory = UMovieGraphSchema::PC_SoftClass;
+		break;
 	default:
 		EdPinType.PinCategory = UMovieGraphSchema::PC_Float;
 		break;
@@ -152,45 +184,29 @@ void UMoviePipelineEdGraphNode::GetNodeContextMenuActions(UToolMenu* Menu, UGrap
 
 void UMoviePipelineEdGraphNode::GetPropertyPromotionContextMenuActions(UToolMenu* Menu, UGraphNodeContextMenuContext* Context) const
 {
-	// Exclude the dynamic properties which are used for EditCondition metadata
-	static const FName EditConditionKey = FName("EditCondition");
-	TArray<FName> ExcludedPropertyNames;
-	for (const FPropertyBagPropertyDesc& PropertyDescription : RuntimeNode->GetDynamicPropertyDescriptions())
-	{
-		for (const FPropertyBagPropertyDescMetaData& Metadata : PropertyDescription.MetaData)
-		{
-			if (Metadata.Key == EditConditionKey)
-			{
-				ExcludedPropertyNames.Add(FName(Metadata.Value));
-			}
-		}
-	}
-
-	// Exclude dynamic properties which have already been promoted
-	ExcludedPropertyNames.Append(RuntimeNode->GetExposedDynamicProperties());
-
 	FToolMenuSection& Section = Menu->AddSection("MoviePipelineGraphExposeAsPin", LOCTEXT("ExposeAsPin", "Expose Property as Pin"));
 
-	const TArray<FPropertyBagPropertyDesc>& PropertyDescriptions = RuntimeNode->GetDynamicPropertyDescriptions();
-	for (const FPropertyBagPropertyDesc& PropertyDescription : PropertyDescriptions)
+	const TArray<FMovieGraphPropertyInfo>& OverrideablePropertyInfo = RuntimeNode->GetOverrideablePropertyInfo();
+	for (const FMovieGraphPropertyInfo& PropertyInfo : OverrideablePropertyInfo)
 	{
-		if (ExcludedPropertyNames.Contains(PropertyDescription.Name))
-		{
-			continue;
-		}
-		
 		Section.AddMenuEntry(
-			PropertyDescription.Name,
-			FText::FromName(PropertyDescription.Name),
+			PropertyInfo.Name,
+			FText::FromName(PropertyInfo.Name),
 			LOCTEXT("PromotePropertyToPin", "Promote this property to a pin on this node."),
 			FSlateIcon(),
 			FUIAction(
-				FExecuteAction::CreateUObject(this, &UMoviePipelineEdGraphNode::PromotePropertyToPin, PropertyDescription.Name),
-				FCanExecuteAction())
+				FExecuteAction::CreateUObject(this, &UMoviePipelineEdGraphNode::TogglePromotePropertyToPin, PropertyInfo.Name),
+				FCanExecuteAction(),
+				FGetActionCheckState::CreateLambda([this, PropertyInfo]()
+				{
+					const TArray<FMovieGraphPropertyInfo>& ExposedProperties = RuntimeNode->GetExposedProperties();
+					return ExposedProperties.Contains(PropertyInfo) ? ECheckBoxState::Checked : ECheckBoxState::Unchecked;
+				})),
+			EUserInterfaceActionType::ToggleButton
 		);
 	}
 
-	if (PropertyDescriptions.Num() == ExcludedPropertyNames.Num())
+	if (OverrideablePropertyInfo.IsEmpty())
 	{
 		Section.AddMenuEntry(
 			"NoPropertiesAvailable",
@@ -204,9 +220,9 @@ void UMoviePipelineEdGraphNode::GetPropertyPromotionContextMenuActions(UToolMenu
 	}
 }
 
-void UMoviePipelineEdGraphNode::PromotePropertyToPin(const FName PropertyName) const
+void UMoviePipelineEdGraphNode::TogglePromotePropertyToPin(const FName PropertyName) const
 {
-	RuntimeNode->PromoteDynamicPropertyToPin(PropertyName);
+	RuntimeNode->TogglePromotePropertyToPin(PropertyName);
 }
 
 bool UMoviePipelineEdGraphNodeBase::ShouldCreatePin(const UMovieGraphPin* InPin) const
diff --git a/Engine/Plugins/MovieScene/MovieRenderPipeline/Source/MovieRenderPipelineEditor/Private/Graph/MovieEdGraphNode.h b/Engine/Plugins/MovieScene/MovieRenderPipeline/Source/MovieRenderPipelineEditor/Private/Graph/MovieEdGraphNode.h
--- a/Engine/Plugins/MovieScene/MovieRenderPipeline/Source/MovieRenderPipelineEditor/Private/Graph/MovieEdGraphNode.h
+++ b/Engine/Plugins/MovieScene/MovieRenderPipeline/Source/MovieRenderPipelineEditor/Private/Graph/MovieEdGraphNode.h
@@ -84,8 +84,8 @@ public:
 	//~ End UEdGraphNode Interface
 
 protected:
-	/** Promote the property with the given name to a pin exposed on the node. */
-	void PromotePropertyToPin(const FName PropertyName) const;
+	/** Toggle the promotion of the property with the given name to a pin. */
+	void TogglePromotePropertyToPin(const FName PropertyName) const;
 
 private:
 	void GetPropertyPromotionContextMenuActions(UToolMenu* Menu, UGraphNodeContextMenuContext* Context) const;
diff --git a/Engine/Plugins/MovieScene/MovieRenderPipeline/Source/MovieRenderPipelineEditor/Private/Graph/MovieGraphSchema.cpp b/Engine/Plugins/MovieScene/MovieRenderPipeline/Source/MovieRenderPipelineEditor/Private/Graph/MovieGraphSchema.cpp
--- a/Engine/Plugins/MovieScene/MovieRenderPipeline/Source/MovieRenderPipelineEditor/Private/Graph/MovieGraphSchema.cpp
+++ b/Engine/Plugins/MovieScene/MovieRenderPipeline/Source/MovieRenderPipelineEditor/Private/Graph/MovieGraphSchema.cpp
@@ -20,11 +20,21 @@ TArray<UClass*> UMovieGraphSchema::MoviePipelineNodeClasses;
 #define LOCTEXT_NAMESPACE "MoviePipelineGraphSchema"
 
 const FName UMovieGraphSchema::PC_Branch(TEXT("branch"));
-const FName UMovieGraphSchema::PC_Float(TEXT("float"));
-const FName UMovieGraphSchema::PC_Integer(TEXT("integer"));
 const FName UMovieGraphSchema::PC_Boolean(TEXT("boolean"));
+const FName UMovieGraphSchema::PC_Byte(TEXT("byte"));
+const FName UMovieGraphSchema::PC_Integer(TEXT("integer"));
+const FName UMovieGraphSchema::PC_Int64(TEXT("int64"));
+const FName UMovieGraphSchema::PC_Float(TEXT("float"));
+const FName UMovieGraphSchema::PC_Double(TEXT("double"));
+const FName UMovieGraphSchema::PC_Name(TEXT("name"));
 const FName UMovieGraphSchema::PC_String(TEXT("string"));
-const FName UMovieGraphSchema::PC_IntPoint(TEXT("intpoint"));
+const FName UMovieGraphSchema::PC_Text(TEXT("text"));
+const FName UMovieGraphSchema::PC_Enum(TEXT("enum"));
+const FName UMovieGraphSchema::PC_Struct(TEXT("struct"));
+const FName UMovieGraphSchema::PC_Object(TEXT("object"));
+const FName UMovieGraphSchema::PC_SoftObject(TEXT("softobject"));
+const FName UMovieGraphSchema::PC_Class(TEXT("class"));
+const FName UMovieGraphSchema::PC_SoftClass(TEXT("softclass"));
 
 void UMovieGraphSchema::CreateDefaultNodesForGraph(UEdGraph& Graph) const
 {
@@ -259,9 +269,14 @@ FLinearColor UMovieGraphSchema::GetPinTypeColor(const FEdGraphPinType& PinType)
 		return Settings->ExecutionPinTypeColor;
 	}
 
-	if (PinType.PinCategory == PC_Float)
+	if (PinType.PinCategory == PC_Boolean)
 	{
-		return Settings->FloatPinTypeColor;
+		return Settings->BooleanPinTypeColor;
+	}
+	
+	if (PinType.PinCategory == PC_Byte)
+	{
+		return Settings->BytePinTypeColor;
 	}
 
 	if (PinType.PinCategory == PC_Integer)
@@ -269,9 +284,24 @@ FLinearColor UMovieGraphSchema::GetPinTypeColor(const FEdGraphPinType& PinType)
 		return Settings->IntPinTypeColor;
 	}
 
-	if (PinType.PinCategory == PC_Boolean)
+	if (PinType.PinCategory == PC_Int64)
 	{
-		return Settings->BooleanPinTypeColor;
+		return Settings->Int64PinTypeColor;
+	}
+
+	if (PinType.PinCategory == PC_Float)
+	{
+		return Settings->FloatPinTypeColor;
+	}
+
+	if (PinType.PinCategory == PC_Double)
+	{
+		return Settings->DoublePinTypeColor;
+	}
+
+	if (PinType.PinCategory == PC_Name)
+	{
+		return Settings->NamePinTypeColor;
 	}
 
 	if (PinType.PinCategory == PC_String)
@@ -279,9 +309,39 @@ FLinearColor UMovieGraphSchema::GetPinTypeColor(const FEdGraphPinType& PinType)
 		return Settings->StringPinTypeColor;
 	}
 
-	if (PinType.PinCategory == PC_IntPoint)
+	if (PinType.PinCategory == PC_Text)
+	{
+		return Settings->TextPinTypeColor;
+	}
+
+	if (PinType.PinCategory == PC_Enum)
+	{
+		return Settings->BytePinTypeColor;
+	}
+
+	if (PinType.PinCategory == PC_Struct)
+	{
+		return Settings->StructPinTypeColor;
+	}
+	
+	if (PinType.PinCategory == PC_Object)
+	{
+		return Settings->ObjectPinTypeColor;
+	}
+
+	if (PinType.PinCategory == PC_SoftObject)
+	{
+		return Settings->SoftObjectPinTypeColor;
+	}
+	
+	if (PinType.PinCategory == PC_Class)
+    {
+    	return Settings->ClassPinTypeColor;
+    }
+
+	if (PinType.PinCategory == PC_SoftClass)
 	{
-		return Settings->VectorPinTypeColor;
+		return Settings->SoftClassPinTypeColor;
 	}
 	
 	return Settings->DefaultPinTypeColor;
diff --git a/Engine/Plugins/MovieScene/MovieRenderPipeline/Source/MovieRenderPipelineEditor/Private/Graph/MovieGraphSchema.h b/Engine/Plugins/MovieScene/MovieRenderPipeline/Source/MovieRenderPipelineEditor/Private/Graph/MovieGraphSchema.h
--- a/Engine/Plugins/MovieScene/MovieRenderPipeline/Source/MovieRenderPipelineEditor/Private/Graph/MovieGraphSchema.h
+++ b/Engine/Plugins/MovieScene/MovieRenderPipeline/Source/MovieRenderPipelineEditor/Private/Graph/MovieGraphSchema.h
@@ -38,11 +38,21 @@ public:
 public:
 	// Allowed "PinCategory" values for use on EdGraphPin
 	static const FName PC_Branch;
-	static const FName PC_Float;
-	static const FName PC_Integer;
 	static const FName PC_Boolean;
+	static const FName PC_Byte;
+	static const FName PC_Integer;
+	static const FName PC_Int64;
+	static const FName PC_Float;
+	static const FName PC_Double;
+	static const FName PC_Name;
 	static const FName PC_String;
-	static const FName PC_IntPoint;
+	static const FName PC_Text;
+	static const FName PC_Enum;
+	static const FName PC_Struct;
+	static const FName PC_Object;
+	static const FName PC_SoftObject;
+	static const FName PC_Class;
+	static const FName PC_SoftClass;
 	
 private:
 	static TArray<UClass*> MoviePipelineNodeClasses;
diff --git a/Engine/Plugins/MovieScene/MovieRenderPipeline/Source/MovieRenderPipelineEditor/Private/Graph/SMovieGraphActiveRenderSettingsTabContent.cpp b/Engine/Plugins/MovieScene/MovieRenderPipeline/Source/MovieRenderPipelineEditor/Private/Graph/SMovieGraphActiveRenderSettingsTabContent.cpp
--- a/Engine/Plugins/MovieScene/MovieRenderPipeline/Source/MovieRenderPipelineEditor/Private/Graph/SMovieGraphActiveRenderSettingsTabContent.cpp
+++ b/Engine/Plugins/MovieScene/MovieRenderPipeline/Source/MovieRenderPipelineEditor/Private/Graph/SMovieGraphActiveRenderSettingsTabContent.cpp
@@ -23,14 +23,26 @@ FActiveRenderSettingsTreeElement::FActiveRenderSettingsTreeElement(const FName&
 
 FString FActiveRenderSettingsTreeElement::GetValue() const
 {
-	if (SettingsProperty && SettingsNode)
+	FString ValueString;
+	
+	if (!SettingsProperty || !SettingsNode)
 	{
-		FString ValueString;
-		SettingsProperty->ExportTextItem_InContainer(ValueString, SettingsNode, nullptr, nullptr, PPF_None);
 		return ValueString;
 	}
 
-	return FString();
+	// If the property is a dynamic property, then get its value from the node
+	for (const FPropertyBagPropertyDesc& PropertyDesc : SettingsNode->GetDynamicPropertyDescriptions())
+	{
+		if (SettingsProperty->GetFName() == PropertyDesc.Name)
+		{
+			SettingsNode->GetDynamicPropertyValue(PropertyDesc.Name, ValueString);
+			return ValueString;
+		}
+	}
+
+	// Otherwise, ask the property for its value directly
+	SettingsProperty->ExportTextItem_InContainer(ValueString, SettingsNode, nullptr, nullptr, PPF_None);
+	return ValueString;
 }
 
 bool FActiveRenderSettingsTreeElement::IsBranchRenderable() const
@@ -128,19 +140,15 @@ const TArray<TSharedPtr<FActiveRenderSettingsTreeElement>>& FActiveRenderSetting
 	// For node elements: all overrideable properties on the node should be children
 	if (Type == EElementType::Node)
 	{
-		for (TFieldIterator<FProperty> PropertyIterator(SettingsNode->GetClass()); PropertyIterator; ++PropertyIterator)
+		for (const FProperty* Property : SettingsNode->GetAllOverrideableProperties())
 		{
-			FProperty* NodeProperty = *PropertyIterator;
-			if (UMovieGraphConfig::FindOverridePropertyForRealProperty(SettingsNode->GetClass(), NodeProperty))
-			{
-				TSharedPtr<FActiveRenderSettingsTreeElement> Element =
-					MakeShared<FActiveRenderSettingsTreeElement>(NodeProperty->GetFName(), EElementType::Property);
-				Element->SettingsNode = SettingsNode;
-				Element->SettingsProperty = NodeProperty;
-				Element->FlattenedGraph = FlattenedGraph;
-				
-				ChildrenCache.Add(MoveTemp(Element));
-			}
+			TSharedPtr<FActiveRenderSettingsTreeElement> Element =
+				MakeShared<FActiveRenderSettingsTreeElement>(Property->GetFName(), EElementType::Property);
+			Element->SettingsNode = SettingsNode;
+			Element->SettingsProperty = Property;
+			Element->FlattenedGraph = FlattenedGraph;
+			
+			ChildrenCache.Add(MoveTemp(Element));
 		}
 	}
 	
diff --git a/Engine/Plugins/MovieScene/MovieRenderPipeline/Source/MovieRenderPipelineEditor/Private/Graph/SMovieGraphActiveRenderSettingsTabContent.h b/Engine/Plugins/MovieScene/MovieRenderPipeline/Source/MovieRenderPipelineEditor/Private/Graph/SMovieGraphActiveRenderSettingsTabContent.h
--- a/Engine/Plugins/MovieScene/MovieRenderPipeline/Source/MovieRenderPipelineEditor/Private/Graph/SMovieGraphActiveRenderSettingsTabContent.h
+++ b/Engine/Plugins/MovieScene/MovieRenderPipeline/Source/MovieRenderPipelineEditor/Private/Graph/SMovieGraphActiveRenderSettingsTabContent.h
@@ -61,7 +61,7 @@ public:
 	TObjectPtr<UMovieGraphNode> SettingsNode = nullptr;
 
 	/** If this element represents a property on a node, this is the pointer to that property. */
-	FProperty* SettingsProperty = nullptr;
+	const FProperty* SettingsProperty = nullptr;
 
 	/** The result of the graph traversal. */
 	TWeakObjectPtr<UMovieGraphEvaluatedConfig> FlattenedGraph = nullptr;
