MoviePipeline: Updated the Active Render Settings tree to remember expansion state across refreshes.

diff --git a/Engine/Plugins/MovieScene/MovieRenderPipeline/Source/MovieRenderPipelineEditor/Private/Graph/SMovieGraphActiveRenderSettingsTabContent.cpp b/Engine/Plugins/MovieScene/MovieRenderPipeline/Source/MovieRenderPipelineEditor/Private/Graph/SMovieGraphActiveRenderSettingsTabContent.cpp
--- a/Engine/Plugins/MovieScene/MovieRenderPipeline/Source/MovieRenderPipelineEditor/Private/Graph/SMovieGraphActiveRenderSettingsTabContent.cpp
+++ b/Engine/Plugins/MovieScene/MovieRenderPipeline/Source/MovieRenderPipelineEditor/Private/Graph/SMovieGraphActiveRenderSettingsTabContent.cpp
@@ -72,6 +72,7 @@ const TArray<TSharedPtr<FActiveRenderSettingsTreeElement>>& FActiveRenderSetting
 			MakeShared<FActiveRenderSettingsTreeElement>(FName(*Node->GetNodeTitle().ToString()), EElementType::Node);
 		Element->SettingsNode = Node;
 		Element->FlattenedGraph = FlattenedGraph;
+		Element->ParentElement = AsWeak();
 
 		return MoveTemp(Element);
 	};
@@ -112,6 +113,7 @@ const TArray<TSharedPtr<FActiveRenderSettingsTreeElement>>& FActiveRenderSetting
 				TSharedPtr<FActiveRenderSettingsTreeElement> Element =
 					MakeShared<FActiveRenderSettingsTreeElement>(FName(Pair.Key), EElementType::NamedBranch);
 				Element->FlattenedGraph = FlattenedGraph;
+				Element->ParentElement = AsWeak();
 				
 				ChildrenCache.Add(MoveTemp(Element));
 			}
@@ -147,6 +149,7 @@ const TArray<TSharedPtr<FActiveRenderSettingsTreeElement>>& FActiveRenderSetting
 			Element->SettingsNode = SettingsNode;
 			Element->SettingsProperty = Property;
 			Element->FlattenedGraph = FlattenedGraph;
+			Element->ParentElement = AsWeak();
 			
 			ChildrenCache.Add(MoveTemp(Element));
 		}
@@ -160,6 +163,22 @@ void FActiveRenderSettingsTreeElement::ClearCachedChildren() const
 	ChildrenCache.Empty();
 }
 
+uint32 FActiveRenderSettingsTreeElement::GetHash() const
+{
+	if (ElementHash != 0)
+	{
+		return ElementHash;
+	}
+
+	ElementHash = GetTypeHash(Name);
+	if (ParentElement.IsValid())
+	{
+		ElementHash = HashCombine(ElementHash, ParentElement.Pin()->GetHash());
+	}
+
+	return ElementHash;
+}
+
 void SMovieGraphActiveRenderSettingsTreeItem::Construct(
 	const FArguments& InArgs, const TSharedRef<STableViewBase>& InOwnerTable, const TSharedPtr<FActiveRenderSettingsTreeElement>& InTreeElement)
 {
@@ -288,6 +307,7 @@ void SMovieGraphActiveRenderSettingsTabContent::Construct(const FArguments& InAr
 			.SelectionMode(ESelectionMode::Single)
 			.OnGenerateRow(this, &SMovieGraphActiveRenderSettingsTabContent::GenerateTreeRow)
 			.OnGetChildren(this, &SMovieGraphActiveRenderSettingsTabContent::GetChildrenForTree)
+			.OnExpansionChanged(this, &SMovieGraphActiveRenderSettingsTabContent::OnExpansionChanged)
 			.HeaderRow
 			(
 				SNew(SHeaderRow)
@@ -346,6 +366,12 @@ void SMovieGraphActiveRenderSettingsTabContent::TraverseGraph()
 	}
 	
 	TreeView->RequestTreeRefresh();
+
+	// Restore tree expansion state after the refresh
+	for (const TSharedPtr<FActiveRenderSettingsTreeElement>& RootElement : RootElements)
+	{
+		RestoreExpansionStateRecursive(RootElement);
+	}
 }
 
 FReply SMovieGraphActiveRenderSettingsTabContent::OnEvaluateGraphClicked()
@@ -354,6 +380,28 @@ FReply SMovieGraphActiveRenderSettingsTabContent::OnEvaluateGraphClicked()
 	return FReply::Handled();
 }
 
+void SMovieGraphActiveRenderSettingsTabContent::OnExpansionChanged(TSharedPtr<FActiveRenderSettingsTreeElement> InElement, bool bInExpanded)
+{
+	if (bInExpanded)
+	{
+		ExpandedElements.Add(InElement->GetHash());
+	}
+	else
+	{
+		ExpandedElements.Remove(InElement->GetHash());
+	}
+}
+
+void SMovieGraphActiveRenderSettingsTabContent::RestoreExpansionStateRecursive(const TSharedPtr<FActiveRenderSettingsTreeElement>& InElement)
+{
+	TreeView->SetItemExpansion(InElement, ExpandedElements.Contains(InElement->GetHash()));
+	
+	for (const TSharedPtr<FActiveRenderSettingsTreeElement>& ChildElement : InElement->GetChildren())
+	{
+		RestoreExpansionStateRecursive(ChildElement);
+	}
+}
+
 TSharedRef<ITableRow> SMovieGraphActiveRenderSettingsTabContent::GenerateTreeRow(
 	TSharedPtr<FActiveRenderSettingsTreeElement> InTreeElement, const TSharedRef<STableViewBase>& OwnerTable)
 {
diff --git a/Engine/Plugins/MovieScene/MovieRenderPipeline/Source/MovieRenderPipelineEditor/Private/Graph/SMovieGraphActiveRenderSettingsTabContent.h b/Engine/Plugins/MovieScene/MovieRenderPipeline/Source/MovieRenderPipelineEditor/Private/Graph/SMovieGraphActiveRenderSettingsTabContent.h
--- a/Engine/Plugins/MovieScene/MovieRenderPipeline/Source/MovieRenderPipelineEditor/Private/Graph/SMovieGraphActiveRenderSettingsTabContent.h
+++ b/Engine/Plugins/MovieScene/MovieRenderPipeline/Source/MovieRenderPipelineEditor/Private/Graph/SMovieGraphActiveRenderSettingsTabContent.h
@@ -14,7 +14,7 @@ class UMovieGraphNode;
  * Represents the data of one element displayed in the tree, which is one of the types in EElementType. Provides a way
  * of getting the children under the element.
  */
-class FActiveRenderSettingsTreeElement
+class FActiveRenderSettingsTreeElement : public TSharedFromThis<FActiveRenderSettingsTreeElement>
 {
 public:
 	/** The type of this element. */
@@ -47,6 +47,9 @@ public:
 	/** Clears the cached result of GetChildren(). */
 	void ClearCachedChildren() const;
 
+	/** Gets the hash that uniquely identifies this element in the tree. */
+	uint32 GetHash() const;
+
 public:
 	/** The name of the element as it appears in the tree. */
 	FName Name;
@@ -66,9 +69,15 @@ public:
 	/** The result of the graph traversal. */
 	TWeakObjectPtr<UMovieGraphEvaluatedConfig> FlattenedGraph = nullptr;
 
+	/** This element's parent element in the tree.  */
+	TWeakPtr<const FActiveRenderSettingsTreeElement> ParentElement = nullptr;
+
 private:
 	/** The (cached) child elements nested under this element. */
 	mutable TArray<TSharedPtr<FActiveRenderSettingsTreeElement>> ChildrenCache;
+
+	/** The (cached) hash of this element. */
+	mutable uint32 ElementHash = 0;
 };
 
 /**
@@ -119,6 +128,12 @@ private:
 	/** Handles the button click for evaluating the graph. */
 	FReply OnEvaluateGraphClicked();
 
+	/** Handles the expansion state change of an element in the tree. */
+	void OnExpansionChanged(TSharedPtr<FActiveRenderSettingsTreeElement> InElement, bool bInExpanded);
+
+	/** Recursively restores the cached expansion state of the tree, starting at the specified element. */
+	void RestoreExpansionStateRecursive(const TSharedPtr<FActiveRenderSettingsTreeElement>& InElement);
+
 	/** Generates a row in the tree widget based on the provided element. */
 	TSharedRef<ITableRow> GenerateTreeRow(TSharedPtr<FActiveRenderSettingsTreeElement> InTreeElement, const TSharedRef<STableViewBase>& OwnerTable);
 
@@ -132,6 +147,9 @@ private:
 	/** The main tree widget displayed in the tab. */
 	TSharedPtr<STreeView<TSharedPtr<FActiveRenderSettingsTreeElement>>> TreeView = nullptr;
 
+	/** The hashes of elements that are currently expanded in the tree. */
+	TSet<uint32> ExpandedElements;
+
 	/** The last result of a graph traversal. */
 	TStrongObjectPtr<UMovieGraphEvaluatedConfig> FlattenedGraph = nullptr;
 
