Speculative include fix in movie pipeline code

#rnx

diff --git a/Engine/Plugins/MovieScene/MovieRenderPipeline/Source/MovieRenderPipelineCore/Public/MoviePipeline.h b/Engine/Plugins/MovieScene/MovieRenderPipeline/Source/MovieRenderPipelineCore/Public/MoviePipeline.h
--- a/Engine/Plugins/MovieScene/MovieRenderPipeline/Source/MovieRenderPipelineCore/Public/MoviePipeline.h
+++ b/Engine/Plugins/MovieScene/MovieRenderPipeline/Source/MovieRenderPipelineCore/Public/MoviePipeline.h
@@ -2,6 +2,7 @@
 #pragma once
 
 #include "UObject/Object.h"
+#include "HAL/ThreadSafeBool.h"
 #include "Engine/EngineCustomTimeStep.h"
 #include "MovieRenderPipelineDataTypes.h"
 #include "MoviePipelineBlueprintLibrary.h"
diff --git a/Engine/Plugins/MovieScene/MovieRenderPipeline/Source/MovieRenderPipelineRenderPasses/Public/MoviePipelineImagePassBase.h b/Engine/Plugins/MovieScene/MovieRenderPipeline/Source/MovieRenderPipelineRenderPasses/Public/MoviePipelineImagePassBase.h
--- a/Engine/Plugins/MovieScene/MovieRenderPipeline/Source/MovieRenderPipelineRenderPasses/Public/MoviePipelineImagePassBase.h
+++ b/Engine/Plugins/MovieScene/MovieRenderPipeline/Source/MovieRenderPipelineRenderPasses/Public/MoviePipelineImagePassBase.h
@@ -9,6 +9,7 @@
 #include "UObject/GCObject.h"
 #include "Async/AsyncWork.h"
 #include "Async/TaskGraphInterfaces.h"
+#include "HAL/ThreadSafeBool.h"
 #include "Templates/Function.h"
 #include "Stats/Stats.h"
 #include "CanvasTypes.h"
