Speculative include fix in movierenderpipeline code

#rnx

diff --git a/Engine/Plugins/MovieScene/MovieRenderPipeline/Source/MovieRenderPipelineCore/Public/Graph/MovieGraphDefaultRenderer.h b/Engine/Plugins/MovieScene/MovieRenderPipeline/Source/MovieRenderPipelineCore/Public/Graph/MovieGraphDefaultRenderer.h
--- a/Engine/Plugins/MovieScene/MovieRenderPipeline/Source/MovieRenderPipelineCore/Public/Graph/MovieGraphDefaultRenderer.h
+++ b/Engine/Plugins/MovieScene/MovieRenderPipeline/Source/MovieRenderPipelineCore/Public/Graph/MovieGraphDefaultRenderer.h
@@ -5,6 +5,7 @@
 #include "PixelFormat.h"
 #include "MovieRenderPipelineCoreModule.h"
 #include "Async/TaskGraphInterfaces.h"
+#include "HAL/ThreadSafeBool.h"
 #include "Camera/CameraTypes.h"
 #include "Tasks/Task.h"
 #include "MovieGraphDefaultRenderer.generated.h"
diff --git a/Engine/Plugins/MovieScene/MovieRenderPipeline/Source/MovieRenderPipelineCore/Public/Graph/MovieGraphPipeline.h b/Engine/Plugins/MovieScene/MovieRenderPipeline/Source/MovieRenderPipelineCore/Public/Graph/MovieGraphPipeline.h
--- a/Engine/Plugins/MovieScene/MovieRenderPipeline/Source/MovieRenderPipelineCore/Public/Graph/MovieGraphPipeline.h
+++ b/Engine/Plugins/MovieScene/MovieRenderPipeline/Source/MovieRenderPipelineCore/Public/Graph/MovieGraphPipeline.h
@@ -2,6 +2,7 @@
 #pragma once
 
 #include "UObject/Object.h"
+#include "HAL/ThreadSafeBool.h"
 #include "MoviePipelineBase.h"
 #include "MovieGraphDataTypes.h"
 #include "MovieGraphConfig.h"
