MoviePipeline: Cleaned up a few variable-related annoyances in the graph.
* Fixed an issue which caused a variable node to be deselected when the variable is renamed.
* Prefixed all variable entries in the node creation menu with "Get" to be more consistent w/ other graphs (eg, "Get FooVar").
* Variable nodes are now selected on creation like other non-variable nodes.

diff --git a/Engine/Plugins/MovieScene/MovieRenderPipeline/Source/MovieRenderPipelineEditor/Private/Graph/MovieGraphSchema.cpp b/Engine/Plugins/MovieScene/MovieRenderPipeline/Source/MovieRenderPipelineEditor/Private/Graph/MovieGraphSchema.cpp
--- a/Engine/Plugins/MovieScene/MovieRenderPipeline/Source/MovieRenderPipelineEditor/Private/Graph/MovieGraphSchema.cpp
+++ b/Engine/Plugins/MovieScene/MovieRenderPipeline/Source/MovieRenderPipelineEditor/Private/Graph/MovieGraphSchema.cpp
@@ -136,7 +136,7 @@ void UMovieGraphSchema::GetGraphContextActions(FGraphContextMenuBuilder& Context
 	const bool bIncludeGlobal = true;
 	for (const UMovieGraphVariable* Variable : RuntimeGraph->GetVariables(bIncludeGlobal))
 	{
-		const FText Name = FText::FromString(Variable->Name);
+		const FText Name = FText::Format(LOCTEXT("CreateVariable_Name", "Get {0}"), FText::FromString(Variable->Name));
 		const FText Category = Variable->IsGlobal() ? LOCTEXT("CreateGlobalVariable_Category", "Global Variables") : LOCTEXT("CreateVariable_Category", "Variables");
 		const FText Tooltip = LOCTEXT("CreateVariable_Tooltip", "Create an accessor node for this variable.");
 		
@@ -407,7 +407,7 @@ UEdGraphNode* FMovieGraphSchemaAction_NewVariableNode::PerformAction(UEdGraph* P
 
 	// Now create the variable node
 	FGraphNodeCreator<UMoviePipelineEdGraphVariableNode> NodeCreator(*ParentGraph);
-	UMoviePipelineEdGraphVariableNode* GraphNode = NodeCreator.CreateNode(bSelectNewNode);
+	UMoviePipelineEdGraphVariableNode* GraphNode = NodeCreator.CreateUserInvokedNode(bSelectNewNode);
 	GraphNode->Construct(RuntimeNode);
 	GraphNode->NodePosX = Location.X;
 	GraphNode->NodePosY = Location.Y;
diff --git a/Engine/Plugins/MovieScene/MovieRenderPipeline/Source/MovieRenderPipelineEditor/Private/Graph/SMovieGraphMembersTabContent.cpp b/Engine/Plugins/MovieScene/MovieRenderPipeline/Source/MovieRenderPipelineEditor/Private/Graph/SMovieGraphMembersTabContent.cpp
--- a/Engine/Plugins/MovieScene/MovieRenderPipeline/Source/MovieRenderPipelineEditor/Private/Graph/SMovieGraphMembersTabContent.cpp
+++ b/Engine/Plugins/MovieScene/MovieRenderPipeline/Source/MovieRenderPipelineEditor/Private/Graph/SMovieGraphMembersTabContent.cpp
@@ -11,6 +11,28 @@
 
 #define LOCTEXT_NAMESPACE "SMoviePipelineMembersTabContent"
 
+namespace UE::MovieGraph::Private
+{
+	/** Gets a graph member from a FEdGraphSchemaAction. */
+	UMovieGraphMember* GetMemberFromAction(FEdGraphSchemaAction* SchemaAction)
+	{
+		if (!SchemaAction)
+		{
+			return nullptr;
+		}
+
+		const FMovieGraphSchemaAction* SelectedGraphAction = static_cast<FMovieGraphSchemaAction*>(SchemaAction);
+
+		// Currently the action menu only includes UMovieGraphMember targets
+		if (UMovieGraphMember* SelectedMember = Cast<UMovieGraphMember>(SelectedGraphAction->ActionTarget))
+		{
+			return SelectedMember;
+		}
+
+		return nullptr;
+	}
+}
+
 const TArray<FText> SMovieGraphMembersTabContent::ActionMenuSectionNames
 {
 	LOCTEXT("ActionMenuSectionName_Invalid", "INVALID"),
@@ -37,6 +59,7 @@ void SMovieGraphMembersTabContent::Construct(const FArguments& InArgs)
 		.OnGetSectionWidget(this, &SMovieGraphMembersTabContent::GetSectionWidget)
 		.UseSectionStyling(true)
 		.OnCollectAllActions(this, &SMovieGraphMembersTabContent::CollectAllActions)
+		.OnActionMatchesName(this, &SMovieGraphMembersTabContent::ActionMatchesName)
 	];
 }
 
@@ -59,9 +82,8 @@ void SMovieGraphMembersTabContent::DeleteSelectedMembers() const
 	ActionMenu->GetSelectedActions(SelectedActions);
 	for (TSharedPtr<FEdGraphSchemaAction> SelectedAction : SelectedActions)
 	{
-		if (const FMovieGraphSchemaAction* MovieGraphAction = static_cast<FMovieGraphSchemaAction*>(SelectedAction.Get()))
+		if (UMovieGraphMember* GraphMember = UE::MovieGraph::Private::GetMemberFromAction(SelectedAction.Get()))
 		{
-			UMovieGraphMember* GraphMember = Cast<UMovieGraphMember>(MovieGraphAction->ActionTarget.Get());
 			CurrentGraph->DeleteMember(GraphMember);
 		}
 	}
@@ -81,10 +103,8 @@ bool SMovieGraphMembersTabContent::CanDeleteSelectedMembers() const
 	
 	for (const TSharedPtr<FEdGraphSchemaAction>& SelectedAction : SelectedActions)
 	{
-		const FMovieGraphSchemaAction* MovieAction = static_cast<FMovieGraphSchemaAction*>(SelectedAction.Get());
-
 		// Don't allow deletion if the member was explicitly marked as non-deletable
-		if (const UMovieGraphMember* Member = Cast<UMovieGraphMember>(MovieAction->ActionTarget))
+		if (const UMovieGraphMember* Member = UE::MovieGraph::Private::GetMemberFromAction(SelectedAction.Get()))
 		{
 			if (!Member->IsDeletable())
 			{
@@ -170,6 +190,16 @@ void SMovieGraphMembersTabContent::CollectAllActions(FGraphActionListBuilderBase
 	OutAllActions.Append(ActionMenuBuilder);
 }
 
+bool SMovieGraphMembersTabContent::ActionMatchesName(FEdGraphSchemaAction* InAction, const FName& InName) const
+{
+	if (const UMovieGraphMember* Member = UE::MovieGraph::Private::GetMemberFromAction(InAction))
+	{
+		return InName == Member->Name;
+	}
+
+	return false;
+}
+
 void SMovieGraphMembersTabContent::CollectStaticSections(TArray<int32>& StaticSectionIDs)
 {
 	// Start at index 1 to skip the invalid section
@@ -275,9 +305,29 @@ FReply SMovieGraphMembersTabContent::OnAddButtonClickedOnSection(const int32 InS
 void SMovieGraphMembersTabContent::RefreshMemberActions(UMovieGraphMember* UpdatedMember) const
 {
 	// Currently the entire action menu is refreshed rather than a specific action being targeted
-	
+
+	// Cache the currently selected member, so it can be referenced after the refresh
+	const UMovieGraphMember* SelectedMember = nullptr;
+	TArray<TSharedPtr<FEdGraphSchemaAction>> SelectedActions;
+	ActionMenu->GetSelectedActions(SelectedActions);
+	if (!SelectedActions.IsEmpty())
+	{
+		if (const UMovieGraphMember* Member = UE::MovieGraph::Private::GetMemberFromAction(SelectedActions[0].Get()))
+		{
+			SelectedMember = Member;
+		}
+	}
+
+	// Do the refresh
 	const bool bPreserveExpansion = true;
 	ActionMenu->RefreshAllActions(bPreserveExpansion);
+
+	// Re-select the updated member if it was previously selected. The action menu performs a reselect after a refresh
+	// automatically based on action name, but the member may have been renamed (so an explicit re-select is needed).
+	if (SelectedMember && (SelectedMember == UpdatedMember))
+	{
+		ActionMenu->SelectItemByName(FName(SelectedMember->Name));
+	}
 }
 
 #undef LOCTEXT_NAMESPACE
\ No newline at end of file
diff --git a/Engine/Plugins/MovieScene/MovieRenderPipeline/Source/MovieRenderPipelineEditor/Private/Graph/SMovieGraphMembersTabContent.h b/Engine/Plugins/MovieScene/MovieRenderPipeline/Source/MovieRenderPipelineEditor/Private/Graph/SMovieGraphMembersTabContent.h
--- a/Engine/Plugins/MovieScene/MovieRenderPipeline/Source/MovieRenderPipelineEditor/Private/Graph/SMovieGraphMembersTabContent.h
+++ b/Engine/Plugins/MovieScene/MovieRenderPipeline/Source/MovieRenderPipelineEditor/Private/Graph/SMovieGraphMembersTabContent.h
@@ -67,6 +67,9 @@ private:
 	FText GetSectionTitle(int32 InSectionID);
 	TSharedRef<SWidget> GetSectionWidget(TSharedRef<SWidget> RowWidget, int32 InSectionID);
 
+	/** Returns true if the action matches the specified name, else false. */
+	bool ActionMatchesName(FEdGraphSchemaAction* InAction, const FName& InName) const;
+
 	/** Handler which deals with populating the context menu. */
 	TSharedPtr<SWidget> OnContextMenuOpening();
 
