MoviePipeline: When using the "Render all Cameras" feature, Spawnable cameras will now be skipped if their default spawn track value is false.

diff --git a/Engine/Plugins/MovieScene/MovieRenderPipeline/Source/MovieRenderPipelineCore/Private/MoviePipelineBlueprintLibrary.cpp b/Engine/Plugins/MovieScene/MovieRenderPipeline/Source/MovieRenderPipelineCore/Private/MoviePipelineBlueprintLibrary.cpp
--- a/Engine/Plugins/MovieScene/MovieRenderPipeline/Source/MovieRenderPipelineCore/Private/MoviePipelineBlueprintLibrary.cpp
+++ b/Engine/Plugins/MovieScene/MovieRenderPipeline/Source/MovieRenderPipelineCore/Private/MoviePipelineBlueprintLibrary.cpp
@@ -17,8 +17,10 @@
 #include "Tracks/MovieSceneSubTrack.h"
 #include "Tracks/MovieSceneCameraCutTrack.h"
 #include "Tracks/MovieSceneCinematicShotTrack.h"
+#include "Tracks/MovieSceneSpawnTrack.h"
 #include "Sections/MovieSceneCinematicShotSection.h"
 #include "Sections/MovieSceneCameraCutSection.h"
+#include "Sections/MovieSceneSpawnSection.h"
 #include "MovieSceneTimeHelpers.h"
 #include "MoviePipelineQueue.h"
 #include "MoviePipelineSetting.h"
@@ -579,9 +581,10 @@ void UMoviePipelineBlueprintLibrary::UpdateJobShotListFromSequence(ULevelSequenc
 
 					if (ValidBinding.IsValid())
 					{
-						FMoviePipelineSidecarCamera& SidecarCamera = Entity.SidecarCameras.AddDefaulted_GetRef();
+						FMoviePipelineSidecarCamera SidecarCamera;
 						SidecarCamera.BindingId = ValidBinding;
-						SidecarCamera.SequenceId = Entity.SequenceID;
+						SidecarCamera.SequenceId = Entity.SequenceID; 
+						bool bAddCamera = true;
 						
 						FGuid ParentGuid = Possessable.GetParent();
 						while (ParentGuid.IsValid())
@@ -592,7 +595,7 @@ void UMoviePipelineBlueprintLibrary::UpdateJobShotListFromSequence(ULevelSequenc
 							if (ParentGuid == MainBinding.GetGuid())
 							{
 								// We want to know which sidecar camera this is, not the possessable index.
-								MainCameraIndex = Entity.SidecarCameras.Num() - 1;
+								MainCameraIndex = Entity.SidecarCameras.Num();
 							}
 
 							if (FMovieScenePossessable* ParentAsPossessable = LeafNode->MovieScene->FindPossessable(ParentGuid))
@@ -602,11 +605,36 @@ void UMoviePipelineBlueprintLibrary::UpdateJobShotListFromSequence(ULevelSequenc
 							}
 							else if (FMovieSceneSpawnable* ParentAsSpawnable = LeafNode->MovieScene->FindSpawnable(ParentGuid))
 							{
+								FMovieSceneBinding* Binding = LeafNode->MovieScene->FindBinding(ParentAsSpawnable->GetGuid());
+								if (Binding)
+								{
+									const TArray<UMovieSceneTrack*> AllBindingTracks = Binding->GetTracks();
+									for (const UMovieSceneTrack* Track : AllBindingTracks)
+									{
+										if (const UMovieSceneSpawnTrack* SpawnTrack = Cast<UMovieSceneSpawnTrack>(Track))
+										{
+											for (const UMovieSceneSection* Section : SpawnTrack->GetAllSections())
+											{
+												if (const UMovieSceneSpawnSection* BoolSection = Cast<UMovieSceneSpawnSection>(Section))
+												{
+													bAddCamera = BoolSection->GetChannel().GetDefault().Get(true);
+												}
+											}
+										}
+									}
+								}
+
 								// Spawnables will never have a parent
 								ParentGuid.Invalidate(); 
+								
 								SidecarCamera.Name = ParentAsSpawnable->GetName();
 							}
 						}
+
+						if (bAddCamera)
+						{
+							Entity.SidecarCameras.Add(SidecarCamera);
+						}
 					}
 				}
 
