MoviePipeline: Added queue and job selection in a new "Evaluation Context" settings menu in Active Render Settings. The selected job is provided to the traversal context (which is used during evaluation).

To take advantage of the added job context, graph variables will resolve to the value set on the job context if the job has an enabled override for the variable.

diff --git a/Engine/Plugins/MovieScene/MovieRenderPipeline/Source/MovieRenderPipelineCore/Private/Graph/MovieGraphConfig.cpp b/Engine/Plugins/MovieScene/MovieRenderPipeline/Source/MovieRenderPipelineCore/Private/Graph/MovieGraphConfig.cpp
--- a/Engine/Plugins/MovieScene/MovieRenderPipeline/Source/MovieRenderPipelineCore/Private/Graph/MovieGraphConfig.cpp
+++ b/Engine/Plugins/MovieScene/MovieRenderPipeline/Source/MovieRenderPipelineCore/Private/Graph/MovieGraphConfig.cpp
@@ -668,7 +668,7 @@ void UMovieGraphConfig::InitializeFlattenedNode(UMovieGraphNode* InNode)
 	InNode->UpdateDynamicProperties();
 }
 
-void UMovieGraphConfig::CopyOverriddenProperties(UMovieGraphNode* FromNode, UMovieGraphNode* ToNode)
+void UMovieGraphConfig::CopyOverriddenProperties(UMovieGraphNode* FromNode, UMovieGraphNode* ToNode, const FMovieGraphTraversalContext* InContext)
 {
 	if (!ensure(FromNode && ToNode))
 	{
@@ -736,7 +736,7 @@ void UMovieGraphConfig::CopyOverriddenProperties(UMovieGraphNode* FromNode, UMov
 				{
 					// There was a valid connection to the input pin; resolve the value from the connected output and set the
 					// value on this property
-					const FString ResolvedValue = ConnectedPin->Node->GetResolvedValueForOutputPin(ConnectedPin->Properties.Label);
+					const FString ResolvedValue = ConnectedPin->Node->GetResolvedValueForOutputPin(ConnectedPin->Properties.Label, InContext);
 					if (bIsDynamic)
 					{
 						ToNode->SetDynamicPropertyValue(PropertyName, ResolvedValue);
@@ -837,7 +837,7 @@ void UMovieGraphConfig::CreateFlattenedGraph_Recursive(UMovieGraphEvaluatedConfi
 
 		// ToDo: Handle "Disable" ndoes that can disable upstream types of nodes, etc. Push disable types into
 		// the currently evaluating context
-		CopyOverriddenProperties(Node, ExistingNode);
+		CopyOverriddenProperties(Node, ExistingNode, &InEvaluationContext.UserContext);
 	}
 	
 	// Now that we've potentially resolved the values on this node, continue to travel up-stream along any execution pins,
diff --git a/Engine/Plugins/MovieScene/MovieRenderPipeline/Source/MovieRenderPipelineCore/Private/Graph/Nodes/MovieGraphVariableNode.cpp b/Engine/Plugins/MovieScene/MovieRenderPipeline/Source/MovieRenderPipelineCore/Private/Graph/Nodes/MovieGraphVariableNode.cpp
--- a/Engine/Plugins/MovieScene/MovieRenderPipeline/Source/MovieRenderPipelineCore/Private/Graph/Nodes/MovieGraphVariableNode.cpp
+++ b/Engine/Plugins/MovieScene/MovieRenderPipeline/Source/MovieRenderPipelineCore/Private/Graph/Nodes/MovieGraphVariableNode.cpp
@@ -3,6 +3,7 @@
 #include "Graph/Nodes/MovieGraphVariableNode.h"
 
 #include "Graph/MovieGraphConfig.h"
+#include "MoviePipelineQueue.h"
 #include "Styling/AppStyle.h"
 
 UMovieGraphVariableNode::UMovieGraphVariableNode()
@@ -29,10 +30,25 @@ TArray<FMovieGraphPinProperties> UMovieGraphVariableNode::GetOutputPinProperties
 	return Properties;
 }
 
-FString UMovieGraphVariableNode::GetResolvedValueForOutputPin(const FName& InPinName) const
+FString UMovieGraphVariableNode::GetResolvedValueForOutputPin(const FName& InPinName, const FMovieGraphTraversalContext* InContext) const
 {
 	if (GraphVariable && (GraphVariable->Name == InPinName))
 	{
+		// If there's a valid job in the traversal context, and the job has an enabled variable assignment, use that
+		// instead of the variable value set in the graph
+		if (InContext && InContext->Job)
+		{
+			bool bIsEnabled = false;
+			if (InContext->Job->VariableAssignments->GetVariableAssignmentEnableState(GraphVariable, bIsEnabled))
+			{
+				if (bIsEnabled)
+				{
+					return InContext->Job->VariableAssignments->GetValueSerializedString(FName(GraphVariable->Name));
+				}
+			}
+		}
+
+		// No valid job context: just get the value from the variable
 		return GraphVariable->GetValueSerializedString();
 	}
 	
diff --git a/Engine/Plugins/MovieScene/MovieRenderPipeline/Source/MovieRenderPipelineCore/Private/MovieJobVariableAssignmentContainer.cpp b/Engine/Plugins/MovieScene/MovieRenderPipeline/Source/MovieRenderPipelineCore/Private/MovieJobVariableAssignmentContainer.cpp
--- a/Engine/Plugins/MovieScene/MovieRenderPipeline/Source/MovieRenderPipelineCore/Private/MovieJobVariableAssignmentContainer.cpp
+++ b/Engine/Plugins/MovieScene/MovieRenderPipeline/Source/MovieRenderPipelineCore/Private/MovieJobVariableAssignmentContainer.cpp
@@ -223,15 +223,20 @@ bool UMovieJobVariableAssignmentContainer::FindOrGenerateVariableOverride(
 {
 	if (InGraphVariable)
 	{
-		const FString EditCondPropName = FString::Format(TEXT("{0}{1}"), {EditConditionPrefix, InGraphVariable->Name});
-		
 		if (const FPropertyBagPropertyDesc* Desc = Value.FindPropertyDescByName(FName(InGraphVariable->Name)))
 		{
-			*OutPropDesc = *Desc;
+			if (OutPropDesc)
+			{
+				*OutPropDesc = *Desc;
+			}
 
-			if (const FPropertyBagPropertyDesc* EditCondPropDesc = Value.FindPropertyDescByName(FName(EditCondPropName)))
+			if (OutEditConditionPropDesc)
 			{
-				*OutEditConditionPropDesc = *EditCondPropDesc;
+				const FString EditCondPropName = FString::Format(TEXT("{0}{1}"), {EditConditionPrefix, Desc->ID.ToString()});
+				if (const FPropertyBagPropertyDesc* EditCondPropDesc = Value.FindPropertyDescByName(FName(EditCondPropName)))
+				{
+					*OutEditConditionPropDesc = *EditCondPropDesc;
+				}
 			}
 			
 			return true;
@@ -450,15 +455,17 @@ bool UMovieJobVariableAssignmentContainer::SetVariableAssignmentEnableState(cons
 bool UMovieJobVariableAssignmentContainer::GetVariableAssignmentEnableState(const UMovieGraphVariable* InGraphVariable, bool& bOutIsEnabled)
 {
 	FPropertyBagPropertyDesc* PropDesc = nullptr;
-	FPropertyBagPropertyDesc* EditConditionPropDesc = nullptr;
+	FPropertyBagPropertyDesc EditConditionPropDesc;
 	const bool bAddIfNotExists = false;
-	if (FindOrGenerateVariableOverride(InGraphVariable, PropDesc, EditConditionPropDesc, bAddIfNotExists))
+	if (FindOrGenerateVariableOverride(InGraphVariable, PropDesc, &EditConditionPropDesc, bAddIfNotExists))
 	{
-		const FString EditCondPropName = FString::Format(TEXT("{0}{1}"), {EditConditionPrefix, InGraphVariable->Name});
-		TValueOrError<bool, EPropertyBagResult> Result = Value.GetValueBool(FName(EditCondPropName));
-		bOutIsEnabled = Result.GetValue();
+		TValueOrError<bool, EPropertyBagResult> Result = Value.GetValueBool(EditConditionPropDesc.Name);
+		if (Result.HasValue())
+		{
+			bOutIsEnabled = Result.GetValue();
+		}
 		
-		return !Result.HasError();
+		return Result.HasValue();
 	}
 
 	return false;
diff --git a/Engine/Plugins/MovieScene/MovieRenderPipeline/Source/MovieRenderPipelineCore/Public/Graph/MovieGraphConfig.h b/Engine/Plugins/MovieScene/MovieRenderPipeline/Source/MovieRenderPipelineCore/Public/Graph/MovieGraphConfig.h
--- a/Engine/Plugins/MovieScene/MovieRenderPipeline/Source/MovieRenderPipelineCore/Public/Graph/MovieGraphConfig.h
+++ b/Engine/Plugins/MovieScene/MovieRenderPipeline/Source/MovieRenderPipelineCore/Public/Graph/MovieGraphConfig.h
@@ -420,7 +420,7 @@ public:
 
 protected:
 	/** Copies properties in FromNode that are marked for override into ToNode, but only if ToNode doesn't already override that value. */
-	void CopyOverriddenProperties(UMovieGraphNode* FromNode, UMovieGraphNode* ToNode);
+	void CopyOverriddenProperties(UMovieGraphNode* FromNode, UMovieGraphNode* ToNode, const FMovieGraphTraversalContext* InContext);
 	
 	/** Find all "Overrideable" marked properties, then find their edit condition properties, then set those to false. */
 	void InitializeFlattenedNode(UMovieGraphNode* InNode);
diff --git a/Engine/Plugins/MovieScene/MovieRenderPipeline/Source/MovieRenderPipelineCore/Public/Graph/MovieGraphNode.h b/Engine/Plugins/MovieScene/MovieRenderPipeline/Source/MovieRenderPipelineCore/Public/Graph/MovieGraphNode.h
--- a/Engine/Plugins/MovieScene/MovieRenderPipeline/Source/MovieRenderPipelineCore/Public/Graph/MovieGraphNode.h
+++ b/Engine/Plugins/MovieScene/MovieRenderPipeline/Source/MovieRenderPipelineCore/Public/Graph/MovieGraphNode.h
@@ -23,6 +23,7 @@ class UMovieGraphOutput;
 class UMovieGraphPin;
 class UMovieGraphVariable;
 struct FMovieGraphEvaluationContext;
+struct FMovieGraphTraversalContext;
 
 #if WITH_EDITOR
 class UEdGraphNode;
@@ -106,7 +107,7 @@ public:
 	 * Gets the resolved value of a named output pin (one that is returned in GetOutputPinProperties()). If the resolved
 	 * value could not be determined, an empty string is returned.
 	 */
-	virtual FString GetResolvedValueForOutputPin(const FName& InPinName) const
+	virtual FString GetResolvedValueForOutputPin(const FName& InPinName, const FMovieGraphTraversalContext* InContext) const
 	{
 		return FString();
 	}
diff --git a/Engine/Plugins/MovieScene/MovieRenderPipeline/Source/MovieRenderPipelineCore/Public/Graph/Nodes/MovieGraphVariableNode.h b/Engine/Plugins/MovieScene/MovieRenderPipeline/Source/MovieRenderPipelineCore/Public/Graph/Nodes/MovieGraphVariableNode.h
--- a/Engine/Plugins/MovieScene/MovieRenderPipeline/Source/MovieRenderPipelineCore/Public/Graph/Nodes/MovieGraphVariableNode.h
+++ b/Engine/Plugins/MovieScene/MovieRenderPipeline/Source/MovieRenderPipelineCore/Public/Graph/Nodes/MovieGraphVariableNode.h
@@ -17,7 +17,7 @@ public:
 	UMovieGraphVariableNode();
 
 	virtual TArray<FMovieGraphPinProperties> GetOutputPinProperties() const override;
-	virtual FString GetResolvedValueForOutputPin(const FName& InPinName) const override;
+	virtual FString GetResolvedValueForOutputPin(const FName& InPinName, const FMovieGraphTraversalContext* InContext) const override;
 
 	/** Gets the variable that this node represents. */
 	UMovieGraphVariable* GetVariable() const { return GraphVariable; }
diff --git a/Engine/Plugins/MovieScene/MovieRenderPipeline/Source/MovieRenderPipelineEditor/Private/Graph/SMovieGraphActiveRenderSettingsTabContent.cpp b/Engine/Plugins/MovieScene/MovieRenderPipeline/Source/MovieRenderPipelineEditor/Private/Graph/SMovieGraphActiveRenderSettingsTabContent.cpp
--- a/Engine/Plugins/MovieScene/MovieRenderPipeline/Source/MovieRenderPipelineEditor/Private/Graph/SMovieGraphActiveRenderSettingsTabContent.cpp
+++ b/Engine/Plugins/MovieScene/MovieRenderPipeline/Source/MovieRenderPipelineEditor/Private/Graph/SMovieGraphActiveRenderSettingsTabContent.cpp
@@ -2,10 +2,17 @@
 
 #include "SMovieGraphActiveRenderSettingsTabContent.h"
 
+#include "ContentBrowserModule.h"
+#include "Editor.h"
+#include "Framework/MultiBox/MultiBoxBuilder.h"
 #include "Graph/MovieGraphConfig.h"
 #include "Graph/MovieGraphNode.h"
 #include "Graph/Nodes/MovieGraphRenderPassNode.h"
+#include "IContentBrowserSingleton.h"
+#include "MoviePipelineQueue.h"
+#include "MoviePipelineQueueSubsystem.h"
 #include "SPositiveActionButton.h"
+#include "Widgets/Input/SComboBox.h"
 
 #define LOCTEXT_NAMESPACE "SMoviePipelineActiveRenderSettingsTabContent"
 
@@ -274,6 +281,22 @@ void SMovieGraphActiveRenderSettingsTabContent::Construct(const FArguments& InAr
 {
 	CurrentGraph = InArgs._Graph;
 
+	const UMoviePipelineQueueSubsystem* Subsystem = GEditor->GetEditorSubsystem<UMoviePipelineQueueSubsystem>();
+	check(Subsystem);
+
+	// Initialize with the current queue (if there is one), and pick the first job by default (if a job exists)
+	if (UMoviePipelineQueue* Queue = Subsystem->GetQueue())
+	{
+		// Note: The queue from the queue subsystem is probably transient
+		TraversalQueue = Queue;
+
+		const TArray<UMoviePipelineExecutorJob*>& Jobs = TraversalQueue->GetJobs();
+		if (!Jobs.IsEmpty())
+		{
+			TraversalJob = Jobs[0];
+		}
+	}
+
 	// Add the two default root elements, "Globals" and "Branches", which are always visible. These elements will
 	// generate the children which should be displayed under them in the tree.
 	RootElements.Add(MakeShared<FActiveRenderSettingsTreeElement>(
@@ -293,10 +316,20 @@ void SMovieGraphActiveRenderSettingsTabContent::Construct(const FArguments& InAr
 			.AutoWidth()
 			[
 				SNew(SPositiveActionButton)
-				.Text(FText::FromString("Evaluate Graph"))
+				.Text(LOCTEXT("Button_EvaluateGraph", "Evaluate Graph"))
 				.Icon(FAppStyle::GetBrush("Icons.Refresh"))
 				.OnClicked(this, &SMovieGraphActiveRenderSettingsTabContent::OnEvaluateGraphClicked)
 			]
+
+			+SHorizontalBox::Slot()
+			.Padding(5, 0, 0, 0)
+			.AutoWidth()
+			[
+				SNew(SPositiveActionButton)
+				.Text(LOCTEXT("Button_EvaluationContext", "Evaluation Context..."))
+				.Icon(FAppStyle::GetBrush("EditorPreferences.TabIcon"))
+				.OnGetMenuContent(this, &SMovieGraphActiveRenderSettingsTabContent::GenerateEvaluationContextMenu)
+			]
 		]
 
 		+ SVerticalBox::Slot()
@@ -355,8 +388,8 @@ void SMovieGraphActiveRenderSettingsTabContent::TraverseGraph()
 		RootElement->ClearCachedChildren();
 	}
 
-	// TODO: Fill in the context via the UI (there are no widgets to provide this info yet)
 	FMovieGraphTraversalContext Context;
+	Context.Job = TraversalJob.Get();
 
 	// Traverse the graph, and update the root elements
 	FlattenedGraph = TStrongObjectPtr(CurrentGraph->CreateFlattenedGraph(Context));
@@ -380,6 +413,83 @@ FReply SMovieGraphActiveRenderSettingsTabContent::OnEvaluateGraphClicked()
 	return FReply::Handled();
 }
 
+TSharedRef<SWidget> SMovieGraphActiveRenderSettingsTabContent::GenerateEvaluationContextMenu()
+{
+	FMenuBuilder MenuBuilder(true, nullptr);
+
+	MenuBuilder.BeginSection(NAME_None, LOCTEXT("JobContext_MenuSection", "Job Context"));
+	{
+		const TSharedRef<SWidget> JobPicker = SNew(SBox)
+		.WidthOverride(400)
+		.Padding(10)
+		[
+			SNew(SVerticalBox)
+			+SVerticalBox::Slot()
+			[
+				SNew(SHorizontalBox)
+				+SHorizontalBox::Slot()
+				.FillWidth(0.3f)
+				[
+					SNew(STextBlock)
+					.Text(LOCTEXT("JobContext_QueueLabel", "Queue"))
+				]
+
+				+SHorizontalBox::Slot()
+				.FillWidth(0.7f)
+				[
+					SNew(SPositiveActionButton)
+					.Text(this, &SMovieGraphActiveRenderSettingsTabContent::GetQueueButtonText)
+					.OnGetMenuContent(this, &SMovieGraphActiveRenderSettingsTabContent::MakeQueueButtonMenuContents)
+				]
+			]
+
+			+SVerticalBox::Slot()
+			.AutoHeight()
+			[
+				SNew(SHorizontalBox)
+				+SHorizontalBox::Slot()
+				.FillWidth(0.3f)
+				[
+					SNew(STextBlock)
+					.Text(LOCTEXT("JobContext_JobLabel", "Job"))
+				]
+
+				+SHorizontalBox::Slot()
+				.FillWidth(0.7f)
+				[
+					SNew(SComboBox<TSharedPtr<FString>>)
+					.IsEnabled_Lambda([this]() { return TraversalQueue != nullptr; })
+					.OptionsSource(&AvailableJobsInQueue)
+					.OnGenerateWidget_Lambda([](TSharedPtr<FString> Item)
+					{
+						return SNew(STextBlock)
+							.Text(FText::FromString(*Item));
+					})
+					.OnSelectionChanged(this, &SMovieGraphActiveRenderSettingsTabContent::HandleJobSelected)
+					[
+						SNew(STextBlock)
+						.MinDesiredWidth(50.f)
+						.Text_Lambda([this]()
+						{
+							if (TraversalJob.IsValid())
+							{
+								return FText::FromString(TraversalJob->JobName);
+							}
+
+							return LOCTEXT("JobContext_NoJobSelected", "No job selected");
+						})
+					]
+				]
+			]
+		];
+
+		MenuBuilder.AddWidget(JobPicker, FText(), true, false);
+	}
+	MenuBuilder.EndSection();
+
+	return MenuBuilder.MakeWidget();
+}
+
 void SMovieGraphActiveRenderSettingsTabContent::OnExpansionChanged(TSharedPtr<FActiveRenderSettingsTreeElement> InElement, bool bInExpanded)
 {
 	if (bInExpanded)
@@ -427,4 +537,116 @@ void SMovieGraphActiveRenderSettingsTabContent::GetChildrenForTree(
 	OutChildren.Append(InItem->GetChildren());
 }
 
+FText SMovieGraphActiveRenderSettingsTabContent::GetQueueButtonText() const
+{
+	if (!TraversalQueue.IsValid())
+	{
+		return LOCTEXT("JobContext_PickAQueue", "Pick a queue...");
+	}
+	
+	// If the queue has an origin queue, use the origin's name. A queue will have a defined origin
+	// if the queue is coming from the queue subsystem (since the subsystem usually points to a transient
+	// queue which is being displayed in the queue UI).
+	if (const UMoviePipelineQueue* QueueOrigin = TraversalQueue->GetQueueOrigin())
+	{
+		return FText::FromString(QueueOrigin->GetName());
+	}
+
+	// If the queue doesn't have an origin, this queue was probably picked here in the graph. If it has a
+	// non-transient package, it has been saved at some point, so we can use the queue name.
+	const UPackage* Package = TraversalQueue->GetPackage();
+	if (Package && !Package->HasAnyPackageFlags(PKG_TransientFlags))
+	{
+		return FText::FromString(TraversalQueue->GetName());
+	}
+
+	// This is most likely an unsaved transient queue from the queue UI
+	return LOCTEXT("JobContext_UnnamedQueue", "Unnamed queue");
+}
+
+TSharedRef<SWidget> SMovieGraphActiveRenderSettingsTabContent::MakeQueueButtonMenuContents()
+{
+	FMenuBuilder MenuBuilder(true, nullptr);
+    IContentBrowserSingleton& ContentBrowser = FModuleManager::LoadModuleChecked<FContentBrowserModule>("ContentBrowser").Get();
+
+    FAssetPickerConfig AssetPickerConfig;
+    {
+		// If a transient queue (with a queue origin) is active, select the origin queue. The origin queue will be saved
+		// so it will properly show up as selected in the asset picker (whereas a transient queue will not).
+		AssetPickerConfig.InitialAssetSelection = (TraversalQueue.IsValid() && TraversalQueue->GetQueueOrigin())
+			? FAssetData(TraversalQueue->GetQueueOrigin())
+			: FAssetData(TraversalQueue.Get());
+		
+    	AssetPickerConfig.SelectionMode = ESelectionMode::Single;
+    	AssetPickerConfig.InitialAssetViewType = EAssetViewType::Column;
+    	AssetPickerConfig.bFocusSearchBoxWhenOpened = true;
+    	AssetPickerConfig.bAllowNullSelection = false;
+    	AssetPickerConfig.bShowBottomToolbar = true;
+    	AssetPickerConfig.bAutohideSearchBar = false;
+    	AssetPickerConfig.bAllowDragging = false;
+    	AssetPickerConfig.bCanShowClasses = false;
+    	AssetPickerConfig.bShowPathInColumnView = false;
+    	AssetPickerConfig.bShowTypeInColumnView = false;
+    	AssetPickerConfig.bSortByPathInColumnView = false;
+    	AssetPickerConfig.SaveSettingsName = TEXT("MovieRenderGraphTraversalContextQueue");
+    	AssetPickerConfig.AssetShowWarningText = LOCTEXT("JobContext_NoQueuesWarning", "No Queues Found");
+    	AssetPickerConfig.Filter.ClassPaths.Add(UMoviePipelineQueue::StaticClass()->GetClassPathName());
+    	AssetPickerConfig.OnAssetSelected.BindLambda([this](const FAssetData& InAsset)
+    	{
+    		TraversalQueue = Cast<UMoviePipelineQueue>(InAsset.GetAsset());
+    		if (TraversalQueue.IsValid())
+    		{
+    			const UMoviePipelineQueueSubsystem* Subsystem = GEditor->GetEditorSubsystem<UMoviePipelineQueueSubsystem>();
+    			check(Subsystem);
+
+    			// If the selected queue matches the current queue being edited in the queue UI,
+    			// use the transient queue from the queue UI. This will let the user edit the queue
+    			// without saving, and have those updates reflected in the traversal.
+    			UMoviePipelineQueue* CurrentQueue = Subsystem->GetQueue();
+    			if (CurrentQueue && (CurrentQueue->GetQueueOrigin() == TraversalQueue))
+    			{
+    				TraversalQueue = CurrentQueue;
+    			}
+
+    			// Reset job selection state
+    			AvailableJobsInQueue.Empty();
+    			TraversalJob = nullptr;
+    			Algo::Transform(TraversalQueue->GetJobs(), AvailableJobsInQueue, [](const UMoviePipelineExecutorJob* Job)
+    			{
+    				return MakeShared<FString>(Job->JobName);
+    			});
+    		}
+
+    		if (QueuePickerWidget.IsValid())
+    		{
+    			FSlateApplication::Get().DismissMenuByWidget(QueuePickerWidget.ToSharedRef());
+    		}
+    	});
+    }
+
+    SAssignNew(QueuePickerWidget, SBox)
+    	.WidthOverride(300.f)
+    	.HeightOverride(300.f)
+    	[
+    		ContentBrowser.CreateAssetPicker(AssetPickerConfig)
+    	];
+
+    MenuBuilder.AddWidget(QueuePickerWidget.ToSharedRef(), FText(), true, false);
+
+    return MenuBuilder.MakeWidget();
+}
+
+void SMovieGraphActiveRenderSettingsTabContent::HandleJobSelected(TSharedPtr<FString> Item, ESelectInfo::Type SelectInfo)
+{
+	// Find the associated job object
+	for (UMoviePipelineExecutorJob* Job : TraversalQueue->GetJobs())
+	{
+		if (Job && (Job->JobName == *Item))
+		{
+			TraversalJob = Job;
+			break;
+		}
+	}
+}
+
 #undef LOCTEXT_NAMESPACE
diff --git a/Engine/Plugins/MovieScene/MovieRenderPipeline/Source/MovieRenderPipelineEditor/Private/Graph/SMovieGraphActiveRenderSettingsTabContent.h b/Engine/Plugins/MovieScene/MovieRenderPipeline/Source/MovieRenderPipelineEditor/Private/Graph/SMovieGraphActiveRenderSettingsTabContent.h
--- a/Engine/Plugins/MovieScene/MovieRenderPipeline/Source/MovieRenderPipelineEditor/Private/Graph/SMovieGraphActiveRenderSettingsTabContent.h
+++ b/Engine/Plugins/MovieScene/MovieRenderPipeline/Source/MovieRenderPipelineEditor/Private/Graph/SMovieGraphActiveRenderSettingsTabContent.h
@@ -7,6 +7,8 @@
 #include "Widgets/Views/STableRow.h"
 #include "Widgets/Views/STreeView.h"
 
+class UMoviePipelineExecutorJob;
+class UMoviePipelineQueue;
 class UMovieGraphEvaluatedConfig;
 class UMovieGraphNode;
 
@@ -128,6 +130,9 @@ private:
 	/** Handles the button click for evaluating the graph. */
 	FReply OnEvaluateGraphClicked();
 
+	/** Generates the menu for the Evaluation Context button. */
+	TSharedRef<SWidget> GenerateEvaluationContextMenu();
+
 	/** Handles the expansion state change of an element in the tree. */
 	void OnExpansionChanged(TSharedPtr<FActiveRenderSettingsTreeElement> InElement, bool bInExpanded);
 
@@ -140,6 +145,15 @@ private:
 	/** Gets the child elements in the tree for the provided element. */
 	void GetChildrenForTree(TSharedPtr<FActiveRenderSettingsTreeElement> InItem, TArray<TSharedPtr<FActiveRenderSettingsTreeElement>>& OutChildren);
 
+	/** Gets the text that is displayed on the queue picker button (in the evaluation context menu). */
+	FText GetQueueButtonText() const;
+
+	/** Generates the menu contents for the queue picker button. Populates the QueuePickerWidget member. */
+	TSharedRef<SWidget> MakeQueueButtonMenuContents();
+
+	/** Handles a new job selection in the evaluation context menu. */
+	void HandleJobSelected(TSharedPtr<FString> Item, ESelectInfo::Type SelectInfo);
+
 private:
 	/** The runtime graph that this UI gets data from. */
 	TWeakObjectPtr<UMovieGraphConfig> CurrentGraph = nullptr;
@@ -155,4 +169,16 @@ private:
 
 	/** The root-most elements in the tree, which are always present. */
 	TArray<TSharedPtr<FActiveRenderSettingsTreeElement>> RootElements;
+
+	/** The widget that houses the queue asset picker. */
+	TSharedPtr<SWidget> QueuePickerWidget = nullptr;
+
+	/** The queue that is selected in the evaluation context settings. */
+	TWeakObjectPtr<UMoviePipelineQueue> TraversalQueue = nullptr;
+
+	/** The job that is selected in the evaluation context settings. */
+	TWeakObjectPtr<UMoviePipelineExecutorJob> TraversalJob = nullptr;
+
+	/** The jobs that are available to be selected in the evaluation context settings. */
+	TArray<TSharedPtr<FString>> AvailableJobsInQueue;
 };
