MoviePipeline: Added node instance naming so each branch in a graph traversal can contain multiple nodes of the same type (eg, multiple Collection nodes, each with a different collection name/identifier).

Also:
* Moved Collection/Modifier nodes out of placeholder nodes
* Traversal respects nodes that have instance names.
* Slightly updated the Render Layer Subsystem prototype code to be usable (temporarily) with the Modifier and Collection nodes.

[Checking in for Matt.Hoffman]

diff --git a/Engine/Plugins/MovieScene/MovieRenderPipeline/Source/MovieRenderPipelineCore/Private/Graph/MovieGraphConfig.cpp b/Engine/Plugins/MovieScene/MovieRenderPipeline/Source/MovieRenderPipelineCore/Private/Graph/MovieGraphConfig.cpp
--- a/Engine/Plugins/MovieScene/MovieRenderPipeline/Source/MovieRenderPipelineCore/Private/Graph/MovieGraphConfig.cpp
+++ b/Engine/Plugins/MovieScene/MovieRenderPipeline/Source/MovieRenderPipelineCore/Private/Graph/MovieGraphConfig.cpp
@@ -819,12 +819,15 @@ void UMovieGraphConfig::CreateFlattenedGraph_Recursive(UMovieGraphEvaluatedConfi
 
 	if(bShouldIncludeNode)
 	{
-		UMovieGraphNode* ExistingNode = OutBranchConfig.GetNodeByClassExactMatch(Node->GetClass());
+		const UMovieGraphSettingNode* NodeAsSetting = CastChecked<UMovieGraphSettingNode>(Node);
+		const FString& NodeInstanceName = NodeAsSetting->GetNodeInstanceName();
+		
+		UMovieGraphNode* ExistingNode = OutBranchConfig.GetNodeByClassExactMatch(Node->GetClass(), NodeInstanceName);
 		if (!ExistingNode)
 		{
 			// Create a new instance of this node inside our flattened eval graph
 			ExistingNode = NewObject<UMovieGraphNode>(InOwningConfig, Node->GetClass());
-			OutBranchConfig.NodeInstances.Add(ExistingNode);
+			OutBranchConfig.NamedNodes.FindOrAdd(NodeInstanceName).NodeInstances.Add(ExistingNode);
 
 			// Set all of the boolean edit condition values to false, so we can use "true" to indicate
 			// that the value was overridden already during traversal.
diff --git a/Engine/Plugins/MovieScene/MovieRenderPipeline/Source/MovieRenderPipelineCore/Private/Graph/Nodes/MovieGraphCollectionNode.cpp b/Engine/Plugins/MovieScene/MovieRenderPipeline/Source/MovieRenderPipelineCore/Private/Graph/Nodes/MovieGraphCollectionNode.cpp
new file mode 100644
--- /dev/null
+++ b/Engine/Plugins/MovieScene/MovieRenderPipeline/Source/MovieRenderPipelineCore/Private/Graph/Nodes/MovieGraphCollectionNode.cpp
@@ -0,0 +1,57 @@
+// Copyright Epic Games, Inc. All Rights Reserved.
+
+#include "Graph/Nodes/MovieGraphCollectionNode.h"
+
+#include "Graph/MovieGraphConfig.h"
+#include "Styling/AppStyle.h"
+
+#define LOCTEXT_NAMESPACE "MovieGraph"
+
+#if WITH_EDITOR
+FText UMovieGraphCollectionNode::GetNodeTitle(const bool bGetDescriptive) const
+{
+	static const FText CollectionNodeName = LOCTEXT("NodeName_Collection", "Collection");
+	static const FText CollectionNodeDescription = LOCTEXT("NodeDescription_Collection", "Collection\n{0}");
+
+	if (bGetDescriptive && !CollectionName.IsEmpty())
+	{
+		return FText::Format(CollectionNodeDescription, FText::FromString(CollectionName));
+	}
+
+	return CollectionNodeName;
+}
+
+FText UMovieGraphCollectionNode::GetMenuCategory() const
+{
+	return LOCTEXT("CollectionNode_Category", "Utility");
+}
+
+FLinearColor UMovieGraphCollectionNode::GetNodeTitleColor() const
+{
+	static const FLinearColor CollectionNodeColor = FLinearColor(0.047f, 0.501f, 0.654f);
+	return CollectionNodeColor;
+}
+
+FSlateIcon UMovieGraphCollectionNode::GetIconAndTint(FLinearColor& OutColor) const
+{
+	static const FSlateIcon CollectionIcon = FSlateIcon(FAppStyle::GetAppStyleSetName(), "SceneOutliner.NewFolderIcon");
+
+	OutColor = FLinearColor::White;
+	return CollectionIcon;
+}
+
+void UMovieGraphCollectionNode::PostEditChangeProperty(FPropertyChangedEvent& PropertyChangedEvent)
+{
+	Super::PostEditChangeProperty(PropertyChangedEvent);
+
+	// Broadcast a node-changed delegate so that the node title's UI gets updated.
+	if (PropertyChangedEvent.GetPropertyName() == GET_MEMBER_NAME_CHECKED(UMovieGraphCollectionNode, CollectionName) ||
+		PropertyChangedEvent.GetPropertyName() == GET_MEMBER_NAME_CHECKED(UMovieGraphCollectionNode, bOverride_CollectionName))
+	{
+		OnNodeChangedDelegate.Broadcast(this);
+	}
+}
+
+#endif // WITH_EDITOR
+
+#undef LOCTEXT_NAMESPACE // "MovieGraph"
\ No newline at end of file
diff --git a/Engine/Plugins/MovieScene/MovieRenderPipeline/Source/MovieRenderPipelineCore/Private/Graph/Nodes/MovieGraphModifierNode.cpp b/Engine/Plugins/MovieScene/MovieRenderPipeline/Source/MovieRenderPipelineCore/Private/Graph/Nodes/MovieGraphModifierNode.cpp
new file mode 100644
--- /dev/null
+++ b/Engine/Plugins/MovieScene/MovieRenderPipeline/Source/MovieRenderPipelineCore/Private/Graph/Nodes/MovieGraphModifierNode.cpp
@@ -0,0 +1,62 @@
+// Copyright Epic Games, Inc. All Rights Reserved.
+
+#include "Graph/Nodes/MovieGraphModifierNode.h"
+
+#include "Graph/MovieGraphConfig.h"
+#include "Styling/AppStyle.h"
+
+#define LOCTEXT_NAMESPACE "MovieGraph"
+
+#if WITH_EDITOR
+FText UMovieGraphModifierNode::GetNodeTitle(const bool bGetDescriptive) const
+{
+	static const FText ModifierNodeName = LOCTEXT("NodeName_Modifier", "Modifier");
+	static const FText ModifierNodeDescription = LOCTEXT("NodeDescription_Modifier", "Modifier\n{0} - {1}");
+
+	const FString ModifiedCollectionNameDisp = ModifiedCollectionName.IsEmpty() ? TEXT("MISSING") : ModifiedCollectionName;
+	const FString ModifierNameDisp = ModifierName.IsEmpty() ? TEXT("MISSING") : ModifierName;
+
+	if (bGetDescriptive)
+	{
+		return FText::Format(ModifierNodeDescription, FText::FromString(ModifierNameDisp), FText::FromString(ModifiedCollectionNameDisp));
+	}
+
+	return ModifierNodeName;
+}
+
+FText UMovieGraphModifierNode::GetMenuCategory() const
+{
+	return LOCTEXT("CollectionNode_Category", "Utility");
+}
+
+FLinearColor UMovieGraphModifierNode::GetNodeTitleColor() const
+{
+	static const FLinearColor ModifierNodeColor = FLinearColor(0.6f, 0.113f, 0.113f);
+	return ModifierNodeColor;
+}
+
+FSlateIcon UMovieGraphModifierNode::GetIconAndTint(FLinearColor& OutColor) const
+{
+	static const FSlateIcon ModifierIcon = FSlateIcon(FAppStyle::GetAppStyleSetName(), "ContentBrowser.ReferenceViewer");
+
+	OutColor = FLinearColor::White;
+	return ModifierIcon;
+}
+
+void UMovieGraphModifierNode::PostEditChangeProperty(FPropertyChangedEvent& PropertyChangedEvent)
+{
+	Super::PostEditChangeProperty(PropertyChangedEvent);
+
+	// Broadcast a node-changed delegate so that the node title's UI gets updated.
+	if ((PropertyChangedEvent.GetPropertyName() == GET_MEMBER_NAME_CHECKED(UMovieGraphModifierNode, ModifierName)) ||
+		(PropertyChangedEvent.GetPropertyName() == GET_MEMBER_NAME_CHECKED(UMovieGraphModifierNode, bOverride_ModifierName)) ||
+		(PropertyChangedEvent.GetPropertyName() == GET_MEMBER_NAME_CHECKED(UMovieGraphModifierNode, ModifiedCollectionName)) || 
+		(PropertyChangedEvent.GetPropertyName() == GET_MEMBER_NAME_CHECKED(UMovieGraphModifierNode, bOverride_ModifiedCollectionName)))
+	{
+		OnNodeChangedDelegate.Broadcast(this);
+	}
+}
+
+#endif // WITH_EDITOR
+
+#undef LOCTEXT_NAMESPACE // "MovieGraph"
\ No newline at end of file
diff --git a/Engine/Plugins/MovieScene/MovieRenderPipeline/Source/MovieRenderPipelineCore/Public/Graph/MovieGraphConfig.h b/Engine/Plugins/MovieScene/MovieRenderPipeline/Source/MovieRenderPipelineCore/Public/Graph/MovieGraphConfig.h
--- a/Engine/Plugins/MovieScene/MovieRenderPipeline/Source/MovieRenderPipelineCore/Public/Graph/MovieGraphConfig.h
+++ b/Engine/Plugins/MovieScene/MovieRenderPipeline/Source/MovieRenderPipelineCore/Public/Graph/MovieGraphConfig.h
@@ -170,8 +170,17 @@ public:
 	DECLARE_MULTICAST_DELEGATE_OneParam(FOnMovieGraphNodesDeleted, TArray<UMovieGraphNode*>);
 #endif // WITH_EDITOR
 
+USTRUCT()
+struct MOVIERENDERPIPELINECORE_API FMovieGraphEvaluatedSettingsStack
+{
+	GENERATED_BODY()
+
+	UPROPERTY(Transient)
+	TArray<TObjectPtr<UMovieGraphNode>> NodeInstances;
+};
+
 /**
-* A flattened list of configuration values for a given Graph Branch. For named branches, this incldues the "Globals"
+* A flattened list of configuration values for a given Graph Branch. For named branches, this includes the "Globals"
 * branch (for any value not also overridden by the named branch).
 * 
 */
@@ -179,23 +188,32 @@ USTRUCT()
 struct MOVIERENDERPIPELINECORE_API FMovieGraphEvaluatedBranchConfig
 {
 	GENERATED_BODY()
-public:
-	UMovieGraphNode* GetNodeByClassExactMatch(const TSubclassOf<UMovieGraphNode>& InClass)
+
+	UMovieGraphNode* GetNodeByClassExactMatch(const TSubclassOf<UMovieGraphNode>& InClass, const FString& InName)
 	{
-		for (const TObjectPtr<UMovieGraphNode>& Instance : NodeInstances)
+		if (const FMovieGraphEvaluatedSettingsStack* FoundStack = NamedNodes.Find(InName))
 		{
-			if (Instance && Instance->GetClass() == InClass)
+			for (const TObjectPtr<UMovieGraphNode>& Instance : FoundStack->NodeInstances)
 			{
-				return Instance;
+				if (Instance && Instance->GetClass() == InClass)
+				{
+					return Instance;
+				}
 			}
 		}
 
 		return nullptr;
 	}
 
-	const TArray<TObjectPtr<UMovieGraphNode>>& GetNodes() const
+	TArray<TObjectPtr<UMovieGraphNode>> GetNodes() const
 	{
-		return NodeInstances;
+		TArray<TObjectPtr<UMovieGraphNode>> AllNodeInstances;
+		for (const TPair<FString, FMovieGraphEvaluatedSettingsStack>& KVP : NamedNodes)
+		{
+			AllNodeInstances.Append(KVP.Value.NodeInstances);
+		}
+
+		return AllNodeInstances;
 	}
 
 private:
@@ -203,8 +221,12 @@ private:
 	// without going through the graph resolving.
 	friend class UMovieGraphConfig;
 	
+	/**
+	* Nodes that have been evaluated in the branch. Key: the node instance name, value: the nodes that share the
+	* instance name. For nodes that do not have an instance name, an empty string is the key.
+	*/
 	UPROPERTY(Transient)
-	TArray<TObjectPtr<UMovieGraphNode>> NodeInstances;
+	TMap<FString, FMovieGraphEvaluatedSettingsStack> NamedNodes;
 };
 
 /**
diff --git a/Engine/Plugins/MovieScene/MovieRenderPipeline/Source/MovieRenderPipelineCore/Public/Graph/MovieGraphNode.h b/Engine/Plugins/MovieScene/MovieRenderPipeline/Source/MovieRenderPipelineCore/Public/Graph/MovieGraphNode.h
--- a/Engine/Plugins/MovieScene/MovieRenderPipeline/Source/MovieRenderPipelineCore/Public/Graph/MovieGraphNode.h
+++ b/Engine/Plugins/MovieScene/MovieRenderPipeline/Source/MovieRenderPipelineCore/Public/Graph/MovieGraphNode.h
@@ -300,6 +300,13 @@ public:
 	virtual TArray<FMovieGraphPinProperties> GetInputPinProperties() const override;
 	virtual TArray<FMovieGraphPinProperties> GetOutputPinProperties() const override;
 	// ~UMovieGraphNode Interface
+
+	/**
+	 * An identifier that distinguishes this node from other nodes of the same type within a branch. During graph
+	 * traversal, nodes with the same type and identifier are considered the same node. An empty string is a valid
+	 * identifier.
+	 */
+	virtual FString GetNodeInstanceName() const { return FString(); }
 	
 	/*
 	* This is called either on the CDO, or on a "flattened" instance of the node every frame when
diff --git a/Engine/Plugins/MovieScene/MovieRenderPipeline/Source/MovieRenderPipelineCore/Public/Graph/MovieGraphPlaceholderNodes.cpp b/Engine/Plugins/MovieScene/MovieRenderPipeline/Source/MovieRenderPipelineCore/Public/Graph/MovieGraphPlaceholderNodes.cpp
--- a/Engine/Plugins/MovieScene/MovieRenderPipeline/Source/MovieRenderPipelineCore/Public/Graph/MovieGraphPlaceholderNodes.cpp
+++ b/Engine/Plugins/MovieScene/MovieRenderPipeline/Source/MovieRenderPipelineCore/Public/Graph/MovieGraphPlaceholderNodes.cpp
@@ -14,123 +14,6 @@ static const FText NodeCategory_OutputType = LOCTEXT("NodeCategory_OutputType",
 static const FText NodeCategory_Conditionals = LOCTEXT("NodeCategory_Conditionals", "Conditionals");
 static const FText NodeCategory_Settings = LOCTEXT("NodeCategory_Settings", "Settings");
 
-TArray<FMovieGraphPinProperties> UMoviePipelineCollectionNode::GetInputPinProperties() const
-{
-	TArray<FMovieGraphPinProperties> Properties;
-	Properties.Add(FMovieGraphPinProperties::MakeBranchProperties());
-	return Properties;
-}
-
-TArray<FMovieGraphPinProperties> UMoviePipelineCollectionNode::GetOutputPinProperties() const
-{
-	TArray<FMovieGraphPinProperties> Properties;
-	Properties.Add(FMovieGraphPinProperties::MakeBranchProperties());
-	return Properties;
-}
-
-#if WITH_EDITOR
-FText UMoviePipelineCollectionNode::GetNodeTitle(const bool bGetDescriptive) const
-{
-	static const FText CollectionNodeName = LOCTEXT("NodeName_Collection", "Collection");
-	static const FText CollectionNodeDescription = LOCTEXT("NodeDescription_Collection", "Collection\n{0}");
-
-	if (bGetDescriptive && !CollectionName.IsEmpty())
-	{
-		return FText::Format(CollectionNodeDescription, FText::FromString(CollectionName));
-	}
-
-	return CollectionNodeName;
-}
-
-FText UMoviePipelineCollectionNode::GetMenuCategory() const
-{
-	return NodeCategory_Utility;
-}
-
-FLinearColor UMoviePipelineCollectionNode::GetNodeTitleColor() const
-{
-	static const FLinearColor CollectionNodeColor = FLinearColor(0.047f, 0.501f, 0.654f);
-	return CollectionNodeColor;
-}
-
-FSlateIcon UMoviePipelineCollectionNode::GetIconAndTint(FLinearColor& OutColor) const
-{
-	static const FSlateIcon CollectionIcon = FSlateIcon(FAppStyle::GetAppStyleSetName(), "SceneOutliner.NewFolderIcon");
-
-	OutColor = FLinearColor::White;
-	return CollectionIcon;
-}
-
-void UMoviePipelineCollectionNode::PostEditChangeProperty(FPropertyChangedEvent& PropertyChangedEvent)
-{
-	Super::PostEditChangeProperty(PropertyChangedEvent);
-
-	if (PropertyChangedEvent.GetPropertyName() == GET_MEMBER_NAME_CHECKED(UMoviePipelineCollectionNode, CollectionName))
-	{
-		OnNodeChangedDelegate.Broadcast(this);
-	}
-}
-#endif // WITH_EDITOR
-
-TArray<FMovieGraphPinProperties> UMovieGraphModifierNode::GetInputPinProperties() const
-{
-	TArray<FMovieGraphPinProperties> Properties;
-	Properties.Add(FMovieGraphPinProperties::MakeBranchProperties());
-	return Properties;
-}
-
-TArray<FMovieGraphPinProperties> UMovieGraphModifierNode::GetOutputPinProperties() const
-{
-	TArray<FMovieGraphPinProperties> Properties;
-	Properties.Add(FMovieGraphPinProperties::MakeBranchProperties());
-	return Properties;
-}
-
-#if WITH_EDITOR
-FText UMovieGraphModifierNode::GetNodeTitle(const bool bGetDescriptive) const
-{
-	static const FText ModifierNodeName = LOCTEXT("NodeName_Modifier", "Modifier");
-	static const FText ModifierNodeDescription = LOCTEXT("NodeDescription_Modifier", "Modifier\n{0} - {1}");
-
-	if (bGetDescriptive && !ModifiedCollectionName.IsEmpty() && !ModifierName.IsEmpty())
-	{
-		return FText::Format(ModifierNodeDescription, FText::FromString(ModifiedCollectionName), FText::FromString(ModifierName));
-	}
-
-	return ModifierNodeName;
-}
-
-FText UMovieGraphModifierNode::GetMenuCategory() const
-{
-	return NodeCategory_Utility;
-}
-
-FLinearColor UMovieGraphModifierNode::GetNodeTitleColor() const
-{
-	static const FLinearColor ModifierNodeColor = FLinearColor(0.6f, 0.113f, 0.113f);
-	return ModifierNodeColor;
-}
-
-FSlateIcon UMovieGraphModifierNode::GetIconAndTint(FLinearColor& OutColor) const
-{
-	static const FSlateIcon ModifierIcon = FSlateIcon(FAppStyle::GetAppStyleSetName(), "ContentBrowser.ReferenceViewer");
-
-	OutColor = FLinearColor::White;
-	return ModifierIcon;
-}
-
-void UMovieGraphModifierNode::PostEditChangeProperty(FPropertyChangedEvent& PropertyChangedEvent)
-{
-	Super::PostEditChangeProperty(PropertyChangedEvent);
-
-	if ((PropertyChangedEvent.GetPropertyName() == GET_MEMBER_NAME_CHECKED(UMovieGraphModifierNode, ModifierName)) ||
-		(PropertyChangedEvent.GetPropertyName() == GET_MEMBER_NAME_CHECKED(UMovieGraphModifierNode, ModifiedCollectionName)))
-	{
-		OnNodeChangedDelegate.Broadcast(this);
-	}
-}
-#endif // WITH_EDITOR
-
 TArray<FMovieGraphPinProperties> UMovieGraphGlobalGameOverridesNode::GetInputPinProperties() const
 {
 	TArray<FMovieGraphPinProperties> Properties;
diff --git a/Engine/Plugins/MovieScene/MovieRenderPipeline/Source/MovieRenderPipelineCore/Public/Graph/MovieGraphPlaceholderNodes.h b/Engine/Plugins/MovieScene/MovieRenderPipeline/Source/MovieRenderPipelineCore/Public/Graph/MovieGraphPlaceholderNodes.h
--- a/Engine/Plugins/MovieScene/MovieRenderPipeline/Source/MovieRenderPipelineCore/Public/Graph/MovieGraphPlaceholderNodes.h
+++ b/Engine/Plugins/MovieScene/MovieRenderPipeline/Source/MovieRenderPipelineCore/Public/Graph/MovieGraphPlaceholderNodes.h
@@ -6,83 +6,6 @@
 
 #include "MovieGraphPlaceholderNodes.generated.h"
 
-// TODO: This should probably be called UMovieGraphCollectionNode instead
-/** A graph node which creates collections that are available for use in other nodes. */
-UCLASS()
-class MOVIERENDERPIPELINECORE_API UMoviePipelineCollectionNode : public UMovieGraphNode
-{
-	GENERATED_BODY()
-
-public:
-	UMoviePipelineCollectionNode() = default;
-
-	virtual TArray<FMovieGraphPinProperties> GetInputPinProperties() const override;
-	virtual TArray<FMovieGraphPinProperties> GetOutputPinProperties() const override;
-
-	virtual TArray<FPropertyBagPropertyDesc> GetDynamicPropertyDescriptions() const override
-	{
-		TArray<FPropertyBagPropertyDesc> Properties;
-		FPropertyBagPropertyDesc FloatEditConProperty = FPropertyBagPropertyDesc("bOverride_TestPropName", EPropertyBagPropertyType::Bool);
-		FPropertyBagPropertyDesc FloatProperty = FPropertyBagPropertyDesc("TestPropName", EPropertyBagPropertyType::Float);
-#if WITH_EDITOR
-		FloatEditConProperty.MetaData.Add(FPropertyBagPropertyDescMetaData("InlineEditConditionToggle", "true"));
-		FloatProperty.MetaData.Add(FPropertyBagPropertyDescMetaData("EditCondition", "bOverride_TestPropName"));
-#endif
-		
-		Properties.Add(FloatEditConProperty);
-		Properties.Add(FloatProperty);
-		return Properties;
-	}
-
-#if WITH_EDITOR
-	virtual FText GetNodeTitle(const bool bGetDescriptive = false) const override;
-	virtual FText GetMenuCategory() const override;
-	virtual FLinearColor GetNodeTitleColor() const override;
-	virtual FSlateIcon GetIconAndTint(FLinearColor& OutColor) const override;
-
-	//~ Begin UObject interface
-	virtual void PostEditChangeProperty(FPropertyChangedEvent& PropertyChangedEvent) override;
-	//~ End UObject interface
-#endif
-
-	/** The name of this collection, which is used to reference this collection in the graph. */
-	UPROPERTY(EditAnywhere, Category = "General")
-	FString CollectionName;
-};
-
-/** A node which modifies part of the world being rendered. */
-UCLASS()
-class MOVIERENDERPIPELINECORE_API UMovieGraphModifierNode : public UMovieGraphNode
-{
-	GENERATED_BODY()
-
-public:
-	UMovieGraphModifierNode() = default;
-
-	virtual TArray<FMovieGraphPinProperties> GetInputPinProperties() const override;
-	virtual TArray<FMovieGraphPinProperties> GetOutputPinProperties() const override;
-
-#if WITH_EDITOR
-	virtual FText GetNodeTitle(const bool bGetDescriptive = false) const override;
-	virtual FText GetMenuCategory() const override;
-	virtual FLinearColor GetNodeTitleColor() const override;
-	virtual FSlateIcon GetIconAndTint(FLinearColor& OutColor) const override;
-
-	//~ Begin UObject interface
-	virtual void PostEditChangeProperty(FPropertyChangedEvent& PropertyChangedEvent) override;
-	//~ End UObject interface
-#endif
-
-	// TODO: Temporary - The modifier name will come from the specific specialized modifier node type (eg, MaterialModifier) in the future
-	/** The name of this modifier. */
-	UPROPERTY(EditAnywhere, Category = "General")
-	FString ModifierName;
-
-	/** The name of the collection being modified. */
-	UPROPERTY(EditAnywhere, Category = "General")
-	FString ModifiedCollectionName;
-};
-
 /** A node which configures the global game overrides. */
 UCLASS()
 class MOVIERENDERPIPELINECORE_API UMovieGraphGlobalGameOverridesNode : public UMovieGraphSettingNode
diff --git a/Engine/Plugins/MovieScene/MovieRenderPipeline/Source/MovieRenderPipelineCore/Public/Graph/MoviePipelineRenderLayerSubsystem.h b/Engine/Plugins/MovieScene/MovieRenderPipeline/Source/MovieRenderPipelineCore/Public/Graph/MoviePipelineRenderLayerSubsystem.h
--- a/Engine/Plugins/MovieScene/MovieRenderPipeline/Source/MovieRenderPipelineCore/Public/Graph/MoviePipelineRenderLayerSubsystem.h
+++ b/Engine/Plugins/MovieScene/MovieRenderPipeline/Source/MovieRenderPipelineCore/Public/Graph/MoviePipelineRenderLayerSubsystem.h
@@ -1,4 +1,4 @@
-﻿// Copyright Epic Games, Inc. All Rights Reserved.
+// Copyright Epic Games, Inc. All Rights Reserved.
 
 #pragma once
 
@@ -43,7 +43,7 @@ private:
 /**
  * Base class for providing actor modification functionality via collections.
  */
-UCLASS(Abstract)
+UCLASS(Abstract, EditInlineNew)
 class UMoviePipelineCollectionModifier : public UObject
 {
 	GENERATED_BODY()
@@ -104,7 +104,8 @@ private:
 	/** Maps a component to its original material assignments (per index). */
 	FComponentToMaterialMap ModifiedComponents;
 
-	UPROPERTY()
+public:
+	UPROPERTY(BlueprintReadWrite, EditAnywhere, Category = "Render Layers")
 	TSoftObjectPtr<UMaterialInterface> MaterialToApply;
 };
 
@@ -137,14 +138,15 @@ private:
 	UPROPERTY(Transient)
 	TMap<TSoftObjectPtr<AActor>, bool> ModifiedActors;
 
-	UPROPERTY()
+public:
+	UPROPERTY(BlueprintReadWrite, EditAnywhere, Category = "Render Layers")
 	bool bIsHidden = false;
 };
 
 /**
  * Base class for providing the ability to determine if an actor matches a query.
  */
-UCLASS(Abstract)
+UCLASS(Abstract, EditInlineNew)
 class UMoviePipelineCollectionQuery : public UObject
 {
 	GENERATED_BODY()
@@ -186,17 +188,17 @@ public:
 	UFUNCTION(BlueprintCallable, Category = "Render Layers")
 	virtual bool DoesActorMatchQuery(const AActor* Actor) const override;
 
-private:
-	UPROPERTY()
+public:
+	UPROPERTY(BlueprintReadWrite, EditAnywhere, Category = "General")
 	TArray<TObjectPtr<UClass>> ComponentTypes;
 
-	UPROPERTY()
+	UPROPERTY(BlueprintReadWrite, EditAnywhere, Category = "General")
 	EMoviePipelineCollectionCommonQueryMode QueryMode = EMoviePipelineCollectionCommonQueryMode::And;
 
-	UPROPERTY()
+	UPROPERTY(BlueprintReadWrite, EditAnywhere, Category = "General")
 	TArray<FString> ActorNames;
 
-	UPROPERTY()
+	UPROPERTY(BlueprintReadWrite, EditAnywhere, Category = "General")
 	TArray<FName> Tags;
 };
 
diff --git a/Engine/Plugins/MovieScene/MovieRenderPipeline/Source/MovieRenderPipelineCore/Public/Graph/Nodes/MovieGraphCollectionNode.h b/Engine/Plugins/MovieScene/MovieRenderPipeline/Source/MovieRenderPipelineCore/Public/Graph/Nodes/MovieGraphCollectionNode.h
new file mode 100644
--- /dev/null
+++ b/Engine/Plugins/MovieScene/MovieRenderPipeline/Source/MovieRenderPipelineCore/Public/Graph/Nodes/MovieGraphCollectionNode.h
@@ -0,0 +1,49 @@
+// Copyright Epic Games, Inc. All Rights Reserved.
+
+#pragma once
+
+#include "Graph/MovieGraphNode.h"
+
+#include "MovieGraphCollectionNode.generated.h"
+
+// Forward Declare
+class UMoviePipelineCollectionQuery;
+
+/** 
+* A collection node specifies an interface for doing dynamic scene queries for actors in the world. Collections work in tandem with
+* UMovieGraphModifiers to select which actors the modifiers should be run on.
+*/
+UCLASS()
+class MOVIERENDERPIPELINECORE_API UMovieGraphCollectionNode : public UMovieGraphSettingNode
+{
+	GENERATED_BODY()
+
+public:
+#if WITH_EDITOR
+	virtual FText GetNodeTitle(const bool bGetDescriptive = false) const override;
+	virtual FText GetMenuCategory() const override;
+	virtual FLinearColor GetNodeTitleColor() const override;
+	virtual FSlateIcon GetIconAndTint(FLinearColor& OutColor) const override;
+
+	//~ Begin UObject interface
+	virtual void PostEditChangeProperty(FPropertyChangedEvent& PropertyChangedEvent) override;
+	//~ End UObject interface
+#endif
+
+	virtual FString GetNodeInstanceName() const override { return CollectionName; }
+
+public:
+	UPROPERTY(EditAnywhere, BlueprintReadWrite, Category = Overrides, meta = (InlineEditConditionToggle))
+	uint8 bOverride_CollectionName : 1;
+	
+	UPROPERTY(EditAnywhere, BlueprintReadWrite, Category = Overrides, meta = (InlineEditConditionToggle))
+	uint8 bOverride_QueryClass : 1;
+
+	/** The name of this collection, which is used to reference this collection in the graph. */
+	UPROPERTY(EditAnywhere, BlueprintReadWrite, Category = "General", meta=(EditCondition="bOverride_CollectionName"))
+	FString CollectionName;
+	
+	/** The type of query this node should run. */
+	UPROPERTY(BlueprintReadWrite, EditAnywhere, Instanced, Category = "General", meta=(EditCondition="bOverride_QueryClass", EditInline))
+	TObjectPtr<UMoviePipelineCollectionQuery> QueryClass;
+};
\ No newline at end of file
diff --git a/Engine/Plugins/MovieScene/MovieRenderPipeline/Source/MovieRenderPipelineCore/Public/Graph/Nodes/MovieGraphModifierNode.h b/Engine/Plugins/MovieScene/MovieRenderPipeline/Source/MovieRenderPipelineCore/Public/Graph/Nodes/MovieGraphModifierNode.h
new file mode 100644
--- /dev/null
+++ b/Engine/Plugins/MovieScene/MovieRenderPipeline/Source/MovieRenderPipelineCore/Public/Graph/Nodes/MovieGraphModifierNode.h
@@ -0,0 +1,56 @@
+// Copyright Epic Games, Inc. All Rights Reserved.
+
+#pragma once
+
+#include "Graph/MovieGraphNode.h"
+
+#include "MovieGraphModifierNode.generated.h"
+
+// Forward Declare
+class UMoviePipelineCollectionModifier;
+
+/** 
+* A collection node specifies an interface for doing dynamic scene queries for actors in the world. Collections work in tandem with
+* UMovieGraphModifiers to select which actors the modifiers should be run on.
+*/
+UCLASS()
+class MOVIERENDERPIPELINECORE_API UMovieGraphModifierNode : public UMovieGraphSettingNode
+{
+	GENERATED_BODY()
+
+public:
+#if WITH_EDITOR
+	virtual FText GetNodeTitle(const bool bGetDescriptive = false) const override;
+	virtual FText GetMenuCategory() const override;
+	virtual FLinearColor GetNodeTitleColor() const override;
+	virtual FSlateIcon GetIconAndTint(FLinearColor& OutColor) const override;
+
+	//~ Begin UObject interface
+	virtual void PostEditChangeProperty(FPropertyChangedEvent& PropertyChangedEvent) override;
+	//~ End UObject interface
+#endif
+
+	virtual FString GetNodeInstanceName() const override { return ModifierName; }
+
+public:
+	UPROPERTY(EditAnywhere, BlueprintReadWrite, Category = Overrides, meta = (InlineEditConditionToggle))
+	uint8 bOverride_ModifierName : 1;
+	
+	UPROPERTY(EditAnywhere, BlueprintReadWrite, Category = Overrides, meta = (InlineEditConditionToggle))
+	uint8 bOverride_ModifiedCollectionName : 1;
+		
+	UPROPERTY(EditAnywhere, BlueprintReadWrite, Category = Overrides, meta = (InlineEditConditionToggle))
+	uint8 bOverride_ModifierClass : 1;
+
+	/** The name of this modifier. */
+	UPROPERTY(EditAnywhere, BlueprintReadWrite, Category = "General", meta=(EditCondition="bOverride_ModifierName"))
+	FString ModifierName;
+
+	/** The name of the collection being modified. */
+	UPROPERTY(EditAnywhere, BlueprintReadWrite, Category = "General", meta=(EditCondition="bOverride_ModifiedCollectionName")) 
+	FString ModifiedCollectionName;
+	
+	/** The modifier this node should run. */
+	UPROPERTY(BlueprintReadWrite, EditAnywhere, Instanced, Category = "General", meta=(EditCondition="bOverride_ModifierClass", EditInline))
+	TObjectPtr<UMoviePipelineCollectionModifier> ModifierClass;
+};
\ No newline at end of file
diff --git a/Engine/Plugins/MovieScene/MovieRenderPipeline/Source/MovieRenderPipelineEditor/Private/Graph/SMovieGraphActiveRenderSettingsTabContent.cpp b/Engine/Plugins/MovieScene/MovieRenderPipeline/Source/MovieRenderPipelineEditor/Private/Graph/SMovieGraphActiveRenderSettingsTabContent.cpp
--- a/Engine/Plugins/MovieScene/MovieRenderPipeline/Source/MovieRenderPipelineEditor/Private/Graph/SMovieGraphActiveRenderSettingsTabContent.cpp
+++ b/Engine/Plugins/MovieScene/MovieRenderPipeline/Source/MovieRenderPipelineEditor/Private/Graph/SMovieGraphActiveRenderSettingsTabContent.cpp
@@ -75,8 +75,19 @@ const TArray<TSharedPtr<FActiveRenderSettingsTreeElement>>& FActiveRenderSetting
 {
 	auto MakeElementFromNode = [this](const TObjectPtr<UMovieGraphNode>& Node) -> TSharedPtr<FActiveRenderSettingsTreeElement>
 	{
+		// The element name should include the node's instance name (if the instance name isn't empty)
+		FString ElementName = Node->GetNodeTitle().ToString();
+		if (const UMovieGraphSettingNode* SettingNode = Cast<UMovieGraphSettingNode>(Node))
+		{
+			const FString InstanceName = SettingNode->GetNodeInstanceName();
+			if (!InstanceName.IsEmpty())
+			{
+				ElementName = FString::Format(TEXT("{0} - {1}"), {ElementName, InstanceName});
+			}
+		}
+		
 		TSharedPtr<FActiveRenderSettingsTreeElement> Element =
-			MakeShared<FActiveRenderSettingsTreeElement>(FName(*Node->GetNodeTitle().ToString()), EElementType::Node);
+			MakeShared<FActiveRenderSettingsTreeElement>(FName(*ElementName), EElementType::Node);
 		Element->SettingsNode = Node;
 		Element->FlattenedGraph = FlattenedGraph;
 		Element->ParentElement = AsWeak();
