MoviePipeline: Grab bag of miscellaneous tweaks/fixes for the render graph.

* Hide the "Properties" category on nodes in the details panel if there are no dynamic properties available. Having this category present w/o any dynamic properties was messy/confusing.
* Offset the position of the "Outputs" node in new graphs so it doesn't overlap with "Inputs".
* Removed temporary/test dynamic properties from the "Render Layer" node.
* Added a new "Set CVar Value" node.
* Added a new "Get CVar Value" node.
* Make the "Select" node work with traversal. It's still restricted to just string values, though. Moved it out of placeholders.
* Updated the "Branch" node to work with traversal again; the "Condition" pin is read correctly.
* Moved the "Global Game Overrides" node out of placeholders, and gave it a bunch of properties (these may need to be modified in the future).
* Cleaned up the node context menu (category names and node menu titles).
* More nodes inherit from UMovieGraphSettingsNode now, so they appear within traversal results.
* Added variable type indicator icon to the Members panel.

diff --git a/Engine/Plugins/MovieScene/MovieRenderPipeline/Source/MovieRenderPipelineCore/Private/Graph/MovieGraphConfig.cpp b/Engine/Plugins/MovieScene/MovieRenderPipeline/Source/MovieRenderPipelineCore/Private/Graph/MovieGraphConfig.cpp
--- a/Engine/Plugins/MovieScene/MovieRenderPipeline/Source/MovieRenderPipelineCore/Private/Graph/MovieGraphConfig.cpp
+++ b/Engine/Plugins/MovieScene/MovieRenderPipeline/Source/MovieRenderPipelineCore/Private/Graph/MovieGraphConfig.cpp
@@ -72,6 +72,12 @@ UMovieGraphConfig::UMovieGraphConfig()
 		AddDefaultMembers();
 		InputNode->UpdatePins();
 		OutputNode->UpdatePins();
+
+		// Offset the default output node so it doesn't overlap the default input node
+#if WITH_EDITOR
+		constexpr int32 OutputNodeOffset = 300;
+		OutputNode->SetNodePosX(OutputNodeOffset);
+#endif
 	}
 }
 
@@ -99,7 +105,7 @@ void UMovieGraphConfig::PostLoad()
 	}
 }
 
-UMovieGraphVariable* UMovieGraphConfig::AddGlobalVariable(const FName& InName)
+UMovieGraphVariable* UMovieGraphConfig::AddGlobalVariable(const FName& InName, EMovieGraphValueType ValueType)
 {
 	// Don't add duplicate global variables
 	const bool VariableExists = Variables.ContainsByPredicate([&InName](const TObjectPtr<UMovieGraphVariable>& Variable)
@@ -116,6 +122,7 @@ UMovieGraphVariable* UMovieGraphConfig::AddGlobalVariable(const FName& InName)
 	{
 		NewVariable->bIsGlobal = true;
 		NewVariable->bIsEditable = false;
+		NewVariable->SetValueType(ValueType);
 		return NewVariable;
 	}
 
@@ -152,13 +159,19 @@ void UMovieGraphConfig::AddDefaultMembers()
 		OutputNode->UpdatePins();
 	}
 
+	static const TMap<FName, EMovieGraphValueType> GlobalVariableNamesAndTypes =
+	{
+		{GlobalVariable_ShotName, EMovieGraphValueType::String},
+		{GlobalVariable_SequenceName, EMovieGraphValueType::String},
+		{GlobalVariable_FrameNumber, EMovieGraphValueType::Int32},
+		{GlobalVariable_CameraName, EMovieGraphValueType::String},
+		{GlobalVariable_RenderLayerName, EMovieGraphValueType::String}
+	};
+
 	// Add all of the global variables that should be available in the graph
-	static const TArray<FName> GlobalVariableNames =
-		{GlobalVariable_ShotName, GlobalVariable_SequenceName, GlobalVariable_FrameNumber,
-		 GlobalVariable_CameraName, GlobalVariable_RenderLayerName};
-	for (const FName& GlobalVariableName : GlobalVariableNames)
+	for (const TTuple<FName, EMovieGraphValueType>& GlobalVariableInfo : GlobalVariableNamesAndTypes)
 	{
-		AddGlobalVariable(GlobalVariableName);
+		AddGlobalVariable(GlobalVariableInfo.Key, GlobalVariableInfo.Value);
 	}
 }
 
diff --git a/Engine/Plugins/MovieScene/MovieRenderPipeline/Source/MovieRenderPipelineCore/Private/Graph/Nodes/MovieGraphBranchNode.cpp b/Engine/Plugins/MovieScene/MovieRenderPipeline/Source/MovieRenderPipelineCore/Private/Graph/Nodes/MovieGraphBranchNode.cpp
--- a/Engine/Plugins/MovieScene/MovieRenderPipeline/Source/MovieRenderPipelineCore/Private/Graph/Nodes/MovieGraphBranchNode.cpp
+++ b/Engine/Plugins/MovieScene/MovieRenderPipeline/Source/MovieRenderPipelineCore/Private/Graph/Nodes/MovieGraphBranchNode.cpp
@@ -52,26 +52,12 @@ TArray<UMovieGraphPin*> UMovieGraphBranchNode::EvaluatePinsToFollow(FMovieGraphE
 		return PinsToFollow;
 	}
 
-	// Connected Node is the node that the pin is actually connected to.
-	const UMovieGraphValueContainer* ValueContainer = ConnectedNode->GetPropertyValueContainerForPin(OtherPin->GetName());
-
-	if (!ValueContainer)
-	{
-		// Somehow we got connected to node that can't provide a value.
-		UE_LOG(LogMovieRenderPipeline, Warning, TEXT("ConnectedNode: %s provided no ValueContainer to fetch value."), *ConnectedNode->GetName());
-		return PinsToFollow;
-	}
-
-	bool bEvaluatedValue;
-	bool bSuccessfullyFetched = ValueContainer->GetValueBool(bEvaluatedValue);
-	if (!bSuccessfullyFetched)
-	{
-		// There was a node connected but it didn't know how to return a boolean value for us.
-		UE_LOG(LogMovieRenderPipeline, Warning, TEXT("ValueContainer had no bool value."));
-		return PinsToFollow;
-	}
-
-	if (bEvaluatedValue)
+	// TODO: Comparing against a string like this is not ideal. This could use GetPropertyValueContainerForPin() instead,
+	// however that may not work in a general case. This fits variables nicely since they already have a
+	// UMovieGraphValueContainer -- other nodes in the future may not, though, and generating them on every frame may be expensive.
+	// Still an area that needs to be researched.
+	static const FString TrueValue = FString("true");
+	if (ConnectedNode->GetResolvedValueForOutputPin(OtherPin->Properties.Label, &InContext.UserContext) == TrueValue)
 	{
 		PinsToFollow.Add(GetInputPin(UE::MovieGraph::BranchNode::TrueBranch));
 	}
diff --git a/Engine/Plugins/MovieScene/MovieRenderPipeline/Source/MovieRenderPipelineCore/Private/Graph/Nodes/MovieGraphGetCVarValueNode.cpp b/Engine/Plugins/MovieScene/MovieRenderPipeline/Source/MovieRenderPipelineCore/Private/Graph/Nodes/MovieGraphGetCVarValueNode.cpp
new file mode 100644
--- /dev/null
+++ b/Engine/Plugins/MovieScene/MovieRenderPipeline/Source/MovieRenderPipelineCore/Private/Graph/Nodes/MovieGraphGetCVarValueNode.cpp
@@ -0,0 +1,60 @@
+// Copyright Epic Games, Inc. All Rights Reserved.
+
+#include "Graph/Nodes/MovieGraphGetCVarValueNode.h"
+
+#include "Graph/MovieGraphConfig.h"
+#include "Styling/AppStyle.h"
+
+#if WITH_EDITOR
+FText UMovieGraphGetCVarValueNode::GetNodeTitle(const bool bGetDescriptive) const
+{
+	static const FText GetCVarNodeName = NSLOCTEXT("MovieGraphNodes", "NodeName_GetCVar", "Get CVar Value");
+	static const FText GetCVarNodeDescription = NSLOCTEXT("MovieGraphNodes", "NodeDescription_GetCVar", "Get CVar Value\n{0}");
+
+	if (bGetDescriptive && !Name.IsEmpty())
+	{
+		return FText::Format(GetCVarNodeDescription, FText::FromString(Name));
+	}
+	
+	return GetCVarNodeName;
+}
+
+FText UMovieGraphGetCVarValueNode::GetMenuCategory() const
+{
+	return NSLOCTEXT("MovieGraphNodes", "GetCVarGraphNode_Category", "Utility");
+}
+
+FLinearColor UMovieGraphGetCVarValueNode::GetNodeTitleColor() const
+{
+	static const FLinearColor GetCVarNodeColor = FLinearColor(0.04f, 0.22f, 0.36f);
+	return GetCVarNodeColor;
+}
+
+FSlateIcon UMovieGraphGetCVarValueNode::GetIconAndTint(FLinearColor& OutColor) const
+{
+	static const FSlateIcon RenderLayerIcon = FSlateIcon(FAppStyle::GetAppStyleSetName(), "MainFrame.BrowseCVars");
+
+	OutColor = FLinearColor::White;
+	return RenderLayerIcon;
+}
+
+TArray<FMovieGraphPinProperties> UMovieGraphGetCVarValueNode::GetOutputPinProperties() const
+{
+	TArray<FMovieGraphPinProperties> OutputPinProperties = Super::GetOutputPinProperties();
+
+	FMovieGraphPinProperties Properties(FName("Value"), EMovieGraphValueType::String, false);
+	OutputPinProperties.Add(MoveTemp(Properties));
+
+	return OutputPinProperties;
+}
+
+void UMovieGraphGetCVarValueNode::PostEditChangeProperty(FPropertyChangedEvent& PropertyChangedEvent)
+{
+	Super::PostEditChangeProperty(PropertyChangedEvent);
+	
+	if (PropertyChangedEvent.GetPropertyName() == GET_MEMBER_NAME_CHECKED(UMovieGraphGetCVarValueNode, Name))
+	{
+		OnNodeChangedDelegate.Broadcast(this);
+	}
+}
+#endif // WITH_EDITOR
\ No newline at end of file
diff --git a/Engine/Plugins/MovieScene/MovieRenderPipeline/Source/MovieRenderPipelineCore/Private/Graph/Nodes/MovieGraphGlobalGameOverrides.cpp b/Engine/Plugins/MovieScene/MovieRenderPipeline/Source/MovieRenderPipelineCore/Private/Graph/Nodes/MovieGraphGlobalGameOverrides.cpp
new file mode 100644
--- /dev/null
+++ b/Engine/Plugins/MovieScene/MovieRenderPipeline/Source/MovieRenderPipelineCore/Private/Graph/Nodes/MovieGraphGlobalGameOverrides.cpp
@@ -0,0 +1,33 @@
+// Copyright Epic Games, Inc. All Rights Reserved.
+
+#include "Graph/Nodes/MovieGraphGlobalGameOverrides.h"
+
+#include "Styling/AppStyle.h"
+
+#if WITH_EDITOR
+FText UMovieGraphGlobalGameOverridesNode::GetNodeTitle(const bool bGetDescriptive) const
+{
+	static const FText GlobalGameOverridesNodeName = NSLOCTEXT("MovieGraphNodes", "NodeName_GlobalGameOverrides", "Global Game Overrides");
+	return GlobalGameOverridesNodeName;
+}
+
+FText UMovieGraphGlobalGameOverridesNode::GetMenuCategory() const
+{
+	static const FText NodeCategory_Globals = NSLOCTEXT("MovieGraphNodes", "NodeCategory_Globals", "Globals");
+	return NodeCategory_Globals;
+}
+
+FLinearColor UMovieGraphGlobalGameOverridesNode::GetNodeTitleColor() const
+{
+	static const FLinearColor GlobalGameOverridesNodeColor = FLinearColor(0.549f, 0.f, 0.250f);
+	return GlobalGameOverridesNodeColor;
+}
+
+FSlateIcon UMovieGraphGlobalGameOverridesNode::GetIconAndTint(FLinearColor& OutColor) const
+{
+	static const FSlateIcon GlobalGameOverridesIcon = FSlateIcon(FAppStyle::GetAppStyleSetName(), "Launcher.TabIcon");
+
+	OutColor = FLinearColor::White;
+	return GlobalGameOverridesIcon;
+}
+#endif // WITH_EDITOR
\ No newline at end of file
diff --git a/Engine/Plugins/MovieScene/MovieRenderPipeline/Source/MovieRenderPipelineCore/Private/Graph/Nodes/MovieGraphRenderLayerNode.cpp b/Engine/Plugins/MovieScene/MovieRenderPipeline/Source/MovieRenderPipelineCore/Private/Graph/Nodes/MovieGraphRenderLayerNode.cpp
--- a/Engine/Plugins/MovieScene/MovieRenderPipeline/Source/MovieRenderPipelineCore/Private/Graph/Nodes/MovieGraphRenderLayerNode.cpp
+++ b/Engine/Plugins/MovieScene/MovieRenderPipeline/Source/MovieRenderPipelineCore/Private/Graph/Nodes/MovieGraphRenderLayerNode.cpp
@@ -10,21 +10,6 @@ UMovieGraphRenderLayerNode::UMovieGraphRenderLayerNode()
 	LayerName = TEXT("beauty");
 }
 
-TArray<FPropertyBagPropertyDesc> UMovieGraphRenderLayerNode::GetDynamicPropertyDescriptions() const
-{
-	TArray<FPropertyBagPropertyDesc> Properties;
-	FPropertyBagPropertyDesc FloatEditConProperty = FPropertyBagPropertyDesc("bOverride_TestPropName", EPropertyBagPropertyType::Bool);
-	FPropertyBagPropertyDesc FloatProperty = FPropertyBagPropertyDesc("TestPropName", EPropertyBagPropertyType::String);
-#if WITH_EDITOR
-	FloatEditConProperty.MetaData.Add(FPropertyBagPropertyDescMetaData("InlineEditConditionToggle", "true"));
-	FloatProperty.MetaData.Add(FPropertyBagPropertyDescMetaData("EditCondition", "bOverride_TestPropName"));
-#endif
-	
-	Properties.Add(FloatEditConProperty);
-	Properties.Add(FloatProperty);
-	return Properties;
-}
-
 #if WITH_EDITOR
 FText UMovieGraphRenderLayerNode::GetNodeTitle(const bool bGetDescriptive) const
 {
diff --git a/Engine/Plugins/MovieScene/MovieRenderPipeline/Source/MovieRenderPipelineCore/Private/Graph/Nodes/MovieGraphSelectNode.cpp b/Engine/Plugins/MovieScene/MovieRenderPipeline/Source/MovieRenderPipelineCore/Private/Graph/Nodes/MovieGraphSelectNode.cpp
new file mode 100644
--- /dev/null
+++ b/Engine/Plugins/MovieScene/MovieRenderPipeline/Source/MovieRenderPipelineCore/Private/Graph/Nodes/MovieGraphSelectNode.cpp
@@ -0,0 +1,121 @@
+// Copyright Epic Games, Inc. All Rights Reserved.
+
+#include "Graph/Nodes/MovieGraphSelectNode.h"
+
+#include "Graph/MovieGraphConfig.h"
+#include "Styling/AppStyle.h"
+
+namespace UE::MovieGraph::SelectNode
+{
+	static const FName SelectedOption("Selected Option");
+}
+
+TArray<FMovieGraphPinProperties> UMovieGraphSelectNode::GetInputPinProperties() const
+{
+	TArray<FMovieGraphPinProperties> Properties;
+
+	// Generate branch pins for each option
+	for (const FString& SelectOption : SelectOptions)
+	{
+		Properties.Add(FMovieGraphPinProperties::MakeBranchProperties(FName(SelectOption)));
+	}
+
+	Properties.Add(FMovieGraphPinProperties(UE::MovieGraph::SelectNode::SelectedOption, EMovieGraphValueType::String, false));
+	
+	return Properties;
+}
+
+TArray<FMovieGraphPinProperties> UMovieGraphSelectNode::GetOutputPinProperties() const
+{
+	TArray<FMovieGraphPinProperties> Properties;
+	Properties.Add(FMovieGraphPinProperties::MakeBranchProperties());
+	return Properties;
+}
+
+TArray<UMovieGraphPin*> UMovieGraphSelectNode::EvaluatePinsToFollow(FMovieGraphEvaluationContext& InContext) const
+{
+	TArray<UMovieGraphPin*> PinsToFollow;
+
+	// The resolved value of the "Selected Option" property. May come from a connection or a value specified on the node.
+	FString ResolvedSelectValue;
+
+	// Try getting the value from a connection first
+	bool bGotValueFromConnection = false;
+	if (const UMovieGraphPin* SelectPin = GetInputPin(UE::MovieGraph::SelectNode::SelectedOption))
+	{
+		if (const UMovieGraphPin* OtherPin = SelectPin->GetFirstConnectedPin())
+		{
+			if (const UMovieGraphNode* ConnectedNode = OtherPin->Node)
+			{
+				bGotValueFromConnection = true;
+				ResolvedSelectValue = ConnectedNode->GetResolvedValueForOutputPin(OtherPin->Properties.Label, &InContext.UserContext);
+			}
+		}
+	}
+
+	if (!bGotValueFromConnection)
+	{
+		ResolvedSelectValue = SelectedOption;
+	}
+
+	for (const FString& SelectOption : SelectOptions)
+	{
+		if (SelectOption == ResolvedSelectValue)
+		{
+			PinsToFollow.Add(GetInputPin(FName(SelectOption)));
+			break;
+		}
+	}
+
+	return PinsToFollow;
+}
+
+#if WITH_EDITOR
+FText UMovieGraphSelectNode::GetNodeTitle(const bool bGetDescriptive) const
+{
+	static const FText SelectNodeName = NSLOCTEXT("MovieGraphNodes", "NodeName_Select", "Select");
+	static const FText SelectNodeDescription = NSLOCTEXT("MovieGraphNodes", "NodeDescription_Select", "Select\n{0}");
+
+	if (bGetDescriptive && !Description.IsEmpty())
+	{
+		return FText::Format(SelectNodeDescription, FText::FromString(Description));
+	}
+
+	return SelectNodeName;
+}
+
+FText UMovieGraphSelectNode::GetMenuCategory() const
+{
+	static const FText NodeCategory_Conditionals = NSLOCTEXT("MovieGraphNodes", "NodeCategory_Conditionals", "Conditionals");
+	return NodeCategory_Conditionals;
+}
+
+FLinearColor UMovieGraphSelectNode::GetNodeTitleColor() const
+{
+	static const FLinearColor SelectNodeColor = FLinearColor(0.266f, 0.266f, 0.266f);
+	return SelectNodeColor;
+}
+
+FSlateIcon UMovieGraphSelectNode::GetIconAndTint(FLinearColor& OutColor) const
+{
+	static const FSlateIcon SelectIcon = FSlateIcon(FAppStyle::GetAppStyleSetName(), "Icons.Merge");
+
+	OutColor = FLinearColor::White;
+	return SelectIcon;
+}
+
+void UMovieGraphSelectNode::PostEditChangeProperty(FPropertyChangedEvent& PropertyChangedEvent)
+{
+	Super::PostEditChangeProperty(PropertyChangedEvent);
+
+	if (PropertyChangedEvent.GetPropertyName() == GET_MEMBER_NAME_CHECKED(UMovieGraphSelectNode, SelectOptions))
+	{
+		UpdatePins();
+	}
+	
+	if (PropertyChangedEvent.GetPropertyName() == GET_MEMBER_NAME_CHECKED(UMovieGraphSelectNode, Description))
+	{
+		OnNodeChangedDelegate.Broadcast(this);
+	}
+}
+#endif // WITH_EDITOR
\ No newline at end of file
diff --git a/Engine/Plugins/MovieScene/MovieRenderPipeline/Source/MovieRenderPipelineCore/Private/Graph/Nodes/MovieGraphSetCVarValueNode.cpp b/Engine/Plugins/MovieScene/MovieRenderPipeline/Source/MovieRenderPipelineCore/Private/Graph/Nodes/MovieGraphSetCVarValueNode.cpp
new file mode 100644
--- /dev/null
+++ b/Engine/Plugins/MovieScene/MovieRenderPipeline/Source/MovieRenderPipelineCore/Private/Graph/Nodes/MovieGraphSetCVarValueNode.cpp
@@ -0,0 +1,51 @@
+// Copyright Epic Games, Inc. All Rights Reserved.
+
+#include "Graph/Nodes/MovieGraphSetCVarValueNode.h"
+
+#include "Graph/MovieGraphConfig.h"
+#include "Styling/AppStyle.h"
+
+#if WITH_EDITOR
+FText UMovieGraphSetCVarValueNode::GetNodeTitle(const bool bGetDescriptive) const
+{
+	static const FText SetCVarNodeName = NSLOCTEXT("MovieGraphNodes", "NodeName_SetCVar", "Set CVar Value");
+	static const FText SetCVarNodeDescription = NSLOCTEXT("MovieGraphNodes", "NodeDescription_SetCVar", "Set CVar Value\n{0}");
+
+	if (bGetDescriptive && !Name.IsEmpty())
+	{
+		return FText::Format(SetCVarNodeDescription, FText::FromString(Name));
+	}
+	
+	return SetCVarNodeName;
+}
+
+FText UMovieGraphSetCVarValueNode::GetMenuCategory() const
+{
+	return NSLOCTEXT("MovieGraphNodes", "SetCVarGraphNode_Category", "Utility");
+}
+
+FLinearColor UMovieGraphSetCVarValueNode::GetNodeTitleColor() const
+{
+	static const FLinearColor SetCVarNodeColor = FLinearColor(0.04f, 0.22f, 0.36f);
+	return SetCVarNodeColor;
+}
+
+FSlateIcon UMovieGraphSetCVarValueNode::GetIconAndTint(FLinearColor& OutColor) const
+{
+	static const FSlateIcon RenderLayerIcon = FSlateIcon(FAppStyle::GetAppStyleSetName(), "MainFrame.BrowseCVars");
+
+	OutColor = FLinearColor::White;
+	return RenderLayerIcon;
+}
+
+void UMovieGraphSetCVarValueNode::PostEditChangeProperty(FPropertyChangedEvent& PropertyChangedEvent)
+{
+	Super::PostEditChangeProperty(PropertyChangedEvent);
+	
+	if (PropertyChangedEvent.GetPropertyName() == GET_MEMBER_NAME_CHECKED(UMovieGraphSetCVarValueNode, Name))
+	{
+		OnNodeChangedDelegate.Broadcast(this);
+	}
+}
+
+#endif // WITH_EDITOR
\ No newline at end of file
diff --git a/Engine/Plugins/MovieScene/MovieRenderPipeline/Source/MovieRenderPipelineCore/Public/Graph/MovieGraphConfig.h b/Engine/Plugins/MovieScene/MovieRenderPipeline/Source/MovieRenderPipelineCore/Public/Graph/MovieGraphConfig.h
--- a/Engine/Plugins/MovieScene/MovieRenderPipeline/Source/MovieRenderPipelineCore/Public/Graph/MovieGraphConfig.h
+++ b/Engine/Plugins/MovieScene/MovieRenderPipeline/Source/MovieRenderPipelineCore/Public/Graph/MovieGraphConfig.h
@@ -510,8 +510,8 @@ private:
 	template<typename T>
 	T* AddMember(TArray<TObjectPtr<T>>& InMemberArray, const FName& InBaseName);
 
-	/** Adds a global variable to the graph with the provided name. */
-	UMovieGraphVariable* AddGlobalVariable(const FName& InName);
+	/** Adds a global variable to the graph with the provided name and value type. */
+	UMovieGraphVariable* AddGlobalVariable(const FName& InName, EMovieGraphValueType ValueType);
 
 	/** Adds members to the graph that should always be available. */
 	void AddDefaultMembers();
diff --git a/Engine/Plugins/MovieScene/MovieRenderPipeline/Source/MovieRenderPipelineCore/Public/Graph/MovieGraphNode.h b/Engine/Plugins/MovieScene/MovieRenderPipeline/Source/MovieRenderPipelineCore/Public/Graph/MovieGraphNode.h
--- a/Engine/Plugins/MovieScene/MovieRenderPipeline/Source/MovieRenderPipelineCore/Public/Graph/MovieGraphNode.h
+++ b/Engine/Plugins/MovieScene/MovieRenderPipeline/Source/MovieRenderPipelineCore/Public/Graph/MovieGraphNode.h
@@ -259,7 +259,7 @@ protected:
 	TArray<TObjectPtr<UMovieGraphPin>> OutputPins;
 
 	/** Properties which can be dynamically declared on the node (vs. native properties which are always present). */
-	UPROPERTY(EditAnywhere, meta=(FixedLayout, ShowOnlyInnerProperties), Category = "Node")
+	UPROPERTY(EditAnywhere, meta=(FixedLayout, ShowOnlyInnerProperties), Category = "Properties")
 	FInstancedPropertyBag DynamicProperties;
 
 	/** Tracks which properties have been exposed on the node as inputs. */
diff --git a/Engine/Plugins/MovieScene/MovieRenderPipeline/Source/MovieRenderPipelineCore/Public/Graph/MovieGraphPlaceholderNodes.cpp b/Engine/Plugins/MovieScene/MovieRenderPipeline/Source/MovieRenderPipelineCore/Public/Graph/MovieGraphPlaceholderNodes.cpp
--- a/Engine/Plugins/MovieScene/MovieRenderPipeline/Source/MovieRenderPipelineCore/Public/Graph/MovieGraphPlaceholderNodes.cpp
+++ b/Engine/Plugins/MovieScene/MovieRenderPipeline/Source/MovieRenderPipelineCore/Public/Graph/MovieGraphPlaceholderNodes.cpp
@@ -7,68 +7,10 @@
 
 #define LOCTEXT_NAMESPACE "MovieGraphNode"
 
-static const FText NodeCategory_Utility = LOCTEXT("NodeCategory_Utility", "Utility");
-static const FText NodeCategory_Globals = LOCTEXT("NodeCategory_Globals", "Globals");
-static const FText NodeCategory_Renderers = LOCTEXT("NodeCategory_Renderers", "Renderers");
+static const FText NodeCategory_Rendering = LOCTEXT("NodeCategory_Rendering", "Rendering");
 static const FText NodeCategory_OutputType = LOCTEXT("NodeCategory_OutputType", "Output Type");
-static const FText NodeCategory_Conditionals = LOCTEXT("NodeCategory_Conditionals", "Conditionals");
 static const FText NodeCategory_Settings = LOCTEXT("NodeCategory_Settings", "Settings");
 
-TArray<FMovieGraphPinProperties> UMovieGraphGlobalGameOverridesNode::GetInputPinProperties() const
-{
-	TArray<FMovieGraphPinProperties> Properties;
-	Properties.Add(FMovieGraphPinProperties::MakeBranchProperties());
-	return Properties;
-}
-
-TArray<FMovieGraphPinProperties> UMovieGraphGlobalGameOverridesNode::GetOutputPinProperties() const
-{
-	TArray<FMovieGraphPinProperties> Properties;
-	Properties.Add(FMovieGraphPinProperties::MakeBranchProperties());
-	return Properties;
-}
-
-#if WITH_EDITOR
-FText UMovieGraphGlobalGameOverridesNode::GetNodeTitle(const bool bGetDescriptive) const
-{
-	static const FText GlobalGameOverridesNodeName = LOCTEXT("NodeName_GlobalGameOverrides", "Global Game Overrides");
-	return GlobalGameOverridesNodeName;
-}
-
-FText UMovieGraphGlobalGameOverridesNode::GetMenuCategory() const
-{
-	return NodeCategory_Globals;
-}
-
-FLinearColor UMovieGraphGlobalGameOverridesNode::GetNodeTitleColor() const
-{
-	static const FLinearColor GlobalGameOverridesNodeColor = FLinearColor(0.549f, 0.f, 0.250f);
-	return GlobalGameOverridesNodeColor;
-}
-
-FSlateIcon UMovieGraphGlobalGameOverridesNode::GetIconAndTint(FLinearColor& OutColor) const
-{
-	static const FSlateIcon GlobalGameOverridesIcon = FSlateIcon(FAppStyle::GetAppStyleSetName(), "Launcher.TabIcon");
-
-	OutColor = FLinearColor::White;
-	return GlobalGameOverridesIcon;
-}
-#endif // WITH_EDITOR
-
-TArray<FMovieGraphPinProperties> UMovieGraphPathTracedRendererNode::GetInputPinProperties() const
-{
-	TArray<FMovieGraphPinProperties> Properties;
-	Properties.Add(FMovieGraphPinProperties::MakeBranchProperties());
-	return Properties;
-}
-
-TArray<FMovieGraphPinProperties> UMovieGraphPathTracedRendererNode::GetOutputPinProperties() const
-{
-	TArray<FMovieGraphPinProperties> Properties;
-	Properties.Add(FMovieGraphPinProperties::MakeBranchProperties());
-	return Properties;
-}
-
 #if WITH_EDITOR
 FText UMovieGraphPathTracedRendererNode::GetNodeTitle(const bool bGetDescriptive) const
 {
@@ -78,7 +20,7 @@ FText UMovieGraphPathTracedRendererNode::GetNodeTitle(const bool bGetDescriptive
 
 FText UMovieGraphPathTracedRendererNode::GetMenuCategory() const
 {
-	return NodeCategory_Renderers;
+	return NodeCategory_Rendering;
 }
 
 FLinearColor UMovieGraphPathTracedRendererNode::GetNodeTitleColor() const
@@ -96,24 +38,10 @@ FSlateIcon UMovieGraphPathTracedRendererNode::GetIconAndTint(FLinearColor& OutCo
 }
 #endif // WITH_EDITOR
 
-TArray<FMovieGraphPinProperties> UMovieGraphEXRSequenceNode::GetInputPinProperties() const
-{
-	TArray<FMovieGraphPinProperties> Properties;
-	Properties.Add(FMovieGraphPinProperties::MakeBranchProperties());
-	return Properties;
-}
-
-TArray<FMovieGraphPinProperties> UMovieGraphEXRSequenceNode::GetOutputPinProperties() const
-{
-	TArray<FMovieGraphPinProperties> Properties;
-	Properties.Add(FMovieGraphPinProperties::MakeBranchProperties());
-	return Properties;
-}
-
 #if WITH_EDITOR
 FText UMovieGraphEXRSequenceNode::GetNodeTitle(const bool bGetDescriptive) const
 {
-	static const FText EXRSequenceNodeName = LOCTEXT("NodeName_EXRSequence", "EXR Sequence");
+	static const FText EXRSequenceNodeName = LOCTEXT("NodeName_EXRSequence", ".exr Sequence");
 	return EXRSequenceNodeName;
 }
 
@@ -137,91 +65,6 @@ FSlateIcon UMovieGraphEXRSequenceNode::GetIconAndTint(FLinearColor& OutColor) co
 }
 #endif // WITH_EDITOR
 
-TArray<FMovieGraphPinProperties> UMovieGraphSelectNode::GetInputPinProperties() const
-{
-	TArray<FMovieGraphPinProperties> Properties;
-
-	// Generate branch pins for each option
-	for (const FString& SelectOption : SelectOptions)
-	{
-		Properties.Add(FMovieGraphPinProperties::MakeBranchProperties(FName(SelectOption)));
-	}
-
-	Properties.Add(FMovieGraphPinProperties(NAME_None, EMovieGraphValueType::String, false));
-	
-	return Properties;
-}
-
-TArray<FMovieGraphPinProperties> UMovieGraphSelectNode::GetOutputPinProperties() const
-{
-	TArray<FMovieGraphPinProperties> Properties;
-	Properties.Add(FMovieGraphPinProperties::MakeBranchProperties());
-	return Properties;
-}
-
-#if WITH_EDITOR
-FText UMovieGraphSelectNode::GetNodeTitle(const bool bGetDescriptive) const
-{
-	static const FText SelectNodeName = LOCTEXT("NodeName_Select", "Select");
-	static const FText SelectNodeDescription = LOCTEXT("NodeDescription_Select", "Select\n{0}");
-
-	if (bGetDescriptive && !Description.IsEmpty())
-	{
-		return FText::Format(SelectNodeDescription, FText::FromString(Description));
-	}
-
-	return SelectNodeName;
-}
-
-FText UMovieGraphSelectNode::GetMenuCategory() const
-{
-	return NodeCategory_Conditionals;
-}
-
-FLinearColor UMovieGraphSelectNode::GetNodeTitleColor() const
-{
-	static const FLinearColor SelectNodeColor = FLinearColor(0.266f, 0.266f, 0.266f);
-	return SelectNodeColor;
-}
-
-FSlateIcon UMovieGraphSelectNode::GetIconAndTint(FLinearColor& OutColor) const
-{
-	static const FSlateIcon SelectIcon = FSlateIcon(FAppStyle::GetAppStyleSetName(), "Icons.Merge");
-
-	OutColor = FLinearColor::White;
-	return SelectIcon;
-}
-
-void UMovieGraphSelectNode::PostEditChangeProperty(FPropertyChangedEvent& PropertyChangedEvent)
-{
-	Super::PostEditChangeProperty(PropertyChangedEvent);
-
-	if (PropertyChangedEvent.GetPropertyName() == GET_MEMBER_NAME_CHECKED(UMovieGraphSelectNode, SelectOptions))
-	{
-		UpdatePins();
-	}
-	
-	if (PropertyChangedEvent.GetPropertyName() == GET_MEMBER_NAME_CHECKED(UMovieGraphSelectNode, Description))
-	{
-		OnNodeChangedDelegate.Broadcast(this);
-	}
-}
-#endif // WITH_EDITOR
-
-TArray<FMovieGraphPinProperties> UMovieGraphAntiAliasingNode::GetInputPinProperties() const
-{
-	TArray<FMovieGraphPinProperties> Properties;
-	Properties.Add(FMovieGraphPinProperties::MakeBranchProperties());
-	return Properties;
-}
-
-TArray<FMovieGraphPinProperties> UMovieGraphAntiAliasingNode::GetOutputPinProperties() const
-{
-	TArray<FMovieGraphPinProperties> Properties;
-	Properties.Add(FMovieGraphPinProperties::MakeBranchProperties());
-	return Properties;
-}
-
 #if WITH_EDITOR
 FText UMovieGraphAntiAliasingNode::GetNodeTitle(const bool bGetDescriptive) const
 {
diff --git a/Engine/Plugins/MovieScene/MovieRenderPipeline/Source/MovieRenderPipelineCore/Public/Graph/MovieGraphPlaceholderNodes.h b/Engine/Plugins/MovieScene/MovieRenderPipeline/Source/MovieRenderPipelineCore/Public/Graph/MovieGraphPlaceholderNodes.h
--- a/Engine/Plugins/MovieScene/MovieRenderPipeline/Source/MovieRenderPipelineCore/Public/Graph/MovieGraphPlaceholderNodes.h
+++ b/Engine/Plugins/MovieScene/MovieRenderPipeline/Source/MovieRenderPipelineCore/Public/Graph/MovieGraphPlaceholderNodes.h
@@ -6,38 +6,15 @@
 
 #include "MovieGraphPlaceholderNodes.generated.h"
 
-/** A node which configures the global game overrides. */
-UCLASS()
-class MOVIERENDERPIPELINECORE_API UMovieGraphGlobalGameOverridesNode : public UMovieGraphSettingNode
-{
-	GENERATED_BODY()
-
-public:
-	UMovieGraphGlobalGameOverridesNode() = default;
-
-	virtual TArray<FMovieGraphPinProperties> GetInputPinProperties() const override;
-	virtual TArray<FMovieGraphPinProperties> GetOutputPinProperties() const override;
-
-#if WITH_EDITOR
-	virtual FText GetNodeTitle(const bool bGetDescriptive = false) const override;
-	virtual FText GetMenuCategory() const override;
-	virtual FLinearColor GetNodeTitleColor() const override;
-	virtual FSlateIcon GetIconAndTint(FLinearColor& OutColor) const override;
-#endif
-};
-
 /** A node which represents a path traced renderer. */
 UCLASS()
-class MOVIERENDERPIPELINECORE_API UMovieGraphPathTracedRendererNode : public UMovieGraphNode
+class MOVIERENDERPIPELINECORE_API UMovieGraphPathTracedRendererNode : public UMovieGraphSettingNode
 {
 	GENERATED_BODY()
 
 public:
 	UMovieGraphPathTracedRendererNode() = default;
 
-	virtual TArray<FMovieGraphPinProperties> GetInputPinProperties() const override;
-	virtual TArray<FMovieGraphPinProperties> GetOutputPinProperties() const override;
-
 #if WITH_EDITOR
 	virtual FText GetNodeTitle(const bool bGetDescriptive = false) const override;
 	virtual FText GetMenuCategory() const override;
@@ -48,16 +25,13 @@ public:
 
 /** A node which generates an EXR image sequence. */
 UCLASS()
-class MOVIERENDERPIPELINECORE_API UMovieGraphEXRSequenceNode : public UMovieGraphNode
+class MOVIERENDERPIPELINECORE_API UMovieGraphEXRSequenceNode : public UMovieGraphSettingNode
 {
 	GENERATED_BODY()
 
 public:
 	UMovieGraphEXRSequenceNode() = default;
 
-	virtual TArray<FMovieGraphPinProperties> GetInputPinProperties() const override;
-	virtual TArray<FMovieGraphPinProperties> GetOutputPinProperties() const override;
-
 #if WITH_EDITOR
 	virtual FText GetNodeTitle(const bool bGetDescriptive = false) const override;
 	virtual FText GetMenuCategory() const override;
@@ -66,55 +40,15 @@ public:
 #endif
 };
 
-// TODO: Currently this node is restricted to just accepting string-based options.
-/** A node which creates a condition that selects from a set of input branches. */
-UCLASS()
-class MOVIERENDERPIPELINECORE_API UMovieGraphSelectNode : public UMovieGraphNode
-{
-	GENERATED_BODY()
-
-public:
-	UMovieGraphSelectNode() = default;
-
-	virtual TArray<FMovieGraphPinProperties> GetInputPinProperties() const override;
-	virtual TArray<FMovieGraphPinProperties> GetOutputPinProperties() const override;
-
-#if WITH_EDITOR
-	virtual FText GetNodeTitle(const bool bGetDescriptive = false) const override;
-	virtual FText GetMenuCategory() const override;
-	virtual FLinearColor GetNodeTitleColor() const override;
-	virtual FSlateIcon GetIconAndTint(FLinearColor& OutColor) const override;
-
-	//~ Begin UObject interface
-	virtual void PostEditChangeProperty(FPropertyChangedEvent& PropertyChangedEvent) override;
-	//~ End UObject interface
-#endif
-
-	/** The values of options which can be selected. */
-	UPROPERTY(EditAnywhere, Category = "General")
-	TArray<FString> SelectOptions;
-
-	/** The value of the option which has been selected. */
-	UPROPERTY(EditAnywhere, Category = "General")
-	FString SelectedOption;
-
-	/** A description of what this select is doing. */
-	UPROPERTY(EditAnywhere, Category = "General")
-	FString Description;
-};
-
 /** A node which configures anti-aliasing settings. */
 UCLASS()
-class MOVIERENDERPIPELINECORE_API UMovieGraphAntiAliasingNode : public UMovieGraphNode
+class MOVIERENDERPIPELINECORE_API UMovieGraphAntiAliasingNode : public UMovieGraphSettingNode
 {
 	GENERATED_BODY()
 
 public:
 	UMovieGraphAntiAliasingNode() = default;
 
-	virtual TArray<FMovieGraphPinProperties> GetInputPinProperties() const override;
-	virtual TArray<FMovieGraphPinProperties> GetOutputPinProperties() const override;
-
 #if WITH_EDITOR
 	virtual FText GetNodeTitle(const bool bGetDescriptive = false) const override;
 	virtual FText GetMenuCategory() const override;
diff --git a/Engine/Plugins/MovieScene/MovieRenderPipeline/Source/MovieRenderPipelineCore/Public/Graph/Nodes/MovieGraphGetCVarValueNode.h b/Engine/Plugins/MovieScene/MovieRenderPipeline/Source/MovieRenderPipelineCore/Public/Graph/Nodes/MovieGraphGetCVarValueNode.h
new file mode 100644
--- /dev/null
+++ b/Engine/Plugins/MovieScene/MovieRenderPipeline/Source/MovieRenderPipelineCore/Public/Graph/Nodes/MovieGraphGetCVarValueNode.h
@@ -0,0 +1,38 @@
+// Copyright Epic Games, Inc. All Rights Reserved.
+
+#pragma once
+
+#include "Graph/MovieGraphNode.h"
+
+#include "MovieGraphGetCVarValueNode.generated.h"
+
+/** A node which can get a specific console variable's value. */
+UCLASS()
+class MOVIERENDERPIPELINECORE_API UMovieGraphGetCVarValueNode : public UMovieGraphSettingNode
+{
+	GENERATED_BODY()
+public:
+	UMovieGraphGetCVarValueNode() = default;
+
+	virtual TArray<FMovieGraphPinProperties> GetOutputPinProperties() const override;
+	virtual FString GetNodeInstanceName() const override { return Name; }
+
+#if WITH_EDITOR
+	virtual FText GetNodeTitle(const bool bGetDescriptive = false) const override;
+	virtual FText GetMenuCategory() const override;
+	virtual FLinearColor GetNodeTitleColor() const override;
+	virtual FSlateIcon GetIconAndTint(FLinearColor& OutColor) const override;
+	
+	//~ Begin UObject interface
+	virtual void PostEditChangeProperty(FPropertyChangedEvent& PropertyChangedEvent) override;
+	//~ End UObject interface
+#endif
+
+public:
+	UPROPERTY(EditAnywhere, BlueprintReadWrite, Category = Overrides, meta = (InlineEditConditionToggle))
+	uint8 bOverride_Name : 1;
+
+	/** The name of the CVar that will have its value fetched. */
+	UPROPERTY(BlueprintReadWrite, EditAnywhere, Category = "Settings", meta = (EditCondition = "bOverride_Name"))
+	FString Name;
+};
\ No newline at end of file
diff --git a/Engine/Plugins/MovieScene/MovieRenderPipeline/Source/MovieRenderPipelineCore/Public/Graph/Nodes/MovieGraphGlobalGameOverrides.h b/Engine/Plugins/MovieScene/MovieRenderPipeline/Source/MovieRenderPipelineCore/Public/Graph/Nodes/MovieGraphGlobalGameOverrides.h
new file mode 100644
--- /dev/null
+++ b/Engine/Plugins/MovieScene/MovieRenderPipeline/Source/MovieRenderPipelineCore/Public/Graph/Nodes/MovieGraphGlobalGameOverrides.h
@@ -0,0 +1,117 @@
+// Copyright Epic Games, Inc. All Rights Reserved.
+
+#pragma once
+
+#include "Graph/MovieGraphNode.h"
+#include "MoviePipelineGameOverrideSetting.h"
+#include "MovieGraphGlobalGameOverrides.generated.h"
+
+/** A node which configures the global game overrides. */
+UCLASS()
+class MOVIERENDERPIPELINECORE_API UMovieGraphGlobalGameOverridesNode : public UMovieGraphSettingNode
+{
+	GENERATED_BODY()
+
+public:
+	UMovieGraphGlobalGameOverridesNode() = default;
+
+#if WITH_EDITOR
+	virtual FText GetNodeTitle(const bool bGetDescriptive = false) const override;
+	virtual FText GetMenuCategory() const override;
+	virtual FLinearColor GetNodeTitleColor() const override;
+	virtual FSlateIcon GetIconAndTint(FLinearColor& OutColor) const override;
+#endif
+
+public:
+	UPROPERTY(EditAnywhere, BlueprintReadWrite, Category = Overrides, meta = (InlineEditConditionToggle))
+	uint8 bOverride_GameModeOverride : 1;
+
+	UPROPERTY(EditAnywhere, BlueprintReadWrite, Category = Overrides, meta = (InlineEditConditionToggle))
+	uint8 bOverride_CinematicQualitySettings : 1;
+
+	UPROPERTY(EditAnywhere, BlueprintReadWrite, Category = Overrides, meta = (InlineEditConditionToggle))
+	uint8 bOverride_TextureStreaming : 1;
+
+	UPROPERTY(EditAnywhere, BlueprintReadWrite, Category = Overrides, meta = (InlineEditConditionToggle))
+	uint8 bOverride_UseLODZero : 1;
+
+	UPROPERTY(EditAnywhere, BlueprintReadWrite, Category = Overrides, meta = (InlineEditConditionToggle))
+	uint8 bOverride_DisableHLODs : 1;
+
+	UPROPERTY(EditAnywhere, BlueprintReadWrite, Category = Overrides, meta = (InlineEditConditionToggle))
+	uint8 bOverride_UseHighQualityShadows : 1;
+
+	UPROPERTY(EditAnywhere, BlueprintReadWrite, Category = Overrides, meta = (InlineEditConditionToggle))
+	uint8 bOverride_ShadowDistanceScale : 1;
+
+	UPROPERTY(EditAnywhere, BlueprintReadWrite, Category = Overrides, meta = (InlineEditConditionToggle))
+	uint8 bOverride_ShadowRadiusThreshold : 1;
+	
+	UPROPERTY(EditAnywhere, BlueprintReadWrite, Category = Overrides, meta = (InlineEditConditionToggle))
+	uint8 bOverride_OverrideViewDistanceScale : 1;
+
+	UPROPERTY(EditAnywhere, BlueprintReadWrite, Category = Overrides, meta = (InlineEditConditionToggle))
+	uint8 bOverride_ViewDistanceScale : 1;
+
+	UPROPERTY(EditAnywhere, BlueprintReadWrite, Category = Overrides, meta = (InlineEditConditionToggle))
+	uint8 bOverride_FlushGrassStreaming : 1;
+
+	UPROPERTY(EditAnywhere, BlueprintReadWrite, Category = Overrides, meta = (InlineEditConditionToggle))
+	uint8 bOverride_FlushStreamingManagers : 1;
+
+	UPROPERTY(EditAnywhere, BlueprintReadWrite, Category = Overrides, meta = (InlineEditConditionToggle))
+	uint8 bOverride_VirtualTextureFeedbackFactor : 1;
+
+	UPROPERTY(EditAnywhere, BlueprintReadWrite, Category = Overrides, meta = (InlineEditConditionToggle))
+	uint8 bOverride_DisableGPUTimeout : 1;
+	
+	/** Optional Game Mode to override the map's default game mode with. This can be useful if the game's normal mode displays UI elements or loading screens that you don't want captured. */
+	UPROPERTY(EditAnywhere, BlueprintReadWrite, Category = "Game", meta = (EditCondition = "bOverride_GameModeOverride"))
+	TSubclassOf<AGameModeBase> GameModeOverride;
+
+	/** If true, automatically set the engine to the Cinematic Scalability quality settings during render. See the Scalability Reference documentation for information on how to edit cvars to add/change default quality values.*/
+	UPROPERTY(EditAnywhere, BlueprintReadWrite, Category = "Game", meta = (EditCondition = "bOverride_CinematicQualitySettings"))
+	bool bCinematicQualitySettings;
+
+	/** Defines which If true, when using texture streaming fully load the required textures each frame instead of loading them in over time. This solves objects being blurry after camera cuts. */
+	UPROPERTY(EditAnywhere, BlueprintReadWrite, Category = "Rendering", meta = (EditCondition = "bOverride_TextureStreaming"))
+	EMoviePipelineTextureStreamingMethod TextureStreaming;
+
+	/** Should we try to use the highest quality LOD for meshes and particle systems regardless of distance? */
+	UPROPERTY(EditAnywhere, BlueprintReadWrite, Category = "Rendering", meta = (EditCondition = "bOverride_UseLODZero"))
+	bool bUseLODZero;
+
+	/** Should we disable Hierarchical LODs and instead use their real meshes regardless of distance? */
+	UPROPERTY(EditAnywhere, BlueprintReadWrite, Category = "Rendering", meta = (EditCondition = "bOverride_DisableHLODs"))
+	bool bDisableHLODs;
+
+	/** Should we override shadow-related CVars with some high quality preset settings? */
+	UPROPERTY(EditAnywhere, BlueprintReadWrite, Category = "Rendering", meta = (EditCondition = "bOverride_UseHighQualityShadows"))
+	bool bUseHighQualityShadows;
+
+	// TODO: This should also be disabled if bUseHighQualityShadows is false, but EditCondition doesn't work with that *plus* the InlineEditConditionToggle
+	/** Scalability option to trade shadow distance versus performance for directional lights  */
+	UPROPERTY(EditAnywhere, BlueprintReadWrite, Category = "Rendering", meta = (EditCondition = "bOverride_ShadowDistanceScale", ClampMin=0.1, UIMin=0.1, UIMax=10))
+	int32 ShadowDistanceScale;
+
+	// TODO: This should also be disabled if bUseHighQualityShadows is false, but EditCondition doesn't work with that *plus* the InlineEditConditionToggle
+	/** Cull shadow casters if they are too small, value is the minimal screen space bounding sphere radius */
+	UPROPERTY(EditAnywhere, BlueprintReadWrite, Category = "Rendering", meta = (EditCondition = "bOverride_ShadowRadiusThreshold", UIMin=0.001, ClampMin=0.001))
+	float ShadowRadiusThreshold;
+
+	/** Controls the view distance scale. A primitive's MaxDrawDistance is scaled by this value. */
+	UPROPERTY(EditAnywhere, BlueprintReadWrite, Category = "Rendering", meta = (EditCondition = "bOverride_ViewDistanceScale"))
+	int32 ViewDistanceScale;
+	
+	/** Flushing grass streaming (combined with override view distance scale) prevents visible pop-in/culling of grace instances. */
+	UPROPERTY(EditAnywhere, BlueprintReadWrite, Category = "Rendering", meta = (EditCondition = "bOverride_FlushGrassStreaming"))
+	bool bFlushGrassStreaming;
+
+	/** Experimental. If true flush the streaming managers (Texture Streaming) each frame. Allows Texture Streaming to not have visible pop-in in final frames. */
+	UPROPERTY(EditAnywhere, BlueprintReadWrite, Category = "Rendering", meta = (EditCondition = "bOverride_FlushStreamingManagers"))
+	bool bFlushStreamingManagers;
+
+	/** The virtual texture feedback resolution factor. A lower factor will increase virtual texture feedback resolution. */
+	UPROPERTY(EditAnywhere, BlueprintReadWrite, Category = "Rendering", meta = (EditCondition = "bOverride_VirtualTextureFeedbackFactor"))
+	int32 VirtualTextureFeedbackFactor;
+};
\ No newline at end of file
diff --git a/Engine/Plugins/MovieScene/MovieRenderPipeline/Source/MovieRenderPipelineCore/Public/Graph/Nodes/MovieGraphRenderLayerNode.h b/Engine/Plugins/MovieScene/MovieRenderPipeline/Source/MovieRenderPipelineCore/Public/Graph/Nodes/MovieGraphRenderLayerNode.h
--- a/Engine/Plugins/MovieScene/MovieRenderPipeline/Source/MovieRenderPipelineCore/Public/Graph/Nodes/MovieGraphRenderLayerNode.h
+++ b/Engine/Plugins/MovieScene/MovieRenderPipeline/Source/MovieRenderPipelineCore/Public/Graph/Nodes/MovieGraphRenderLayerNode.h
@@ -14,8 +14,6 @@ public:
 
 	FString GetRenderLayerName() const { return LayerName; }
 
-	virtual TArray<FPropertyBagPropertyDesc> GetDynamicPropertyDescriptions() const override;
-
 #if WITH_EDITOR
 	virtual FText GetNodeTitle(const bool bGetDescriptive = false) const override;
 	virtual FText GetMenuCategory() const override;
diff --git a/Engine/Plugins/MovieScene/MovieRenderPipeline/Source/MovieRenderPipelineCore/Public/Graph/Nodes/MovieGraphSelectNode.h b/Engine/Plugins/MovieScene/MovieRenderPipeline/Source/MovieRenderPipelineCore/Public/Graph/Nodes/MovieGraphSelectNode.h
new file mode 100644
--- /dev/null
+++ b/Engine/Plugins/MovieScene/MovieRenderPipeline/Source/MovieRenderPipelineCore/Public/Graph/Nodes/MovieGraphSelectNode.h
@@ -0,0 +1,45 @@
+// Copyright Epic Games, Inc. All Rights Reserved.
+
+#pragma once
+
+#include "Graph/MovieGraphNode.h"
+
+#include "MovieGraphSelectNode.generated.h"
+
+// TODO: Currently this node is restricted to just accepting string-based options.
+/** A node which creates a condition that selects from a set of input branches. */
+UCLASS()
+class MOVIERENDERPIPELINECORE_API UMovieGraphSelectNode : public UMovieGraphNode
+{
+	GENERATED_BODY()
+
+public:
+	UMovieGraphSelectNode() = default;
+
+	virtual TArray<FMovieGraphPinProperties> GetInputPinProperties() const override;
+	virtual TArray<FMovieGraphPinProperties> GetOutputPinProperties() const override;
+	virtual TArray<UMovieGraphPin*> EvaluatePinsToFollow(FMovieGraphEvaluationContext& InContext) const override;
+
+#if WITH_EDITOR
+	virtual FText GetNodeTitle(const bool bGetDescriptive = false) const override;
+	virtual FText GetMenuCategory() const override;
+	virtual FLinearColor GetNodeTitleColor() const override;
+	virtual FSlateIcon GetIconAndTint(FLinearColor& OutColor) const override;
+
+	//~ Begin UObject interface
+	virtual void PostEditChangeProperty(FPropertyChangedEvent& PropertyChangedEvent) override;
+	//~ End UObject interface
+	#endif
+
+	/** The values of options which can be selected. */
+	UPROPERTY(EditAnywhere, Category = "General")
+	TArray<FString> SelectOptions;
+
+	/** The value of the option which has been selected. */
+	UPROPERTY(EditAnywhere, Category = "General")
+	FString SelectedOption;
+
+	/** A description of what this select is doing. */
+	UPROPERTY(EditAnywhere, Category = "General")
+	FString Description;
+};
\ No newline at end of file
diff --git a/Engine/Plugins/MovieScene/MovieRenderPipeline/Source/MovieRenderPipelineCore/Public/Graph/Nodes/MovieGraphSetCVarValueNode.h b/Engine/Plugins/MovieScene/MovieRenderPipeline/Source/MovieRenderPipelineCore/Public/Graph/Nodes/MovieGraphSetCVarValueNode.h
new file mode 100644
--- /dev/null
+++ b/Engine/Plugins/MovieScene/MovieRenderPipeline/Source/MovieRenderPipelineCore/Public/Graph/Nodes/MovieGraphSetCVarValueNode.h
@@ -0,0 +1,44 @@
+// Copyright Epic Games, Inc. All Rights Reserved.
+
+#pragma once
+
+#include "Graph/MovieGraphNode.h"
+
+#include "MovieGraphSetCVarValueNode.generated.h"
+
+/** A node which can set a specific console variable's value. */
+UCLASS()
+class MOVIERENDERPIPELINECORE_API UMovieGraphSetCVarValueNode : public UMovieGraphSettingNode
+{
+	GENERATED_BODY()
+public:
+	UMovieGraphSetCVarValueNode() = default;
+
+	virtual FString GetNodeInstanceName() const override { return Name; }
+
+#if WITH_EDITOR
+	virtual FText GetNodeTitle(const bool bGetDescriptive = false) const override;
+	virtual FText GetMenuCategory() const override;
+	virtual FLinearColor GetNodeTitleColor() const override;
+	virtual FSlateIcon GetIconAndTint(FLinearColor& OutColor) const override;
+
+	//~ Begin UObject interface
+	virtual void PostEditChangeProperty(FPropertyChangedEvent& PropertyChangedEvent) override;
+	//~ End UObject interface
+#endif
+
+public:
+	UPROPERTY(EditAnywhere, BlueprintReadWrite, Category = Overrides, meta = (InlineEditConditionToggle))
+	uint8 bOverride_Name : 1;
+
+	UPROPERTY(EditAnywhere, BlueprintReadWrite, Category = Overrides, meta = (InlineEditConditionToggle))
+	uint8 bOverride_Value : 1;
+
+	/** The name of the CVar having its value set. */
+	UPROPERTY(BlueprintReadWrite, EditAnywhere, Category = "Settings", meta = (EditCondition = "bOverride_Name"))
+	FString Name;
+
+	/** The new value of the CVar. */
+	UPROPERTY(BlueprintReadWrite, EditAnywhere, Category = "Settings", meta = (EditCondition = "bOverride_Value"))
+	FString Value;
+};
\ No newline at end of file
diff --git a/Engine/Plugins/MovieScene/MovieRenderPipeline/Source/MovieRenderPipelineEditor/Private/Customizations/Graph/GraphNodeCustomization.h b/Engine/Plugins/MovieScene/MovieRenderPipeline/Source/MovieRenderPipelineEditor/Private/Customizations/Graph/GraphNodeCustomization.h
new file mode 100644
--- /dev/null
+++ b/Engine/Plugins/MovieScene/MovieRenderPipeline/Source/MovieRenderPipelineEditor/Private/Customizations/Graph/GraphNodeCustomization.h
@@ -0,0 +1,57 @@
+// Copyright Epic Games, Inc. All Rights Reserved.
+
+#pragma once
+
+#include "DetailLayoutBuilder.h"
+#include "IDetailCustomization.h"
+#include "Graph/MovieGraphConfig.h"
+#include "PropertyHandle.h"
+
+#define LOCTEXT_NAMESPACE "MoviePipelineEditor"
+
+/** Customize how graph nodes appear in the details panel. */
+class FGraphNodeCustomization : public IDetailCustomization
+{
+public:
+	static TSharedRef<IDetailCustomization> MakeInstance()
+	{
+		return MakeShared<FGraphNodeCustomization>();
+	}
+
+protected:
+	//~ Begin IDetailCustomization interface
+	virtual void CustomizeDetails(const TSharedPtr<IDetailLayoutBuilder>& DetailBuilder) override
+	{
+		CustomizeDetails(*DetailBuilder);
+	}
+
+	virtual void CustomizeDetails(IDetailLayoutBuilder& DetailBuilder) override
+	{
+		TArray<TWeakObjectPtr<UObject>> ObjectsBeingCustomized;
+		DetailBuilder.GetObjectsBeingCustomized(ObjectsBeingCustomized);
+
+		// Hide the "Properties" category, which houses the dynamic properties, if there are no dynamic properties
+		// on the node. An empty category in the details panel is messy/confusing.
+		for (const TWeakObjectPtr<UObject>& CustomizedObject : ObjectsBeingCustomized)
+		{
+			if (const UMovieGraphNode* Node = Cast<UMovieGraphNode>(CustomizedObject))
+			{
+				if (Node->GetDynamicPropertyDescriptions().IsEmpty())
+				{
+					DetailBuilder.HideCategory("Properties");
+					break;
+				}
+			}
+		}
+
+		// Also hide the "Properties" category if more than one object is selected. There needs to be some more work done
+		// in property bags to be able to handle this condition without crashing in many cases.
+		if (ObjectsBeingCustomized.Num() > 1)
+		{
+			DetailBuilder.HideCategory("Properties");
+		}
+	}
+	//~ End IDetailCustomization interface
+};
+
+#undef LOCTEXT_NAMESPACE
\ No newline at end of file
diff --git a/Engine/Plugins/MovieScene/MovieRenderPipeline/Source/MovieRenderPipelineEditor/Private/Graph/MovieEdGraphNode.cpp b/Engine/Plugins/MovieScene/MovieRenderPipeline/Source/MovieRenderPipelineEditor/Private/Graph/MovieEdGraphNode.cpp
--- a/Engine/Plugins/MovieScene/MovieRenderPipeline/Source/MovieRenderPipelineEditor/Private/Graph/MovieEdGraphNode.cpp
+++ b/Engine/Plugins/MovieScene/MovieRenderPipeline/Source/MovieRenderPipelineEditor/Private/Graph/MovieEdGraphNode.cpp
@@ -48,7 +48,7 @@ void UMoviePipelineEdGraphNodeBase::PostTransacted(const FTransactionObjectEvent
 	}
 }
 
-FEdGraphPinType UMoviePipelineEdGraphNodeBase::GetPinType(const UMovieGraphPin* InPin)
+FEdGraphPinType UMoviePipelineEdGraphNodeBase::GetPinType(EMovieGraphValueType ValueType, bool bIsBranch)
 {
 	FEdGraphPinType EdPinType;
 	EdPinType.ResetToDefaults();
@@ -57,13 +57,13 @@ FEdGraphPinType UMoviePipelineEdGraphNodeBase::GetPinType(const UMovieGraphPin*
 	EdPinType.PinSubCategory = NAME_None;
 
 	// Special case for branch pins
-	if (InPin->Properties.bIsBranch)
+	if (bIsBranch)
 	{
 		EdPinType.PinCategory = UMovieGraphSchema::PC_Branch;
 		return EdPinType;
 	}
 
-	switch (InPin->Properties.Type)
+	switch (ValueType)
 	{
 	case EMovieGraphValueType::Bool:
 		EdPinType.PinCategory = UMovieGraphSchema::PC_Boolean;
@@ -118,6 +118,11 @@ FEdGraphPinType UMoviePipelineEdGraphNodeBase::GetPinType(const UMovieGraphPin*
 	return EdPinType;
 }
 
+FEdGraphPinType UMoviePipelineEdGraphNodeBase::GetPinType(const UMovieGraphPin* InPin)
+{
+	return GetPinType(InPin->Properties.Type, InPin->Properties.bIsBranch);
+}
+
 void UMoviePipelineEdGraphNodeBase::UpdatePosition() const
 {
 	if (RuntimeNode)
diff --git a/Engine/Plugins/MovieScene/MovieRenderPipeline/Source/MovieRenderPipelineEditor/Private/Graph/MovieEdGraphNode.h b/Engine/Plugins/MovieScene/MovieRenderPipeline/Source/MovieRenderPipelineEditor/Private/Graph/MovieEdGraphNode.h
--- a/Engine/Plugins/MovieScene/MovieRenderPipeline/Source/MovieRenderPipelineEditor/Private/Graph/MovieEdGraphNode.h
+++ b/Engine/Plugins/MovieScene/MovieRenderPipeline/Source/MovieRenderPipelineEditor/Private/Graph/MovieEdGraphNode.h
@@ -33,6 +33,9 @@ public:
 	virtual void GetNodeContextMenuActions(UToolMenu* Menu, UGraphNodeContextMenuContext* Context) const override;
 	//~ End UEdGraphNode Interface
 
+	static FEdGraphPinType GetPinType(EMovieGraphValueType ValueType, bool bIsBranch);
+	static FEdGraphPinType GetPinType(const UMovieGraphPin* InPin);
+
 	// Called after PrepareForCopying to restore changes made during preparation. 
 	virtual void PostCopy();
 	virtual void PostPaste();
@@ -47,8 +50,6 @@ protected:
 
 	void CreatePins(const TArray<UMovieGraphPin*>& InInputPins, const TArray<UMovieGraphPin*>& InOutputPins);
 
-	static FEdGraphPinType GetPinType(const UMovieGraphPin* InPin);
-
 	/** Recreate the pins on this node, discarding all existing pins. */
 	void ReconstructPins();
 
diff --git a/Engine/Plugins/MovieScene/MovieRenderPipeline/Source/MovieRenderPipelineEditor/Private/Graph/MovieGraphAssetToolkit.cpp b/Engine/Plugins/MovieScene/MovieRenderPipeline/Source/MovieRenderPipelineEditor/Private/Graph/MovieGraphAssetToolkit.cpp
--- a/Engine/Plugins/MovieScene/MovieRenderPipeline/Source/MovieRenderPipelineEditor/Private/Graph/MovieGraphAssetToolkit.cpp
+++ b/Engine/Plugins/MovieScene/MovieRenderPipeline/Source/MovieRenderPipelineEditor/Private/Graph/MovieGraphAssetToolkit.cpp
@@ -3,6 +3,7 @@
 #include "MovieGraphAssetToolkit.h"
 
 #include "Customizations/Graph/GraphMemberCustomization.h"
+#include "Customizations/Graph/GraphNodeCustomization.h"
 #include "Framework/Commands/GenericCommands.h"
 #include "Graph/MovieGraphConfig.h"
 #include "MovieEdGraphNode.h"
@@ -225,6 +226,10 @@ TSharedRef<SDockTab> FMovieGraphAssetToolkit::SpawnTab_RenderGraphDetails(const
 	SelectedGraphObjectsDetailsWidget->RegisterInstancedCustomPropertyLayout(
 		UMovieGraphMember::StaticClass(),
 		FOnGetDetailCustomizationInstance::CreateStatic(&FGraphMemberCustomization::MakeInstance));
+
+	SelectedGraphObjectsDetailsWidget->RegisterInstancedCustomPropertyLayout(
+		UMovieGraphNode::StaticClass(),
+		FOnGetDetailCustomizationInstance::CreateStatic(&FGraphNodeCustomization::MakeInstance));
 	
 	return SNew(SDockTab)
 		.TabColorScale(GetTabColorScale())
diff --git a/Engine/Plugins/MovieScene/MovieRenderPipeline/Source/MovieRenderPipelineEditor/Private/Graph/SMovieGraphMembersTabContent.cpp b/Engine/Plugins/MovieScene/MovieRenderPipeline/Source/MovieRenderPipelineEditor/Private/Graph/SMovieGraphMembersTabContent.cpp
--- a/Engine/Plugins/MovieScene/MovieRenderPipeline/Source/MovieRenderPipelineEditor/Private/Graph/SMovieGraphMembersTabContent.cpp
+++ b/Engine/Plugins/MovieScene/MovieRenderPipeline/Source/MovieRenderPipelineEditor/Private/Graph/SMovieGraphMembersTabContent.cpp
@@ -6,6 +6,7 @@
 #include "Framework/Commands/GenericCommands.h"
 #include "Graph/MovieGraphConfig.h"
 #include "Graph/MovieGraphSchema.h"
+#include "MovieEdGraphNode.h"
 #include "SGraphActionMenu.h"
 #include "Toolkits/AssetEditorToolkit.h"
 
@@ -53,6 +54,7 @@ void SMovieGraphMembersTabContent::Construct(const FArguments& InArgs)
 		.OnActionSelected(OnActionSelected)
 		.AutoExpandActionMenu(true)
 		.AlphaSortItems(false)
+		.OnCreateWidgetForAction(this, &SMovieGraphMembersTabContent::CreateActionWidget)
 		.OnCollectStaticSections(this, &SMovieGraphMembersTabContent::CollectStaticSections)
 		.OnContextMenuOpening(this, &SMovieGraphMembersTabContent::OnContextMenuOpening)
 		.OnGetSectionTitle(this, &SMovieGraphMembersTabContent::GetSectionTitle)
@@ -63,6 +65,48 @@ void SMovieGraphMembersTabContent::Construct(const FArguments& InArgs)
 	];
 }
 
+TSharedRef<SWidget> SMovieGraphMembersTabContent::CreateActionWidget(FCreateWidgetForActionData* InCreateData) const
+{
+	// For variables, show an icon w/ the color representing the variable's type (in addition to the variable name)
+	if (InCreateData->Action->GetSectionID() == static_cast<uint32>(EActionSection::Variables))
+	{
+		const FMovieGraphSchemaAction_NewVariableNode* VariableAction =
+			static_cast<FMovieGraphSchemaAction_NewVariableNode*>(InCreateData->Action.Get());
+		const UMovieGraphVariable* Variable = Cast<UMovieGraphVariable>(VariableAction->ActionTarget);
+		const FEdGraphPinType PinType = Variable
+			? UMoviePipelineEdGraphNodeBase::GetPinType(Variable->GetValueType(), false)
+			: FEdGraphPinType();
+		const FLinearColor PinColor = CurrentGraph->PipelineEdGraph->GetSchema()->GetPinTypeColor(PinType);
+		
+		return
+			SNew(SHorizontalBox)
+			+ SHorizontalBox::Slot()
+			.AutoWidth()
+			.Padding(0, 0, 5, 0)
+			[
+				SNew(SImage)
+				.Image(FAppStyle::GetBrush("Kismet.AllClasses.VariableIcon"))
+				.ColorAndOpacity(PinColor)
+			]
+			
+			+ SHorizontalBox::Slot()
+			.FillWidth(1.f)
+			[
+				SNew(STextBlock)
+				.Text(InCreateData->Action->GetMenuDescription())
+			];	
+	}
+	
+	return
+		SNew(SHorizontalBox)
+		+ SHorizontalBox::Slot()
+		.AutoWidth()
+		[
+			SNew(STextBlock)
+			.Text(InCreateData->Action->GetMenuDescription())
+		];
+}
+
 void SMovieGraphMembersTabContent::ClearSelection() const
 {
 	if (ActionMenu.IsValid())
diff --git a/Engine/Plugins/MovieScene/MovieRenderPipeline/Source/MovieRenderPipelineEditor/Private/Graph/SMovieGraphMembersTabContent.h b/Engine/Plugins/MovieScene/MovieRenderPipeline/Source/MovieRenderPipelineEditor/Private/Graph/SMovieGraphMembersTabContent.h
--- a/Engine/Plugins/MovieScene/MovieRenderPipeline/Source/MovieRenderPipelineEditor/Private/Graph/SMovieGraphMembersTabContent.h
+++ b/Engine/Plugins/MovieScene/MovieRenderPipeline/Source/MovieRenderPipelineEditor/Private/Graph/SMovieGraphMembersTabContent.h
@@ -2,6 +2,7 @@
 
 #pragma once
 
+#include "SGraphActionMenu.h"
 #include "Widgets/Input/SButton.h"
 #include "Widgets/SCompoundWidget.h"
 
@@ -37,6 +38,9 @@ public:
 	
 	void Construct(const FArguments& InArgs);
 
+	/** Creates the widget for an action within the action menu. */
+	TSharedRef<SWidget> CreateActionWidget(FCreateWidgetForActionData* CreateWidgetForActionData) const;
+
 	/** Resets the selected members in the UI. */
 	void ClearSelection() const;
 
