MoviePipeline: Fixed a bug that caused MoviePipelineQueue assets to not be loaded into the main MRQ queue UI when they're double clicked in the content browser.

The MoviePipelineQueueSubsystem has a new method, LoadQueue(), which largely contains code that was ripped out of SMoviePipelineQueuePanel. The new AssetDefinition for queue assets, as well as the queue panel, can now call this shared code to load in a new queue.

diff --git a/Engine/Plugins/MovieScene/MovieRenderPipeline/Source/MovieRenderPipelineEditor/Private/AssetDefinition_PipelineQueue.cpp b/Engine/Plugins/MovieScene/MovieRenderPipeline/Source/MovieRenderPipelineEditor/Private/AssetDefinition_PipelineQueue.cpp
new file mode 100644
--- /dev/null
+++ b/Engine/Plugins/MovieScene/MovieRenderPipeline/Source/MovieRenderPipelineEditor/Private/AssetDefinition_PipelineQueue.cpp
@@ -0,0 +1,19 @@
+// Copyright Epic Games, Inc. All Rights Reserved.
+
+#include "AssetDefinition_PipelineQueue.h"
+
+#include "Editor.h"
+#include "MoviePipelineQueueSubsystem.h"
+
+EAssetCommandResult UAssetDefinition_PipelineQueue::OpenAssets(const FAssetOpenArgs& OpenArgs) const
+{
+	if (UMoviePipelineQueue* ValidQueue = OpenArgs.LoadFirstValid<UMoviePipelineQueue>())
+	{
+		UMoviePipelineQueueSubsystem* Subsystem = GEditor->GetEditorSubsystem<UMoviePipelineQueueSubsystem>();
+		check(Subsystem);
+
+		Subsystem->LoadQueue(ValidQueue);
+	}
+
+	return EAssetCommandResult::Handled;
+}
\ No newline at end of file
diff --git a/Engine/Plugins/MovieScene/MovieRenderPipeline/Source/MovieRenderPipelineEditor/Private/AssetDefinition_PipelineQueue.h b/Engine/Plugins/MovieScene/MovieRenderPipeline/Source/MovieRenderPipelineEditor/Private/AssetDefinition_PipelineQueue.h
new file mode 100644
--- /dev/null
+++ b/Engine/Plugins/MovieScene/MovieRenderPipeline/Source/MovieRenderPipelineEditor/Private/AssetDefinition_PipelineQueue.h
@@ -0,0 +1,27 @@
+// Copyright Epic Games, Inc. All Rights Reserved.
+
+#pragma once
+
+#include "MoviePipelineQueue.h"
+#include "AssetDefinitionDefault.h"
+
+#include "AssetDefinition_PipelineQueue.generated.h"
+
+UCLASS()
+class UAssetDefinition_PipelineQueue : public UAssetDefinitionDefault
+{
+	GENERATED_BODY()
+
+public:
+	// UAssetDefinition Begin
+	virtual FText GetAssetDisplayName() const override { return NSLOCTEXT("AssetTypeActions", "AssetTypeActions_PipelineQueue", "Movie Pipeline Queue"); }
+	virtual FLinearColor GetAssetColor() const override { return FLinearColor(FColor(78, 40, 165)); }
+	virtual TSoftClassPtr<UObject> GetAssetClass() const override { return UMoviePipelineQueue::StaticClass(); }
+	virtual TConstArrayView<FAssetCategoryPath> GetAssetCategories() const override
+	{
+		static const auto Categories = { EAssetCategoryPaths::Cinematics };
+		return Categories;
+	}
+	virtual EAssetCommandResult OpenAssets(const FAssetOpenArgs& OpenArgs) const override;
+	// UAssetDefinition End
+};
diff --git a/Engine/Plugins/MovieScene/MovieRenderPipeline/Source/MovieRenderPipelineEditor/Private/MoviePipelineQueueSubsystem.cpp b/Engine/Plugins/MovieScene/MovieRenderPipeline/Source/MovieRenderPipelineEditor/Private/MoviePipelineQueueSubsystem.cpp
--- a/Engine/Plugins/MovieScene/MovieRenderPipeline/Source/MovieRenderPipelineEditor/Private/MoviePipelineQueueSubsystem.cpp
+++ b/Engine/Plugins/MovieScene/MovieRenderPipeline/Source/MovieRenderPipelineEditor/Private/MoviePipelineQueueSubsystem.cpp
@@ -1,11 +1,93 @@
 // Copyright Epic Games, Inc. All Rights Reserved.
 
 #include "MoviePipelineQueueSubsystem.h"
+
+#include "Framework/Notifications/NotificationManager.h"
 #include "LevelSequenceEditorModule.h"
+#include "Misc/MessageDialog.h"
 #include "Modules/ModuleManager.h"
+#include "MoviePipelineBlueprintLibrary.h"
+#include "Widgets/Notifications/SNotificationList.h"
 
 #include UE_INLINE_GENERATED_CPP_BY_NAME(MoviePipelineQueueSubsystem)
 
+#define LOCTEXT_NAMESPACE "MoviePipelineQueueSubsystem"
+
+bool UMoviePipelineQueueSubsystem::LoadQueue(UMoviePipelineQueue* QueueToLoad, const bool bPromptOnReplacingDirtyQueue)
+{
+	if (!QueueToLoad)
+	{
+		return false;
+	}
+
+	UMoviePipelineQueue* Queue = GetQueue();
+	if (!Queue)
+	{
+		return false;
+	}
+
+	if (bPromptOnReplacingDirtyQueue && IsQueueDirty())
+	{
+		const FText TitleText = LOCTEXT("UnsavedQueueWarningTitle", "Unsaved Changes to Queue");
+		const FText MessageText = LOCTEXT("UnsavedQueueWarningMessage", "The changes made to the current queue will be lost by importing another queue. Do you want to continue with this import?");
+
+		// If the application is in unattended mode, auto-accept the dialog
+		constexpr EAppReturnType::Type DefaultReturn = EAppReturnType::Yes;
+		
+		if (FMessageDialog::Open(EAppMsgType::YesNo, DefaultReturn, MessageText, TitleText) == EAppReturnType::No)
+		{
+			return false;
+		}
+	}
+	
+	Queue->CopyFrom(QueueToLoad);
+	Queue->SetQueueOrigin(QueueToLoad);
+	Queue->SetIsDirty(false);
+
+	// Update the shot list in case the stored queue being copied is out of date with the sequence
+	for (UMoviePipelineExecutorJob* Job : Queue->GetJobs())
+	{
+		if (ULevelSequence* LoadedSequence = Cast<ULevelSequence>(Job->Sequence.TryLoad()))
+		{
+			bool bShotsChanged = false;
+			UMoviePipelineBlueprintLibrary::UpdateJobShotListFromSequence(LoadedSequence, Job, bShotsChanged);
+
+			if (bShotsChanged)
+			{
+				FNotificationInfo Info(LOCTEXT("QueueShotsUpdated", "Shots have changed since the queue was saved, please resave the queue"));
+				Info.ExpireDuration = 5.0f;
+				FSlateNotificationManager::Get().AddNotification(Info)->SetCompletionState(SNotificationItem::CS_Fail);
+			}
+		}
+	}
+
+	OnQueueLoaded.Broadcast();
+
+	return true;
+}
+
+bool UMoviePipelineQueueSubsystem::IsQueueDirty() const
+{
+	const UMoviePipelineQueue* Queue = GetQueue();
+	if (!Queue)
+	{
+		return false;
+	}
+
+	// The queue is considered dirty if the current queue has no origin (ie, it has never been saved) or it has been
+	// modified since it was loaded. We skip the no-origin check if there are no jobs (to avoid triggering on an
+	// empty, first-load queue), but we don't skip the IsDirty check because a queue that had jobs and then
+	// removed them will be dirty.
+	const bool bShouldCheckQueueOrigin = Queue->GetJobs().Num() > 0;
+	bool bHasNoQueueOrigin = false;
+	if (bShouldCheckQueueOrigin)
+	{
+		bHasNoQueueOrigin = !Queue->GetQueueOrigin();
+	}
+	
+	return bHasNoQueueOrigin || Queue->IsDirty();
+}
+
 UMoviePipelineExecutorBase* UMoviePipelineQueueSubsystem::RenderQueueWithExecutor(TSubclassOf<UMoviePipelineExecutorBase> InExecutorType)
 {
 	if(!ensureMsgf(!IsRendering(), TEXT("RenderQueueWithExecutor cannot be called while already rendering!")))
@@ -62,4 +144,4 @@ void UMoviePipelineQueueSubsystem::OnExecutorFinished(UMoviePipelineExecutorBase
 	ActiveExecutor = nullptr;
 }
 
-
+#undef LOCTEXT_NAMESPACE	// MoviePipelineQueueSubsystem
\ No newline at end of file
diff --git a/Engine/Plugins/MovieScene/MovieRenderPipeline/Source/MovieRenderPipelineEditor/Private/Widgets/SMoviePipelineQueueEditor.cpp b/Engine/Plugins/MovieScene/MovieRenderPipeline/Source/MovieRenderPipelineEditor/Private/Widgets/SMoviePipelineQueueEditor.cpp
--- a/Engine/Plugins/MovieScene/MovieRenderPipeline/Source/MovieRenderPipelineEditor/Private/Widgets/SMoviePipelineQueueEditor.cpp
+++ b/Engine/Plugins/MovieScene/MovieRenderPipeline/Source/MovieRenderPipelineEditor/Private/Widgets/SMoviePipelineQueueEditor.cpp
@@ -1165,6 +1165,11 @@ void SMoviePipelineQueueEditor::Construct(const FArguments& InArgs)
 
 	// When undo occurs, get a notification so we can make sure our view is up to date
 	GEditor->RegisterForUndo(this);
+
+	// React to when a new queue is loaded via the subsystem
+	UMoviePipelineQueueSubsystem* Subsystem = GEditor->GetEditorSubsystem<UMoviePipelineQueueSubsystem>();
+	check(Subsystem);
+	Subsystem->OnQueueLoaded.AddSP(this, &SMoviePipelineQueueEditor::OnQueueLoaded);
 }
 
 
@@ -1588,4 +1593,21 @@ void SMoviePipelineQueueEditor::OnJobSelectionChanged_Impl(TSharedPtr<IMoviePipe
 	OnJobSelectionChanged.ExecuteIfBound(SelectedJobs);
 }
 
+void SMoviePipelineQueueEditor::OnQueueLoaded()
+{
+	const UMoviePipelineQueueSubsystem* Subsystem = GEditor->GetEditorSubsystem<UMoviePipelineQueueSubsystem>();
+	check(Subsystem);
+	
+	// Automatically select the first job in the queue
+	TArray<UMoviePipelineExecutorJob*> Jobs;
+	if (Subsystem->GetQueue()->GetJobs().Num() > 0)
+	{
+		Jobs.Add(Subsystem->GetQueue()->GetJobs()[0]);
+	}
+	
+	// Go through the UI so it updates the UI selection too and then this will loop back
+	// around to OnSelectionChanged to update ourself.
+	SetSelectedJobs(Jobs);
+}
+
 #undef LOCTEXT_NAMESPACE
diff --git a/Engine/Plugins/MovieScene/MovieRenderPipeline/Source/MovieRenderPipelineEditor/Private/Widgets/SMoviePipelineQueueEditor.h b/Engine/Plugins/MovieScene/MovieRenderPipeline/Source/MovieRenderPipelineEditor/Private/Widgets/SMoviePipelineQueueEditor.h
--- a/Engine/Plugins/MovieScene/MovieRenderPipeline/Source/MovieRenderPipelineEditor/Private/Widgets/SMoviePipelineQueueEditor.h
+++ b/Engine/Plugins/MovieScene/MovieRenderPipeline/Source/MovieRenderPipelineEditor/Private/Widgets/SMoviePipelineQueueEditor.h
@@ -88,6 +88,9 @@ private:
 	void SetSelectedJobs_Impl(const TArray<UMoviePipelineExecutorJob*>& InJobs);
 	void OnJobSelectionChanged_Impl(TSharedPtr<IMoviePipelineQueueTreeItem> TreeItem, ESelectInfo::Type SelectInfo);
 
+	/** Called when a new queue is loaded in the queue subsystem. */
+	void OnQueueLoaded();
+
 private:
 	TArray<TSharedPtr<IMoviePipelineQueueTreeItem>> RootNodes;
 	TSharedPtr<STreeView<TSharedPtr<IMoviePipelineQueueTreeItem>>> TreeView;
diff --git a/Engine/Plugins/MovieScene/MovieRenderPipeline/Source/MovieRenderPipelineEditor/Private/Widgets/SMoviePipelineQueuePanel.cpp b/Engine/Plugins/MovieScene/MovieRenderPipeline/Source/MovieRenderPipelineEditor/Private/Widgets/SMoviePipelineQueuePanel.cpp
--- a/Engine/Plugins/MovieScene/MovieRenderPipeline/Source/MovieRenderPipelineEditor/Private/Widgets/SMoviePipelineQueuePanel.cpp
+++ b/Engine/Plugins/MovieScene/MovieRenderPipeline/Source/MovieRenderPipelineEditor/Private/Widgets/SMoviePipelineQueuePanel.cpp
@@ -8,7 +8,6 @@
 #include "SMoviePipelineConfigPanel.h"
 #include "MovieRenderPipelineSettings.h"
 #include "MovieRenderPipelineStyle.h"
-#include "MoviePipelineBlueprintLibrary.h"
 #include "Sections/MovieSceneCinematicShotSection.h"
 #include "MoviePipelineQueue.h"
 #include "MoviePipelinePrimaryConfig.h"
@@ -25,10 +24,8 @@
 #include "Widgets/SWindow.h"
 #include "Styling/AppStyle.h"
 #include "Framework/Application/SlateApplication.h"
-#include "Framework/Notifications/NotificationManager.h"
 #include "Widgets/Layout/SSplitter.h"
 #include "Widgets/Layout/SWidgetSwitcher.h"
-#include "Widgets/Notifications/SNotificationList.h"
 #include "Styling/SlateIconFinder.h"
 #include "Widgets/Images/SImage.h"
 #include "Framework/MultiBox/MultiBoxBuilder.h"
@@ -742,7 +739,8 @@ void SMoviePipelineQueuePanel::SaveTransientQueueToAsset(UMoviePipelineQueue* De
 	FEditorFileUtils::EPromptReturnCode PromptReturnCode = FEditorFileUtils::PromptForCheckoutAndSave({ DestinationQueue->GetPackage() }, bCheckDirty, bPromptToSave);
 	if (PromptReturnCode == FEditorFileUtils::EPromptReturnCode::PR_Success)
 	{
-		LoadQueue(DestinationQueue);
+		constexpr bool bPromptOnReplacingDirtyQueue = false;
+		Subsystem->LoadQueue(DestinationQueue, bPromptOnReplacingDirtyQueue);
 	}
 }
 
@@ -750,86 +748,19 @@ void SMoviePipelineQueuePanel::OnImportSavedQueueAsset(const FAssetData& InPrese
 {
 	FSlateApplication::Get().DismissAllMenus();
 
-	if (IsQueueDirty())
-	{
-		const FText TitleText = LOCTEXT("UnsavedQueueWarningTitle", "Unsaved Changes to Queue");
-		const FText MessageText = LOCTEXT("UnsavedQueueWarningMessage", "The changes made to the current queue will be lost by importing another queue. Do you want to continue with this import?");
-
-		if (FMessageDialog::Open(EAppMsgType::YesNo, MessageText, TitleText) == EAppReturnType::No)
-		{
-			return;
-		}
-	}
+	UMoviePipelineQueueSubsystem* Subsystem = GEditor->GetEditorSubsystem<UMoviePipelineQueueSubsystem>();
+	check(Subsystem);
 
 	UMoviePipelineQueue* SavedQueue = CastChecked<UMoviePipelineQueue>(InPresetAsset.GetAsset());
-	LoadQueue(SavedQueue);
-}
-
-void SMoviePipelineQueuePanel::LoadQueue(UMoviePipelineQueue* SavedQueue) const
-{
-	if (SavedQueue)
-	{
-		UMoviePipelineQueueSubsystem* Subsystem = GEditor->GetEditorSubsystem<UMoviePipelineQueueSubsystem>();
-		check(Subsystem);
-
-		UMoviePipelineQueue* CurrentQueue = Subsystem->GetQueue();
-		CurrentQueue->CopyFrom(SavedQueue);
-		CurrentQueue->SetQueueOrigin(SavedQueue);
-		CurrentQueue->SetIsDirty(false);
-
-		// Update the shot list in case the stored queue being copied is out of date with the sequence
-		for (UMoviePipelineExecutorJob* Job : Subsystem->GetQueue()->GetJobs())
-		{
-			ULevelSequence* LoadedSequence = Cast<ULevelSequence>(Job->Sequence.TryLoad());
-			if (LoadedSequence)
-			{
-				bool bShotsChanged = false;
-				UMoviePipelineBlueprintLibrary::UpdateJobShotListFromSequence(LoadedSequence, Job, bShotsChanged);
-
-				if (bShotsChanged)
-				{
-					FNotificationInfo Info(LOCTEXT("QueueShotsUpdated", "Shots have changed since the queue was saved, please resave the queue"));
-					Info.ExpireDuration = 5.0f;
-					FSlateNotificationManager::Get().AddNotification(Info)->SetCompletionState(SNotificationItem::CS_Fail);
-				}
-			}
-		}
-
-		// Automatically select the first job in the queue
-		TArray<UMoviePipelineExecutorJob*> Jobs;
-		if (Subsystem->GetQueue()->GetJobs().Num() > 0)
-		{
-			Jobs.Add(Subsystem->GetQueue()->GetJobs()[0]);
-		}
-
-		// Go through the UI so it updates the UI selection too and then this will loop back
-		// around to OnSelectionChanged to update ourself.
-		PipelineQueueEditorWidget->SetSelectedJobs(Jobs);
-	}
+	Subsystem->LoadQueue(SavedQueue);
 }
 
 bool SMoviePipelineQueuePanel::IsQueueDirty() const
 {
-	UMoviePipelineQueueSubsystem* Subsystem = GEditor->GetEditorSubsystem<UMoviePipelineQueueSubsystem>();
+	const UMoviePipelineQueueSubsystem* Subsystem = GEditor->GetEditorSubsystem<UMoviePipelineQueueSubsystem>();
 	check(Subsystem);
 
-	const UMoviePipelineQueue* Queue = Subsystem->GetQueue();
-	if(!Queue)
-	{
-		return false;
-	}
-
-	// The queue is considered dirty if the current queue has no origin (ie, it has never been saved) or it has been
-	// modified since it was loaded. We skip the no-origin check if there are no jobs (to avoid triggering on an
-	// empty, first-load queue), but we don't skip the IsDirty check because a queue that had jobs and then
-	// removed them will be dirty.
-	bool bShouldCheckQueueOrigin = Queue->GetJobs().Num() > 0;
-	bool bHasNoQueueOrigin = false;
-	if (bShouldCheckQueueOrigin)
-	{
-		bHasNoQueueOrigin = !GetQueueOrigin();
-	}
-	return bHasNoQueueOrigin || (Queue->IsDirty());
+	return Subsystem->IsQueueDirty();
 }
 
 UMoviePipelineQueue* SMoviePipelineQueuePanel::GetQueueOrigin() const
diff --git a/Engine/Plugins/MovieScene/MovieRenderPipeline/Source/MovieRenderPipelineEditor/Private/Widgets/SMoviePipelineQueuePanel.h b/Engine/Plugins/MovieScene/MovieRenderPipeline/Source/MovieRenderPipelineEditor/Private/Widgets/SMoviePipelineQueuePanel.h
--- a/Engine/Plugins/MovieScene/MovieRenderPipeline/Source/MovieRenderPipelineEditor/Private/Widgets/SMoviePipelineQueuePanel.h
+++ b/Engine/Plugins/MovieScene/MovieRenderPipeline/Source/MovieRenderPipelineEditor/Private/Widgets/SMoviePipelineQueuePanel.h
@@ -60,7 +60,6 @@ private:
 	void OnSaveAsAsset();
 	void SaveTransientQueueToAsset(UMoviePipelineQueue* DestinationQueue) const;
 	void OnImportSavedQueueAsset(const FAssetData& InPresetAsset) const;
-	void LoadQueue(UMoviePipelineQueue* SavedQueue) const;
 	bool IsQueueDirty() const;
 	UMoviePipelineQueue* GetQueueOrigin() const;
 	FString GetQueueOriginName() const;
diff --git a/Engine/Plugins/MovieScene/MovieRenderPipeline/Source/MovieRenderPipelineEditor/Public/MoviePipelineQueueSubsystem.h b/Engine/Plugins/MovieScene/MovieRenderPipeline/Source/MovieRenderPipelineEditor/Public/MoviePipelineQueueSubsystem.h
--- a/Engine/Plugins/MovieScene/MovieRenderPipeline/Source/MovieRenderPipelineEditor/Public/MoviePipelineQueueSubsystem.h
+++ b/Engine/Plugins/MovieScene/MovieRenderPipeline/Source/MovieRenderPipelineEditor/Public/MoviePipelineQueueSubsystem.h
@@ -6,6 +6,8 @@
 #include "MoviePipelineExecutor.h"
 #include "MoviePipelineQueueSubsystem.generated.h"
 
+DECLARE_MULTICAST_DELEGATE(FOnMoviePipelineQueueLoaded)
+
 UCLASS(BlueprintType)
 class MOVIERENDERPIPELINEEDITOR_API UMoviePipelineQueueSubsystem : public UEditorSubsystem
 {
@@ -25,6 +27,24 @@ public:
 		return CurrentQueue;
 	}
 
+	/**
+	 * Loads a new queue by copying it into the queue subsystem's current transient queue (the one returned by GetQueue()).
+	 * 
+	 * If bPromptOnReplacingDirtyQueue is true and the current queue has been modified since being loaded, a dialog will prompt the
+	 * user if they want to discard their changes. If this dialog is rejected, or there was an error loading the queue, returns
+	 * false, else returns true. Note that bPromptOnReplacingDirtyQueue is effectively ignored if the application is in
+	 * "unattended" mode because the dialog is auto-accepted.
+	 */
+	UFUNCTION(BlueprintPure, Category = "Movie Render Pipeline")
+	bool LoadQueue(UMoviePipelineQueue* QueueToLoad, const bool bPromptOnReplacingDirtyQueue = true);
+
+	/**
+	 * Returns true if the current queue has been modified since it was loaded, else returns false. Note that the empty
+	 * queue that is in use upon the initial load of MRQ is not considered dirty.
+	 */
+	UFUNCTION(BlueprintPure, Category = "Movie Render Pipeline")
+	bool IsQueueDirty() const;
+
 	/** Returns the active executor (if there is one). This can be used to subscribe to events on an already in-progress render. May be null. */
 	UFUNCTION(BlueprintPure, Category = "Movie Render Pipeline")
 	UMoviePipelineExecutorBase* GetActiveExecutor() const
@@ -65,6 +85,10 @@ public:
 		return ActiveExecutor ? ActiveExecutor->IsRendering() : false;
 	}
 
+public:
+	/** Called when a queue is loaded into the system. */
+	FOnMoviePipelineQueueLoaded OnQueueLoaded;
+
 private:
 	/** Called when the executor is finished so that we can release it and stop reporting IsRendering() == true. */
 	void OnExecutorFinished(UMoviePipelineExecutorBase* InPipelineExecutor, bool bSuccess);
