MoviePipeline: Added ability to drag variable members from the Members panel into the graph (creating a getter node).

diff --git a/Engine/Plugins/MovieScene/MovieRenderPipeline/Source/MovieRenderPipelineEditor/Private/Graph/MovieGraphSchema.cpp b/Engine/Plugins/MovieScene/MovieRenderPipeline/Source/MovieRenderPipelineEditor/Private/Graph/MovieGraphSchema.cpp
--- a/Engine/Plugins/MovieScene/MovieRenderPipeline/Source/MovieRenderPipelineEditor/Private/Graph/MovieGraphSchema.cpp
+++ b/Engine/Plugins/MovieScene/MovieRenderPipeline/Source/MovieRenderPipelineEditor/Private/Graph/MovieGraphSchema.cpp
@@ -260,86 +260,86 @@ void UMovieGraphSchema::BreakSinglePinLink(UEdGraphPin* SourcePin, UEdGraphPin*
 	RuntimeGraph->RemoveEdge(SourceRuntimeNode, SourcePin->PinName, TargetRuntimeNode, TargetPin->PinName);
 }
 
-FLinearColor UMovieGraphSchema::GetPinTypeColor(const FEdGraphPinType& PinType) const
+FLinearColor UMovieGraphSchema::GetTypeColor(const FName& InType)
 {
 	const UGraphEditorSettings* Settings = GetDefault<UGraphEditorSettings>();
 
-	if (PinType.PinCategory == PC_Branch)
+	if (InType == PC_Branch)
 	{
 		return Settings->ExecutionPinTypeColor;
 	}
 
-	if (PinType.PinCategory == PC_Boolean)
+	if (InType == PC_Boolean)
 	{
 		return Settings->BooleanPinTypeColor;
 	}
 	
-	if (PinType.PinCategory == PC_Byte)
+	if (InType == PC_Byte)
 	{
 		return Settings->BytePinTypeColor;
 	}
 
-	if (PinType.PinCategory == PC_Integer)
+	if (InType == PC_Integer)
 	{
 		return Settings->IntPinTypeColor;
 	}
 
-	if (PinType.PinCategory == PC_Int64)
+	if (InType == PC_Int64)
 	{
 		return Settings->Int64PinTypeColor;
 	}
 
-	if (PinType.PinCategory == PC_Float)
+	if (InType == PC_Float)
 	{
 		return Settings->FloatPinTypeColor;
 	}
 
-	if (PinType.PinCategory == PC_Double)
+	if (InType == PC_Double)
 	{
 		return Settings->DoublePinTypeColor;
 	}
 
-	if (PinType.PinCategory == PC_Name)
+	if (InType == PC_Name)
 	{
 		return Settings->NamePinTypeColor;
 	}
 
-	if (PinType.PinCategory == PC_String)
+	if (InType == PC_String)
 	{
 		return Settings->StringPinTypeColor;
 	}
 
-	if (PinType.PinCategory == PC_Text)
+	if (InType == PC_Text)
 	{
 		return Settings->TextPinTypeColor;
 	}
 
-	if (PinType.PinCategory == PC_Enum)
+	if (InType == PC_Enum)
 	{
 		return Settings->BytePinTypeColor;
 	}
 
-	if (PinType.PinCategory == PC_Struct)
+	if (InType == PC_Struct)
 	{
 		return Settings->StructPinTypeColor;
 	}
 	
-	if (PinType.PinCategory == PC_Object)
+	if (InType == PC_Object)
 	{
 		return Settings->ObjectPinTypeColor;
 	}
 
-	if (PinType.PinCategory == PC_SoftObject)
+	if (InType == PC_SoftObject)
 	{
 		return Settings->SoftObjectPinTypeColor;
 	}
 	
-	if (PinType.PinCategory == PC_Class)
+	if (InType == PC_Class)
     {
     	return Settings->ClassPinTypeColor;
     }
 
-	if (PinType.PinCategory == PC_SoftClass)
+	if (InType == PC_SoftClass)
 	{
 		return Settings->SoftClassPinTypeColor;
 	}
@@ -347,6 +347,11 @@ FLinearColor UMovieGraphSchema::GetPinTypeColor(const FEdGraphPinType& PinType)
 	return Settings->DefaultPinTypeColor;
 }
 
+FLinearColor UMovieGraphSchema::GetPinTypeColor(const FEdGraphPinType& PinType) const
+{
+	return GetTypeColor(PinType.PinCategory);
+}
+
 FConnectionDrawingPolicy* UMovieGraphSchema::CreateConnectionDrawingPolicy(int32 InBackLayerID, int32 InFrontLayerID,
 	float InZoomFactor, const FSlateRect& InClippingRect, FSlateWindowElementList& InDrawElements,
 	UEdGraph* InGraphObj) const
diff --git a/Engine/Plugins/MovieScene/MovieRenderPipeline/Source/MovieRenderPipelineEditor/Private/Graph/MovieGraphSchema.h b/Engine/Plugins/MovieScene/MovieRenderPipeline/Source/MovieRenderPipelineEditor/Private/Graph/MovieGraphSchema.h
--- a/Engine/Plugins/MovieScene/MovieRenderPipeline/Source/MovieRenderPipelineEditor/Private/Graph/MovieGraphSchema.h
+++ b/Engine/Plugins/MovieScene/MovieRenderPipeline/Source/MovieRenderPipelineEditor/Private/Graph/MovieGraphSchema.h
@@ -22,6 +22,7 @@ public:
 	virtual const FPinConnectionResponse CanCreateConnection(const UEdGraphPin* A, const UEdGraphPin* B) const override;
 	virtual bool TryCreateConnection(UEdGraphPin* InA, UEdGraphPin* InB) const override;
 	//virtual const FPinConnectionResponse CanMergeNodes(const UEdGraphNode* A, const UEdGraphNode* B) const override;
+	static FLinearColor GetTypeColor(const FName& InType);
 	virtual FLinearColor GetPinTypeColor(const FEdGraphPinType& PinType) const override;
 	virtual FConnectionDrawingPolicy* CreateConnectionDrawingPolicy(int32 InBackLayerID, int32 InFrontLayerID, float InZoomFactor, const FSlateRect& InClippingRect, FSlateWindowElementList& InDrawElements, UEdGraph* InGraphObj) const override;
 
diff --git a/Engine/Plugins/MovieScene/MovieRenderPipeline/Source/MovieRenderPipelineEditor/Private/Graph/SMovieGraphMembersTabContent.cpp b/Engine/Plugins/MovieScene/MovieRenderPipeline/Source/MovieRenderPipelineEditor/Private/Graph/SMovieGraphMembersTabContent.cpp
--- a/Engine/Plugins/MovieScene/MovieRenderPipeline/Source/MovieRenderPipelineEditor/Private/Graph/SMovieGraphMembersTabContent.cpp
+++ b/Engine/Plugins/MovieScene/MovieRenderPipeline/Source/MovieRenderPipelineEditor/Private/Graph/SMovieGraphMembersTabContent.cpp
@@ -6,6 +6,7 @@
 #include "Framework/Commands/GenericCommands.h"
 #include "Graph/MovieGraphConfig.h"
 #include "Graph/MovieGraphSchema.h"
+#include "Graph/Nodes/MovieGraphVariableNode.h"
 #include "MovieEdGraphNode.h"
 #include "SGraphActionMenu.h"
 #include "Toolkits/AssetEditorToolkit.h"
@@ -32,6 +33,18 @@ namespace UE::MovieGraph::Private
 
 		return nullptr;
 	}
+
+	static bool GetIconAndColorFromDataType(EMovieGraphValueType VariableType, const FSlateBrush*& OutPrimaryBrush, FSlateColor& OutIconColor, const FSlateBrush*& OutSecondaryBrush, FSlateColor& OutSecondaryColor)
+	{
+		constexpr bool bIsBranch = false;
+		const FEdGraphPinType PinType = UMoviePipelineEdGraphNodeBase::GetPinType(VariableType, bIsBranch);
+
+		OutPrimaryBrush = FAppStyle::GetBrush("Kismet.AllClasses.VariableIcon");
+		OutIconColor = UMovieGraphSchema::GetTypeColor(PinType.PinCategory);
+		OutSecondaryBrush = nullptr;
+		
+		return true;
+	}
 }
 
 const TArray<FText> SMovieGraphMembersTabContent::ActionMenuSectionNames
@@ -54,6 +67,7 @@ void SMovieGraphMembersTabContent::Construct(const FArguments& InArgs)
 		.OnActionSelected(OnActionSelected)
 		.AutoExpandActionMenu(true)
 		.AlphaSortItems(false)
+		.OnActionDragged(this, &SMovieGraphMembersTabContent::OnActionDragged)
 		.OnCreateWidgetForAction(this, &SMovieGraphMembersTabContent::CreateActionWidget)
 		.OnCollectStaticSections(this, &SMovieGraphMembersTabContent::CollectStaticSections)
 		.OnContextMenuOpening(this, &SMovieGraphMembersTabContent::OnContextMenuOpening)
@@ -65,6 +79,25 @@ void SMovieGraphMembersTabContent::Construct(const FArguments& InArgs)
 	];
 }
 
+FReply SMovieGraphMembersTabContent::OnActionDragged(const TArray<TSharedPtr<FEdGraphSchemaAction>>& InActions, const FPointerEvent& MouseEvent)
+{
+	const TSharedPtr<FEdGraphSchemaAction> Action(!InActions.IsEmpty() ? InActions[0] : nullptr);
+	if (!Action.IsValid())
+	{
+		return FReply::Unhandled();
+	}
+	
+	const FMovieGraphSchemaAction* GraphSchemaAction = static_cast<FMovieGraphSchemaAction*>(Action.Get());
+	const TObjectPtr<UObject> ActionTarget = GraphSchemaAction->ActionTarget;
+
+	if (UMovieGraphVariable* VariableMember = Cast<UMovieGraphVariable>(ActionTarget.Get()))
+	{
+		return FReply::Handled().BeginDragDrop(FMovieGraphDragAction_Variable::New(Action, VariableMember));
+	}
+	
+	return FReply::Unhandled();
+}
+
 TSharedRef<SWidget> SMovieGraphMembersTabContent::CreateActionWidget(FCreateWidgetForActionData* InCreateData) const
 {
 	// For variables, show an icon w/ the color representing the variable's type (in addition to the variable name)
@@ -374,4 +407,52 @@ void SMovieGraphMembersTabContent::RefreshMemberActions(UMovieGraphMember* Updat
 	}
 }
 
+TSharedRef<FMovieGraphDragAction_Variable> FMovieGraphDragAction_Variable::New(
+	TSharedPtr<FEdGraphSchemaAction> InAction, UMovieGraphVariable* InVariable)
+{
+	TSharedRef<FMovieGraphDragAction_Variable> DragAction = MakeShared<FMovieGraphDragAction_Variable>();
+	DragAction->SourceAction = InAction;
+	DragAction->WeakVariable = InVariable;
+	DragAction->Construct();
+	
+	return DragAction;
+}
+
+void FMovieGraphDragAction_Variable::HoverTargetChanged()
+{
+	FGraphSchemaActionDragDropAction::HoverTargetChanged();
+}
+
+FReply FMovieGraphDragAction_Variable::DroppedOnPanel(
+	const TSharedRef<SWidget>& InPanel, FVector2D InScreenPosition, FVector2D InGraphPosition, UEdGraph& InGraph)
+{
+	const UMovieGraphVariable* VariableMember = WeakVariable.Get();
+	if (!VariableMember)
+	{
+		return FReply::Unhandled();
+	}
+
+	// When creating the new action, since it's only being used to create a node, the category, display name, and tooltip can just be empty
+	const TSharedPtr<FMovieGraphSchemaAction> NewAction = MakeShared<FMovieGraphSchemaAction_NewVariableNode>(
+		FText::GetEmpty(), FText::GetEmpty(), VariableMember->GetGuid(), FText::GetEmpty());
+	NewAction->NodeClass = UMovieGraphVariableNode::StaticClass();
+	
+	constexpr UEdGraphPin* FromPin = nullptr;
+	constexpr bool bSelectNewNode = true;
+	NewAction->PerformAction(&InGraph, FromPin, InGraphPosition, bSelectNewNode);
+
+	return FReply::Handled();
+}
+
+void FMovieGraphDragAction_Variable::GetDefaultStatusSymbol(
+	const FSlateBrush*& OutPrimaryBrush, FSlateColor& OutIconColor, const FSlateBrush*& OutSecondaryBrush, FSlateColor& OutSecondaryColor) const
+{
+	const UMovieGraphVariable* VariableMember = WeakVariable.Get();
+	if (!VariableMember ||
+		!UE::MovieGraph::Private::GetIconAndColorFromDataType(VariableMember->GetValueType(), OutPrimaryBrush, OutIconColor, OutSecondaryBrush, OutSecondaryColor))
+	{
+		return FGraphSchemaActionDragDropAction::GetDefaultStatusSymbol(OutPrimaryBrush, OutIconColor, OutSecondaryBrush, OutSecondaryColor);
+	}
+}
+
 #undef LOCTEXT_NAMESPACE
\ No newline at end of file
diff --git a/Engine/Plugins/MovieScene/MovieRenderPipeline/Source/MovieRenderPipelineEditor/Private/Graph/SMovieGraphMembersTabContent.h b/Engine/Plugins/MovieScene/MovieRenderPipeline/Source/MovieRenderPipelineEditor/Private/Graph/SMovieGraphMembersTabContent.h
--- a/Engine/Plugins/MovieScene/MovieRenderPipeline/Source/MovieRenderPipelineEditor/Private/Graph/SMovieGraphMembersTabContent.h
+++ b/Engine/Plugins/MovieScene/MovieRenderPipeline/Source/MovieRenderPipelineEditor/Private/Graph/SMovieGraphMembersTabContent.h
@@ -2,12 +2,14 @@
 
 #pragma once
 
+#include "GraphEditorDragDropAction.h"
 #include "SGraphActionMenu.h"
 #include "Widgets/Input/SButton.h"
 #include "Widgets/SCompoundWidget.h"
 
 class FAssetEditorToolkit;
 class SGraphActionMenu;
+class UMovieGraphVariable;
 struct FEdGraphSchemaAction;
 struct FGraphActionListBuilderBase;
 
@@ -38,6 +40,9 @@ public:
 	
 	void Construct(const FArguments& InArgs);
 
+	/** Handler for when an action in the action menu is dragged. */
+	FReply OnActionDragged(const TArray<TSharedPtr<FEdGraphSchemaAction>>& InActions, const FPointerEvent& MouseEvent);
+
 	/** Creates the widget for an action within the action menu. */
 	TSharedRef<SWidget> CreateActionWidget(FCreateWidgetForActionData* CreateWidgetForActionData) const;
 
@@ -96,3 +101,25 @@ private:
 	/** Delegate to call when an action is selected */
 	FOnActionSelected OnActionSelected;
 };
+
+/* Drag-and-drop action which handles variable members. */
+class FMovieGraphDragAction_Variable : public FGraphSchemaActionDragDropAction
+{
+public:
+	DRAG_DROP_OPERATOR_TYPE(FMovieGraphDragAction_Variable, FGraphSchemaActionDragDropAction)
+
+	static TSharedRef<FMovieGraphDragAction_Variable> New(
+		TSharedPtr<FEdGraphSchemaAction> InAction, UMovieGraphVariable* InVariable);
+
+	virtual void HoverTargetChanged() override;
+	virtual FReply DroppedOnPanel(
+		const TSharedRef<SWidget>& InPanel, FVector2D InScreenPosition, FVector2D InGraphPosition, UEdGraph& InGraph) override;
+
+protected:
+	virtual void GetDefaultStatusSymbol(
+		const FSlateBrush*& OutPrimaryBrush, FSlateColor& OutIconColor, FSlateBrush const*& OutSecondaryBrush, FSlateColor& OutSecondaryColor) const override;
+
+private:
+	/** The variable member that the drag-and-drop is associated with. */
+	TWeakObjectPtr<UMovieGraphVariable> WeakVariable;
+};
\ No newline at end of file
