MoviePipeline: Added validations (eg, uniqueness check) when renaming graph members, both in the UI and API.

diff --git a/Engine/Plugins/MovieScene/MovieRenderPipeline/Source/MovieRenderPipelineCore/Private/Graph/MovieGraphConfig.cpp b/Engine/Plugins/MovieScene/MovieRenderPipeline/Source/MovieRenderPipelineCore/Private/Graph/MovieGraphConfig.cpp
--- a/Engine/Plugins/MovieScene/MovieRenderPipeline/Source/MovieRenderPipelineCore/Private/Graph/MovieGraphConfig.cpp
+++ b/Engine/Plugins/MovieScene/MovieRenderPipeline/Source/MovieRenderPipelineCore/Private/Graph/MovieGraphConfig.cpp
@@ -14,6 +14,71 @@
 
 static const FName MovieGraphGlobalsMemberName("Globals");
 
+UMovieGraphConfig* UMovieGraphMember::GetOwningGraph() const
+{
+	return GetTypedOuter<UMovieGraphConfig>();
+}
+
+bool UMovieGraphMember::SetMemberName(const FString& InNewName)
+{
+	FText UnusedError;
+	if (CanRename(FText::FromString(InNewName), UnusedError))
+	{
+		Name = InNewName;
+		return true;
+	}
+
+	return false;
+}
+
+bool UMovieGraphMember::CanRename(const FText& InNewName, FText& OutError) const
+{
+	static const FString InvalidChars("\"',\n\r\t");
+	
+	if (InNewName.IsEmptyOrWhitespace())
+	{
+		OutError = LOCTEXT("InvalidMemberRename_Empty", "The name cannot be empty.");
+		return false;
+	}
+
+	const FString NewNameString = InNewName.ToString();
+	if (!FName::IsValidXName(NewNameString, InvalidChars, &OutError))
+	{
+		return false;
+	}
+
+	return true;
+}
+
+bool UMovieGraphVariable::CanRename(const FText& InNewName, FText& OutError) const
+{
+	if (!Super::CanRename(InNewName, OutError))
+	{
+		return false;
+	}
+
+	if (const UMovieGraphConfig* Graph = GetOwningGraph())
+	{
+		constexpr bool bIncludeGlobal = true;
+		if (!IsUniqueNameInMemberArray(InNewName, Graph->GetVariables(bIncludeGlobal)))
+		{
+			OutError = LOCTEXT("InvalidVariableRename_Exists", "A variable with this name already exists.");
+			return false;
+		}
+	}
+
+	return true;
+}
+
+bool UMovieGraphVariable::SetMemberName(const FString& InNewName)
+{
+#if WITH_EDITOR
+	OnMovieGraphVariableChangedDelegate.Broadcast(this);
+#endif
+	
+	return Super::SetMemberName(InNewName);
+}
+
 #if WITH_EDITOR
 void UMovieGraphVariable::PostEditChangeProperty(FPropertyChangedEvent& PropertyChangedEvent)
 {
@@ -29,6 +94,34 @@ bool UMovieGraphInput::IsDeletable() const
 	return Name != MovieGraphGlobalsMemberName;
 }
 
+bool UMovieGraphInput::CanRename(const FText& InNewName, FText& OutError) const
+{
+	if (!Super::CanRename(InNewName, OutError))
+	{
+		return false;
+	}
+
+	if (const UMovieGraphConfig* Graph = GetOwningGraph())
+	{
+		if (!IsUniqueNameInMemberArray(InNewName, Graph->GetInputs()))
+		{
+			OutError = LOCTEXT("InvalidInputRename_Exists", "An input with this name already exists.");
+			return false;
+		}
+	}
+
+	return true;
+}
+
+bool UMovieGraphInput::SetMemberName(const FString& InNewName)
+{
+#if WITH_EDITOR
+	OnMovieGraphInputChangedDelegate.Broadcast(this);
+#endif
+	
+	return Super::SetMemberName(InNewName);
+}
+
 #if WITH_EDITOR
 void UMovieGraphInput::PostEditChangeProperty(FPropertyChangedEvent& PropertyChangedEvent)
 {
@@ -44,6 +137,34 @@ bool UMovieGraphOutput::IsDeletable() const
 	return Name != MovieGraphGlobalsMemberName;
 }
 
+bool UMovieGraphOutput::CanRename(const FText& InNewName, FText& OutError) const
+{
+	if (!Super::CanRename(InNewName, OutError))
+	{
+		return false;
+	}
+
+	if (const UMovieGraphConfig* Graph = GetOwningGraph())
+	{
+		if (!IsUniqueNameInMemberArray(InNewName, Graph->GetOutputs()))
+		{
+			OutError = LOCTEXT("InvalidOutputRename_Exists", "An output with this name already exists.");
+			return false;
+		}
+	}
+
+	return true;
+}
+
+bool UMovieGraphOutput::SetMemberName(const FString& InNewName)
+{
+#if WITH_EDITOR
+	OnMovieGraphOutputChangedDelegate.Broadcast(this);
+#endif
+	
+	return Super::SetMemberName(InNewName);
+}
+
 #if WITH_EDITOR
 void UMovieGraphOutput::PostEditChangeProperty(FPropertyChangedEvent& PropertyChangedEvent)
 {
@@ -110,7 +231,7 @@ UMovieGraphVariable* UMovieGraphConfig::AddGlobalVariable(const FName& InName, E
 	// Don't add duplicate global variables
 	const bool VariableExists = Variables.ContainsByPredicate([&InName](const TObjectPtr<UMovieGraphVariable>& Variable)
 	{
-		return Variable && (Variable->Name == InName);
+		return Variable && (Variable->GetMemberName() == InName);
 	});
 
 	if (VariableExists)
@@ -133,19 +254,19 @@ void UMovieGraphConfig::AddDefaultMembers()
 {
 	const bool InputGlobalsExists = Inputs.ContainsByPredicate([](const UMovieGraphMember* Member)
 	{
-		return Member && (Member->Name == MovieGraphGlobalsMemberName);
+		return Member && (Member->GetMemberName() == MovieGraphGlobalsMemberName);
 	});
 
 	const bool OutputGlobalsExists = Outputs.ContainsByPredicate([](const UMovieGraphMember* Member)
 	{
-		return Member && (Member->Name == MovieGraphGlobalsMemberName);
+		return Member && (Member->GetMemberName() == MovieGraphGlobalsMemberName);
 	});
 
 	// Ensure there is a Globals input member
 	if (!InputGlobalsExists)
 	{
 		UMovieGraphInput* NewInput = AddInput();
-		NewInput->Name = MovieGraphGlobalsMemberName.ToString();
+		NewInput->SetMemberName(MovieGraphGlobalsMemberName.ToString());
 
 		InputNode->UpdatePins();
 	}
@@ -154,7 +275,7 @@ void UMovieGraphConfig::AddDefaultMembers()
 	if (!OutputGlobalsExists)
 	{
 		UMovieGraphOutput* NewOutput = AddOutput();
-		NewOutput->Name = MovieGraphGlobalsMemberName.ToString();
+		NewOutput->SetMemberName(MovieGraphGlobalsMemberName.ToString());
 
 		OutputNode->UpdatePins();
 	}
@@ -424,8 +545,8 @@ T* UMovieGraphConfig::AddMember(TArray<TObjectPtr<T>>& InMemberArray, const FNam
 
 	// Generate and set a unique name
 	TArray<FString> ExistingMemberNames;
-	Algo::Transform(InMemberArray, ExistingMemberNames, [](const T* Member) { return Member->Name; });
-	NewMember->Name = GetUniqueName(ExistingMemberNames, InBaseName.ToString());
+	Algo::Transform(InMemberArray, ExistingMemberNames, [](const T* Member) { return Member->GetMemberName(); });
+	NewMember->SetMemberName(GetUniqueName(ExistingMemberNames, InBaseName.ToString()));
 
 	return NewMember;
 }
@@ -525,7 +646,7 @@ bool UMovieGraphConfig::DeleteMember(UMovieGraphMember* MemberToDelete)
 
 	if (!MemberToDelete->IsDeletable())
 	{
-		UE_LOG(LogMovieRenderPipeline, Error, TEXT("DeleteMember: The member '%s' cannot be deleted because it is flagged as non-deletable."), *MemberToDelete->Name);
+		UE_LOG(LogMovieRenderPipeline, Error, TEXT("DeleteMember: The member '%s' cannot be deleted because it is flagged as non-deletable."), *MemberToDelete->GetMemberName());
 		return false;
 	}
 
@@ -607,7 +728,7 @@ bool UMovieGraphConfig::DeleteInputMember(UMovieGraphInput* InputMemberToDelete)
 	if (InputMemberToDelete)
 	{
 		Inputs.RemoveSingle(InputMemberToDelete);
-		RemoveOutboundEdges(InputNode, FName(InputMemberToDelete->Name));
+		RemoveOutboundEdges(InputNode, FName(InputMemberToDelete->GetMemberName()));
 
 		// This calls OnNodeChangedDelegate to update the graph
 		InputNode->UpdatePins();
@@ -623,7 +744,7 @@ bool UMovieGraphConfig::DeleteOutputMember(UMovieGraphOutput* OutputMemberToDele
 	if (OutputMemberToDelete)
 	{
 		Outputs.RemoveSingle(OutputMemberToDelete);
-		RemoveInboundEdges(OutputNode, FName(OutputMemberToDelete->Name));
+		RemoveInboundEdges(OutputNode, FName(OutputMemberToDelete->GetMemberName()));
 
 		// This calls OnNodeChangedDelegate to update the graph
 		OutputNode->UpdatePins();
diff --git a/Engine/Plugins/MovieScene/MovieRenderPipeline/Source/MovieRenderPipelineCore/Private/Graph/Nodes/MovieGraphInputNode.cpp b/Engine/Plugins/MovieScene/MovieRenderPipeline/Source/MovieRenderPipelineCore/Private/Graph/Nodes/MovieGraphInputNode.cpp
--- a/Engine/Plugins/MovieScene/MovieRenderPipeline/Source/MovieRenderPipelineCore/Private/Graph/Nodes/MovieGraphInputNode.cpp
+++ b/Engine/Plugins/MovieScene/MovieRenderPipeline/Source/MovieRenderPipelineCore/Private/Graph/Nodes/MovieGraphInputNode.cpp
@@ -28,7 +28,7 @@ TArray<FMovieGraphPinProperties> UMovieGraphInputNode::GetOutputPinProperties()
 	{
 		for (const UMovieGraphInput* Input : ParentGraph->GetInputs())
 		{
-			FMovieGraphPinProperties PinProperties(FName(Input->Name), Input->GetValueType(), false);
+			FMovieGraphPinProperties PinProperties(FName(Input->GetMemberName()), Input->GetValueType(), false);
 			PinProperties.bIsBranch = Input->bIsBranch;
 			Properties.Add(MoveTemp(PinProperties));
 		}
@@ -137,7 +137,7 @@ void UMovieGraphInputNode::UpdateExistingPins(UMovieGraphMember* ChangedInput) c
 		{
 			for (int32 Index = 0; Index < InputMembers.Num(); ++Index)
 			{
-				OutputPins[Index]->Properties.Label = FName(InputMembers[Index]->Name);
+				OutputPins[Index]->Properties.Label = FName(InputMembers[Index]->GetMemberName());
 				OutputPins[Index]->Properties.Type = InputMembers[Index]->GetValueType();
 				OutputPins[Index]->Properties.bIsBranch = InputMembers[Index]->bIsBranch;
 			}
diff --git a/Engine/Plugins/MovieScene/MovieRenderPipeline/Source/MovieRenderPipelineCore/Private/Graph/Nodes/MovieGraphOutputNode.cpp b/Engine/Plugins/MovieScene/MovieRenderPipeline/Source/MovieRenderPipelineCore/Private/Graph/Nodes/MovieGraphOutputNode.cpp
--- a/Engine/Plugins/MovieScene/MovieRenderPipeline/Source/MovieRenderPipelineCore/Private/Graph/Nodes/MovieGraphOutputNode.cpp
+++ b/Engine/Plugins/MovieScene/MovieRenderPipeline/Source/MovieRenderPipelineCore/Private/Graph/Nodes/MovieGraphOutputNode.cpp
@@ -27,7 +27,7 @@ TArray<FMovieGraphPinProperties> UMovieGraphOutputNode::GetInputPinProperties()
 	{
 		for (const UMovieGraphOutput* Output : ParentGraph->GetOutputs())
 		{
-			FMovieGraphPinProperties PinProperties(FName(Output->Name), Output->GetValueType(), false);
+			FMovieGraphPinProperties PinProperties(FName(Output->GetMemberName()), Output->GetValueType(), false);
 			PinProperties.bIsBranch = Output->bIsBranch;
 			Properties.Add(MoveTemp(PinProperties));
 		}
@@ -93,7 +93,7 @@ void UMovieGraphOutputNode::UpdateExistingPins(UMovieGraphMember* ChangedOutput)
 		{
 			for (int32 Index = 0; Index < OutputMembers.Num(); ++Index)
 			{
-				InputPins[Index]->Properties.Label = FName(OutputMembers[Index]->Name);
+				InputPins[Index]->Properties.Label = FName(OutputMembers[Index]->GetMemberName());
 				InputPins[Index]->Properties.Type = OutputMembers[Index]->GetValueType();
 				InputPins[Index]->Properties.bIsBranch = OutputMembers[Index]->bIsBranch;
 			}
diff --git a/Engine/Plugins/MovieScene/MovieRenderPipeline/Source/MovieRenderPipelineCore/Private/Graph/Nodes/MovieGraphSubgraphNode.cpp b/Engine/Plugins/MovieScene/MovieRenderPipeline/Source/MovieRenderPipelineCore/Private/Graph/Nodes/MovieGraphSubgraphNode.cpp
--- a/Engine/Plugins/MovieScene/MovieRenderPipeline/Source/MovieRenderPipelineCore/Private/Graph/Nodes/MovieGraphSubgraphNode.cpp
+++ b/Engine/Plugins/MovieScene/MovieRenderPipeline/Source/MovieRenderPipelineCore/Private/Graph/Nodes/MovieGraphSubgraphNode.cpp
@@ -22,7 +22,7 @@ namespace UE::MovieGraph::Private
 			[](const UMovieGraphInterfaceBase* Member) -> FMovieGraphPinProperties
 			{
 				constexpr bool bAllowMultipleConnections = false;
-				FMovieGraphPinProperties Properties(FName(Member->Name), Member->GetValueType(), bAllowMultipleConnections);
+				FMovieGraphPinProperties Properties(FName(Member->GetMemberName()), Member->GetValueType(), bAllowMultipleConnections);
 				Properties.bIsBranch = Member->bIsBranch;
 				return Properties;
 			});
diff --git a/Engine/Plugins/MovieScene/MovieRenderPipeline/Source/MovieRenderPipelineCore/Private/Graph/Nodes/MovieGraphVariableNode.cpp b/Engine/Plugins/MovieScene/MovieRenderPipeline/Source/MovieRenderPipelineCore/Private/Graph/Nodes/MovieGraphVariableNode.cpp
--- a/Engine/Plugins/MovieScene/MovieRenderPipeline/Source/MovieRenderPipelineCore/Private/Graph/Nodes/MovieGraphVariableNode.cpp
+++ b/Engine/Plugins/MovieScene/MovieRenderPipeline/Source/MovieRenderPipelineCore/Private/Graph/Nodes/MovieGraphVariableNode.cpp
@@ -20,7 +20,7 @@ TArray<FMovieGraphPinProperties> UMovieGraphVariableNode::GetOutputPinProperties
 	
 	if (GraphVariable)
 	{
-		Properties.Add(FMovieGraphPinProperties(FName(GraphVariable->Name), GraphVariable->GetValueType(), false));
+		Properties.Add(FMovieGraphPinProperties(FName(GraphVariable->GetMemberName()), GraphVariable->GetValueType(), false));
 	}
 	else
 	{
@@ -32,7 +32,7 @@ TArray<FMovieGraphPinProperties> UMovieGraphVariableNode::GetOutputPinProperties
 
 FString UMovieGraphVariableNode::GetResolvedValueForOutputPin(const FName& InPinName, const FMovieGraphTraversalContext* InContext) const
 {
-	if (GraphVariable && (GraphVariable->Name == InPinName))
+	if (GraphVariable && (GraphVariable->GetMemberName() == InPinName))
 	{
 		// If there's a valid job in the traversal context, and the job has an enabled variable assignment, use that
 		// instead of the variable value set in the graph
@@ -43,7 +43,7 @@ FString UMovieGraphVariableNode::GetResolvedValueForOutputPin(const FName& InPin
 			{
 				if (bIsEnabled)
 				{
-					return InContext->Job->VariableAssignments->GetValueSerializedString(FName(GraphVariable->Name));
+					return InContext->Job->VariableAssignments->GetValueSerializedString(FName(GraphVariable->GetMemberName()));
 				}
 			}
 		}
@@ -72,7 +72,7 @@ void UMovieGraphVariableNode::SetVariable(UMovieGraphVariable* InVariable)
 #if WITH_EDITOR
 FText UMovieGraphVariableNode::GetNodeTitle(const bool bGetDescriptive) const
 {
-	return GraphVariable ? FText::FromString(GraphVariable->Name) : FText();
+	return GraphVariable ? FText::FromString(GraphVariable->GetMemberName()) : FText();
 }
 	
 FText UMovieGraphVariableNode::GetMenuCategory() const
@@ -112,7 +112,7 @@ void UMovieGraphVariableNode::UpdateOutputPin(UMovieGraphMember* ChangedVariable
 	if (!OutputPins.IsEmpty() && ChangedVariable)
 	{
 		// Update the output pin to reflect the variable data model
-		OutputPins[0]->Properties.Label = FName(ChangedVariable->Name);
+		OutputPins[0]->Properties.Label = FName(ChangedVariable->GetMemberName());
 		OutputPins[0]->Properties.Type = ChangedVariable->GetValueType();
 	}
 
diff --git a/Engine/Plugins/MovieScene/MovieRenderPipeline/Source/MovieRenderPipelineCore/Private/MovieJobVariableAssignmentContainer.cpp b/Engine/Plugins/MovieScene/MovieRenderPipeline/Source/MovieRenderPipelineCore/Private/MovieJobVariableAssignmentContainer.cpp
--- a/Engine/Plugins/MovieScene/MovieRenderPipeline/Source/MovieRenderPipelineCore/Private/MovieJobVariableAssignmentContainer.cpp
+++ b/Engine/Plugins/MovieScene/MovieRenderPipeline/Source/MovieRenderPipelineCore/Private/MovieJobVariableAssignmentContainer.cpp
@@ -223,7 +223,7 @@ bool UMovieJobVariableAssignmentContainer::FindOrGenerateVariableOverride(
 {
 	if (InGraphVariable)
 	{
-		if (const FPropertyBagPropertyDesc* Desc = Value.FindPropertyDescByName(FName(InGraphVariable->Name)))
+		if (const FPropertyBagPropertyDesc* Desc = Value.FindPropertyDescByName(FName(InGraphVariable->GetMemberName())))
 		{
 			if (OutPropDesc)
 			{
@@ -260,7 +260,7 @@ bool UMovieJobVariableAssignmentContainer::GenerateVariableOverride(const UMovie
 	NewProperty.ValueType = static_cast<EPropertyBagPropertyType>(InGraphVariable->GetValueType());
 	NewProperty.ValueTypeObject = InGraphVariable->GetValueTypeObject();
 	NewProperty.ContainerTypes = { static_cast<EPropertyBagContainerType>(InGraphVariable->GetValueContainerType()) };
-	NewProperty.Name = FName(InGraphVariable->Name);
+	NewProperty.Name = FName(InGraphVariable->GetMemberName());
 	NewProperty.ID = FGuid::NewGuid();
 #if WITH_EDITOR
 	NewProperty.MetaData.Add(FPropertyBagPropertyDescMetaData("VariableGUID", InGraphVariable->GetGuid().ToString()));
@@ -371,10 +371,10 @@ void UMovieJobVariableAssignmentContainer::UpdateGraphVariableOverrides()
 
 			const UMovieGraphVariable* Variable = *FoundVariable;
 
-			if (Desc.Name != FName(Variable->Name))
+			if (Desc.Name != FName(Variable->GetMemberName()))
 			{
 				bNeedsToRegenerate = true;
-				Desc.Name = FName(Variable->Name);
+				Desc.Name = FName(Variable->GetMemberName());
 			}
 
 			if ((uint8)Desc.ValueType != (uint8)(Variable->GetValueType()))
@@ -443,7 +443,7 @@ bool UMovieJobVariableAssignmentContainer::SetVariableAssignmentEnableState(cons
 {
 	if (FindOrGenerateVariableOverride(InGraphVariable))
 	{
-		const FString EditCondPropName = FString::Format(TEXT("{0}{1}"), {EditConditionPrefix, InGraphVariable->Name});
+		const FString EditCondPropName = FString::Format(TEXT("{0}{1}"), {EditConditionPrefix, InGraphVariable->GetMemberName()});
 		const EPropertyBagResult Result = Value.SetValueBool(FName(EditCondPropName), bIsEnabled);
 		
 		return Result == EPropertyBagResult::Success;
diff --git a/Engine/Plugins/MovieScene/MovieRenderPipeline/Source/MovieRenderPipelineCore/Public/Graph/MovieGraphConfig.h b/Engine/Plugins/MovieScene/MovieRenderPipeline/Source/MovieRenderPipelineCore/Public/Graph/MovieGraphConfig.h
--- a/Engine/Plugins/MovieScene/MovieRenderPipeline/Source/MovieRenderPipelineCore/Public/Graph/MovieGraphConfig.h
+++ b/Engine/Plugins/MovieScene/MovieRenderPipeline/Source/MovieRenderPipelineCore/Public/Graph/MovieGraphConfig.h
@@ -30,6 +30,21 @@ class MOVIERENDERPIPELINECORE_API UMovieGraphMember : public UMovieGraphValueCon
 public:
 	UMovieGraphMember() = default;
 
+	/** Gets the graph that owns this member, or nullptr if one was not found. */
+	UMovieGraphConfig* GetOwningGraph() const;
+
+	/** Gets the name of this member. */
+	FString GetMemberName() const { return Name; }
+
+	/** Sets the name of this member. Returns true if the rename was successful, else false. */
+	virtual bool SetMemberName(const FString& InNewName);
+
+	/**
+	 * Determines if this member can be renamed to the specified name. If the rename is not possible, returns false
+	 * and OutError is populated with the reason, else returns true.
+	 */
+	virtual bool CanRename(const FText& InNewName, FText& OutError) const;
+
 	/** Gets the GUID that uniquely identifies this member. */
 	const FGuid& GetGuid() const { return Guid; }
 
@@ -41,18 +56,34 @@ public:
 
 	/** Gets whether this member is editable via the UI. */
 	virtual bool IsEditable() const { return bIsEditable; }
+	
+protected:
+    /** Determines if InName is a unique name within the members in MemberArray. */
+	template<class T>
+    bool IsUniqueNameInMemberArray(const FText& InName, const TArray<T*>& InMemberArray) const
+	{
+		const FString NameString = InName.ToString();
+	
+		const bool bExists = InMemberArray.ContainsByPredicate([&NameString](const T* Member)
+		{
+			return Member->GetMemberName() == NameString;
+		});
 
-public:
-	// TODO: Need a details customization that validates whether or not the name is valid/unique
-	/** The name of this member, which is user-facing. */
-	UPROPERTY(EditAnywhere, Category = "General", meta=(EditCondition="bIsEditable", HideEditConditionToggle))
-	FString Name;
+		// Check against the current name as well; this method shouldn't flag the provided name as non-unique if it's
+		// the member's current name
+		return !bExists || (NameString == Name);
+	}
 
+public:
 	/** The optional description of this member, which is user-facing. */
 	UPROPERTY(EditAnywhere, Category = "General", meta=(EditCondition="bIsEditable", HideEditConditionToggle))
 	FString Description;
 
-private:
+protected:
+	/** The name of this member, which is user-facing. */
+	UPROPERTY()
+	FString Name;
+	
 	/** A GUID that uniquely identifies this member within its graph. */
 	UPROPERTY()
 	FGuid Guid;
@@ -85,6 +116,8 @@ public:
 
 	//~ Begin UMovieGraphMember interface
 	virtual bool IsDeletable() const override { return !bIsGlobal; }
+	virtual bool CanRename(const FText& InNewName, FText& OutError) const override;
+	virtual bool SetMemberName(const FString& InNewName) override;
 	//~ End UMovieGraphMember interface
 
 public:
@@ -127,7 +160,11 @@ class MOVIERENDERPIPELINECORE_API UMovieGraphInput : public UMovieGraphInterface
 public:
 	UMovieGraphInput() = default;
 
+	//~ Begin UMovieGraphMember interface
 	virtual bool IsDeletable() const override;
+	virtual bool CanRename(const FText& InNewName, FText& OutError) const override;
+	virtual bool SetMemberName(const FString& InNewName) override;
+	//~ End UMovieGraphMember interface
 
 public:
 #if WITH_EDITOR
@@ -150,7 +187,11 @@ class MOVIERENDERPIPELINECORE_API UMovieGraphOutput : public UMovieGraphInterfac
 public:
 	UMovieGraphOutput() = default;
 
+	//~ Begin UMovieGraphMember interface
 	virtual bool IsDeletable() const override;
+	virtual bool CanRename(const FText& InNewName, FText& OutError) const override;
+	virtual bool SetMemberName(const FString& InNewName) override;
+	//~ End UMovieGraphMember interface
 
 public:
 #if WITH_EDITOR
diff --git a/Engine/Plugins/MovieScene/MovieRenderPipeline/Source/MovieRenderPipelineEditor/Private/Customizations/Graph/GraphMemberCustomization.h b/Engine/Plugins/MovieScene/MovieRenderPipeline/Source/MovieRenderPipelineEditor/Private/Customizations/Graph/GraphMemberCustomization.h
--- a/Engine/Plugins/MovieScene/MovieRenderPipeline/Source/MovieRenderPipelineEditor/Private/Customizations/Graph/GraphMemberCustomization.h
+++ b/Engine/Plugins/MovieScene/MovieRenderPipeline/Source/MovieRenderPipelineEditor/Private/Customizations/Graph/GraphMemberCustomization.h
@@ -2,10 +2,13 @@
 
 #pragma once
 
+#include "DetailCategoryBuilder.h"
 #include "DetailLayoutBuilder.h"
+#include "DetailWidgetRow.h"
 #include "IDetailCustomization.h"
 #include "Graph/MovieGraphConfig.h"
 #include "PropertyHandle.h"
+#include "Widgets/Input/SEditableTextBox.h"
 
 #define LOCTEXT_NAMESPACE "MoviePipelineEditor"
 
@@ -25,6 +28,25 @@ protected:
 		CustomizeDetails(*DetailBuilder);
 	}
 
+	void OnNameChanged(const FText& InText, UMovieGraphMember* MovieGraphMember) const
+	{
+		NameEditableTextBox->SetError(FText::GetEmpty());
+
+		FText Error;
+		if (!MovieGraphMember->CanRename(InText, Error))
+		{
+			NameEditableTextBox->SetError(Error);
+		}
+	}
+	
+	void OnNameCommitted(const FText& InText, ETextCommit::Type Arg, UMovieGraphMember* MovieGraphMember) const
+	{
+		if (MovieGraphMember->SetMemberName(InText.ToString()))
+		{
+			NameEditableTextBox->SetError(FText::GetEmpty());
+		}
+	}
+
 	virtual void CustomizeDetails(IDetailLayoutBuilder& DetailBuilder) override
 	{
 		TArray<TWeakObjectPtr<UObject>> ObjectsBeingCustomized;
@@ -54,9 +76,39 @@ protected:
 					DetailBuilder.EditDefaultProperty(ValueProperty)->IsEnabled(Variable->IsEditable());
 				}
 			}
+
+			UMovieGraphMember* MemberObject = Cast<UMovieGraphMember>(CustomizedObject);
+			if (!MemberObject)
+			{
+				continue;
+			}
+
+			// Add a custom row for the Name property (to allow for proper validation)
+			IDetailCategoryBuilder& GeneralCategory = DetailBuilder.EditCategory("General");
+			GeneralCategory.AddCustomRow(FText::GetEmpty())
+			.NameContent()
+			[
+				SNew(STextBlock)
+				.Text(LOCTEXT("MemberPropertyLabel_Name", "Name"))
+				.Font(DetailBuilder.GetDetailFont())
+			]
+			.ValueContent()
+			[
+				SAssignNew(NameEditableTextBox, SEditableTextBox)
+				.Text_Lambda([MemberObject]() { return FText::FromString(MemberObject->GetMemberName()); })
+				.OnTextChanged(this, &FGraphMemberCustomization::OnNameChanged, MemberObject)
+				.OnTextCommitted(this, &FGraphMemberCustomization::OnNameCommitted, MemberObject)
+				.IsReadOnly_Lambda([MemberObject]() { return !MemberObject->IsEditable(); })
+				.SelectAllTextWhenFocused(true)
+				.Font(IDetailLayoutBuilder::GetDetailFont())
+			];
 		}
 	}
 	//~ End IDetailCustomization interface
+
+private:
+	/** Text box for the "Name" property. */
+	TSharedPtr<SEditableTextBox> NameEditableTextBox;
 };
 
 #undef LOCTEXT_NAMESPACE
\ No newline at end of file
diff --git a/Engine/Plugins/MovieScene/MovieRenderPipeline/Source/MovieRenderPipelineEditor/Private/Graph/MovieEdGraphVariableNode.cpp b/Engine/Plugins/MovieScene/MovieRenderPipeline/Source/MovieRenderPipelineEditor/Private/Graph/MovieEdGraphVariableNode.cpp
--- a/Engine/Plugins/MovieScene/MovieRenderPipeline/Source/MovieRenderPipelineEditor/Private/Graph/MovieEdGraphVariableNode.cpp
+++ b/Engine/Plugins/MovieScene/MovieRenderPipeline/Source/MovieRenderPipelineEditor/Private/Graph/MovieEdGraphVariableNode.cpp
@@ -25,7 +25,7 @@ void UMoviePipelineEdGraphVariableNode::AllocateDefaultPins()
 		const TArray<TObjectPtr<UMovieGraphPin>>& OutputPins = RuntimeNode->GetOutputPins();
 		if (!OutputPins.IsEmpty())
 		{
-			CreatePin(EGPD_Output, GetPinType(OutputPins[0].Get()), FName(VariableNode->GetVariable()->Name));
+			CreatePin(EGPD_Output, GetPinType(OutputPins[0].Get()), FName(VariableNode->GetVariable()->GetMemberName()));
 		}
 	}
 }
diff --git a/Engine/Plugins/MovieScene/MovieRenderPipeline/Source/MovieRenderPipelineEditor/Private/Graph/MovieGraphSchema.cpp b/Engine/Plugins/MovieScene/MovieRenderPipeline/Source/MovieRenderPipelineEditor/Private/Graph/MovieGraphSchema.cpp
--- a/Engine/Plugins/MovieScene/MovieRenderPipeline/Source/MovieRenderPipelineEditor/Private/Graph/MovieGraphSchema.cpp
+++ b/Engine/Plugins/MovieScene/MovieRenderPipeline/Source/MovieRenderPipelineEditor/Private/Graph/MovieGraphSchema.cpp
@@ -136,7 +136,7 @@ void UMovieGraphSchema::GetGraphContextActions(FGraphContextMenuBuilder& Context
 	const bool bIncludeGlobal = true;
 	for (const UMovieGraphVariable* Variable : RuntimeGraph->GetVariables(bIncludeGlobal))
 	{
-		const FText Name = FText::Format(LOCTEXT("CreateVariable_Name", "Get {0}"), FText::FromString(Variable->Name));
+		const FText Name = FText::Format(LOCTEXT("CreateVariable_Name", "Get {0}"), FText::FromString(Variable->GetMemberName()));
 		const FText Category = Variable->IsGlobal() ? LOCTEXT("CreateGlobalVariable_Category", "Global Variables") : LOCTEXT("CreateVariable_Category", "Variables");
 		const FText Tooltip = LOCTEXT("CreateVariable_Tooltip", "Create an accessor node for this variable.");
 		
diff --git a/Engine/Plugins/MovieScene/MovieRenderPipeline/Source/MovieRenderPipelineEditor/Private/Graph/SMovieGraphMembersTabContent.cpp b/Engine/Plugins/MovieScene/MovieRenderPipeline/Source/MovieRenderPipelineEditor/Private/Graph/SMovieGraphMembersTabContent.cpp
--- a/Engine/Plugins/MovieScene/MovieRenderPipeline/Source/MovieRenderPipelineEditor/Private/Graph/SMovieGraphMembersTabContent.cpp
+++ b/Engine/Plugins/MovieScene/MovieRenderPipeline/Source/MovieRenderPipelineEditor/Private/Graph/SMovieGraphMembersTabContent.cpp
@@ -209,7 +209,7 @@ void SMovieGraphMembersTabContent::CollectAllActions(FGraphActionListBuilderBase
 	// Creates a new action in the action menu under a specific section w/ the provided action target
 	auto AddToActionMenu = [&ActionMenuBuilder](UMovieGraphMember* ActionTarget, const EActionSection Section, const FText& Category) -> void
 	{
-		const FText MemberActionDesc = FText::FromString(ActionTarget->Name);
+		const FText MemberActionDesc = FText::FromString(ActionTarget->GetMemberName());
 		const FText MemberActionTooltip;
 		const FText MemberActionKeywords;
 		const int32 MemberActionSectionID = static_cast<int32>(Section);
@@ -271,7 +271,7 @@ bool SMovieGraphMembersTabContent::ActionMatchesName(FEdGraphSchemaAction* InAct
 {
 	if (const UMovieGraphMember* Member = UE::MovieGraph::Private::GetMemberFromAction(InAction))
 	{
-		return InName == Member->Name;
+		return InName == Member->GetMemberName();
 	}
 
 	return false;
@@ -403,7 +403,7 @@ void SMovieGraphMembersTabContent::RefreshMemberActions(UMovieGraphMember* Updat
 	// automatically based on action name, but the member may have been renamed (so an explicit re-select is needed).
 	if (SelectedMember && (SelectedMember == UpdatedMember))
 	{
-		ActionMenu->SelectItemByName(FName(SelectedMember->Name));
+		ActionMenu->SelectItemByName(FName(SelectedMember->GetMemberName()));
 	}
 }
 
