Movie Pipeline: Fix deprecation of MoviePipelineMasterConfig.h in favor of MoviePipelinePrimaryConfig.h

#tests Tested that #include MoviePipelineMasterConfig.h gives a warning but still compiles

diff --git a/Engine/Plugins/MovieScene/MovieRenderPipeline/Source/MovieRenderPipelineCore/Public/MoviePipelineMasterConfig.h b/Engine/Plugins/MovieScene/MovieRenderPipeline/Source/MovieRenderPipelineCore/Public/MoviePipelineMasterConfig.h
new file mode 100644
--- /dev/null
+++ b/Engine/Plugins/MovieScene/MovieRenderPipeline/Source/MovieRenderPipelineCore/Public/MoviePipelineMasterConfig.h
@@ -0,0 +1,17 @@
+/**
+ * DEPRECATED MoviePipelineMasterConfig.h Use MoviePipelinePrimaryConfig.h instead 
+ */
+
+#pragma once
+
+#ifdef _MSC_VER
+	#pragma message(__FILE__"(9): warning: use MoviePipelinePrimaryConfig.h instead of MoviePipelineMasterConfig.h")
+#else
+	#pragma message("#include MoviePipelinePrimaryConfig.h instead of MoviePipelineMasterConfig")
+#endif
+
+#include "MoviePipelinePrimaryConfig.h"
+
+ // For backwards compatibility 5.2
+UE_DEPRECATED(5.2, "UMoviePipelineMasterConfig is deprecated. Please use UMoviePipelinePrimaryConfig.")
+typedef UMoviePipelinePrimaryConfig UMoviePipelineMasterConfig;
\ No newline at end of file
diff --git a/Engine/Plugins/MovieScene/MovieRenderPipeline/Source/MovieRenderPipelineEditor/Public/MoviePipelineConfigFactory.h b/Engine/Plugins/MovieScene/MovieRenderPipeline/Source/MovieRenderPipelineEditor/Public/MoviePipelineConfigFactory.h
--- a/Engine/Plugins/MovieScene/MovieRenderPipeline/Source/MovieRenderPipelineEditor/Public/MoviePipelineConfigFactory.h
+++ b/Engine/Plugins/MovieScene/MovieRenderPipeline/Source/MovieRenderPipelineEditor/Public/MoviePipelineConfigFactory.h
@@ -5,6 +5,10 @@
 #include "Factories/Factory.h"
 #include "MoviePipelineConfigFactory.generated.h"
 
+// For backwards compatibility 5.2
+UE_DEPRECATED(5.2, "UMoviePipelineMasterConfigFactory is deprecated. Please use UMoviePipelinePrimaryConfigFactory.")
+typedef UMoviePipelinePrimaryConfigFactory UMoviePipelineMasterConfigFactory;
+
 // Forward Declare
 UCLASS(BlueprintType)
 class MOVIERENDERPIPELINEEDITOR_API UMoviePipelinePrimaryConfigFactory : public UFactory
