MoviePipeline: UI should set the Graph Config not the Graph Preset for now so that it actually gets picked up by the runtime render.

diff --git a/Engine/Plugins/MovieScene/MovieRenderPipeline/Source/MovieRenderPipelineEditor/Private/Widgets/SMoviePipelineQueueEditor.cpp b/Engine/Plugins/MovieScene/MovieRenderPipeline/Source/MovieRenderPipelineEditor/Private/Widgets/SMoviePipelineQueueEditor.cpp
--- a/Engine/Plugins/MovieScene/MovieRenderPipeline/Source/MovieRenderPipelineEditor/Private/Widgets/SMoviePipelineQueueEditor.cpp
+++ b/Engine/Plugins/MovieScene/MovieRenderPipeline/Source/MovieRenderPipelineEditor/Private/Widgets/SMoviePipelineQueueEditor.cpp
@@ -286,7 +286,7 @@ public:
 
 		// Note: Setting the graph preset will transition the job to use a graph-based configuration
 		UMovieGraphConfig* NewGraph = NewObject<UMovieGraphConfig>(Job);
-		Job->SetGraphPreset(NewGraph);
+		Job->SetGraphConfig(NewGraph);
 	}
 
 	EVisibility GetPrimaryConfigModifiedVisibility() const
@@ -926,7 +926,7 @@ struct FMoviePipelineShotItem : IMoviePipelineQueueTreeItem
 
 		// Note: Setting the graph preset will transition the job to use a graph-based configuration
 		UMovieGraphConfig* NewGraph = NewObject<UMovieGraphConfig>(Shot);
-		Shot->SetGraphPreset(NewGraph);
+		Shot->SetGraphConfig(NewGraph);
 	}
 
 	EVisibility GetShotConfigModifiedVisibility() const
