MoviePipeline: Fixed crash when creating a new node via a connection (autowire).

The graph was being rebuilt in between creating the node and autowiring it, causing pointers to become invalid when autowiring.

diff --git a/Engine/Plugins/MovieScene/MovieRenderPipeline/Source/MovieRenderPipelineCore/Private/Graph/MovieGraphConfig.cpp b/Engine/Plugins/MovieScene/MovieRenderPipeline/Source/MovieRenderPipelineCore/Private/Graph/MovieGraphConfig.cpp
--- a/Engine/Plugins/MovieScene/MovieRenderPipeline/Source/MovieRenderPipelineCore/Private/Graph/MovieGraphConfig.cpp
+++ b/Engine/Plugins/MovieScene/MovieRenderPipeline/Source/MovieRenderPipelineCore/Private/Graph/MovieGraphConfig.cpp
@@ -479,10 +479,6 @@ void UMovieGraphConfig::AddNode(UMovieGraphNode* InNode)
 	InNode->Rename(nullptr, this);
 
 	AllNodes.Add(InNode);
-
-#if WITH_EDITOR
-	OnGraphChangedDelegate.Broadcast();
-#endif
 }
 
 bool UMovieGraphConfig::RemoveNodes(TArray<UMovieGraphNode*> InNodes)
