MoviePipeline: Added checks to ensure validity of the Globals branch in the graph.

* Only allow Globals branches to be connected to other Globals branches.
* Non-Globals branches cannot be connected to Globals branches.
* Nodes that are designated as Globals-only cannot be connected to non-Globals branches (and vice-versa).
* Added an additional rename validation check for members so they cannot be named 'Globals'.

diff --git a/Engine/Plugins/MovieScene/MovieRenderPipeline/Source/MovieRenderPipelineCore/Private/Graph/MovieGraphConfig.cpp b/Engine/Plugins/MovieScene/MovieRenderPipeline/Source/MovieRenderPipelineCore/Private/Graph/MovieGraphConfig.cpp
--- a/Engine/Plugins/MovieScene/MovieRenderPipeline/Source/MovieRenderPipelineCore/Private/Graph/MovieGraphConfig.cpp
+++ b/Engine/Plugins/MovieScene/MovieRenderPipeline/Source/MovieRenderPipelineCore/Private/Graph/MovieGraphConfig.cpp
@@ -12,8 +12,6 @@
 
 #define LOCTEXT_NAMESPACE "MovieGraphConfig"
 
-static const FName MovieGraphGlobalsMemberName("Globals");
-
 UMovieGraphConfig* UMovieGraphMember::GetOwningGraph() const
 {
 	return GetTypedOuter<UMovieGraphConfig>();
@@ -41,6 +39,12 @@ bool UMovieGraphMember::CanRename(const FText& InNewName, FText& OutError) const
 		return false;
 	}
 
+	if (InNewName.ToString() == UMovieGraphNode::GlobalsPinNameString)
+	{
+		OutError = LOCTEXT("InvalidMemberRename_Globals", "The name cannot be 'Globals'.");
+		return false;
+	}
+
 	const FString NewNameString = InNewName.ToString();
 	if (!FName::IsValidXName(NewNameString, InvalidChars, &OutError))
 	{
@@ -91,7 +95,7 @@ void UMovieGraphVariable::PostEditChangeProperty(FPropertyChangedEvent& Property
 bool UMovieGraphInput::IsDeletable() const
 {
 	// The input is deletable as long as it's not the Globals input
-	return Name != MovieGraphGlobalsMemberName;
+	return Name != UMovieGraphNode::GlobalsPinNameString;
 }
 
 bool UMovieGraphInput::CanRename(const FText& InNewName, FText& OutError) const
@@ -134,7 +138,7 @@ void UMovieGraphInput::PostEditChangeProperty(FPropertyChangedEvent& PropertyCha
 bool UMovieGraphOutput::IsDeletable() const
 {
 	// The output is deletable as long as it's not the Globals output
-	return Name != MovieGraphGlobalsMemberName;
+	return Name != UMovieGraphNode::GlobalsPinNameString;
 }
 
 bool UMovieGraphOutput::CanRename(const FText& InNewName, FText& OutError) const
@@ -254,19 +258,21 @@ void UMovieGraphConfig::AddDefaultMembers()
 {
 	const bool InputGlobalsExists = Inputs.ContainsByPredicate([](const UMovieGraphMember* Member)
 	{
-		return Member && (Member->GetMemberName() == MovieGraphGlobalsMemberName);
+		return Member && (Member->GetMemberName() == UMovieGraphNode::GlobalsPinNameString);
 	});
 
 	const bool OutputGlobalsExists = Outputs.ContainsByPredicate([](const UMovieGraphMember* Member)
 	{
-		return Member && (Member->GetMemberName() == MovieGraphGlobalsMemberName);
+		return Member && (Member->GetMemberName() == UMovieGraphNode::GlobalsPinNameString);
 	});
 
 	// Ensure there is a Globals input member
 	if (!InputGlobalsExists)
 	{
 		UMovieGraphInput* NewInput = AddInput();
-		NewInput->SetMemberName(MovieGraphGlobalsMemberName.ToString());
+
+		// Don't call SetMemberName() here, because that will reject setting the name to Globals
+		NewInput->Name = UMovieGraphNode::GlobalsPinNameString;
 
 		InputNode->UpdatePins();
 	}
@@ -275,7 +281,9 @@ void UMovieGraphConfig::AddDefaultMembers()
 	if (!OutputGlobalsExists)
 	{
 		UMovieGraphOutput* NewOutput = AddOutput();
-		NewOutput->SetMemberName(MovieGraphGlobalsMemberName.ToString());
+
+		// Don't call SetMemberName() here, because that will reject setting the name to Globals
+		NewOutput->Name = UMovieGraphNode::GlobalsPinNameString;
 
 		OutputNode->UpdatePins();
 	}
@@ -777,6 +785,40 @@ FBoolProperty* UMovieGraphConfig::FindOverridePropertyForRealProperty(UClass* In
 	return nullptr;
 }
 
+void UMovieGraphConfig::VisitUpstreamNodes(UMovieGraphNode* FromNode, const FVisitNodesCallback& VisitCallback) const
+{
+	TSet<UMovieGraphNode*> VisitedNodes;
+	VisitUpstreamNodes_Recursive(FromNode, VisitCallback, VisitedNodes);
+}
+
+void UMovieGraphConfig::VisitDownstreamNodes(UMovieGraphNode* FromNode, const FVisitNodesCallback& VisitCallback) const
+{
+	TSet<UMovieGraphNode*> VisitedNodes;
+	VisitDownstreamNodes_Recursive(FromNode, VisitCallback, VisitedNodes);
+}
+
+TArray<FString> UMovieGraphConfig::GetDownstreamBranchNames(UMovieGraphNode* FromNode, const UMovieGraphPin* FromPin) const
+{
+	TArray<FString> BranchNames;
+
+	// FromNode itself might be the Outputs node, so check before visiting the downstream nodes
+	if (FromNode->IsA<UMovieGraphOutputNode>() && FromPin)
+	{
+		BranchNames.Add(FromPin->Properties.Label.ToString());
+	}
+
+	VisitDownstreamNodes(FromNode, FVisitNodesCallback::CreateLambda(
+		[&BranchNames](UMovieGraphNode* VisitedNode, const UMovieGraphPin* VisitedPin)
+		{
+			if (VisitedNode->IsA<UMovieGraphOutputNode>() && VisitedPin)
+			{
+				BranchNames.Add(VisitedPin->Properties.Label.ToString());
+			}
+		}));
+
+	return BranchNames;
+}
+
 void UMovieGraphConfig::InitializeFlattenedNode(UMovieGraphNode* InNode)
 {
 	// We go through each of the bOverride_ properties on this new instance and set
@@ -999,6 +1041,62 @@ void UMovieGraphConfig::CreateFlattenedGraph_Recursive(UMovieGraphEvaluatedConfi
 	}
 }
 
+void UMovieGraphConfig::VisitUpstreamNodes_Recursive(UMovieGraphNode* FromNode,	const FVisitNodesCallback& VisitCallback, TSet<UMovieGraphNode*>& VisitedNodes) const
+{
+	if (VisitedNodes.Contains(FromNode))
+	{
+		// Cycle detected, stop recursing down this path. This is not necessarily an error, so don't log. For example:
+		//  |----| ---> N2 
+		//  | N1 | ---> N3 
+		//  |----| ---> N4
+		// Where nodes N2, N3, and N4 are nodes feeding out of node N1. N2 may have visited N1 already, so when N3 visits
+		// N1, just stop.
+		return;
+	}
+	
+	VisitedNodes.Add(FromNode);
+	
+	for (const UMovieGraphPin* Pin : FromNode->GetInputPins())
+	{
+		for (const UMovieGraphPin* ConnectedPin : Pin->GetAllConnectedPins())
+		{
+			if (ConnectedPin->Properties.bIsBranch)
+			{
+				VisitCallback.ExecuteIfBound(ConnectedPin->Node, ConnectedPin);
+				VisitUpstreamNodes_Recursive(ConnectedPin->Node, VisitCallback, VisitedNodes);
+			}
+		}
+	}
+}
+
+void UMovieGraphConfig::VisitDownstreamNodes_Recursive(UMovieGraphNode* FromNode, const FVisitNodesCallback& VisitCallback, TSet<UMovieGraphNode*>& VisitedNodes) const
+{
+	if (VisitedNodes.Contains(FromNode))
+	{
+		// Cycle detected, stop recursing down this path. This is not necessarily an error, so don't log. For example:
+		// N1 ---> |----|
+		// N2 ---> | N4 |
+		// N3 ---> |----|
+		// Where nodes N1, N2, and N3 are nodes feeding into node N4. N1 may have visited N4 already, so when N2 visits
+		// N4, just stop.
+		return;
+	}
+
+	VisitedNodes.Add(FromNode);
+
+	for (const UMovieGraphPin* Pin : FromNode->GetOutputPins())
+	{
+		for (const UMovieGraphPin* ConnectedPin : Pin->GetAllConnectedPins())
+		{
+			if (ConnectedPin->Properties.bIsBranch)
+			{
+				VisitCallback.ExecuteIfBound(ConnectedPin->Node, ConnectedPin);
+				VisitDownstreamNodes_Recursive(ConnectedPin->Node, VisitCallback, VisitedNodes);
+			}
+		}
+	}
+}
+
 UMovieGraphEvaluatedConfig* UMovieGraphConfig::CreateFlattenedGraph(const FMovieGraphTraversalContext& InContext)
 {
 	TRACE_CPUPROFILER_EVENT_SCOPE(MRQ_CreateFlattenedGraph);
diff --git a/Engine/Plugins/MovieScene/MovieRenderPipeline/Source/MovieRenderPipelineCore/Private/Graph/MovieGraphNode.cpp b/Engine/Plugins/MovieScene/MovieRenderPipeline/Source/MovieRenderPipelineCore/Private/Graph/MovieGraphNode.cpp
--- a/Engine/Plugins/MovieScene/MovieRenderPipeline/Source/MovieRenderPipelineCore/Private/Graph/MovieGraphNode.cpp
+++ b/Engine/Plugins/MovieScene/MovieRenderPipeline/Source/MovieRenderPipelineCore/Private/Graph/MovieGraphNode.cpp
@@ -6,6 +6,7 @@
 #include "UObject/TextProperty.h"
 
 FName UMovieGraphNode::GlobalsPinName = FName("Globals");
+FString UMovieGraphNode::GlobalsPinNameString = FString("Globals");
 
 namespace UE::MovieGraph::Private
 {
diff --git a/Engine/Plugins/MovieScene/MovieRenderPipeline/Source/MovieRenderPipelineCore/Private/Graph/MovieGraphPin.cpp b/Engine/Plugins/MovieScene/MovieRenderPipeline/Source/MovieRenderPipelineCore/Private/Graph/MovieGraphPin.cpp
--- a/Engine/Plugins/MovieScene/MovieRenderPipeline/Source/MovieRenderPipelineCore/Private/Graph/MovieGraphPin.cpp
+++ b/Engine/Plugins/MovieScene/MovieRenderPipeline/Source/MovieRenderPipelineCore/Private/Graph/MovieGraphPin.cpp
@@ -139,4 +139,19 @@ UMovieGraphPin* UMovieGraphPin::GetFirstConnectedPin() const
 	}
 
 	return nullptr;
+}
+
+TArray<UMovieGraphPin*> UMovieGraphPin::GetAllConnectedPins() const
+{
+	TArray<UMovieGraphPin*> ConnectedPins;
+
+	for (UMovieGraphEdge* Edge : Edges)
+	{
+		if (ensureMsgf(Edge, TEXT("Null edge found when trying to get connected pin!")))
+		{
+			ConnectedPins.Add(Edge->GetOtherPin(this));
+		}
+	}
+
+	return ConnectedPins;
 }
\ No newline at end of file
diff --git a/Engine/Plugins/MovieScene/MovieRenderPipeline/Source/MovieRenderPipelineCore/Public/Graph/MovieGraphConfig.h b/Engine/Plugins/MovieScene/MovieRenderPipeline/Source/MovieRenderPipelineCore/Public/Graph/MovieGraphConfig.h
--- a/Engine/Plugins/MovieScene/MovieRenderPipeline/Source/MovieRenderPipelineCore/Public/Graph/MovieGraphConfig.h
+++ b/Engine/Plugins/MovieScene/MovieRenderPipeline/Source/MovieRenderPipelineCore/Public/Graph/MovieGraphConfig.h
@@ -407,6 +407,13 @@ class MOVIERENDERPIPELINECORE_API UMovieGraphConfig : public UObject
 public:
 	UMovieGraphConfig();
 
+	/**
+	 * Callback for when a node is visited. The node is the node being visited, and the pin is the pin which the node
+	 * was accessed by (eg, if visiting downstream nodes, the pin will be the input pin that connects to the node that
+	 * the traversal started from, or the node that was previously visited).
+	 */
+	DECLARE_DELEGATE_TwoParams(FVisitNodesCallback, UMovieGraphNode*, const UMovieGraphPin*);
+
 	//~ UObject interface
 	virtual void PostLoad() override;
 	//~ End UObject interface
@@ -481,6 +488,21 @@ public:
 	/** Given a class and FProperty that belongs to that class, search for a FBoolProperty that matches the name "bOverride_<name of InRealProperty>. */
 	static FBoolProperty* FindOverridePropertyForRealProperty(UClass* InClass, const FProperty* InRealProperty);
 
+	/**
+	 * Visits all nodes upstream from FromNode, running VisitCallback on each one. Note this only follows branch connections,
+	 * and does not recurse into subgraphs.
+	 */
+	void VisitUpstreamNodes(UMovieGraphNode* FromNode, const FVisitNodesCallback& VisitCallback) const;
+
+	/**
+	 * Visits all nodes downstream from FromNode, running VisitCallback on each one. Note this only follows branch connections,
+	 * and does not recurse into subgraphs.
+	 */
+	void VisitDownstreamNodes(UMovieGraphNode* FromNode, const FVisitNodesCallback& VisitCallback) const;
+
+	/** Determines the name(s) of the branches downstream from FromNode, starting at FromPin. */
+	TArray<FString> GetDownstreamBranchNames(UMovieGraphNode* FromNode, const UMovieGraphPin* FromPin) const;
+
 protected:
 	/** Copies properties in FromNode that are marked for override into ToNode, but only if ToNode doesn't already override that value. */
 	void CopyOverriddenProperties(UMovieGraphNode* FromNode, UMovieGraphNode* ToNode, const FMovieGraphTraversalContext* InContext);
@@ -491,6 +513,12 @@ protected:
 	/** Traverse the graph, generating a combined "flatten" graph as it goes. */
 	void CreateFlattenedGraph_Recursive(UMovieGraphEvaluatedConfig* InOwningConfig, FMovieGraphEvaluatedBranchConfig& OutBranchConfig, FMovieGraphEvaluationContext& InEvaluationContext, UMovieGraphPin* InPinToFollow);
 
+	/** Recursive helper for VisitUpstreamNodes(). */
+	void VisitUpstreamNodes_Recursive(UMovieGraphNode* FromNode, const FVisitNodesCallback& VisitCallback, TSet<UMovieGraphNode*>& VisitedNodes) const;
+
+	/** Recursive helper for VisitDownstreamNodes(). */
+	void VisitDownstreamNodes_Recursive(UMovieGraphNode* FromNode, const FVisitNodesCallback& VisitCallback, TSet<UMovieGraphNode*>& VisitedNodes) const;
+
 public:
 	// Names of global variables that are provided by the graph
 	static FName GlobalVariable_ShotName;
diff --git a/Engine/Plugins/MovieScene/MovieRenderPipeline/Source/MovieRenderPipelineCore/Public/Graph/MovieGraphNode.h b/Engine/Plugins/MovieScene/MovieRenderPipeline/Source/MovieRenderPipelineCore/Public/Graph/MovieGraphNode.h
--- a/Engine/Plugins/MovieScene/MovieRenderPipeline/Source/MovieRenderPipelineCore/Public/Graph/MovieGraphNode.h
+++ b/Engine/Plugins/MovieScene/MovieRenderPipeline/Source/MovieRenderPipelineCore/Public/Graph/MovieGraphNode.h
@@ -72,6 +72,14 @@ struct FMovieGraphPropertyInfo
 	}
 };
 
+/** Describes a restriction on what kind of branch a node can be created in within the graph. */
+enum class EMovieGraphBranchRestriction : uint8
+{
+	Any,			///< The node can be created in any type of branch
+	Globals,		///< The node must be created in the Globals branch
+	RenderLayer		///< The node must be created in a branch representing a render layer
+};
+
 /**
 * This is a base class for all nodes that can exist in the UMovieGraphConfig network.
 * In the editor, each node in the network will have an editor-only representation too 
@@ -88,6 +96,8 @@ class MOVIERENDERPIPELINECORE_API UMovieGraphNode : public UObject
 	
 public:
 	static FName GlobalsPinName;
+	static FString GlobalsPinNameString;
+	
 	UMovieGraphNode();
 
 	const TArray<TObjectPtr<UMovieGraphPin>>& GetInputPins() const { return InputPins; }
@@ -194,6 +204,9 @@ public:
 
 	/** Gets the GUID which uniquely identifies this node. */
 	const FGuid& GetGuid() const { return Guid; }
+	
+	/** Determines which types of branches the node can be created in. */
+	virtual EMovieGraphBranchRestriction GetBranchRestriction() const { return EMovieGraphBranchRestriction::Any; }
 
 #if WITH_EDITOR
 	int32 GetNodePosX() const { return NodePosX; }
diff --git a/Engine/Plugins/MovieScene/MovieRenderPipeline/Source/MovieRenderPipelineCore/Public/Graph/MovieGraphPin.h b/Engine/Plugins/MovieScene/MovieRenderPipeline/Source/MovieRenderPipelineCore/Public/Graph/MovieGraphPin.h
--- a/Engine/Plugins/MovieScene/MovieRenderPipeline/Source/MovieRenderPipelineCore/Public/Graph/MovieGraphPin.h
+++ b/Engine/Plugins/MovieScene/MovieRenderPipeline/Source/MovieRenderPipelineCore/Public/Graph/MovieGraphPin.h
@@ -76,6 +76,9 @@ public:
 	/** Gets the first pin connected to this pin. Returns nullptr if no valid connection exists. */
 	UMovieGraphPin* GetFirstConnectedPin() const;
 
+	/** Gets all connected pins. */
+	TArray<UMovieGraphPin*> GetAllConnectedPins() const;
+
 public:
 	// The node that this pin belongs to.
 	UPROPERTY(BlueprintReadOnly, Category = "Properties")
diff --git a/Engine/Plugins/MovieScene/MovieRenderPipeline/Source/MovieRenderPipelineCore/Public/Graph/Nodes/MovieGraphGlobalGameOverrides.h b/Engine/Plugins/MovieScene/MovieRenderPipeline/Source/MovieRenderPipelineCore/Public/Graph/Nodes/MovieGraphGlobalGameOverrides.h
--- a/Engine/Plugins/MovieScene/MovieRenderPipeline/Source/MovieRenderPipelineCore/Public/Graph/Nodes/MovieGraphGlobalGameOverrides.h
+++ b/Engine/Plugins/MovieScene/MovieRenderPipeline/Source/MovieRenderPipelineCore/Public/Graph/Nodes/MovieGraphGlobalGameOverrides.h
@@ -15,6 +15,8 @@ class MOVIERENDERPIPELINECORE_API UMovieGraphGlobalGameOverridesNode : public UM
 public:
 	UMovieGraphGlobalGameOverridesNode() = default;
 
+	virtual EMovieGraphBranchRestriction GetBranchRestriction() const override { return EMovieGraphBranchRestriction::Globals; }
+
 #if WITH_EDITOR
 	virtual FText GetNodeTitle(const bool bGetDescriptive = false) const override;
 	virtual FText GetMenuCategory() const override;
diff --git a/Engine/Plugins/MovieScene/MovieRenderPipeline/Source/MovieRenderPipelineEditor/Private/Graph/MovieGraphSchema.cpp b/Engine/Plugins/MovieScene/MovieRenderPipeline/Source/MovieRenderPipelineEditor/Private/Graph/MovieGraphSchema.cpp
--- a/Engine/Plugins/MovieScene/MovieRenderPipeline/Source/MovieRenderPipelineEditor/Private/Graph/MovieGraphSchema.cpp
+++ b/Engine/Plugins/MovieScene/MovieRenderPipeline/Source/MovieRenderPipelineEditor/Private/Graph/MovieGraphSchema.cpp
@@ -4,6 +4,7 @@
 #include "Graph/MovieGraphConfig.h"
 #include "Graph/Nodes/MovieGraphInputNode.h"
 #include "Graph/Nodes/MovieGraphOutputNode.h"
+#include "Graph/Nodes/MovieGraphSubgraphNode.h"
 #include "Graph/Nodes/MovieGraphVariableNode.h"
 #include "Graph/MovieGraphNode.h"
 #include "Graph/MovieEdGraph.h"
@@ -36,6 +37,38 @@ const FName UMovieGraphSchema::PC_SoftObject(TEXT("softobject"));
 const FName UMovieGraphSchema::PC_Class(TEXT("class"));
 const FName UMovieGraphSchema::PC_SoftClass(TEXT("softclass"));
 
+namespace UE::MovieGraph::Private
+{
+	UMovieGraphNode* GetGraphNodeFromEdPin(const UEdGraphPin* InPin)
+	{
+		const UMoviePipelineEdGraphNodeBase* EdGraphNode = CastChecked<UMoviePipelineEdGraphNodeBase>(InPin->GetOwningNode());
+
+		UMovieGraphNode* RuntimeNode = EdGraphNode->GetRuntimeNode();
+		check(RuntimeNode);
+
+		return RuntimeNode;
+	}
+
+	UMovieGraphPin* GetGraphPinFromEdPin(const UEdGraphPin* InPin)
+	{
+		const UMovieGraphNode* GraphNode = GetGraphNodeFromEdPin(InPin);
+		UMovieGraphPin* GraphPin = (InPin->Direction == EGPD_Input) ? GraphNode->GetInputPin(InPin->PinName) : GraphNode->GetOutputPin(InPin->PinName);
+		check(GraphPin);
+
+		return GraphPin;
+	}
+	
+	UMovieGraphConfig* GetGraphFromEdPin(const UEdGraphPin* InPin)
+	{
+		const UMovieGraphNode* RuntimeNode = GetGraphNodeFromEdPin(InPin);
+
+		UMovieGraphConfig* RuntimeGraph = RuntimeNode->GetGraph();
+		check(RuntimeGraph);
+
+		return RuntimeGraph;
+	}
+}
+
 void UMovieGraphSchema::CreateDefaultNodesForGraph(UEdGraph& Graph) const
 {
 	/*UMovieGraphConfig* RuntimeGraph = CastChecked<UMoviePipelineEdGraph>(&Graph)->GetPipelineGraph();
@@ -85,6 +118,78 @@ void UMovieGraphSchema::InitMoviePipelineNodeClasses()
 	MoviePipelineNodeClasses.Sort();
 }
 
+bool UMovieGraphSchema::IsConnectionToBranchAllowed(const UEdGraphPin* InputPin, const UEdGraphPin* OutputPin, FText& OutError) const
+{
+	UMovieGraphNode* ToNode = UE::MovieGraph::Private::GetGraphNodeFromEdPin(InputPin);
+	UMovieGraphNode* FromNode = UE::MovieGraph::Private::GetGraphNodeFromEdPin(OutputPin);
+	const UMovieGraphPin* ToPin = UE::MovieGraph::Private::GetGraphPinFromEdPin(InputPin);
+	const UMovieGraphPin* FromPin = UE::MovieGraph::Private::GetGraphPinFromEdPin(OutputPin);
+	const UMovieGraphConfig* GraphConfig = UE::MovieGraph::Private::GetGraphFromEdPin(InputPin);
+
+	// Early exit if a Globals output is not being connected to a Globals input.
+	if ((FromPin->Properties.Label == UMovieGraphNode::GlobalsPinName) && (ToPin->Properties.Label != UMovieGraphNode::GlobalsPinName))
+	{
+		OutError = NSLOCTEXT("MoviePipeline", "GlobalsBranchMismatchError", "Globals branches can only be connected to other Globals branches.");
+		return false;
+	}
+
+	// Early exit if a non-Globals output is being connected to a Globals input.
+	if ((FromPin->Properties.Label != UMovieGraphNode::GlobalsPinName) && (ToPin->Properties.Label == UMovieGraphNode::GlobalsPinName))
+	{
+		OutError = FText::Format(
+			NSLOCTEXT("MoviePipeline", "NonGlobalsBranchMismatchError", "The branch '{0}' cannot be connected to the Globals branch."),
+				FText::FromName(FromPin->Properties.Label));
+		return false;
+	}
+
+	// Get all upstream nodes for FromNode -- these are the nodes that need to be checked for branch restrictions.
+	// FromNode itself also needs to be part of the validation checks.
+	TArray<UMovieGraphNode*> UpstreamNodes = {FromNode};
+	GraphConfig->VisitUpstreamNodes(FromNode, UMovieGraphConfig::FVisitNodesCallback::CreateLambda(
+		[&UpstreamNodes](UMovieGraphNode* VisitedNode, const UMovieGraphPin* VisitedPin)
+		{
+			UpstreamNodes.Add(VisitedNode);
+		}));
+
+	// Determine which branch(es) are connected to this node downstream.
+	const TArray<FString> UpstreamBranchNames = GraphConfig->GetDownstreamBranchNames(ToNode, ToPin);
+	const bool bGlobalsIsOnlyDownstreamBranch = (UpstreamBranchNames.Num() == 1) && (UpstreamBranchNames[0] == UMovieGraphNode::GlobalsPinNameString);
+	const bool bGlobalsIsDownstream = UpstreamBranchNames.Contains(UMovieGraphNode::GlobalsPinNameString);
+
+	// Error out if any of the upstream nodes cannot be connected to the downstream branch(es).
+	for (const UMovieGraphNode* UpstreamNode : UpstreamNodes)
+	{
+		if (UpstreamNode->GetBranchRestriction() == EMovieGraphBranchRestriction::Globals)
+		{
+			// Globals-specific nodes have to be connected such that the only upstream branch is Globals, OR the group
+			// of nodes is not yet connected to a branch (and this restriction will be enforced when they are).
+			if (bGlobalsIsOnlyDownstreamBranch || UpstreamBranchNames.IsEmpty())
+			{
+				continue;
+			}
+
+			OutError = FText::Format(
+				NSLOCTEXT("MoviePipeline", "GlobalsBranchRestrictionError", "The node '{0}' can only be connected to the Globals branch."),
+					UpstreamNode->GetNodeTitle());
+			return false;
+		}
+
+		// Check that render-layer-only nodes aren't connected to Globals.
+		if (UpstreamNode->GetBranchRestriction() == EMovieGraphBranchRestriction::RenderLayer)
+		{
+			if (bGlobalsIsDownstream)
+			{
+				OutError = FText::Format(
+					NSLOCTEXT("MoviePipeline", "RenderLayerBranchRestrictionError", "The node '{0}' can only be connected to a render layer branch."),
+						UpstreamNode->GetNodeTitle());
+				return false;
+			}
+		}
+	}
+
+	return true;
+}
+
 void UMovieGraphSchema::GetGraphContextActions(FGraphContextMenuBuilder& ContextMenuBuilder) const
 {
 	InitMoviePipelineNodeClasses();
@@ -171,6 +276,13 @@ const FPinConnectionResponse UMovieGraphSchema::CanCreateConnection(const UEdGra
 		return FPinConnectionResponse(CONNECT_RESPONSE_DISALLOW, NSLOCTEXT("MoviePipeline", "PinDirectionMismatchError", "Directions are not compatible!"));
 	}
 
+	// Determine if the connection would violate branch restrictions enforced by the nodes involved in the connection.
+	FText BranchRestrictionError;
+	if (!IsConnectionToBranchAllowed(InputPin, OutputPin, BranchRestrictionError))
+	{
+		return FPinConnectionResponse(CONNECT_RESPONSE_DISALLOW, BranchRestrictionError);
+	}
+
 	// We don't allow multiple things to be connected to an Input Pin
 	if(InputPin->HasAnyConnections())
 	{
diff --git a/Engine/Plugins/MovieScene/MovieRenderPipeline/Source/MovieRenderPipelineEditor/Private/Graph/MovieGraphSchema.h b/Engine/Plugins/MovieScene/MovieRenderPipeline/Source/MovieRenderPipelineEditor/Private/Graph/MovieGraphSchema.h
--- a/Engine/Plugins/MovieScene/MovieRenderPipeline/Source/MovieRenderPipelineEditor/Private/Graph/MovieGraphSchema.h
+++ b/Engine/Plugins/MovieScene/MovieRenderPipeline/Source/MovieRenderPipelineEditor/Private/Graph/MovieGraphSchema.h
@@ -36,6 +36,13 @@ public:
 	//~ End EdGraphSchema Interface
 	static void InitMoviePipelineNodeClasses();
 
+private:
+	/**
+	 * Determines if the connection between InputPin and OutputPin follows branch restriction rules. OutError is populated
+	 * with an error if the connection should be rejected.
+	 */
+	bool IsConnectionToBranchAllowed(const UEdGraphPin* InputPin, const UEdGraphPin* OutputPin, FText& OutError) const;
+
 public:
 	// Allowed "PinCategory" values for use on EdGraphPin
 	static const FName PC_Branch;
