[Core]
* Fixed compile errors when compiling .h files in isolation

diff --git a/Engine/Plugins/Editor/AssetManagerEditor/Source/AssetManagerEditor/Private/Insights/Common/AsyncOperationProgress.h b/Engine/Plugins/Editor/AssetManagerEditor/Source/AssetManagerEditor/Private/Insights/Common/AsyncOperationProgress.h
--- a/Engine/Plugins/Editor/AssetManagerEditor/Source/AssetManagerEditor/Private/Insights/Common/AsyncOperationProgress.h
+++ b/Engine/Plugins/Editor/AssetManagerEditor/Source/AssetManagerEditor/Private/Insights/Common/AsyncOperationProgress.h
@@ -3,6 +3,7 @@
 #pragma once
 
 #include "CoreTypes.h"
+#include <atomic>
 
 namespace UE
 {
diff --git a/Engine/Plugins/Editor/AssetManagerEditor/Source/AssetManagerEditor/Private/Insights/Common/AsyncWorkUtils.h b/Engine/Plugins/Editor/AssetManagerEditor/Source/AssetManagerEditor/Private/Insights/Common/AsyncWorkUtils.h
--- a/Engine/Plugins/Editor/AssetManagerEditor/Source/AssetManagerEditor/Private/Insights/Common/AsyncWorkUtils.h
+++ b/Engine/Plugins/Editor/AssetManagerEditor/Source/AssetManagerEditor/Private/Insights/Common/AsyncWorkUtils.h
@@ -4,6 +4,8 @@
 
 #include "CoreTypes.h"
 
+class FText;
+
 namespace UE
 {
 namespace Insights
diff --git a/Engine/Plugins/Editor/GameplayTagsEditor/Source/GameplayTagsEditor/Public/GameplayTagContainerCustomizationOptions.h b/Engine/Plugins/Editor/GameplayTagsEditor/Source/GameplayTagsEditor/Public/GameplayTagContainerCustomizationOptions.h
--- a/Engine/Plugins/Editor/GameplayTagsEditor/Source/GameplayTagsEditor/Public/GameplayTagContainerCustomizationOptions.h
+++ b/Engine/Plugins/Editor/GameplayTagsEditor/Source/GameplayTagsEditor/Public/GameplayTagContainerCustomizationOptions.h
@@ -3,6 +3,7 @@
 #pragma once
 
 #include "HAL/Platform.h"
+#include "Misc/CoreMiscDefines.h"
 
 struct UE_DEPRECATED(5.3, "Options are not used anymore to customize the property customization behavior.") GAMEPLAYTAGSEDITOR_API FGameplayTagContainerCustomizationOptions
 {
diff --git a/Engine/Plugins/Editor/GameplayTagsEditor/Source/GameplayTagsEditor/Public/GameplayTagCustomizationOptions.h b/Engine/Plugins/Editor/GameplayTagsEditor/Source/GameplayTagsEditor/Public/GameplayTagCustomizationOptions.h
--- a/Engine/Plugins/Editor/GameplayTagsEditor/Source/GameplayTagsEditor/Public/GameplayTagCustomizationOptions.h
+++ b/Engine/Plugins/Editor/GameplayTagsEditor/Source/GameplayTagsEditor/Public/GameplayTagCustomizationOptions.h
@@ -3,6 +3,7 @@
 #pragma once
 
 #include "HAL/Platform.h"
+#include "Misc/CoreMiscDefines.h"
 
 struct UE_DEPRECATED(5.3, "Options are not used anymore to customize the property customization behavior.") GAMEPLAYTAGSEDITOR_API FGameplayTagCustomizationOptions
 {
diff --git a/Engine/Plugins/Editor/ObjectMixer/ObjectMixer/Source/ObjectMixer/Public/Views/List/Modes/SFilterClassMenuItem.h b/Engine/Plugins/Editor/ObjectMixer/ObjectMixer/Source/ObjectMixer/Public/Views/List/Modes/SFilterClassMenuItem.h
--- a/Engine/Plugins/Editor/ObjectMixer/ObjectMixer/Source/ObjectMixer/Public/Views/List/Modes/SFilterClassMenuItem.h
+++ b/Engine/Plugins/Editor/ObjectMixer/ObjectMixer/Source/ObjectMixer/Public/Views/List/Modes/SFilterClassMenuItem.h
@@ -7,6 +7,7 @@
 #include "IContentBrowserSingleton.h"
 #include "UObject/ObjectPtr.h"
 #include "UObject/UObjectGlobals.h"
+#include "Views/List/Modes/ObjectMixerOutlinerMode.h"
 #include "Widgets/Input/SButton.h"
 #include "Widgets/Input/SCheckBox.h"
 
diff --git a/Engine/Plugins/Experimental/Animation/PoseSearch/Source/Editor/Private/PoseSearchDebuggerDatabaseRow.h b/Engine/Plugins/Experimental/Animation/PoseSearch/Source/Editor/Private/PoseSearchDebuggerDatabaseRow.h
--- a/Engine/Plugins/Experimental/Animation/PoseSearch/Source/Editor/Private/PoseSearchDebuggerDatabaseRow.h
+++ b/Engine/Plugins/Experimental/Animation/PoseSearch/Source/Editor/Private/PoseSearchDebuggerDatabaseRow.h
@@ -2,6 +2,7 @@
 
 #pragma once
 
+#include "Components/ListView.h"
 #include "PoseSearchDebuggerDatabaseColumns.h"
 #include "PoseSearchDebuggerDatabaseView.h"
 #include "Widgets/Views/STableRow.h"
diff --git a/Engine/Plugins/Experimental/Mutable/Source/MutableRuntime/Private/MuR/OpLayoutPack.h b/Engine/Plugins/Experimental/Mutable/Source/MutableRuntime/Private/MuR/OpLayoutPack.h
--- a/Engine/Plugins/Experimental/Mutable/Source/MutableRuntime/Private/MuR/OpLayoutPack.h
+++ b/Engine/Plugins/Experimental/Mutable/Source/MutableRuntime/Private/MuR/OpLayoutPack.h
@@ -3,13 +3,15 @@
 #pragma once
 
 #include "MuR/Layout.h"
+#include "MuR/MutableMath.h"
 #include "MuR/MutableTrace.h"
 #include "MuR/Platform.h"
 #include "GenericPlatform/GenericPlatformMath.h"
-
+#include <string>
 
 namespace mu
 {
+	typedef std::string string;
 
 	//---------------------------------------------------------------------------------------------
     // inline void SimpleLayoutPack( Layout* pResult, const Layout* pSource )
diff --git a/Engine/Plugins/Media/MediaIOFramework/Source/MediaIOCore/Public/MediaIOCoreUtilities.h b/Engine/Plugins/Media/MediaIOFramework/Source/MediaIOCore/Public/MediaIOCoreUtilities.h
--- a/Engine/Plugins/Media/MediaIOFramework/Source/MediaIOCore/Public/MediaIOCoreUtilities.h
+++ b/Engine/Plugins/Media/MediaIOFramework/Source/MediaIOCore/Public/MediaIOCoreUtilities.h
@@ -9,6 +9,8 @@
 #include "Logging/LogVerbosity.h"
 #include "HAL/PlatformTime.h"
 #include "HAL/CriticalSection.h"
+#include "Misc/OutputDeviceRedirector.h"
+#include "Misc/ScopeLock.h"
 #include "Misc/Timespan.h"
 
 namespace UE::MediaIO
diff --git a/Engine/Plugins/Media/MediaIOFramework/Source/MediaIOCore/Public/OrderedAsyncGate.h b/Engine/Plugins/Media/MediaIOFramework/Source/MediaIOCore/Public/OrderedAsyncGate.h
--- a/Engine/Plugins/Media/MediaIOFramework/Source/MediaIOCore/Public/OrderedAsyncGate.h
+++ b/Engine/Plugins/Media/MediaIOFramework/Source/MediaIOCore/Public/OrderedAsyncGate.h
@@ -3,7 +3,7 @@
 #pragma once
 
 #include "GenericPlatform/GenericPlatformProcess.h"
-
+#include "HAL/PlatformProcess.h"
 #include <atomic>
 
 namespace UE::MediaIO
diff --git a/Engine/Plugins/MovieScene/MovieRenderPipeline/Source/MovieRenderPipelineCore/Public/MoviePipelineMasterConfig.h b/Engine/Plugins/MovieScene/MovieRenderPipeline/Source/MovieRenderPipelineCore/Public/MoviePipelineMasterConfig.h
--- a/Engine/Plugins/MovieScene/MovieRenderPipeline/Source/MovieRenderPipelineCore/Public/MoviePipelineMasterConfig.h
+++ b/Engine/Plugins/MovieScene/MovieRenderPipeline/Source/MovieRenderPipelineCore/Public/MoviePipelineMasterConfig.h
@@ -1,5 +1,6 @@
 // Copyright Epic Games, Inc. All Rights Reserved.
 
+// HEADER_UNIT_SKIP - Deprecated
 /**
  * DEPRECATED MoviePipelineMasterConfig.h Use MoviePipelinePrimaryConfig.h instead 
  */
diff --git a/Engine/Plugins/MovieScene/MovieRenderPipeline/Source/MovieRenderPipelineEditor/Private/Customizations/ConsoleVariableSettingCustomization.h b/Engine/Plugins/MovieScene/MovieRenderPipeline/Source/MovieRenderPipelineEditor/Private/Customizations/ConsoleVariableSettingCustomization.h
--- a/Engine/Plugins/MovieScene/MovieRenderPipeline/Source/MovieRenderPipelineEditor/Private/Customizations/ConsoleVariableSettingCustomization.h
+++ b/Engine/Plugins/MovieScene/MovieRenderPipeline/Source/MovieRenderPipelineEditor/Private/Customizations/ConsoleVariableSettingCustomization.h
@@ -4,6 +4,7 @@
 
 #include "DetailCategoryBuilder.h"
 #include "DetailLayoutBuilder.h"
+#include "Framework/MultiBox/MultiBoxBuilder.h"
 #include "IDetailChildrenBuilder.h"
 #include "IDetailCustomization.h"
 #include "IDetailGroup.h"
@@ -11,6 +12,8 @@
 #include "PropertyCustomizationHelpers.h"
 #include "PropertyHandle.h"
 #include "Sections/MovieSceneConsoleVariableTrackInterface.h"
+#include "Widgets/Input/SComboButton.h"
+#include "Widgets/Input/SEditableTextBox.h"
 
 #define LOCTEXT_NAMESPACE "SMoviePipelineEditor"
 
diff --git a/Engine/Plugins/MovieScene/MovieRenderPipeline/Source/MovieRenderPipelineEditor/Private/Customizations/Graph/GraphMemberCustomization.h b/Engine/Plugins/MovieScene/MovieRenderPipeline/Source/MovieRenderPipelineEditor/Private/Customizations/Graph/GraphMemberCustomization.h
--- a/Engine/Plugins/MovieScene/MovieRenderPipeline/Source/MovieRenderPipelineEditor/Private/Customizations/Graph/GraphMemberCustomization.h
+++ b/Engine/Plugins/MovieScene/MovieRenderPipeline/Source/MovieRenderPipelineEditor/Private/Customizations/Graph/GraphMemberCustomization.h
@@ -9,6 +9,7 @@
 #include "Graph/MovieGraphConfig.h"
 #include "PropertyHandle.h"
 #include "Widgets/Input/SEditableTextBox.h"
+#include "Widgets/Text/STextBlock.h"
 
 #define LOCTEXT_NAMESPACE "MoviePipelineEditor"
 
diff --git a/Engine/Plugins/MovieScene/MovieRenderPipeline/Source/MovieRenderPipelineEditor/Private/Graph/SMovieGraphActiveRenderSettingsTabContent.h b/Engine/Plugins/MovieScene/MovieRenderPipeline/Source/MovieRenderPipelineEditor/Private/Graph/SMovieGraphActiveRenderSettingsTabContent.h
--- a/Engine/Plugins/MovieScene/MovieRenderPipeline/Source/MovieRenderPipelineEditor/Private/Graph/SMovieGraphActiveRenderSettingsTabContent.h
+++ b/Engine/Plugins/MovieScene/MovieRenderPipeline/Source/MovieRenderPipelineEditor/Private/Graph/SMovieGraphActiveRenderSettingsTabContent.h
@@ -2,6 +2,7 @@
 
 #pragma once
 
+#include "Graph/MovieGraphConfig.h"
 #include "UObject/StrongObjectPtr.h"
 #include "Widgets/SCompoundWidget.h"
 #include "Widgets/Views/STableRow.h"
@@ -9,7 +10,6 @@
 
 class UMoviePipelineExecutorJob;
 class UMoviePipelineQueue;
-class UMovieGraphEvaluatedConfig;
 class UMovieGraphNode;
 
 /**
diff --git a/Engine/Plugins/Runtime/CommonUI/Source/CommonUI/Public/Input/UIActionBinding.h b/Engine/Plugins/Runtime/CommonUI/Source/CommonUI/Public/Input/UIActionBinding.h
--- a/Engine/Plugins/Runtime/CommonUI/Source/CommonUI/Public/Input/UIActionBinding.h
+++ b/Engine/Plugins/Runtime/CommonUI/Source/CommonUI/Public/Input/UIActionBinding.h
@@ -2,6 +2,7 @@
 
 #pragma once
 
+#include "Containers/Ticker.h"
 #include "Engine/DataTable.h"
 #include "Input/CommonUIInputSettings.h" // IWYU pragma: keep
 #include "Input/UIActionBindingHandle.h"
diff --git a/Engine/Plugins/Runtime/MeshModelingToolset/Source/ModelingOperators/Public/SkinningOps/SkinBindingOp.h b/Engine/Plugins/Runtime/MeshModelingToolset/Source/ModelingOperators/Public/SkinningOps/SkinBindingOp.h
--- a/Engine/Plugins/Runtime/MeshModelingToolset/Source/ModelingOperators/Public/SkinningOps/SkinBindingOp.h
+++ b/Engine/Plugins/Runtime/MeshModelingToolset/Source/ModelingOperators/Public/SkinningOps/SkinBindingOp.h
@@ -3,11 +3,11 @@
 #pragma once
 
 #include "ModelingOperators.h"
+#include "ReferenceSkeleton.h"
 #include "SkeletalMeshAttributes.h"
 
 
 struct FReferenceSkeleton;
-struct FMeshBoneInfo;
 
 namespace UE::Geometry
 {
diff --git a/Engine/Plugins/Runtime/Metasound/Source/MetasoundFrontend/Private/Interfaces/MetasoundFrontendInterfaceBindingRegistryPrivate.h b/Engine/Plugins/Runtime/Metasound/Source/MetasoundFrontend/Private/Interfaces/MetasoundFrontendInterfaceBindingRegistryPrivate.h
--- a/Engine/Plugins/Runtime/Metasound/Source/MetasoundFrontend/Private/Interfaces/MetasoundFrontendInterfaceBindingRegistryPrivate.h
+++ b/Engine/Plugins/Runtime/Metasound/Source/MetasoundFrontend/Private/Interfaces/MetasoundFrontendInterfaceBindingRegistryPrivate.h
@@ -2,6 +2,7 @@
 #pragma once
 
 
+#include "Interfaces/MetasoundFrontendInterfaceBindingRegistry.h"
 #include "Interfaces/MetasoundFrontendInterfaceRegistry.h"
 #include "MetasoundFrontendDocument.h"
 #include "MetasoundFrontendRegistryTransaction.h"
diff --git a/Engine/Plugins/Runtime/Metasound/Source/MetasoundFrontend/Public/MetasoundFrontendArchetypeRegistry.h b/Engine/Plugins/Runtime/Metasound/Source/MetasoundFrontend/Public/MetasoundFrontendArchetypeRegistry.h
--- a/Engine/Plugins/Runtime/Metasound/Source/MetasoundFrontend/Public/MetasoundFrontendArchetypeRegistry.h
+++ b/Engine/Plugins/Runtime/Metasound/Source/MetasoundFrontend/Public/MetasoundFrontendArchetypeRegistry.h
@@ -2,6 +2,8 @@
 // DEPRECATED MetasoundFrontendArchetypeRegistry.h. Use Interfaces/MetasoundFrontendInterfaceRegistry.h instead
 #pragma once
 
+// HEADER_UNIT_SKIP - Deprecated
+
 #ifdef _MSC_VER
 #pragma message(__FILE__"(9): warning: use Interfaces/MetasoundFrontendInterfaceRegistry.h instead instead of MetasoundFrontendArchetypeRegistry.h")
 #else
diff --git a/Engine/Plugins/Runtime/Metasound/Source/MetasoundFrontend/Public/MetasoundParameterPackFixedArray.h b/Engine/Plugins/Runtime/Metasound/Source/MetasoundFrontend/Public/MetasoundParameterPackFixedArray.h
--- a/Engine/Plugins/Runtime/Metasound/Source/MetasoundFrontend/Public/MetasoundParameterPackFixedArray.h
+++ b/Engine/Plugins/Runtime/Metasound/Source/MetasoundFrontend/Public/MetasoundParameterPackFixedArray.h
@@ -2,6 +2,7 @@
 
 #pragma once
 
+#include "Containers/Array.h"
 #include "Templates/TypeCompatibleBytes.h"
 
 namespace Metasound
diff --git a/Engine/Plugins/Runtime/NetworkPrediction/Source/NetworkPrediction/Public/NetworkPredictionSerialization.h b/Engine/Plugins/Runtime/NetworkPrediction/Source/NetworkPrediction/Public/NetworkPredictionSerialization.h
--- a/Engine/Plugins/Runtime/NetworkPrediction/Source/NetworkPrediction/Public/NetworkPredictionSerialization.h
+++ b/Engine/Plugins/Runtime/NetworkPrediction/Source/NetworkPrediction/Public/NetworkPredictionSerialization.h
@@ -3,6 +3,7 @@
 #pragma once
 
 #include "Net/UnrealNetwork.h" // For MakeRelative
+#include "NetworkPredictionCVars.h"
 #include "NetworkPredictionReplicationProxy.h"
 #include "NetworkPredictionTickState.h"
 #include "NetworkPredictionTrace.h"
diff --git a/Engine/Plugins/TraceUtilities/Source/EditorTraceUtilities/Private/UnrealInsightsLauncher.h b/Engine/Plugins/TraceUtilities/Source/EditorTraceUtilities/Private/UnrealInsightsLauncher.h
--- a/Engine/Plugins/TraceUtilities/Source/EditorTraceUtilities/Private/UnrealInsightsLauncher.h
+++ b/Engine/Plugins/TraceUtilities/Source/EditorTraceUtilities/Private/UnrealInsightsLauncher.h
@@ -3,6 +3,7 @@
 #pragma once
 
 #include "Async/TaskGraphFwd.h"
+#include "Async/TaskGraphInterfaces.h"
 #include "CoreMinimal.h"
 
 class FUnrealInsightsLauncher : public TSharedFromThis<FUnrealInsightsLauncher>
diff --git a/Engine/Source/Developer/ShaderPreprocessor/Private/StbConfig.h b/Engine/Source/Developer/ShaderPreprocessor/Private/StbConfig.h
--- a/Engine/Source/Developer/ShaderPreprocessor/Private/StbConfig.h
+++ b/Engine/Source/Developer/ShaderPreprocessor/Private/StbConfig.h
@@ -2,6 +2,8 @@
 
 #pragma once
 
+#include <stddef.h>
+
 #ifdef __cplusplus
 extern "C"
 {
diff --git a/Engine/Source/Editor/LevelEditor/Private/SActorPilotViewportToolbar.h b/Engine/Source/Editor/LevelEditor/Private/SActorPilotViewportToolbar.h
--- a/Engine/Source/Editor/LevelEditor/Private/SActorPilotViewportToolbar.h
+++ b/Engine/Source/Editor/LevelEditor/Private/SActorPilotViewportToolbar.h
@@ -8,6 +8,7 @@
 #include "Widgets/DeclarativeSyntaxSupport.h"
 #include "Textures/SlateIcon.h"
 #include "Layout/Margin.h"
+#include "Widgets/Images/SImage.h"
 #include "Widgets/SBoxPanel.h"
 #include "SLevelViewport.h"
 #include "Widgets/Text/STextBlock.h"
diff --git a/Engine/Source/Editor/PropertyEditor/Private/UserInterface/PropertyTable/ObjectNameTableCellPresenter.h b/Engine/Source/Editor/PropertyEditor/Private/UserInterface/PropertyTable/ObjectNameTableCellPresenter.h
--- a/Engine/Source/Editor/PropertyEditor/Private/UserInterface/PropertyTable/ObjectNameTableCellPresenter.h
+++ b/Engine/Source/Editor/PropertyEditor/Private/UserInterface/PropertyTable/ObjectNameTableCellPresenter.h
@@ -2,6 +2,7 @@
 #pragma once
 
 #include "CoreMinimal.h"
+#include "IPropertyTable.h"
 #include "Widgets/SNullWidget.h"
 #include "Widgets/DeclarativeSyntaxSupport.h"
 #include "Layout/Margin.h"
diff --git a/Engine/Source/Editor/PropertyEditor/Public/IStructureDataProvider.h b/Engine/Source/Editor/PropertyEditor/Public/IStructureDataProvider.h
--- a/Engine/Source/Editor/PropertyEditor/Public/IStructureDataProvider.h
+++ b/Engine/Source/Editor/PropertyEditor/Public/IStructureDataProvider.h
@@ -1,7 +1,9 @@
-﻿// Copyright Epic Games, Inc. All Rights Reserved.
+// Copyright Epic Games, Inc. All Rights Reserved.
 #pragma once
 
-class FStructOnScope;
+#include "Containers/Array.h"
+#include "Templates/SharedPointer.h"
+#include "UObject/StructOnScope.h"
 
 //-----------------------------------------------------------------------------
 //	IStructureDataProvider - Used to provide struct data for a property node.
diff --git a/Engine/Source/Editor/PropertyEditor/Public/NumericPropertyParams.h b/Engine/Source/Editor/PropertyEditor/Public/NumericPropertyParams.h
--- a/Engine/Source/Editor/PropertyEditor/Public/NumericPropertyParams.h
+++ b/Engine/Source/Editor/PropertyEditor/Public/NumericPropertyParams.h
@@ -3,6 +3,8 @@
 #pragma once
 
 #include "CoreMinimal.h"
+#include "Misc/Attribute.h"
+#include "UObject/UnrealType.h"
 
 template<typename NumericType>
 struct TNumericPropertyParams
diff --git a/Engine/Source/Editor/Sequencer/Private/SequencerSelectionCurveFilter.h b/Engine/Source/Editor/Sequencer/Private/SequencerSelectionCurveFilter.h
--- a/Engine/Source/Editor/Sequencer/Private/SequencerSelectionCurveFilter.h
+++ b/Engine/Source/Editor/Sequencer/Private/SequencerSelectionCurveFilter.h
@@ -5,6 +5,7 @@
 #include "CoreMinimal.h"
 #include "ISequencerModule.h"
 #include "MVVM/Extensions/IObjectBindingExtension.h"
+#include "MVVM/Selection/Selection.h"
 #include "MVVM/ViewModels/ViewModel.h"
 #include "Tree/CurveEditorTreeFilter.h"
 
diff --git a/Engine/Source/Runtime/Core/Private/AutoRTFM/ObjectPool.h b/Engine/Source/Runtime/Core/Private/AutoRTFM/ObjectPool.h
--- a/Engine/Source/Runtime/Core/Private/AutoRTFM/ObjectPool.h
+++ b/Engine/Source/Runtime/Core/Private/AutoRTFM/ObjectPool.h
@@ -2,6 +2,8 @@
 
 #pragma once
 
+#include "AutoRTFM/Utils.h"
+
 namespace AutoRTFM
 {
 
diff --git a/Engine/Source/Runtime/Core/Private/AutoRTFM/Transaction.h b/Engine/Source/Runtime/Core/Private/AutoRTFM/Transaction.h
--- a/Engine/Source/Runtime/Core/Private/AutoRTFM/Transaction.h
+++ b/Engine/Source/Runtime/Core/Private/AutoRTFM/Transaction.h
@@ -7,6 +7,7 @@
 #include "Stats.h"
 #include "TaggedPtr.h"
 #include "TaskArray.h"
+#include "Templates/Function.h"
 #include "WriteLog.h"
 #include "WriteLogBumpAllocator.h"
 
diff --git a/Engine/Source/Runtime/Core/Private/AutoRTFM/WriteLog.h b/Engine/Source/Runtime/Core/Private/AutoRTFM/WriteLog.h
--- a/Engine/Source/Runtime/Core/Private/AutoRTFM/WriteLog.h
+++ b/Engine/Source/Runtime/Core/Private/AutoRTFM/WriteLog.h
@@ -2,9 +2,9 @@
 
 #pragma once
 
-#include <stdint.h>
-
+#include "AutoRTFM/HitSet.h"
 #include "Utils.h"
+#include <stdint.h>
 
 namespace AutoRTFM
 {
diff --git a/Engine/Source/Runtime/Core/Public/Algo/Partition.h b/Engine/Source/Runtime/Core/Public/Algo/Partition.h
--- a/Engine/Source/Runtime/Core/Public/Algo/Partition.h
+++ b/Engine/Source/Runtime/Core/Public/Algo/Partition.h
@@ -2,6 +2,8 @@
 
 #pragma once
 
+#include "HAL/Platform.h"
+
 namespace Algo
 {
 	/**
diff --git a/Engine/Source/Runtime/Core/Public/Containers/ConsumeAllMpmcQueue.h b/Engine/Source/Runtime/Core/Public/Containers/ConsumeAllMpmcQueue.h
--- a/Engine/Source/Runtime/Core/Public/Containers/ConsumeAllMpmcQueue.h
+++ b/Engine/Source/Runtime/Core/Public/Containers/ConsumeAllMpmcQueue.h
@@ -5,6 +5,7 @@
 #include "CoreTypes.h"
 #include "HAL/UnrealMemory.h"
 #include "Templates/TypeCompatibleBytes.h"
+#include "Templates/UnrealTemplate.h"
 #include <atomic>
 #include <type_traits>
 
diff --git a/Engine/Source/Runtime/Core/Public/Misc/LowLevelTestAdapter.h b/Engine/Source/Runtime/Core/Public/Misc/LowLevelTestAdapter.h
--- a/Engine/Source/Runtime/Core/Public/Misc/LowLevelTestAdapter.h
+++ b/Engine/Source/Runtime/Core/Public/Misc/LowLevelTestAdapter.h
@@ -4,8 +4,8 @@
 
 #if !WITH_LOW_LEVEL_TESTS
 
-#include "CoreTypes.h"
-
+#include "Containers/StringConv.h"
+#include "Containers/UnrealString.h"
 #include <sstream>
 #include <vector>
 
diff --git a/Engine/Source/Runtime/Core/Public/Misc/ReverseIterate.h b/Engine/Source/Runtime/Core/Public/Misc/ReverseIterate.h
--- a/Engine/Source/Runtime/Core/Public/Misc/ReverseIterate.h
+++ b/Engine/Source/Runtime/Core/Public/Misc/ReverseIterate.h
@@ -2,6 +2,7 @@
 
 #pragma once
 
+#include "HAL/Platform.h"
 #include <iterator>
 
 /**
diff --git a/Engine/Source/Runtime/Core/Public/Virtualization/VirtualizationTypes.h b/Engine/Source/Runtime/Core/Public/Virtualization/VirtualizationTypes.h
--- a/Engine/Source/Runtime/Core/Public/Virtualization/VirtualizationTypes.h
+++ b/Engine/Source/Runtime/Core/Public/Virtualization/VirtualizationTypes.h
@@ -3,6 +3,7 @@
 #pragma once
 
 #include "HAL/Platform.h"
+#include "Misc/CoreMiscDefines.h"
 #include "Misc/EnumClassFlags.h"
 
 namespace UE::Virtualization
diff --git a/Engine/Source/Runtime/CoreUObject/Public/Serialization/ArchiveScriptReferenceCollector.h b/Engine/Source/Runtime/CoreUObject/Public/Serialization/ArchiveScriptReferenceCollector.h
--- a/Engine/Source/Runtime/CoreUObject/Public/Serialization/ArchiveScriptReferenceCollector.h
+++ b/Engine/Source/Runtime/CoreUObject/Public/Serialization/ArchiveScriptReferenceCollector.h
@@ -5,6 +5,7 @@
 #include "CoreMinimal.h"
 #include "UObject/Field.h"
 #include "UObject/Object.h"
+#include "UObject/ObjectPtr.h"
 #include "UObject/UObjectGlobals.h"
 #include "Serialization/ArchiveUObject.h"
 
diff --git a/Engine/Source/Runtime/CoreUObject/Public/Templates/IsTObjectPtr.h b/Engine/Source/Runtime/CoreUObject/Public/Templates/IsTObjectPtr.h
--- a/Engine/Source/Runtime/CoreUObject/Public/Templates/IsTObjectPtr.h
+++ b/Engine/Source/Runtime/CoreUObject/Public/Templates/IsTObjectPtr.h
@@ -2,6 +2,9 @@
 
 #pragma once
 
+template <typename T>
+struct TObjectPtr;
+
 template <typename T>
 struct TIsTObjectPtr
 {
diff --git a/Engine/Source/Runtime/CoreUObject/Public/UObject/ErrorException.h b/Engine/Source/Runtime/CoreUObject/Public/UObject/ErrorException.h
--- a/Engine/Source/Runtime/CoreUObject/Public/UObject/ErrorException.h
+++ b/Engine/Source/Runtime/CoreUObject/Public/UObject/ErrorException.h
@@ -2,4 +2,6 @@
 
 #pragma once
 
+// HEADER_UNIT_SKIP - Deprecated
+
 #pragma message (__FILE__"(5): warning: ErrorException.h has been deprecated")
diff --git a/Engine/Source/Runtime/CoreUObject/Public/UObject/ObjectHandleDefines.h b/Engine/Source/Runtime/CoreUObject/Public/UObject/ObjectHandleDefines.h
--- a/Engine/Source/Runtime/CoreUObject/Public/UObject/ObjectHandleDefines.h
+++ b/Engine/Source/Runtime/CoreUObject/Public/UObject/ObjectHandleDefines.h
@@ -2,6 +2,8 @@
 
 #pragma once
 
+// HEADER_UNIT_SKIP - Included through other header
+
 #include "Misc/CoreMiscDefines.h"
 
 #ifndef UE_WITH_OBJECT_HANDLE_LATE_RESOLVE
diff --git a/Engine/Source/Runtime/CoreUObject/Public/UObject/PackedObjectRef.h b/Engine/Source/Runtime/CoreUObject/Public/UObject/PackedObjectRef.h
--- a/Engine/Source/Runtime/CoreUObject/Public/UObject/PackedObjectRef.h
+++ b/Engine/Source/Runtime/CoreUObject/Public/UObject/PackedObjectRef.h
@@ -3,6 +3,8 @@
 #pragma once
 #include "UObject/ObjectHandleDefines.h"
 
+// HEADER_UNIT_SKIP - Included through other header
+
 #if UE_WITH_OBJECT_HANDLE_LATE_RESOLVE
 
 namespace UE::CoreUObject::Private
diff --git a/Engine/Source/Runtime/CoreUObject/Public/UObject/PersistentObjectPtr.h b/Engine/Source/Runtime/CoreUObject/Public/UObject/PersistentObjectPtr.h
--- a/Engine/Source/Runtime/CoreUObject/Public/UObject/PersistentObjectPtr.h
+++ b/Engine/Source/Runtime/CoreUObject/Public/UObject/PersistentObjectPtr.h
@@ -6,6 +6,7 @@
 
 #pragma once
 
+#include "UObject/Object.h"
 #include "UObject/WeakObjectPtr.h"
 
 /**
diff --git a/Engine/Source/Runtime/CoreUObject/Public/UObject/PrintStaleReferencesOptions.h b/Engine/Source/Runtime/CoreUObject/Public/UObject/PrintStaleReferencesOptions.h
--- a/Engine/Source/Runtime/CoreUObject/Public/UObject/PrintStaleReferencesOptions.h
+++ b/Engine/Source/Runtime/CoreUObject/Public/UObject/PrintStaleReferencesOptions.h
@@ -2,6 +2,9 @@
 
 #pragma once
 
+#include "HAL/Platform.h"
+#include "Misc/EnumClassFlags.h"
+
 /* Options for FReferenceChainSearch::FindAndPrintStaleReferencesToObject function */
 enum class EPrintStaleReferencesOptions
 {
diff --git a/Engine/Source/Runtime/Engine/Classes/Animation/AnimBoneDecompressionData.h b/Engine/Source/Runtime/Engine/Classes/Animation/AnimBoneDecompressionData.h
--- a/Engine/Source/Runtime/Engine/Classes/Animation/AnimBoneDecompressionData.h
+++ b/Engine/Source/Runtime/Engine/Classes/Animation/AnimBoneDecompressionData.h
@@ -2,7 +2,7 @@
 
 #pragma once
 
-#include "CoreMinimal.h"
+#include "AnimEncoding.h"
 
 namespace UE::Anim
 {
diff --git a/Engine/Source/Runtime/Engine/Public/ObjectCacheEventSink.h b/Engine/Source/Runtime/Engine/Public/ObjectCacheEventSink.h
--- a/Engine/Source/Runtime/Engine/Public/ObjectCacheEventSink.h
+++ b/Engine/Source/Runtime/Engine/Public/ObjectCacheEventSink.h
@@ -9,6 +9,7 @@ class UActorComponent;
 class UMaterialInterface;
 class UPrimitiveComponent;
 class UStaticMeshComponent;
+class UTexture;
 
 struct ENGINE_API FObjectCacheEventSink
 {
diff --git a/Engine/Source/Runtime/Engine/Public/SkeletalMeshSceneProxy.h b/Engine/Source/Runtime/Engine/Public/SkeletalMeshSceneProxy.h
--- a/Engine/Source/Runtime/Engine/Public/SkeletalMeshSceneProxy.h
+++ b/Engine/Source/Runtime/Engine/Public/SkeletalMeshSceneProxy.h
@@ -9,6 +9,8 @@
 FSkeletalMeshSceneProxy
 -----------------------------------------------------------------------------*/
 
+class FSkeletalMeshObject;
+
 /**
  * A skeletal mesh component scene proxy.
  */
@@ -114,7 +116,7 @@ public:
 
 protected:
 	AActor* Owner;
-	class FSkeletalMeshObject* MeshObject;
+	FSkeletalMeshObject* MeshObject;
 	FSkeletalMeshRenderData* SkeletalMeshRenderData;
 
 	/** The points to the skeletal mesh and physics assets are purely for debug purposes. Access is NOT thread safe! */
diff --git a/Engine/Source/Runtime/Experimental/Chaos/Public/Chaos/ImplicitObjectTransformed.h b/Engine/Source/Runtime/Experimental/Chaos/Public/Chaos/ImplicitObjectTransformed.h
--- a/Engine/Source/Runtime/Experimental/Chaos/Public/Chaos/ImplicitObjectTransformed.h
+++ b/Engine/Source/Runtime/Experimental/Chaos/Public/Chaos/ImplicitObjectTransformed.h
@@ -3,6 +3,7 @@
 
 #include "Chaos/Box.h"
 #include "Chaos/ImplicitObject.h"
+#include "Chaos/ShapeInstanceFwd.h"
 #include "Chaos/Transform.h"
 #include "ChaosArchive.h"
 #include "Templates/EnableIf.h"
diff --git a/Engine/Source/Runtime/Experimental/Iris/Core/Private/Iris/Core/MemoryLayoutUtil.h b/Engine/Source/Runtime/Experimental/Iris/Core/Private/Iris/Core/MemoryLayoutUtil.h
--- a/Engine/Source/Runtime/Experimental/Iris/Core/Private/Iris/Core/MemoryLayoutUtil.h
+++ b/Engine/Source/Runtime/Experimental/Iris/Core/Private/Iris/Core/MemoryLayoutUtil.h
@@ -2,7 +2,8 @@
 
 #pragma once
 
-#include "HAL/Platform.h"
+#include "Math/UnrealMathUtility.h"
+#include "Templates/AlignmentTemplates.h"
 
 /** Helper to deal with manually setting up a block of memory containing multiple allocations. */
 struct FMemoryLayoutUtil
diff --git a/Engine/Source/Runtime/Experimental/Iris/Core/Private/Iris/Serialization/IntNetSerializerBase.h b/Engine/Source/Runtime/Experimental/Iris/Core/Private/Iris/Serialization/IntNetSerializerBase.h
--- a/Engine/Source/Runtime/Experimental/Iris/Core/Private/Iris/Serialization/IntNetSerializerBase.h
+++ b/Engine/Source/Runtime/Experimental/Iris/Core/Private/Iris/Serialization/IntNetSerializerBase.h
@@ -6,6 +6,8 @@
 #include "Iris/Serialization/BitPacking.h"
 #include "Iris/Serialization/NetBitStreamReader.h"
 #include "Iris/Serialization/NetBitStreamWriter.h"
+#include "Iris/Serialization/NetSerializationContext.h"
+#include "Iris/Serialization/NetSerializer.h"
 #include "Iris/Core/BitTwiddling.h"
 #include <type_traits>
 
diff --git a/Engine/Source/Runtime/GeometryCore/Public/Util/IteratorUtil.h b/Engine/Source/Runtime/GeometryCore/Public/Util/IteratorUtil.h
--- a/Engine/Source/Runtime/GeometryCore/Public/Util/IteratorUtil.h
+++ b/Engine/Source/Runtime/GeometryCore/Public/Util/IteratorUtil.h
@@ -4,6 +4,7 @@
 
 #pragma once
 #include "IndexTypes.h"
+#include "Templates/Function.h"
 
 namespace UE
 {
diff --git a/Engine/Source/Runtime/MovieScene/Public/EntitySystem/RelativePtr.h b/Engine/Source/Runtime/MovieScene/Public/EntitySystem/RelativePtr.h
--- a/Engine/Source/Runtime/MovieScene/Public/EntitySystem/RelativePtr.h
+++ b/Engine/Source/Runtime/MovieScene/Public/EntitySystem/RelativePtr.h
@@ -2,6 +2,7 @@
 
 #pragma once
 
+#include "Misc/AssertionMacros.h"
 #include <limits>
 #include <type_traits>
 
diff --git a/Engine/Source/Runtime/Net/Core/Classes/Net/Serialization/FastArraySerializerImplementation.h b/Engine/Source/Runtime/Net/Core/Classes/Net/Serialization/FastArraySerializerImplementation.h
--- a/Engine/Source/Runtime/Net/Core/Classes/Net/Serialization/FastArraySerializerImplementation.h
+++ b/Engine/Source/Runtime/Net/Core/Classes/Net/Serialization/FastArraySerializerImplementation.h
@@ -2,6 +2,8 @@
 
 #pragma once
 
+// HEADER_UNIT_SKIP - Bad include. This header can't compile standalone
+
 #if UE_WITH_IRIS
 #include "Iris/ReplicationSystem/FastArrayReplicationFragment.h"
 
diff --git a/Engine/Source/Runtime/Net/Core/Public/Net/Core/PropertyConditions/PropertyConditionsDelegates.h b/Engine/Source/Runtime/Net/Core/Public/Net/Core/PropertyConditions/PropertyConditionsDelegates.h
--- a/Engine/Source/Runtime/Net/Core/Public/Net/Core/PropertyConditions/PropertyConditionsDelegates.h
+++ b/Engine/Source/Runtime/Net/Core/Public/Net/Core/PropertyConditions/PropertyConditionsDelegates.h
@@ -10,6 +10,8 @@
 
 #include "Delegates/Delegate.h"
 
+enum ELifetimeCondition : int;
+
 namespace UE::Net::Private
 {
 
diff --git a/Engine/Source/Runtime/Online/BuildPatchServices/Public/BuildPatchMessage.h b/Engine/Source/Runtime/Online/BuildPatchServices/Public/BuildPatchMessage.h
--- a/Engine/Source/Runtime/Online/BuildPatchServices/Public/BuildPatchMessage.h
+++ b/Engine/Source/Runtime/Online/BuildPatchServices/Public/BuildPatchMessage.h
@@ -2,6 +2,7 @@
 #pragma once
 
 #include "Containers/UnrealString.h"
+#include "Templates/Function.h"
 
 namespace BuildPatchServices
 {
diff --git a/Engine/Source/Runtime/RHI/Public/RHIValidationCommon.h b/Engine/Source/Runtime/RHI/Public/RHIValidationCommon.h
--- a/Engine/Source/Runtime/RHI/Public/RHIValidationCommon.h
+++ b/Engine/Source/Runtime/RHI/Public/RHIValidationCommon.h
@@ -20,6 +20,8 @@ const bool GRHIValidationEnabled = false;
 
 #if ENABLE_RHI_VALIDATION
 
+class FRHIShader;
+class FRHIShaderResourceView;
 class FRHIUniformBuffer;
 class FRHIViewableResource;
 class FRHIUnorderedAccessView;
@@ -32,6 +34,7 @@ struct FRHITextureCreateDesc;
 struct FRHITransitionInfo;
 struct FRHIViewDesc;
 
+
 namespace RHIValidation
 {
 	struct FStaticUniformBuffers
@@ -976,7 +979,7 @@ namespace RHIValidation
 	extern RHI_API void* CaptureBacktrace();
 
 	/** Validates that the SRV is conform to what the shader expects */
-	extern RHI_API void ValidateShaderResourceView(const class FRHIShader* RHIShaderBase, uint32 BindIndex, class FRHIShaderResourceView* SRV);
+	extern RHI_API void ValidateShaderResourceView(const FRHIShader* RHIShaderBase, uint32 BindIndex, FRHIShaderResourceView* SRV);
 
 }
 
diff --git a/Engine/Source/Runtime/RHICore/Internal/RHIUniformBufferDataShared.h b/Engine/Source/Runtime/RHICore/Internal/RHIUniformBufferDataShared.h
--- a/Engine/Source/Runtime/RHICore/Internal/RHIUniformBufferDataShared.h
+++ b/Engine/Source/Runtime/RHICore/Internal/RHIUniformBufferDataShared.h
@@ -3,6 +3,7 @@
 #pragma once
 
 #include "RHI.h"
+#include "RenderGraphResources.h"
 
 namespace UE::RHICore
 {
diff --git a/Engine/Source/Runtime/Renderer/Public/RectLightTexture.h b/Engine/Source/Runtime/Renderer/Public/RectLightTexture.h
--- a/Engine/Source/Runtime/Renderer/Public/RectLightTexture.h
+++ b/Engine/Source/Runtime/Renderer/Public/RectLightTexture.h
@@ -2,6 +2,8 @@
 
 #pragma once
 
+#include "HAL/Platform.h"
+
 class UTexture;
 
 namespace RectLightAtlas
diff --git a/Engine/Source/Runtime/Slate/Public/Widgets/Views/IItemsSource.h b/Engine/Source/Runtime/Slate/Public/Widgets/Views/IItemsSource.h
--- a/Engine/Source/Runtime/Slate/Public/Widgets/Views/IItemsSource.h
+++ b/Engine/Source/Runtime/Slate/Public/Widgets/Views/IItemsSource.h
@@ -5,8 +5,7 @@
 #include "CoreMinimal.h"
 #include "Containers/ArrayView.h"
 #include "Containers/ObservableArray.h"
-
-class STableViewBase;
+#include "Widgets/Views/STableViewBase.h"
 
  namespace UE::Slate::ItemsSource
  {
