MoviePipeline: Fixed a crash where a Spawnable camera that had the spawn track partially keyed to not visible would crash when adding the job to the queue. Fixed an issue where camera names were not resolving correctly due to multi-camera support.

The fix is to only set the MainCameraIndex if we actually choose to add the camera, as otherwise we were running into situations where the MainCameraIndex would be set, but the camera would not be added, and it would result in an Out of Bounds array access.

Fixed an issue where camera names tokens did not always resolve correctly when using non multi-camera support. This was caused by the placeholder "main camera" entry being pushed into the sidecar array, and then some parts of the rendering were looking up data from the sidecar array (when only trying to render a single camera). Instead, when using single-camera rendering mode we should always use the existing Inner Name built into the UMoviePipelineExecutorShot, and runtime lookups should use the player camera (as Sequencer is animating onto this). We should only look at the sidecar data when rendering all cameras.

diff --git a/Engine/Plugins/MovieScene/MovieRenderPipeline/Source/MovieRenderPipelineRenderPasses/Private/MoviePipelineDeferredPasses.cpp b/Engine/Plugins/MovieScene/MovieRenderPipeline/Source/MovieRenderPipelineRenderPasses/Private/MoviePipelineDeferredPasses.cpp
--- a/Engine/Plugins/MovieScene/MovieRenderPipeline/Source/MovieRenderPipelineRenderPasses/Private/MoviePipelineDeferredPasses.cpp
+++ b/Engine/Plugins/MovieScene/MovieRenderPipeline/Source/MovieRenderPipelineRenderPasses/Private/MoviePipelineDeferredPasses.cpp
@@ -405,7 +405,10 @@ FSceneViewStateInterface* UMoviePipelineDeferredPassBase::GetSceneViewStateInter
 	UE::MoviePipeline::FDeferredPassRenderStatePayload* Payload = (UE::MoviePipeline::FDeferredPassRenderStatePayload*)OptPayload;
 	check(Payload);
 
-	FMultiCameraViewStateData& CameraData = CameraViewStateData[Payload->CameraIndex];
+	// When not using multi-camera support, CameraIndex is -1, but means "first entry" in the array.
+	int32 LocalCameraIndex = FMath::Clamp(Payload->CameraIndex, 0, Payload->CameraIndex);
+
+	FMultiCameraViewStateData& CameraData = CameraViewStateData[LocalCameraIndex];
 	if (FMultiCameraViewStateData::FPerTile* TileData = CameraData.TileData.Find(Payload->TileIndex))
 	{
 		return TileData->SceneViewStates[Payload->SceneViewIndex].GetReference();
@@ -419,12 +422,18 @@ void UMoviePipelineDeferredPassBase::GatherOutputPassesImpl(TArray<FMoviePipelin
 	// No super call here because multiple cameras makes this all complicated
 	// Super::GatherOutputPassesImpl(ExpectedRenderPasses);
 
+	UMoviePipelineExecutorShot* CurrentShot = GetPipeline()->GetActiveShotList()[GetPipeline()->GetCurrentShotIndex()];
+	UMoviePipelineCameraSetting* CameraSettings = GetPipeline()->FindOrAddSettingForShot<UMoviePipelineCameraSetting>(CurrentShot);
+
 	const int32 NumCameras = GetNumCamerasToRender();
 	for (int32 CameraIndex = 0; CameraIndex < NumCameras; CameraIndex++)
 	{
 		FMoviePipelinePassIdentifier PassIdentifierForCurrentCamera;
 		PassIdentifierForCurrentCamera.Name = PassIdentifier.Name;
-		PassIdentifierForCurrentCamera.CameraName = GetCameraName(CameraIndex);
+
+		// If we're not rendering all cameras, we need to pass -1 so we pick up the real camera name.
+		int32 LocalCameraIndex = CameraSettings->bRenderAllCameras ? CameraIndex : -1;
+		PassIdentifierForCurrentCamera.CameraName = GetCameraName(LocalCameraIndex);
 
 		// Add the default backbuffer
 		if (bRenderMainPass)
@@ -484,22 +493,28 @@ void UMoviePipelineDeferredPassBase::RenderSample_GameThreadImpl(const FMoviePip
 	Super::RenderSample_GameThreadImpl(InSampleState);
 
 	const int32 NumCameras = GetNumCamerasToRender();
+	UMoviePipelineExecutorShot* CurrentShot = GetPipeline()->GetActiveShotList()[GetPipeline()->GetCurrentShotIndex()];
+	UMoviePipelineCameraSetting* CameraSettings = GetPipeline()->FindOrAddSettingForShot<UMoviePipelineCameraSetting>(CurrentShot);
+
 	for (int32 CameraIndex = 0; CameraIndex < NumCameras; CameraIndex++)
 	{
 		FMoviePipelinePassIdentifier PassIdentifierForCurrentCamera;
 		PassIdentifierForCurrentCamera.Name = PassIdentifier.Name;
-		PassIdentifierForCurrentCamera.CameraName = GetCameraName(CameraIndex);
+
+		// If we're not rendering all cameras, we need to pass -1 so we pick up the real camera name.
+		int32 LocalCameraIndex = CameraSettings->bRenderAllCameras ? CameraIndex : -1;
+		PassIdentifierForCurrentCamera.CameraName = GetCameraName(LocalCameraIndex);
 
 		// Main Render Pass
 		if (bRenderMainPass)
 		{
-			FMoviePipelineRenderPassMetrics InOutSampleState = GetRenderPassMetricsForCamera(CameraIndex, InSampleState);
+			FMoviePipelineRenderPassMetrics InOutSampleState = GetRenderPassMetricsForCamera(LocalCameraIndex, InSampleState);
 			// InOutSampleState.OutputState.CameraCount = NumCameras;
-			InOutSampleState.OutputState.CameraIndex = CameraIndex;
-			InOutSampleState.OutputState.CameraNameOverride = GetCameraNameOverride(CameraIndex);
+			InOutSampleState.OutputState.CameraIndex = LocalCameraIndex;
+			InOutSampleState.OutputState.CameraNameOverride = GetCameraNameOverride(LocalCameraIndex);
 
 			UE::MoviePipeline::FDeferredPassRenderStatePayload Payload;
-			Payload.CameraIndex = CameraIndex;
+			Payload.CameraIndex = LocalCameraIndex;
 			Payload.TileIndex = InOutSampleState.TileIndexes;
 
 			// Main renders use index 0.
@@ -556,9 +571,9 @@ void UMoviePipelineDeferredPassBase::RenderSample_GameThreadImpl(const FMoviePip
 
 		// Now do the stencil layer submission (which doesn't support additional post processing materials)
 		{
-			FMoviePipelineRenderPassMetrics InOutSampleState = GetRenderPassMetricsForCamera(CameraIndex, InSampleState);
-			InOutSampleState.OutputState.CameraIndex = CameraIndex;
-			InOutSampleState.OutputState.CameraNameOverride = GetCameraNameOverride(CameraIndex);
+			FMoviePipelineRenderPassMetrics InOutSampleState = GetRenderPassMetricsForCamera(LocalCameraIndex, InSampleState);
+			InOutSampleState.OutputState.CameraIndex = LocalCameraIndex;
+			InOutSampleState.OutputState.CameraNameOverride = GetCameraNameOverride(LocalCameraIndex);
 
 			struct FStencilValues
 			{
@@ -655,7 +670,7 @@ void UMoviePipelineDeferredPassBase::RenderSample_GameThreadImpl(const FMoviePip
 				if (StencilLayerMaterial)
 				{
 					UE::MoviePipeline::FDeferredPassRenderStatePayload Payload;
-					Payload.CameraIndex = CameraIndex;
+					Payload.CameraIndex = LocalCameraIndex;
 					Payload.TileIndex = InOutSampleState.TileIndexes;
 					Payload.SceneViewIndex = StencilLayerIndex + (bRenderMainPass ? 1 : 0);
 					TSharedPtr<FSceneViewFamilyContext> ViewFamily = CalculateViewFamily(InOutSampleState, &Payload);
@@ -818,16 +833,6 @@ void UMoviePipelineDeferredPassBase::BlendPostProcessSettings(FSceneView* InView
 	UMoviePipelineExecutorShot* CurrentShot = GetPipeline()->GetActiveShotList()[GetPipeline()->GetCurrentShotIndex()];
 	const int32 NumCameras = GetNumCamerasToRender();
 
-	UCameraComponent* OutCamera = nullptr;
-	FMinimalViewInfo OutViewInfo;
-
-	GetPipeline()->GetSidecarCameraData(CurrentShot, InOutSampleState.OutputState.CameraIndex, OutViewInfo, &OutCamera);
-	if (!OutCamera)
-	{
-		// GetCameraInfo will have already printed a warning
-		return;
-	}
-
 	// The primary camera should still respect the world post processing volumes and should already be the viewtarget.
 	if (NumCameras == 1)
 	{
@@ -836,6 +841,16 @@ void UMoviePipelineDeferredPassBase::BlendPostProcessSettings(FSceneView* InView
 	}
 	else
 	{
+		UCameraComponent* OutCamera = nullptr;
+		FMinimalViewInfo OutViewInfo;
+
+		GetPipeline()->GetSidecarCameraData(CurrentShot, InOutSampleState.OutputState.CameraIndex, OutViewInfo, &OutCamera);
+		if (!OutCamera)
+		{
+			// GetCameraInfo will have already printed a warning
+			return;
+		}
+
 		// For sidecar cameras we need to do the blending of PP volumes
 		FVector ViewLocation = OutCamera->GetComponentLocation();
 		for (IInterface_PostProcessVolume* PPVolume : GetWorld()->PostProcessVolumes)
