MoviePipeline: Fixed a few cases where branch restrictions were not properly enforced when making connections in the graph.

For example, if nothing is connected to the Globals branch yet, then all connections to it were rejected.

diff --git a/Engine/Plugins/MovieScene/MovieRenderPipeline/Source/MovieRenderPipelineCore/Private/Graph/MovieGraphConfig.cpp b/Engine/Plugins/MovieScene/MovieRenderPipeline/Source/MovieRenderPipelineCore/Private/Graph/MovieGraphConfig.cpp
--- a/Engine/Plugins/MovieScene/MovieRenderPipeline/Source/MovieRenderPipelineCore/Private/Graph/MovieGraphConfig.cpp
+++ b/Engine/Plugins/MovieScene/MovieRenderPipeline/Source/MovieRenderPipelineCore/Private/Graph/MovieGraphConfig.cpp
@@ -819,6 +819,28 @@ TArray<FString> UMovieGraphConfig::GetDownstreamBranchNames(UMovieGraphNode* Fro
 	return BranchNames;
 }
 
+TArray<FString> UMovieGraphConfig::GetUpstreamBranchNames(UMovieGraphNode* FromNode, const UMovieGraphPin* FromPin) const
+{
+	TArray<FString> BranchNames;
+
+	// FromNode itself might be the Inputs node, so check before visiting the upstream nodes
+	if (FromNode->IsA<UMovieGraphInputNode>() && FromPin)
+	{
+		BranchNames.Add(FromPin->Properties.Label.ToString());
+	}
+
+	VisitUpstreamNodes(FromNode, FVisitNodesCallback::CreateLambda(
+		[&BranchNames](UMovieGraphNode* VisitedNode, const UMovieGraphPin* VisitedPin)
+		{
+			if (VisitedNode->IsA<UMovieGraphInputNode>() && VisitedPin)
+			{
+				BranchNames.Add(VisitedPin->Properties.Label.ToString());
+			}
+		}));
+
+	return BranchNames;
+}
+
 void UMovieGraphConfig::InitializeFlattenedNode(UMovieGraphNode* InNode)
 {
 	// We go through each of the bOverride_ properties on this new instance and set
diff --git a/Engine/Plugins/MovieScene/MovieRenderPipeline/Source/MovieRenderPipelineCore/Public/Graph/MovieGraphConfig.h b/Engine/Plugins/MovieScene/MovieRenderPipeline/Source/MovieRenderPipelineCore/Public/Graph/MovieGraphConfig.h
--- a/Engine/Plugins/MovieScene/MovieRenderPipeline/Source/MovieRenderPipelineCore/Public/Graph/MovieGraphConfig.h
+++ b/Engine/Plugins/MovieScene/MovieRenderPipeline/Source/MovieRenderPipelineCore/Public/Graph/MovieGraphConfig.h
@@ -503,6 +503,9 @@ public:
 	/** Determines the name(s) of the branches downstream from FromNode, starting at FromPin. */
 	TArray<FString> GetDownstreamBranchNames(UMovieGraphNode* FromNode, const UMovieGraphPin* FromPin) const;
 
+	/** Determines the name(s) of the branches upstream from FromNode, starting at FromPin. */
+	TArray<FString> GetUpstreamBranchNames(UMovieGraphNode* FromNode, const UMovieGraphPin* FromPin) const;
+
 protected:
 	/** Copies properties in FromNode that are marked for override into ToNode, but only if ToNode doesn't already override that value. */
 	void CopyOverriddenProperties(UMovieGraphNode* FromNode, UMovieGraphNode* ToNode, const FMovieGraphTraversalContext* InContext);
diff --git a/Engine/Plugins/MovieScene/MovieRenderPipeline/Source/MovieRenderPipelineEditor/Private/Graph/MovieGraphSchema.cpp b/Engine/Plugins/MovieScene/MovieRenderPipeline/Source/MovieRenderPipelineEditor/Private/Graph/MovieGraphSchema.cpp
--- a/Engine/Plugins/MovieScene/MovieRenderPipeline/Source/MovieRenderPipelineEditor/Private/Graph/MovieGraphSchema.cpp
+++ b/Engine/Plugins/MovieScene/MovieRenderPipeline/Source/MovieRenderPipelineEditor/Private/Graph/MovieGraphSchema.cpp
@@ -126,62 +126,64 @@ bool UMovieGraphSchema::IsConnectionToBranchAllowed(const UEdGraphPin* InputPin,
 	const UMovieGraphPin* FromPin = UE::MovieGraph::Private::GetGraphPinFromEdPin(OutputPin);
 	const UMovieGraphConfig* GraphConfig = UE::MovieGraph::Private::GetGraphFromEdPin(InputPin);
 
-	// Early exit if a Globals output is not being connected to a Globals input.
-	if ((FromPin->Properties.Label == UMovieGraphNode::GlobalsPinName) && (ToPin->Properties.Label != UMovieGraphNode::GlobalsPinName))
-	{
-		OutError = NSLOCTEXT("MoviePipeline", "GlobalsBranchMismatchError", "Globals branches can only be connected to other Globals branches.");
-		return false;
-	}
-
-	// Early exit if a non-Globals output is being connected to a Globals input.
-	if ((FromPin->Properties.Label != UMovieGraphNode::GlobalsPinName) && (ToPin->Properties.Label == UMovieGraphNode::GlobalsPinName))
-	{
-		OutError = FText::Format(
-			NSLOCTEXT("MoviePipeline", "NonGlobalsBranchMismatchError", "The branch '{0}' cannot be connected to the Globals branch."),
-				FText::FromName(FromPin->Properties.Label));
-		return false;
-	}
-
-	// Get all upstream nodes for FromNode -- these are the nodes that need to be checked for branch restrictions.
-	// FromNode itself also needs to be part of the validation checks.
-	TArray<UMovieGraphNode*> UpstreamNodes = {FromNode};
+	// Get all upstream/downstream nodes that occur on the connection -- these are the nodes that need to be checked for branch restrictions.
+	// FromNode/ToNode themselves also needs to be part of the validation checks.
+	TArray<UMovieGraphNode*> NodesToCheck = {FromNode, ToNode};
 	GraphConfig->VisitUpstreamNodes(FromNode, UMovieGraphConfig::FVisitNodesCallback::CreateLambda(
-		[&UpstreamNodes](UMovieGraphNode* VisitedNode, const UMovieGraphPin* VisitedPin)
+		[&NodesToCheck](UMovieGraphNode* VisitedNode, const UMovieGraphPin* VisitedPin)
+		{
+			NodesToCheck.Add(VisitedNode);
+		}));
+
+	GraphConfig->VisitDownstreamNodes(ToNode, UMovieGraphConfig::FVisitNodesCallback::CreateLambda(
+		[&NodesToCheck](UMovieGraphNode* VisitedNode, const UMovieGraphPin* VisitedPin)
 		{
-			UpstreamNodes.Add(VisitedNode);
+			NodesToCheck.Add(VisitedNode);
 		}));
 
-	// Determine which branch(es) are connected to this node downstream.
-	const TArray<FString> UpstreamBranchNames = GraphConfig->GetDownstreamBranchNames(ToNode, ToPin);
-	const bool bGlobalsIsOnlyDownstreamBranch = (UpstreamBranchNames.Num() == 1) && (UpstreamBranchNames[0] == UMovieGraphNode::GlobalsPinNameString);
-	const bool bGlobalsIsDownstream = UpstreamBranchNames.Contains(UMovieGraphNode::GlobalsPinNameString);
+	// Determine which branch(es) are connected to this node up/downstream.
+	const TArray<FString> DownstreamBranchNames = GraphConfig->GetDownstreamBranchNames(ToNode, ToPin);
+	const TArray<FString> UpstreamBranchNames = GraphConfig->GetUpstreamBranchNames(FromNode, FromPin);
+	const bool bGlobalsIsDownstream = DownstreamBranchNames.Contains(UMovieGraphNode::GlobalsPinNameString);
+	const bool bGlobalsIsUpstream = UpstreamBranchNames.Contains(UMovieGraphNode::GlobalsPinNameString);
+	const bool bDownstreamBranchExistsAndIsntOnlyGlobals =
+		!DownstreamBranchNames.IsEmpty() && ((DownstreamBranchNames.Num() != 1) || (DownstreamBranchNames[0] != UMovieGraphNode::GlobalsPinNameString));
+	const bool bUpstreamBranchExistsAndIsntOnlyGlobals =
+		!UpstreamBranchNames.IsEmpty() && ((UpstreamBranchNames.Num() != 1) || (UpstreamBranchNames[0] != UMovieGraphNode::GlobalsPinNameString));
+
+	// Globals branches can only be connected to Globals branches
+	if ((bGlobalsIsDownstream && bUpstreamBranchExistsAndIsntOnlyGlobals) || (bGlobalsIsUpstream && bDownstreamBranchExistsAndIsntOnlyGlobals))
+	{
+		OutError = NSLOCTEXT("MoviePipeline", "GlobalsBranchMismatchError", "Globals branches can only be connected to other Globals branches.");
+		return false;
+	}
 
-	// Error out if any of the upstream nodes cannot be connected to the downstream branch(es).
-	for (const UMovieGraphNode* UpstreamNode : UpstreamNodes)
+	// Error out if any of the nodes that are part of the connection cannot be connected to the upstream/downstream branches.
+	for (const UMovieGraphNode* NodeToCheck : NodesToCheck)
 	{
-		if (UpstreamNode->GetBranchRestriction() == EMovieGraphBranchRestriction::Globals)
+		if (NodeToCheck->GetBranchRestriction() == EMovieGraphBranchRestriction::Globals)
 		{
-			// Globals-specific nodes have to be connected such that the only upstream branch is Globals, OR the group
-			// of nodes is not yet connected to a branch (and this restriction will be enforced when they are).
-			if (bGlobalsIsOnlyDownstreamBranch || UpstreamBranchNames.IsEmpty())
+			// Globals-specific nodes have to be connected such that the only upstream/downstream branches are Globals.
+			// If either the upstream/downstream branches are empty (ie, the node isn't connected to Inputs/Outputs yet)
+			// then the connection is OK for now -- the branch restriction will be enforced when nodes are connected to
+			// Inputs/Outputs.
+			if (bDownstreamBranchExistsAndIsntOnlyGlobals || bUpstreamBranchExistsAndIsntOnlyGlobals)
 			{
-				continue;
+				OutError = FText::Format(
+					NSLOCTEXT("MoviePipeline", "GlobalsBranchRestrictionError", "The node '{0}' can only be connected to the Globals branch."),
+						NodeToCheck->GetNodeTitle());
+				return false;
 			}
-
-			OutError = FText::Format(
-				NSLOCTEXT("MoviePipeline", "GlobalsBranchRestrictionError", "The node '{0}' can only be connected to the Globals branch."),
-					UpstreamNode->GetNodeTitle());
-			return false;
 		}
 
 		// Check that render-layer-only nodes aren't connected to Globals.
-		if (UpstreamNode->GetBranchRestriction() == EMovieGraphBranchRestriction::RenderLayer)
+		if (NodeToCheck->GetBranchRestriction() == EMovieGraphBranchRestriction::RenderLayer)
 		{
-			if (bGlobalsIsDownstream)
+			if (bGlobalsIsDownstream || bGlobalsIsUpstream)
 			{
 				OutError = FText::Format(
 					NSLOCTEXT("MoviePipeline", "RenderLayerBranchRestrictionError", "The node '{0}' can only be connected to a render layer branch."),
-						UpstreamNode->GetNodeTitle());
+						NodeToCheck->GetNodeTitle());
 				return false;
 			}
 		}
