MoviePipeline: Added a new Burn-in node to the graph, as well as its rendering implementation.

Also added a new default burn-in widget blueprint that's specific for the graph and reads from the graph pipeline. There are some things in the new blueprint that are placeholders for now though (eg, camera info), pending updates to the graph's BP function library.

diff --git a/Engine/Plugins/MovieScene/MovieRenderPipeline/Source/MovieRenderPipelineCore/MovieRenderPipelineCore.Build.cs b/Engine/Plugins/MovieScene/MovieRenderPipeline/Source/MovieRenderPipelineCore/MovieRenderPipelineCore.Build.cs
--- a/Engine/Plugins/MovieScene/MovieRenderPipeline/Source/MovieRenderPipelineCore/MovieRenderPipelineCore.Build.cs
+++ b/Engine/Plugins/MovieScene/MovieRenderPipeline/Source/MovieRenderPipelineCore/MovieRenderPipelineCore.Build.cs
@@ -21,8 +21,8 @@ public class MovieRenderPipelineCore : ModuleRules
 				"Networking",
 				"HTTP",
 				"DeveloperSettings",
-				"ClothingSystemRuntimeInterface"
-
+				"ClothingSystemRuntimeInterface",
+				"Slate",
 			}
 		);
 
diff --git a/Engine/Plugins/MovieScene/MovieRenderPipeline/Source/MovieRenderPipelineCore/Private/Graph/MovieGraphBlueprintLibrary.cpp b/Engine/Plugins/MovieScene/MovieRenderPipeline/Source/MovieRenderPipelineCore/Private/Graph/MovieGraphBlueprintLibrary.cpp
--- a/Engine/Plugins/MovieScene/MovieRenderPipeline/Source/MovieRenderPipelineCore/Private/Graph/MovieGraphBlueprintLibrary.cpp
+++ b/Engine/Plugins/MovieScene/MovieRenderPipeline/Source/MovieRenderPipelineCore/Private/Graph/MovieGraphBlueprintLibrary.cpp
@@ -1,8 +1,10 @@
 // Copyright Epic Games, Inc. All Rights Reserved.
 
 #include "Graph/MovieGraphBlueprintLibrary.h"
+
 #include "Graph/Nodes/MovieGraphOutputSettingNode.h"
 #include "Graph/Nodes/MovieGraphRenderLayerNode.h"
+#include "MoviePipelineBlueprintLibrary.h"
 #include "MoviePipelineUtils.h"
 
 FFrameRate UMovieGraphBlueprintLibrary::GetEffectiveFrameRate(UMovieGraphOutputSettingNode* InNode, const FFrameRate& InDefaultRate)
@@ -197,4 +199,22 @@ FString UMovieGraphBlueprintLibrary::ResolveFilenameFormatArguments(const FStrin
 
 		++DuplicateIndex;
 	}
+}
+
+FIntPoint UMovieGraphBlueprintLibrary::GetEffectiveOutputResolution(UMovieGraphEvaluatedConfig* InEvaluatedGraph, const FName& InBranchName)
+{
+	if (!InEvaluatedGraph)
+	{
+		FFrame::KismetExecutionMessage(TEXT("InEvaluatedGraph cannot be null."), ELogVerbosity::Error);
+		return FIntPoint();
+	}
+	
+	constexpr bool bIncludeCDOs = true;
+	const UMovieGraphOutputSettingNode* OutputSetting = InEvaluatedGraph->GetSettingForBranch<UMovieGraphOutputSettingNode>(InBranchName, bIncludeCDOs);
+	if (!ensure(OutputSetting))
+	{
+		return FIntPoint();
+	}
+	
+	return UMoviePipelineBlueprintLibrary::Utility_GetEffectiveOutputResolution(0 /* TODO: Overscan percentage needs to be provided */, OutputSetting->OutputResolution);
 }
\ No newline at end of file
diff --git a/Engine/Plugins/MovieScene/MovieRenderPipeline/Source/MovieRenderPipelineCore/Private/Graph/MovieGraphDefaultRenderer.cpp b/Engine/Plugins/MovieScene/MovieRenderPipeline/Source/MovieRenderPipelineCore/Private/Graph/MovieGraphDefaultRenderer.cpp
--- a/Engine/Plugins/MovieScene/MovieRenderPipeline/Source/MovieRenderPipelineCore/Private/Graph/MovieGraphDefaultRenderer.cpp
+++ b/Engine/Plugins/MovieScene/MovieRenderPipeline/Source/MovieRenderPipelineCore/Private/Graph/MovieGraphDefaultRenderer.cpp
@@ -19,7 +19,9 @@ void UMovieGraphDefaultRenderer::SetupRenderingPipelineForShot(UMoviePipelineExe
 	struct FMovieGraphPass
 	{
 		TSubclassOf<UMovieGraphRenderPassNode> ClassType;
-		TArray<FName> LayerBranchNames;
+
+		/** Maps a named branch to the specific render pass node that is assigned to render it. */
+		TMap<FName, TWeakObjectPtr<UMovieGraphRenderPassNode>> BranchRenderers;
 	};
 
 	TArray<FMovieGraphPass> OutputPasses;
@@ -62,7 +64,8 @@ void UMovieGraphDefaultRenderer::SetupRenderingPipelineForShot(UMoviePipelineExe
 				ExistingPass = &OutputPasses.AddDefaulted_GetRef();
 				ExistingPass->ClassType = RenderPassNode->GetClass();
 			}
-			ExistingPass->LayerBranchNames.AddUnique(Branch);
+			
+			ExistingPass->BranchRenderers.Add(Branch, RenderPassNode);
 		}
 	}
 
@@ -72,17 +75,18 @@ void UMovieGraphDefaultRenderer::SetupRenderingPipelineForShot(UMoviePipelineExe
 	{
 		// ToDo: This should probably come from the Renderers themselves, as they can internally produce multiple
 		// renders (such as ObjectID passes).
-		TotalLayerCount += Pass.LayerBranchNames.Num();
+		TotalLayerCount += Pass.BranchRenderers.Num();
 		
 		FMovieGraphRenderPassSetupData SetupData;
 		SetupData.Renderer = this;
 		UE_LOG(LogMovieRenderPipeline, Warning, TEXT("\tRenderer Class: %s"), *Pass.ClassType->GetName());
-		for (const FName& LayerBranchName : Pass.LayerBranchNames)
+		for (const TTuple<FName, TWeakObjectPtr<UMovieGraphRenderPassNode>>& BranchRenderer : Pass.BranchRenderers)
 		{
 			FMovieGraphRenderPassLayerData& LayerData = SetupData.Layers.AddDefaulted_GetRef();
-			LayerData.BranchName = LayerBranchName;
+			LayerData.BranchName = BranchRenderer.Key;
+			LayerData.RenderPassNode = BranchRenderer.Value;
 
-			UE_LOG(LogMovieRenderPipeline, Warning, TEXT("\t\tBranch Name: %s"), *LayerBranchName.ToString());
+			UE_LOG(LogMovieRenderPipeline, Warning, TEXT("\t\tBranch Name: %s"), *BranchRenderer.Key.ToString());
 		}
 
 		UMovieGraphRenderPassNode* RenderPassCDO = Pass.ClassType->GetDefaultObject<UMovieGraphRenderPassNode>();
diff --git a/Engine/Plugins/MovieScene/MovieRenderPipeline/Source/MovieRenderPipelineCore/Private/Graph/Nodes/MovieGraphBurnInNode.cpp b/Engine/Plugins/MovieScene/MovieRenderPipeline/Source/MovieRenderPipelineCore/Private/Graph/Nodes/MovieGraphBurnInNode.cpp
new file mode 100644
--- /dev/null
+++ b/Engine/Plugins/MovieScene/MovieRenderPipeline/Source/MovieRenderPipelineCore/Private/Graph/Nodes/MovieGraphBurnInNode.cpp
@@ -0,0 +1,243 @@
+// Copyright Epic Games, Inc. All Rights Reserved.
+
+#include "Graph/Nodes/MovieGraphBurnInNode.h"
+
+#include "Engine/TextureRenderTarget2D.h"
+#include "Framework/Application/SlateApplication.h"
+#include "Graph/MovieGraphBlueprintLibrary.h"
+#include "Graph/MovieGraphBurnInWidget.h"
+#include "Graph/MovieGraphDataTypes.h"
+#include "Graph/MovieGraphDefaultRenderer.h"
+#include "Graph/MovieGraphPipeline.h"
+#include "MovieRenderPipelineCoreModule.h"
+#include "OpenColorIODisplayExtension.h"
+#include "Slate/WidgetRenderer.h"
+#include "Styling/AppStyle.h"
+#include "TextureResource.h"
+#include "Widgets/SVirtualWindow.h"
+
+const FString UMovieGraphBurnInNode::RendererName = FString("BurnIn");
+
+#if WITH_EDITOR
+FText UMovieGraphBurnInNode::GetNodeTitle(const bool bGetDescriptive) const
+{
+	return NSLOCTEXT("MovieGraphNodes", "BurnInGraphNode_Description", "Burn In");
+}
+
+FLinearColor UMovieGraphBurnInNode::GetNodeTitleColor() const
+{
+	return FLinearColor(0.572f, 0.274f, 1.f);
+}
+
+FSlateIcon UMovieGraphBurnInNode::GetIconAndTint(FLinearColor& OutColor) const
+{
+	static const FSlateIcon DeferredRendererIcon = FSlateIcon(FAppStyle::GetAppStyleSetName(), "SequenceRecorder.TabIcon");
+
+	OutColor = FLinearColor::White;
+	return DeferredRendererIcon;
+}
+#endif	// WITH_EDITOR
+
+void UMovieGraphBurnInNode::SetupImpl(const FMovieGraphRenderPassSetupData& InSetupData)
+{
+	for (const FMovieGraphRenderPassLayerData& LayerData : InSetupData.Layers)
+	{
+		TUniquePtr<FMovieGraphBurnInPass> RendererInstance = MakeUnique<FMovieGraphBurnInPass>();
+		RendererInstance->Setup(InSetupData.Renderer, LayerData);
+		CurrentInstances.Add(MoveTemp(RendererInstance));
+	}
+
+	constexpr bool bApplyGammaCorrection = false;
+	WidgetRenderer = MakeShared<FWidgetRenderer>(bApplyGammaCorrection);
+}
+
+void UMovieGraphBurnInNode::TeardownImpl()
+{
+	for (const TUniquePtr<FMovieGraphBurnInPass>& Instance : CurrentInstances)
+	{
+		Instance->Teardown();
+	}
+	
+	CurrentInstances.Reset();
+
+	if (FSlateApplication::IsInitialized())
+	{
+		for (const TTuple<const FIntPoint, TSharedPtr<SVirtualWindow>>& Entry : SharedVirtualWindows)
+		{
+			const TSharedPtr<SVirtualWindow>& VirtualWindow = Entry.Value;
+			if (VirtualWindow.IsValid())
+			{
+				FSlateApplication::Get().UnregisterVirtualWindow(VirtualWindow.ToSharedRef());
+			}
+		}
+	}
+
+	SharedVirtualWindows.Empty();
+	WidgetRenderer = nullptr;
+	BurnInWidgetInstances.Empty();
+}
+
+void UMovieGraphBurnInNode::RenderImpl(const FMovieGraphTraversalContext& InFrameTraversalContext, const FMovieGraphTimeStepData& InTimeData)
+{
+	for (const TUniquePtr<FMovieGraphBurnInPass>& Instance : CurrentInstances)
+	{
+		Instance->Render(InFrameTraversalContext, InTimeData);
+	}
+}
+
+void UMovieGraphBurnInNode::GatherOutputPassesImpl(TArray<FMovieGraphRenderDataIdentifier>& OutExpectedPasses) const
+{
+	for (const TUniquePtr<FMovieGraphBurnInPass>& Instance : CurrentInstances)
+    {
+    	Instance->GatherOutputPassesImpl(OutExpectedPasses);
+    }
+}
+
+TSharedPtr<SVirtualWindow> UMovieGraphBurnInNode::GetOrCreateVirtualWindow(const FIntPoint& InResolution)
+{
+	if (const TSharedPtr<SVirtualWindow>* ExistingVirtualWindow = SharedVirtualWindows.Find(InResolution))
+	{
+		return *ExistingVirtualWindow;
+	}
+
+	const TSharedPtr<SVirtualWindow> NewVirtualWindow = SNew(SVirtualWindow).Size(FVector2D(InResolution.X, InResolution.Y));
+	SharedVirtualWindows.Emplace(InResolution, NewVirtualWindow);
+
+	if (FSlateApplication::IsInitialized())
+	{
+		FSlateApplication::Get().RegisterVirtualWindow(NewVirtualWindow.ToSharedRef());
+	}
+
+	return NewVirtualWindow;
+}
+
+TObjectPtr<UMovieGraphBurnInWidget> UMovieGraphBurnInNode::GetOrCreateBurnInWidget(UClass* InWidgetClass, UWorld* InOwner)
+{
+	if (const TObjectPtr<UMovieGraphBurnInWidget>* ExistingBurnInWidget = BurnInWidgetInstances.Find(InWidgetClass))
+	{
+		return *ExistingBurnInWidget;
+	}
+
+	TObjectPtr<UMovieGraphBurnInWidget> BurnInWidget = CreateWidget<UMovieGraphBurnInWidget>(InOwner, InWidgetClass);
+	if (BurnInWidget)
+	{
+		BurnInWidgetInstances.Emplace(InWidgetClass, BurnInWidget);
+	}
+
+	return BurnInWidget;
+}
+
+void UMovieGraphBurnInNode::FMovieGraphBurnInPass::Setup(TWeakObjectPtr<UMovieGraphDefaultRenderer> InRenderer, const FMovieGraphRenderPassLayerData& InLayer)
+{
+	LayerData = InLayer;
+	Renderer = InRenderer;
+	RenderPassNode = CastChecked<UMovieGraphBurnInNode>(InLayer.RenderPassNode);
+
+	RenderDataIdentifier.RootBranchName = LayerData.BranchName;
+	RenderDataIdentifier.RendererName = RenderPassNode->GetRendererName();
+	RenderDataIdentifier.SubResourceName = TEXT("burnin");
+
+	UE::MovieGraph::DefaultRenderer::FCameraInfo CameraInfo = Renderer->GetCameraInfo(LayerData.CameraIdentifier);
+	RenderDataIdentifier.CameraName = CameraInfo.CameraName;
+}
+
+void UMovieGraphBurnInNode::FMovieGraphBurnInPass::Teardown()
+{
+	// Nothing to do
+}
+
+void UMovieGraphBurnInNode::FMovieGraphBurnInPass::Render(const FMovieGraphTraversalContext& InFrameTraversalContext, const FMovieGraphTimeStepData& InTimeData)
+{
+	UMovieGraphPipeline* Pipeline = Renderer->GetOwningGraph();
+	
+	const FIntPoint OutputResolution = UMovieGraphBlueprintLibrary::GetEffectiveOutputResolution(InTimeData.EvaluatedConfig, LayerData.BranchName);
+	const int32 MaxResolution = GetMax2DTextureDimension();
+	if ((OutputResolution.X > MaxResolution) || (OutputResolution.Y > MaxResolution))
+	{
+		UE_LOG(LogMovieRenderPipeline, Error, TEXT("Resolution %dx%d exceeds maximum allowed by GPU. Burn-ins do not support high-resolution tiling and thus can't exceed %dx%d."), OutputResolution.X, OutputResolution.Y, MaxResolution, MaxResolution);
+		return;
+	}
+
+	UClass* LoadedBurnInClass = RenderPassNode->BurnInClass.TryLoadClass<UMovieGraphBurnInWidget>();
+	if (!LoadedBurnInClass)
+	{
+		UE_LOG(LogMovieRenderPipeline, Error, TEXT("The burn-in widget provided in layer '%s' for renderer '%s' is not valid."), *LayerData.BranchName.ToString(), *Renderer->GetClass()->GetName());
+		return;
+	}
+
+	// Create the render target the widget will be rendered into
+	UE::MovieGraph::DefaultRenderer::FRenderTargetInitParams RenderTargetInitParams;
+	RenderTargetInitParams.Size = OutputResolution;
+	RenderTargetInitParams.TargetGamma = FOpenColorIODisplayExtension::DefaultDisplayGamma;
+	RenderTargetInitParams.PixelFormat = PF_B8G8R8A8;
+	UTextureRenderTarget2D* RenderTarget = Renderer->GetOrCreateViewRenderTarget(RenderTargetInitParams);
+	
+	// The CDO contains the resources which are shared with all FMovieGraphBurnInPass instances
+	UMovieGraphBurnInNode* BurnInCDO = RenderPassNode->GetClass()->GetDefaultObject<UMovieGraphBurnInNode>();
+	const TSharedPtr<SVirtualWindow> VirtualWindow = BurnInCDO->GetOrCreateVirtualWindow(OutputResolution);
+	const TSharedPtr<FWidgetRenderer> BurnInWidgetRenderer = BurnInCDO->WidgetRenderer;
+	UMovieGraphBurnInWidget* BurnInWidget = BurnInCDO->GetOrCreateBurnInWidget(LoadedBurnInClass, Pipeline->GetWorld());
+	if (!BurnInWidget)
+	{
+		UE_LOG(LogMovieRenderPipeline, Error, TEXT("Unable to load burn-in widget at path: %s"), *RenderPassNode->BurnInClass.GetAssetPath().ToString());
+		return;
+	}
+
+	if (InFrameTraversalContext.Time.bIsFirstTemporalSampleForFrame)
+	{
+		// Put the widget in our window
+		VirtualWindow->SetContent(BurnInWidget->TakeWidget());
+
+		// Update the widget with the latest frame information
+		BurnInWidget->UpdateForGraph(Pipeline);
+		
+		// Draw the widget to the render target
+		BurnInWidgetRenderer->DrawWindow(RenderTarget, VirtualWindow->GetHittestGrid(), VirtualWindow.ToSharedRef(),
+			1.f, OutputResolution, InTimeData.FrameDeltaTime);
+
+		FRenderTarget* BackbufferRenderTarget = RenderTarget->GameThread_GetRenderTargetResource();
+		TSharedPtr<UE::MovieGraph::IMovieGraphOutputMerger> OutputMerger = Pipeline->GetOutputMerger();
+
+		ENQUEUE_RENDER_COMMAND(BurnInRenderTargetResolveCommand)(
+			[InTimeData, InFrameTraversalContext, RenderTargetInitParams, RenderDataIdentifier = this->RenderDataIdentifier, bComposite = RenderPassNode->bCompositeOntoFinalImage, BackbufferRenderTarget, OutputMerger, OutputResolution](FRHICommandListImmediate& RHICmdList)
+			{
+				const FIntRect SourceRect = FIntRect(0, 0, BackbufferRenderTarget->GetSizeXY().X, BackbufferRenderTarget->GetSizeXY().Y);
+
+				// Read the data back to the CPU
+				TArray<FColor> RawPixels;
+				RawPixels.SetNumUninitialized(SourceRect.Width() * SourceRect.Height());
+
+				FReadSurfaceDataFlags ReadDataFlags(ERangeCompressionMode::RCM_MinMax);
+				ReadDataFlags.SetLinearToGamma(false);
+
+				{
+					// TODO: The readback is taking ~37ms on a 4k image. This is definitely an area that should be a target for optimization in the future.
+					TRACE_CPUPROFILER_EVENT_SCOPE(MRQ::FMovieGraphBurnInPass::ReadSurfaceData);
+					RHICmdList.ReadSurfaceData(BackbufferRenderTarget->GetRenderTargetTexture(), SourceRect, RawPixels, ReadDataFlags);
+				}
+
+				// Take our per-frame Traversal Context and update it with context specific to this sample.
+				FMovieGraphTraversalContext UpdatedTraversalContext = InFrameTraversalContext;
+				UpdatedTraversalContext.Time = InTimeData;
+				UpdatedTraversalContext.RenderDataIdentifier = RenderDataIdentifier;
+
+				TSharedRef<UE::MovieGraph::FMovieGraphSampleState> SampleStatePayload =
+					MakeShared<UE::MovieGraph::FMovieGraphSampleState>();
+				SampleStatePayload->TraversalContext = MoveTemp(UpdatedTraversalContext);
+				SampleStatePayload->BackbufferResolution = RenderTargetInitParams.Size;
+				SampleStatePayload->bRequiresAccumulator = false;
+				SampleStatePayload->bFetchFromAccumulator = false;
+				SampleStatePayload->bCompositeOnOtherRenders = bComposite;
+
+				TUniquePtr<FImagePixelData> PixelData =
+					MakeUnique<TImagePixelData<FColor>>(OutputResolution, TArray64<FColor>(MoveTemp(RawPixels)), SampleStatePayload);
+
+				OutputMerger->OnCompleteRenderPassDataAvailable_AnyThread(MoveTemp(PixelData));
+			});
+	}
+}
+
+void UMovieGraphBurnInNode::FMovieGraphBurnInPass::GatherOutputPassesImpl(TArray<FMovieGraphRenderDataIdentifier>& OutExpectedPasses) const
+{
+	OutExpectedPasses.Add(RenderDataIdentifier);
+}
\ No newline at end of file
diff --git a/Engine/Plugins/MovieScene/MovieRenderPipeline/Source/MovieRenderPipelineCore/Private/Graph/Nodes/MovieGraphImageSequenceOutputNode.cpp b/Engine/Plugins/MovieScene/MovieRenderPipeline/Source/MovieRenderPipelineCore/Private/Graph/Nodes/MovieGraphImageSequenceOutputNode.cpp
--- a/Engine/Plugins/MovieScene/MovieRenderPipeline/Source/MovieRenderPipelineCore/Private/Graph/Nodes/MovieGraphImageSequenceOutputNode.cpp
+++ b/Engine/Plugins/MovieScene/MovieRenderPipeline/Source/MovieRenderPipelineCore/Private/Graph/Nodes/MovieGraphImageSequenceOutputNode.cpp
@@ -1,6 +1,7 @@
 // Copyright Epic Games, Inc. All Rights Reserved.
 
 #include "Graph/Nodes/MovieGraphImageSequenceOutputNode.h"
+
 #include "Graph/Nodes/MovieGraphOutputSettingNode.h"
 #include "Graph/Nodes/MovieGraphRenderLayerNode.h"
 #include "Graph/MovieGraphDataTypes.h"
@@ -33,6 +34,23 @@ void UMovieGraphImageSequenceOutputNode::OnReceiveImageDataImpl(UMovieGraphPipel
 {
 	check(InRawFrameData);
 
+	// Gather the passes that need to be composited
+	TArray<TPair<FMovieGraphRenderDataIdentifier, TUniquePtr<FImagePixelData>>> CompositingPasses;
+	for (TPair<FMovieGraphRenderDataIdentifier, TUniquePtr<FImagePixelData>>& RenderData : InRawFrameData->ImageOutputData)
+	{
+		UE::MovieGraph::FMovieGraphSampleState* Payload = RenderData.Value->GetPayload<UE::MovieGraph::FMovieGraphSampleState>();
+		check(Payload);
+		if (!Payload->bCompositeOnOtherRenders)
+		{
+			continue;
+		}
+
+		TPair<FMovieGraphRenderDataIdentifier, TUniquePtr<FImagePixelData>> CompositePass;
+		CompositePass.Key = RenderData.Key;
+		CompositePass.Value = RenderData.Value->CopyImageData();
+		CompositingPasses.Add(MoveTemp(CompositePass));
+	}
+
 	// ToDo:
 	// The ImageWriteQueue is set up in a fire-and-forget manner. This means that the data needs to be placed in the WriteQueue
 	// as a TUniquePtr (so it can free the data when its done). Unfortunately we can have multiple output formats at once,
@@ -48,6 +66,22 @@ void UMovieGraphImageSequenceOutputNode::OnReceiveImageDataImpl(UMovieGraphPipel
 	// The base ImageSequenceOutputNode doesn't support any multilayer formats, so we write out each render pass separately.
 	for (TPair<FMovieGraphRenderDataIdentifier, TUniquePtr<FImagePixelData>>& RenderData : InRawFrameData->ImageOutputData)
 	{
+		// If this pass is composited, skip it for now
+		bool bSkip = false;
+		for (TPair<FMovieGraphRenderDataIdentifier, TUniquePtr<FImagePixelData>>& CompositedPass : CompositingPasses)
+		{
+			if (CompositedPass.Key == RenderData.Key)
+			{
+				bSkip = true;
+				break;
+			}
+		}
+
+		if (bSkip)
+		{
+			continue;
+		}
+		
 		// A layer within this output data may have chosen to not be written to disk by this CDO node
 		if (!InMask.Contains(RenderData.Key))
 		{
@@ -130,6 +164,33 @@ void UMovieGraphImageSequenceOutputNode::OnReceiveImageDataImpl(UMovieGraphPipel
 		TileImageTask->Filename = FileName;
 		TileImageTask->PixelData = RenderData.Value->CopyImageData();
 
+		EImagePixelType PixelType = TileImageTask->PixelData->GetType();
+
+		// Perform compositing if any compositing passes were found earlier
+		for (TPair<FMovieGraphRenderDataIdentifier, TUniquePtr<FImagePixelData>>& CompositedPass : CompositingPasses)
+		{
+			// This compositing pass will only composite on top of renders w/ the same branch and camera
+			const FMovieGraphRenderDataIdentifier& Id = CompositedPass.Key;
+			if ((Id.CameraName != RenderData.Key.CameraName) || (Id.RootBranchName != RenderData.Key.RootBranchName))
+			{
+				continue;
+			}
+
+			// There could be multiple renders within this branch using the composited pass, so we have to copy the image data
+			switch (PixelType)
+			{
+			case EImagePixelType::Color:
+				TileImageTask->PixelPreProcessors.Add(TAsyncCompositeImage<FColor>(CompositedPass.Value->CopyImageData()));
+				break;
+			case EImagePixelType::Float16:
+				TileImageTask->PixelPreProcessors.Add(TAsyncCompositeImage<FFloat16Color>(CompositedPass.Value->CopyImageData()));
+				break;
+			case EImagePixelType::Float32:
+				TileImageTask->PixelPreProcessors.Add(TAsyncCompositeImage<FLinearColor>(CompositedPass.Value->CopyImageData()));
+				break;
+			}
+		}
+
 		UE::MovieGraph::FMovieGraphOutputFutureData OutputData;
 		OutputData.Shot = nullptr;
 		OutputData.FilePath = FileName;
diff --git a/Engine/Plugins/MovieScene/MovieRenderPipeline/Source/MovieRenderPipelineCore/Private/MoviePipelineBlueprintLibrary.cpp b/Engine/Plugins/MovieScene/MovieRenderPipeline/Source/MovieRenderPipelineCore/Private/MoviePipelineBlueprintLibrary.cpp
--- a/Engine/Plugins/MovieScene/MovieRenderPipeline/Source/MovieRenderPipelineCore/Private/MoviePipelineBlueprintLibrary.cpp
+++ b/Engine/Plugins/MovieScene/MovieRenderPipeline/Source/MovieRenderPipelineCore/Private/MoviePipelineBlueprintLibrary.cpp
@@ -895,25 +895,29 @@ FIntPoint UMoviePipelineBlueprintLibrary::GetEffectiveOutputResolution(UMoviePip
 {
 	if (InPrimaryConfig && InPipelineExecutorShot)
 	{
-		UMoviePipelineOutputSetting* OutputSetting = InPrimaryConfig->FindSetting<UMoviePipelineOutputSetting>();
+		const UMoviePipelineOutputSetting* OutputSetting = InPrimaryConfig->FindSetting<UMoviePipelineOutputSetting>();
 		const UMoviePipelineCameraSetting* CameraSetting = Cast<const UMoviePipelineCameraSetting>(UMoviePipelineBlueprintLibrary::FindOrGetDefaultSettingForShot(UMoviePipelineCameraSetting::StaticClass(), InPrimaryConfig, InPipelineExecutorShot));
 		check(OutputSetting);
 		check(CameraSetting);
-		FIntPoint EffectiveResolution(OutputSetting->OutputResolution);
-		float ClampedOverscanPercentage = FMath::Clamp(CameraSetting->OverscanPercentage, 0.0f, 1.0f);
-		if (ClampedOverscanPercentage > 0.f)
-		{
-			float Scale = 1.f + ClampedOverscanPercentage;
-			EffectiveResolution.X = FMath::CeilToInt(((float)EffectiveResolution.X) * Scale);
-			EffectiveResolution.Y = FMath::CeilToInt(((float)EffectiveResolution.Y) * Scale);
-		}
-
-		return EffectiveResolution;
+		return Utility_GetEffectiveOutputResolution(CameraSetting->OverscanPercentage, OutputSetting->OutputResolution);
 	}
 
 	return FIntPoint();
 }
 
+FIntPoint UMoviePipelineBlueprintLibrary::Utility_GetEffectiveOutputResolution(const float OverscanPercentage, const FIntPoint& InOutputResolution)
+{
+	const float ClampedOverscanPercentage = FMath::Clamp(OverscanPercentage, 0.0f, 1.0f);
+	FIntPoint EffectiveResolution(InOutputResolution);
+	if (ClampedOverscanPercentage > 0.f)
+	{
+		const float Scale = 1.f + ClampedOverscanPercentage;
+		EffectiveResolution.X = FMath::CeilToInt(static_cast<float>(EffectiveResolution.X) * Scale);
+		EffectiveResolution.Y = FMath::CeilToInt(static_cast<float>(EffectiveResolution.Y) * Scale);
+	}
+
+	return EffectiveResolution;
+}
 
 UMoviePipelineSetting* UMoviePipelineBlueprintLibrary::FindOrGetDefaultSettingForShot(TSubclassOf<UMoviePipelineSetting> InSettingType, const UMoviePipelinePrimaryConfig* InPrimaryConfig, const UMoviePipelineExecutorShot* InShot)
 {
diff --git a/Engine/Plugins/MovieScene/MovieRenderPipeline/Source/MovieRenderPipelineCore/Public/Graph/MovieGraphBlueprintLibrary.h b/Engine/Plugins/MovieScene/MovieRenderPipeline/Source/MovieRenderPipelineCore/Public/Graph/MovieGraphBlueprintLibrary.h
--- a/Engine/Plugins/MovieScene/MovieRenderPipeline/Source/MovieRenderPipelineCore/Public/Graph/MovieGraphBlueprintLibrary.h
+++ b/Engine/Plugins/MovieScene/MovieRenderPipeline/Source/MovieRenderPipelineCore/Public/Graph/MovieGraphBlueprintLibrary.h
@@ -42,4 +42,12 @@ public:
 	UFUNCTION(BlueprintCallable, Category = "Movie Graph")
 	static FString ResolveFilenameFormatArguments(const FString& InFormatString, const FMovieGraphFilenameResolveParams& InParams, FMovieGraphResolveArgs& OutMergedFormatArgs);
 
+	/**
+	* In case of overscan percentage being higher than 0, additional pixels are rendered. This function returns the resolution with overscan taken into account.
+	* @param	InEvaluatedGraph	- The evaluated graph that will provide context for resolving the resolution
+	* @param	InBranchName		- The graph branch that the output resolution should be resolved on
+	* @return						- The output resolution, taking into account overscan
+	*/
+	UFUNCTION(BlueprintCallable, Category = "Movie Graph")
+	static FIntPoint GetEffectiveOutputResolution(UMovieGraphEvaluatedConfig* InEvaluatedGraph, const FName& InBranchName);
 };
\ No newline at end of file
diff --git a/Engine/Plugins/MovieScene/MovieRenderPipeline/Source/MovieRenderPipelineCore/Public/Graph/MovieGraphBurnInWidget.h b/Engine/Plugins/MovieScene/MovieRenderPipeline/Source/MovieRenderPipelineCore/Public/Graph/MovieGraphBurnInWidget.h
new file mode 100644
--- /dev/null
+++ b/Engine/Plugins/MovieScene/MovieRenderPipeline/Source/MovieRenderPipelineCore/Public/Graph/MovieGraphBurnInWidget.h
@@ -0,0 +1,24 @@
+﻿// Copyright Epic Games, Inc. All Rights Reserved.
+
+#pragma once
+
+#include "Blueprint/UserWidget.h"
+
+#include "MovieGraphBurnInWidget.generated.h"
+
+/**
+ * Base class for graph-based level sequence burn-ins.
+ */
+UCLASS(Blueprintable, Abstract)
+class MOVIERENDERPIPELINECORE_API UMovieGraphBurnInWidget : public UUserWidget
+{
+public:
+	GENERATED_BODY()
+
+	/** 
+	* Called on the first temporal and first spatial sample of each output frame.
+	* @param	InGraphPipeline		The graph pipeline the burn-in is for. This will be consistent throughout a burn-in widget's life.
+	*/
+	UFUNCTION(BlueprintImplementableEvent)
+	void UpdateForGraph(UMovieGraphPipeline* InGraphPipeline);
+};
diff --git a/Engine/Plugins/MovieScene/MovieRenderPipeline/Source/MovieRenderPipelineCore/Public/Graph/MovieGraphDataTypes.h b/Engine/Plugins/MovieScene/MovieRenderPipeline/Source/MovieRenderPipelineCore/Public/Graph/MovieGraphDataTypes.h
--- a/Engine/Plugins/MovieScene/MovieRenderPipeline/Source/MovieRenderPipelineCore/Public/Graph/MovieGraphDataTypes.h
+++ b/Engine/Plugins/MovieScene/MovieRenderPipeline/Source/MovieRenderPipelineCore/Public/Graph/MovieGraphDataTypes.h
@@ -166,6 +166,7 @@ struct FMovieGraphRenderPassLayerData
 {
 	FName BranchName;
 	FGuid CameraIdentifier;
+	TWeakObjectPtr<class UMovieGraphRenderPassNode> RenderPassNode;
 };
 // ToDo: Both of these can probably go into the Default Renderer implementation.
 
@@ -219,6 +220,7 @@ namespace UE::MovieGraph
 			: bWriteSampleToDisk(false)
 			, bRequiresAccumulator(false)
 			, bFetchFromAccumulator(false)
+			, bCompositeOnOtherRenders(false)
 		{}
 
 		/** The traversal context used to read graph values at the time of submission. */
@@ -241,6 +243,9 @@ namespace UE::MovieGraph
 
 		/** Set this to true for the last sample for a given frame (if it requires accumulation) to fetch the accumulated data out of the accumulator. */
 		bool bFetchFromAccumulator;
+
+		/** Set this to true if this pass should be composited on top of other renders. */
+		bool bCompositeOnOtherRenders;
 	};
 
 	/**
diff --git a/Engine/Plugins/MovieScene/MovieRenderPipeline/Source/MovieRenderPipelineCore/Public/Graph/Nodes/MovieGraphBurnInNode.h b/Engine/Plugins/MovieScene/MovieRenderPipeline/Source/MovieRenderPipelineCore/Public/Graph/Nodes/MovieGraphBurnInNode.h
new file mode 100644
--- /dev/null
+++ b/Engine/Plugins/MovieScene/MovieRenderPipeline/Source/MovieRenderPipelineCore/Public/Graph/Nodes/MovieGraphBurnInNode.h
@@ -0,0 +1,101 @@
+// Copyright Epic Games, Inc. All Rights Reserved.
+
+#pragma once
+
+#include "Graph/Nodes/MovieGraphRenderPassNode.h"
+
+#include "MovieGraphBurnInNode.generated.h"
+
+// Forward Declares
+class FWidgetRenderer;
+class SVirtualWindow;
+class UMovieGraphBurnInWidget;
+class UMovieGraphDefaultRenderer;
+class UMoviePipelineBurnInWidget;
+struct FMovieGraphRenderPassLayerData;
+struct FMovieGraphRenderPassSetupData;
+
+/** A node which generates a widget burn-in, rendered to a standalone image or composited on top of a render layer. */
+UCLASS()
+class MOVIERENDERPIPELINECORE_API UMovieGraphBurnInNode : public UMovieGraphRenderPassNode
+{
+	GENERATED_BODY()
+
+public:
+#if WITH_EDITOR
+	virtual FText GetNodeTitle(const bool bGetDescriptive = false) const override;
+	virtual FLinearColor GetNodeTitleColor() const override;
+	virtual FSlateIcon GetIconAndTint(FLinearColor& OutColor) const override;
+#endif
+
+	/** Gets an existing virtual window with the specified resolution, or creates a new one if one does not exist yet. */
+	TSharedPtr<SVirtualWindow> GetOrCreateVirtualWindow(const FIntPoint& InResolution);
+
+	/**
+	 * Gets an existing widget instance of type WidgetClass if one has been created. Otherwise, returns a new instance
+	 * of WidgetClass with owner InOwner.
+	 */
+	TObjectPtr<UMovieGraphBurnInWidget> GetOrCreateBurnInWidget(UClass* WidgetClass, UWorld* InOwner);
+
+public:
+	static const FString RendererName;
+	
+	UPROPERTY(EditAnywhere, BlueprintReadWrite, Category = Overrides, meta = (InlineEditConditionToggle))
+	uint8 bOverride_BurnInClass : 1;
+
+	UPROPERTY(EditAnywhere, BlueprintReadWrite, Category = Overrides, meta = (InlineEditConditionToggle))
+	uint8 bOverride_bCompositeOntoFinalImage : 1;
+
+	/** The widget that the burn-in should use. */
+	UPROPERTY(EditAnywhere, BlueprintReadWrite, meta=(MetaClass="/Script/MovieRenderPipelineCore.MovieGraphBurnInWidget", EditCondition="bOverride_BurnInClass"), Category="Widget Settings")
+	FSoftClassPath BurnInClass;
+
+	/** If true, the burn-in image will be composited into the final image pass. Does not apply to multi-layer EXR files. */
+	UPROPERTY(EditAnywhere, BlueprintReadWrite, meta=(EditCondition="bOverride_bCompositeOntoFinalImage"), Category="Widget Settings")
+	bool bCompositeOntoFinalImage;
+
+protected:
+	/** A render pass for a specific render layer, generated by the UMovieGraphBurnInNode CDO. */
+	struct FMovieGraphBurnInPass
+	{
+		void Setup(TWeakObjectPtr<UMovieGraphDefaultRenderer> InRenderer, const FMovieGraphRenderPassLayerData& InLayer);
+		void Teardown();
+		void Render(const FMovieGraphTraversalContext& InFrameTraversalContext, const FMovieGraphTimeStepData& InTimeData);
+		void GatherOutputPassesImpl(TArray<FMovieGraphRenderDataIdentifier>& OutExpectedPasses) const;
+
+	protected:
+		FMovieGraphRenderPassLayerData LayerData;
+
+		/** Unique identifier passed in GatherOutputPasses and with each render that identifies the data produced by this renderer. */
+		FMovieGraphRenderDataIdentifier RenderDataIdentifier;
+
+		/** The renderer associated with this pass. */
+		TWeakObjectPtr<UMovieGraphDefaultRenderer> Renderer;
+
+		/** The node instance (not CDO) associated with this pass. */
+		TWeakObjectPtr<UMovieGraphBurnInNode> RenderPassNode;
+	};
+
+private:
+	/** Widget renderer shared with all FMovieGraphBurnInPass instances. */
+	TSharedPtr<FWidgetRenderer> WidgetRenderer;
+
+	/** Virtual window instances shared with all FMovieGraphBurnInPass instances, keyed by resolution. */
+	TMap<const FIntPoint, TSharedPtr<SVirtualWindow>> SharedVirtualWindows;
+
+	/** Burn-in widget instances shared with all FMovieGraphBurnInPass instances, keyed by burn-in class. */
+	UPROPERTY(Transient)
+	TMap<const UClass*, TObjectPtr<UMovieGraphBurnInWidget>> BurnInWidgetInstances;
+
+	/** Render passes generated by the UMovieGraphBurnInNode CDO. */
+	TArray<TUniquePtr<FMovieGraphBurnInPass>> CurrentInstances;
+
+protected:
+	// UMovieGraphRenderPassNode Interface
+	virtual FString GetRendererNameImpl() const override { return RendererName; }
+	virtual void SetupImpl(const FMovieGraphRenderPassSetupData& InSetupData) override;
+	virtual void TeardownImpl() override;
+	virtual void RenderImpl(const FMovieGraphTraversalContext& InFrameTraversalContext, const FMovieGraphTimeStepData& InTimeData) override;
+	virtual void GatherOutputPassesImpl(TArray<FMovieGraphRenderDataIdentifier>& OutExpectedPasses) const override;
+	// ~UMovieGraphRenderPassNode Interface
+};
\ No newline at end of file
diff --git a/Engine/Plugins/MovieScene/MovieRenderPipeline/Source/MovieRenderPipelineCore/Public/MoviePipelineBlueprintLibrary.h b/Engine/Plugins/MovieScene/MovieRenderPipeline/Source/MovieRenderPipelineCore/Public/MoviePipelineBlueprintLibrary.h
--- a/Engine/Plugins/MovieScene/MovieRenderPipeline/Source/MovieRenderPipelineCore/Public/MoviePipelineBlueprintLibrary.h
+++ b/Engine/Plugins/MovieScene/MovieRenderPipeline/Source/MovieRenderPipelineCore/Public/MoviePipelineBlueprintLibrary.h
@@ -163,4 +163,10 @@ public:
 	UFUNCTION(BlueprintPure, Category = "Movie Render Pipeline")
 	static FText GetMoviePipelineEngineChangelistLabel(const UMoviePipeline* InMoviePipeline);
 
+private:
+	// Give the graph BP library access to shared utility methods
+	friend class UMovieGraphBlueprintLibrary;
+
+	/** Get the effective output resolution, taking into account overscan. */
+	static FIntPoint Utility_GetEffectiveOutputResolution(const float OverscanPercentage, const FIntPoint& InOutputResolution);
 };
\ No newline at end of file
