OCIO: Fixed issues caused by a recent MRQ CL merge.

diff --git a/Engine/Plugins/MovieScene/MovieRenderPipeline/Source/MovieRenderPipelineCore/Private/Graph/Nodes/MovieGraphBurnInNode.cpp b/Engine/Plugins/MovieScene/MovieRenderPipeline/Source/MovieRenderPipelineCore/Private/Graph/Nodes/MovieGraphBurnInNode.cpp
--- a/Engine/Plugins/MovieScene/MovieRenderPipeline/Source/MovieRenderPipelineCore/Private/Graph/Nodes/MovieGraphBurnInNode.cpp
+++ b/Engine/Plugins/MovieScene/MovieRenderPipeline/Source/MovieRenderPipelineCore/Private/Graph/Nodes/MovieGraphBurnInNode.cpp
@@ -168,7 +168,7 @@ void UMovieGraphBurnInNode::FMovieGraphBurnInPass::Render(const FMovieGraphTrave
 	// Create the render target the widget will be rendered into
 	UE::MovieGraph::DefaultRenderer::FRenderTargetInitParams RenderTargetInitParams;
 	RenderTargetInitParams.Size = OutputResolution;
-	RenderTargetInitParams.TargetGamma = FOpenColorIODisplayExtension::DefaultDisplayGamma;
+	RenderTargetInitParams.TargetGamma = FOpenColorIORendering::DefaultDisplayGamma;
 	RenderTargetInitParams.PixelFormat = PF_B8G8R8A8;
 	UTextureRenderTarget2D* RenderTarget = Renderer->GetOrCreateViewRenderTarget(RenderTargetInitParams);
 	
