MoviePipeline: Fixed an issue where when rendering with the UI Renderer or the Burn In, it would warn you about multiple cameras but no {camera_name} token, despite only using one camera. This also fixes an issue where widgets were not being composited onto the final images.

Fallout from CL 25881669.

diff --git a/Engine/Plugins/MovieScene/MovieRenderPipeline/Source/MovieRenderPipelineCore/Private/MoviePipelineBlueprintLibrary.cpp b/Engine/Plugins/MovieScene/MovieRenderPipeline/Source/MovieRenderPipelineCore/Private/MoviePipelineBlueprintLibrary.cpp
--- a/Engine/Plugins/MovieScene/MovieRenderPipeline/Source/MovieRenderPipelineCore/Private/MoviePipelineBlueprintLibrary.cpp
+++ b/Engine/Plugins/MovieScene/MovieRenderPipeline/Source/MovieRenderPipelineCore/Private/MoviePipelineBlueprintLibrary.cpp
@@ -586,11 +586,15 @@ void UMoviePipelineBlueprintLibrary::UpdateJobShotListFromSequence(ULevelSequenc
 						SidecarCamera.SequenceId = Entity.SequenceID; 
 						bool bAddCamera = true;
 						
+						// We break out of the while loop below when the parent guid is invalid, so we track
+						// the last known valid GUID too, so we can compare the root binding against the camera.
+						FGuid LastValidGuid = Possessable.GetGuid();
 						FGuid ParentGuid = Possessable.GetParent();
 						while (ParentGuid.IsValid())
 						{
 							if (FMovieScenePossessable* ParentAsPossessable = LeafNode->MovieScene->FindPossessable(ParentGuid))
 							{
+								LastValidGuid = ParentGuid;
 								ParentGuid = ParentAsPossessable->GetParent();
 								SidecarCamera.Name = ParentAsPossessable->GetName();
 							}
@@ -616,6 +620,7 @@ void UMoviePipelineBlueprintLibrary::UpdateJobShotListFromSequence(ULevelSequenc
 								}
 
 								// Spawnables will never have a parent
+								LastValidGuid = ParentGuid;
 								ParentGuid.Invalidate(); 
 								
 								SidecarCamera.Name = ParentAsSpawnable->GetName();
@@ -627,7 +632,7 @@ void UMoviePipelineBlueprintLibrary::UpdateJobShotListFromSequence(ULevelSequenc
 							// This comparison doesn't support cross-sequence references, technically we should create a 
 							// FMovieSceneObjectBindingID etc, but because we only support finding sidecar cameras in the
 							// same sequence as the main camera cut track, this is fine for now.
-							if (ParentGuid == MainBinding.GetGuid())
+							if (LastValidGuid == MainBinding.GetGuid())
 							{
 								// We want to know which sidecar camera this is, not the possessable index.
 								MainCameraIndex = Entity.SidecarCameras.Num();
diff --git a/Engine/Plugins/MovieScene/MovieRenderPipeline/Source/MovieRenderPipelineSettings/Private/MoviePipelineBurnInSetting.cpp b/Engine/Plugins/MovieScene/MovieRenderPipeline/Source/MovieRenderPipelineSettings/Private/MoviePipelineBurnInSetting.cpp
--- a/Engine/Plugins/MovieScene/MovieRenderPipeline/Source/MovieRenderPipelineSettings/Private/MoviePipelineBurnInSetting.cpp
+++ b/Engine/Plugins/MovieScene/MovieRenderPipeline/Source/MovieRenderPipelineSettings/Private/MoviePipelineBurnInSetting.cpp
@@ -34,7 +34,11 @@ void UMoviePipelineBurnInSetting::GatherOutputPassesImpl(TArray<FMoviePipelinePa
 		{
 			FMoviePipelinePassIdentifier PassIdentifierForCurrentCamera;
 			PassIdentifierForCurrentCamera.Name = TEXT("BurnInOverlay");
-			PassIdentifierForCurrentCamera.CameraName = CurrentShot->GetCameraName(CameraIndex);
+
+			// If we're not rendering all cameras, we need to pass -1 so we pick up the real camera name.
+			int32 LocalCameraIndex = CameraSettings->bRenderAllCameras ? CameraIndex : -1;
+			PassIdentifierForCurrentCamera.CameraName = CurrentShot->GetCameraName(LocalCameraIndex);
+
 			ExpectedRenderPasses.Add(PassIdentifierForCurrentCamera);
 		}
 	}
@@ -63,11 +67,14 @@ void UMoviePipelineBurnInSetting::RenderSample_GameThreadImpl(const FMoviePipeli
 		int32 NumCameras = CameraSettings->bRenderAllCameras ? CurrentShot->SidecarCameras.Num() : 1;
 		for (int32 CameraIndex = 0; CameraIndex < NumCameras; CameraIndex++)
 		{
+			// If we're not rendering all cameras, we need to pass -1 so we pick up the real camera name.
+			int32 LocalCameraIndex = CameraSettings->bRenderAllCameras ? CameraIndex : -1;
+
 			FMoviePipelinePassIdentifier PassIdentifierForCurrentCamera;
 			PassIdentifierForCurrentCamera.Name = TEXT("BurnInOverlay");
-			PassIdentifierForCurrentCamera.CameraName = CurrentShot->GetCameraName(CameraIndex);
+			PassIdentifierForCurrentCamera.CameraName = CurrentShot->GetCameraName(LocalCameraIndex);
 
-			UMoviePipelineBurnInWidget* CurrentWidget = BurnInWidgetInstances[CameraIndex];
+			UMoviePipelineBurnInWidget* CurrentWidget = BurnInWidgetInstances[FMath::Clamp(LocalCameraIndex, 0, LocalCameraIndex)];
 			// Put the widget in our window
 			VirtualWindow->SetContent(CurrentWidget->TakeWidget());
 
diff --git a/Engine/Plugins/MovieScene/MovieRenderPipeline/Source/MovieRenderPipelineSettings/Private/MoviePipelineWidgetRenderSetting.cpp b/Engine/Plugins/MovieScene/MovieRenderPipeline/Source/MovieRenderPipelineSettings/Private/MoviePipelineWidgetRenderSetting.cpp
--- a/Engine/Plugins/MovieScene/MovieRenderPipeline/Source/MovieRenderPipelineSettings/Private/MoviePipelineWidgetRenderSetting.cpp
+++ b/Engine/Plugins/MovieScene/MovieRenderPipeline/Source/MovieRenderPipelineSettings/Private/MoviePipelineWidgetRenderSetting.cpp
@@ -32,7 +32,10 @@ void UMoviePipelineWidgetRenderer::GatherOutputPassesImpl(TArray<FMoviePipelineP
 	{
 		FMoviePipelinePassIdentifier PassIdentifierForCurrentCamera;
 		PassIdentifierForCurrentCamera.Name = TEXT("ViewportUI");
-		PassIdentifierForCurrentCamera.CameraName = CurrentShot->GetCameraName(CameraIndex);
+		// If we're not rendering all cameras, we need to pass -1 so we pick up the real camera name.
+		int32 LocalCameraIndex = CameraSettings->bRenderAllCameras ? CameraIndex : -1;
+		PassIdentifierForCurrentCamera.CameraName = CurrentShot->GetCameraName(LocalCameraIndex);
+
 		ExpectedRenderPasses.Add(PassIdentifierForCurrentCamera);
 	}
 }
@@ -57,7 +60,10 @@ void UMoviePipelineWidgetRenderer::RenderSample_GameThreadImpl(const FMoviePipel
 		{
 			FMoviePipelinePassIdentifier PassIdentifierForCurrentCamera;
 			PassIdentifierForCurrentCamera.Name = TEXT("ViewportUI");
-			PassIdentifierForCurrentCamera.CameraName = CurrentShot->GetCameraName(CameraIndex);
+
+			// If we're not rendering all cameras, we need to pass -1 so we pick up the real camera name.
+			int32 LocalCameraIndex = CameraSettings->bRenderAllCameras ? CameraIndex : -1;
+			PassIdentifierForCurrentCamera.CameraName = CurrentShot->GetCameraName(LocalCameraIndex);
 
 			// Draw the widget to the render target
 			FRenderTarget* BackbufferRenderTarget = RenderTarget->GameThread_GetRenderTargetResource();
