MoviePipeline: Optimized the overlapped accumulator for RGBA images. Improvements mostly center around going wider on high-core count machines with potential improvement going from 70ms->5.7ms per accumulation sample (at 4k). The MRQ pipeline stalls from submitting new renders when no accumulators are available (due to being busy accumulating samples in linear fashion). This is generally noticeable as a few frames will render and then MRQ appears to stall while the accumulators finish and catch up, then a few more frames are submitted, then stall, etc.

Timings for some of the various improvements:

ZeroPlanes (used to reset an accumulator after each output frame):
7.8ms -> 3ms (ParallelFor to loop over RGBA(+Weight) planes all at the same time)
3ms -> 1ms (Memset instead of a for loop setting the values pixel by pixel)

PlaneAccumulation (Unpack FP16 value and add results to existing image plane)
15.27ms -> 6.57ms (ParallelFor to loop over RGBA(+Weight) planes all at the same time)
4.48ms -> 0.90ms (Accumulate by Scanline within each plane accumulation)

Fetch (Get the now accumulated data back out of the Accumulator, re-packing it into RGBA)
38ms->7ms (ParallelFor)

AccumulatePixelData
25ms -> 16.1ms
This was the trickiest improvement as the real cost in this function was hiding in other functions. The accumulation function was allocating a 32mb array each frame, but the allocation only appeared to take 315us so it didn't seem high priority compared to the rest of the function taking 9.8ms (9800us). What was not obvious though is that the OS returns the allocation before actually physically allocating the memory and instead it can wait until you access the memory to actually allocate (and potentially clear the allocated memory for security purposes) which means the cost was being hidden in another function (the cost was about 5ms for the allocation). And then when freeing this temporarily allocated 32mb buffer, that takes another 2ms, but that was attributed to the parent function call as it was only happening during the destructor being run. So in reality about 7 of the 9.8ms was the per-frame allocation (which was moved to a persistent member-level allocation).

There remains a few areas for improvement that were not pursued due to the scope of changes they would require. Currently the Surface Readback code creates a copy which costs us about 3ms, and eventually we have to delete this copy (which costs us another 2ms), but both of these costs are on the Task Thread, and because our overall accumulation time has lowered significantly it was deemed not worth pursuing due to the invasive changes required. Another place we could consider optimizing is the extra "WeightPlaneValue". It's always filled with 1.0 (for all weight planes, every time) and is structured this way so we can re-use the RGBA accumulation code. Ideally we wouldn't allocate a spare memory block just to fill it with 1's but the code is too intrusive to change, so we pay the 8mb RAM cost (1080p, per accumulator, so ~50mb total) for now.

diff --git a/Engine/Plugins/MovieScene/MovieRenderPipeline/Source/MovieRenderPipelineCore/Private/MovieRenderOverlappedImage.cpp b/Engine/Plugins/MovieScene/MovieRenderPipeline/Source/MovieRenderPipelineCore/Private/MovieRenderOverlappedImage.cpp
--- a/Engine/Plugins/MovieScene/MovieRenderPipeline/Source/MovieRenderPipelineCore/Private/MovieRenderOverlappedImage.cpp
+++ b/Engine/Plugins/MovieScene/MovieRenderPipeline/Source/MovieRenderPipelineCore/Private/MovieRenderOverlappedImage.cpp
@@ -23,17 +23,13 @@ void FImageOverlappedPlane::Init(FIntPoint InSize)
 {
 	Size = InSize;
 
-	// leaves the data as is
+	// Don't initialize now as we'll call ZeroPlane later when appropriate to reset the values.
 	ChannelData.SetNumUninitialized(Size.X * Size.Y);
 }
 
 void FImageOverlappedPlane::ZeroPlane()
 {
-	int32 Num = Size.X * Size.Y;
-	for (int32 Index = 0; Index < Num; Index++)
-	{
-		ChannelData[Index] = 0.0f;
-	}
+	FMemory::Memset(ChannelData.GetData(), 0, ChannelData.Num() * sizeof(float));
 }
 
 void FImageOverlappedPlane::Reset()
@@ -224,8 +220,14 @@ void FImageOverlappedPlane::AccumulateSinglePlane(const TArray64<float>& InRawDa
 			VectorRegister VecPixelWeight10 = VectorSetFloat1(PixelWeight[1][0]);
 			VectorRegister VecPixelWeight11 = VectorSetFloat1(PixelWeight[1][1]);
 
-			for (int32 DstY = ActualDstY0; DstY < ActualDstY1; DstY++)
+			int32 Offset = ActualDstY0;
+			int32 Count = ActualDstY1 - ActualDstY0;
+			ParallelFor(Count, [Offset, StartY, InSize, ActualDstX0, ActualDstX1,StartX, 
+				PixelWeight00, PixelWeight01, PixelWeight10, PixelWeight11, 
+				VecPixelWeight00, VecPixelWeight01, VecPixelWeight10, VecPixelWeight11,
+				Size = this->Size, &ChannelData = this->ChannelData, &InRawData, &WeightDataX, &WeightDataY](int32 InIndexY)
 			{
+				int32 DstY = InIndexY + Offset;
 				int32 CurrY = DstY - StartY;
 				int32 SrcY0 = FMath::Max<int32>(CurrY + (0 - 1), 0); // we need the max for the first pixel which could go off the edge
 				int32 SrcY1 = FMath::Max<int32>(CurrY + (1 - 1), 0); 
@@ -340,7 +342,7 @@ void FImageOverlappedPlane::AccumulateSinglePlane(const TArray64<float>& InRawDa
 					DstLine[DstX] = Sum;
 				}
 
-			}
+			});
 		}
 	}
 }
@@ -354,12 +356,18 @@ void FImageOverlappedAccumulator::InitMemory(FIntPoint InPlaneSize, int32 InNumC
 
 	ChannelPlanes.SetNum(NumChannels);
 
-	for (int32 Channel = 0; Channel < NumChannels; Channel++)
+	// Initialize the RGBA channels and the weight plane all at the same time.
+	ParallelFor(NumChannels + 1, [&](int32 InChannelIndex)
 	{
-		ChannelPlanes[Channel].Init(PlaneSize);
-	}
-
-	WeightPlane.Init(PlaneSize);
+		if (InChannelIndex < NumChannels)
+		{
+			ChannelPlanes[InChannelIndex].Init(PlaneSize);
+		}
+		else
+		{
+			WeightPlane.Init(PlaneSize);
+		}
+	});
 }
 
 void FImageOverlappedAccumulator::ZeroPlanes()
@@ -367,23 +375,32 @@ void FImageOverlappedAccumulator::ZeroPlanes()
 	SCOPE_CYCLE_COUNTER(STAT_ZeroPlanes);
 	check(ChannelPlanes.Num() == NumChannels);
 
-	for (int32 Channel = 0; Channel < NumChannels; Channel++)
+	// Zero our RGBA channels and the weight plane all at the same time.
+	ParallelFor(NumChannels + 1, [&](int32 InChannelIndex)
 	{
-		ChannelPlanes[Channel].ZeroPlane();
-	}
+		if (InChannelIndex < NumChannels)
+		{
+			ChannelPlanes[InChannelIndex].ZeroPlane();
+		}
+		else
+		{
+			WeightPlane.ZeroPlane();
+		}
+	});
 
-	WeightPlane.ZeroPlane();
 }
 
 void FImageOverlappedAccumulator::Reset()
 {
+	TRACE_CPUPROFILER_EVENT_SCOPE(MoviePipeline::ResetAccumulator);
+	
+	// We don't free the memory on the individual channels because
+	// Reset is called at the end of each output frame and the next
+	// output frame is likely to be the same resolution and need the
+	// same memory. The individual channels gracefully handle resizing.
 	PlaneSize.X = 0;
 	PlaneSize.Y = 0;
 	NumChannels = 0;
-
-	// Let the desctructor clean up
-	ChannelPlanes.Empty();
-	WeightPlane.Reset();
 }
 
 #if 0
@@ -509,15 +526,34 @@ void FImageOverlappedAccumulator::AccumulatePixelData(const FImagePixelData& InP
 		FIntPoint SubRectOffset(WeightX.X0, WeightY.X0);
 		FIntPoint SubRectSize(WeightX.X3 - WeightX.X0, WeightY.X3 - WeightY.X0);
 
-		// hardcode to 4 channels (RGBA), even if we are only saving fewer channels
-		TArray64<float> RawData[4];
 
 		check(NumChannels >= 1);
 		check(NumChannels <= 4);
 
-		for (int32 ChanIter = 0; ChanIter < 4; ChanIter++)
 		{
-			RawData[ChanIter].SetNumUninitialized(RawSize.X * RawSize.Y);
+			// We always make room for all 4 channels (RGBA) even if we only accumulate 3 of them.
+			if (UnpackedDataStorage[0].Num() != RawSize.X * RawSize.Y)
+			{
+				TRACE_CPUPROFILER_EVENT_SCOPE(AllocDataStorage);
+				ParallelFor(4, [&](int32 InChannelIndex)
+				{
+					UnpackedDataStorage[InChannelIndex].SetNumUninitialized(RawSize.X * RawSize.Y);
+				});
+			}
+
+			// Initialize the weight channel with 1's so that we can accumulate a bunch of 1 values together
+			if (WeightChannelStorage.Num() != RawSize.X * RawSize.Y)
+			{
+				TRACE_CPUPROFILER_EVENT_SCOPE(AllocWeightChannelStorage);
+				WeightChannelStorage.SetNumUninitialized(RawSize.X * RawSize.Y);
+				ParallelFor(RawSize.Y, [&](int32 InIndexY)
+				{
+					for (int32 IndexX = 0; IndexX < RawSize.X; IndexX++)
+					{
+						WeightChannelStorage[(InIndexY * RawSize.X) + IndexX] = 1.f;
+					}
+				});
+			}
 		}
 
 		const double AccumulateBeginTime = FPlatformTime::Seconds();
@@ -541,7 +577,7 @@ void FImageOverlappedAccumulator::AccumulatePixelData(const FImagePixelData& InP
 							int32 RawValue = RawDataPtr[(Y*RawSize.X + X)*RawNumChan + ChanIter];
 							float Value = float(RawValue) / 255.0f;
 							int32 Reorder = ChanReorder[ChanIter];
-							RawData[Reorder][Y*RawSize.X + X] = Value;
+							UnpackedDataStorage[Reorder][Y*RawSize.X + X] = Value;
 						}
 					}
 				}
@@ -553,10 +589,10 @@ void FImageOverlappedAccumulator::AccumulatePixelData(const FImagePixelData& InP
 				{
 					const uint8* SrcRowDataPtr = &RawDataPtr[Y*RawSize.X*RawNumChan];
 
-					float* DstRowDataR = &RawData[0][Y*RawSize.X];
-					float* DstRowDataG = &RawData[1][Y*RawSize.X];
-					float* DstRowDataB = &RawData[2][Y*RawSize.X];
-					float* DstRowDataA = &RawData[3][Y*RawSize.X];
+					float* DstRowDataR = &UnpackedDataStorage[0][Y*RawSize.X];
+					float* DstRowDataG = &UnpackedDataStorage[1][Y*RawSize.X];
+					float* DstRowDataB = &UnpackedDataStorage[2][Y*RawSize.X];
+					float* DstRowDataA = &UnpackedDataStorage[3][Y*RawSize.X];
 
 					VectorRegister ColorScale = MakeVectorRegister(1.0f / 255.0f, 1.0f / 255.0f, 1.0f / 255.0f, 1.0f / 255.0f);
 
@@ -577,24 +613,23 @@ void FImageOverlappedAccumulator::AccumulatePixelData(const FImagePixelData& InP
 		}
 		else if (Fmt == EImagePixelType::Float16 && RawNumChan == 4 && RawBitDepth == 16)
 		{
+			TRACE_CPUPROFILER_EVENT_SCOPE(MoviePipeline::Float16Unpack);
 			const uint16* RawDataPtr = static_cast<const uint16*>(SrcRawDataPtr);
 
-			// simple, slow unpack version for fp16, could make this faster
-			for (int32 Y = 0; Y < RawSize.Y; Y++)
+			ParallelFor(RawSize.Y, [&](int32 Y)
 			{
 				for (int32 X = 0; X < RawSize.X; X++)
 				{
+					float UnpackedColors[4];
+					uint16_t* HalfColor = (uint16_t*)&RawDataPtr[(Y * RawSize.X + X) * RawNumChan];
+					FPlatformMath::VectorLoadHalf(&UnpackedColors[0], HalfColor);
+
 					for (int32 ChanIter = 0; ChanIter < 4; ChanIter++)
 					{
-						FFloat16 RawValue;
-						RawValue.Encoded = RawDataPtr[(Y*RawSize.X + X)*RawNumChan + ChanIter];
-						
-						// c cast does the conversion from fp16 bits to float
-						float Value = float(RawValue);
-						RawData[ChanIter][Y*RawSize.X + X] = Value;
+						UnpackedDataStorage[ChanIter][Y * RawSize.X + X] = UnpackedColors[ChanIter];
 					}
 				}
-			}
+			});
 		}
 		else if (Fmt == EImagePixelType::Float32 && RawNumChan == 4 && RawBitDepth == 32)
 		{
@@ -612,7 +647,7 @@ void FImageOverlappedAccumulator::AccumulatePixelData(const FImagePixelData& InP
 						for (int32 ChanIter = 0; ChanIter < 4; ChanIter++)
 						{
 							float Value = RawDataPtr[(Y*RawSize.X + X)*RawNumChan + ChanIter];
-							RawData[ChanIter][Y*RawSize.X + X] = Value;
+							UnpackedDataStorage[ChanIter][Y*RawSize.X + X] = Value;
 						}
 					}
 				}
@@ -626,10 +661,10 @@ void FImageOverlappedAccumulator::AccumulatePixelData(const FImagePixelData& InP
 				{
 					const float* SrcRowDataPtr = &RawDataPtr[Y*RawSize.X*RawNumChan];
 
-					float* DstRowDataR = &RawData[0][Y*RawSize.X];
-					float* DstRowDataG = &RawData[1][Y*RawSize.X];
-					float* DstRowDataB = &RawData[2][Y*RawSize.X];
-					float* DstRowDataA = &RawData[3][Y*RawSize.X];
+					float* DstRowDataR = &UnpackedDataStorage[0][Y*RawSize.X];
+					float* DstRowDataG = &UnpackedDataStorage[1][Y*RawSize.X];
+					float* DstRowDataB = &UnpackedDataStorage[2][Y*RawSize.X];
+					float* DstRowDataA = &UnpackedDataStorage[3][Y*RawSize.X];
 
 					// simple, one pixel at a time vectorized version, we could do better
 					for (int32 X = 0; X < RawSize.X; X++)
@@ -658,12 +693,13 @@ void FImageOverlappedAccumulator::AccumulatePixelData(const FImagePixelData& InP
 	
 		if (AccumulationGamma != 1.0f)
 		{
+			TRACE_CPUPROFILER_EVENT_SCOPE(MoviePipeline::ApplyAccumulationGamma);
 			// Unfortunately, we don't have an SSE optimized pow function. This function is quite slow (about 30-40ms).
 			float Gamma = AccumulationGamma;
 
 			for (int32 ChanIter = 0; ChanIter < NumChannels; ChanIter++)
 			{
-				float* DstData = RawData[ChanIter].GetData();
+				float* DstData = UnpackedDataStorage[ChanIter].GetData();
 
 				int32 DstSize = RawSize.X * RawSize.Y;
 
@@ -712,24 +748,28 @@ void FImageOverlappedAccumulator::AccumulatePixelData(const FImagePixelData& InP
 			UE_LOG(LogTemp, Log, TEXT("    [%8.2f] Tile time."), ElapsedMs);
 		}
 
-
-		for (int32 ChanIter = 0; ChanIter < NumChannels; ChanIter++)
 		{
-			ChannelPlanes[ChanIter].AccumulateSinglePlane(RawData[ChanIter], RawSize, InTileOffset, InSubpixelOffset.X, InSubpixelOffset.Y,
-				SubRectOffset, SubRectSize, WeightDataX, WeightDataY);
-		}
+			TRACE_CPUPROFILER_EVENT_SCOPE(MoviePipeline::PlaneAccumulation);
 
-		// Accumulate weights as well.
-		{
-			TArray64<float> VecOne;
-			VecOne.SetNum(RawSize.X * RawSize.Y);
-			for (int32 Index = 0; Index < VecOne.Num(); Index++)
-			{
-				VecOne[Index] = 1.0f;
-			}
+			// We accumulate an RGBA plane, and an additional weight plane (for dividing per-pixel later)
+			// but we do them all in the same ParallelFor through a slightly confusing bit of logic since
+			// they're stored in separate arrays. This allows us to go wider than doing the RGBA, waiting
+			// for them to finish, and then the weight plane.
+			ParallelFor(NumChannels+1, [&](int32 InChannelIndex)
+				{
+					if (InChannelIndex < NumChannels)
+					{
+						ChannelPlanes[InChannelIndex].AccumulateSinglePlane(UnpackedDataStorage[InChannelIndex], RawSize, InTileOffset, InSubpixelOffset.X, InSubpixelOffset.Y,
+							SubRectOffset, SubRectSize, WeightDataX, WeightDataY);
+					}
+					else
+					{
+						WeightPlane.AccumulateSinglePlane(WeightChannelStorage, RawSize, InTileOffset, InSubpixelOffset.X, InSubpixelOffset.Y,
+							SubRectOffset, SubRectSize, WeightDataX, WeightDataY);
+
+					}
+				});
 
-			WeightPlane.AccumulateSinglePlane(VecOne, RawSize, InTileOffset, InSubpixelOffset.X, InSubpixelOffset.Y,
-				SubRectOffset, SubRectSize, WeightDataX, WeightDataY);
 		}
 
 		const double PlaneEndTime = FPlatformTime::Seconds();
@@ -738,13 +778,11 @@ void FImageOverlappedAccumulator::AccumulatePixelData(const FImagePixelData& InP
 			const float ElapsedMs = float((PlaneEndTime - TileEndTime)*1000.0f);
 			UE_LOG(LogTemp, Log, TEXT("    [%8.2f] Plane time."), ElapsedMs);
 		}
-
 	}
 }
 
 void FImageOverlappedAccumulator::FetchFullImageValue(float Rgba[4], int32 FullX, int32 FullY) const
 {
-	SCOPE_CYCLE_COUNTER(STAT_FetchFullImageValue);
 
 	Rgba[0] = 0.0f;
 	Rgba[1] = 0.0f;
@@ -802,20 +840,20 @@ void FImageOverlappedAccumulator::FetchFinalPixelDataHalfFloat(TArray64<FFloat16
 	int32 FullSizeY = PlaneSize.Y;
 	OutPixelData.SetNumUninitialized(FullSizeX * FullSizeY);
 
-	for (int32 FullY = 0L; FullY < FullSizeY; FullY++)
+	ParallelFor(FullSizeY, [&](int32 InIndexY)
 	{
 		for (int32 FullX = 0L; FullX < FullSizeX; FullX++)
 		{
 			float Rgba[4] = { 0.0f, 0.0f, 0.0f, 1.0f };
 
-			FetchFullImageValue(Rgba, FullX, FullY);
+			FetchFullImageValue(Rgba, FullX, InIndexY);
 			FLinearColor Color = FLinearColor(Rgba[0],Rgba[1],Rgba[2],Rgba[3]);
 
 			// be careful with this index, make sure to use 64bit math, not 32bit
-			int64 DstIndex = int64(FullY) * int64(FullSizeX) + int64(FullX);
+			int64 DstIndex = int64(InIndexY) * int64(FullSizeX) + int64(FullX);
 			OutPixelData[DstIndex] = FFloat16Color(Color);
 		}
-	}
+	});
 }
 
 void FImageOverlappedAccumulator::FetchFinalPixelDataLinearColor(TArray64<FLinearColor> & OutPixelData) const
diff --git a/Engine/Plugins/MovieScene/MovieRenderPipeline/Source/MovieRenderPipelineCore/Public/MovieRenderOverlappedImage.h b/Engine/Plugins/MovieScene/MovieRenderPipeline/Source/MovieRenderPipelineCore/Public/MovieRenderOverlappedImage.h
--- a/Engine/Plugins/MovieScene/MovieRenderPipeline/Source/MovieRenderPipelineCore/Public/MovieRenderOverlappedImage.h
+++ b/Engine/Plugins/MovieScene/MovieRenderPipeline/Source/MovieRenderPipelineCore/Public/MovieRenderOverlappedImage.h
@@ -111,7 +111,9 @@ public:
 	void ZeroPlanes();
 
 	/**
-	 * Resets the memory.
+	 * Resets the PlaneSize and NumChannels but does NOT free the memory
+	 * because most of the time we're going to just need the same resolution
+	 * next frame anyways when InitMemory is called again.
 	 */
 	void Reset();
 
@@ -200,6 +202,19 @@ public:
 	/** Gamma for accumulation. Typical values are 1.0 and 2.2. */
 	float AccumulationGamma;
 
+
+	/** 
+	* A temporary buffer that is only initialized to the size of a single sample used during RGBA unpacking.
+	* We persist it as a member level variable to avoid having to allocate/free it for every sample (which
+	* is about a 32mb allocation at 4k) as that is quite slow.
+	*/
+	TArray64<float> UnpackedDataStorage[4];
+	/** 
+	* A temporary buffer that is only initialized to the size of a single sample used during weight plane accumulation.
+	* It is currently filled with 1.0f for all weights, and exists so that we can use the same code as we do for RGBA accumulation.
+	*/
+	TArray64<float> WeightChannelStorage;
+
 	TArray64<FImageOverlappedPlane> ChannelPlanes;
 	
 	FImageOverlappedPlane WeightPlane;
diff --git a/Engine/Plugins/MovieScene/MovieRenderPipeline/Source/MovieRenderPipelineRenderPasses/Private/MoviePipelineImagePassBase.cpp b/Engine/Plugins/MovieScene/MovieRenderPipeline/Source/MovieRenderPipelineRenderPasses/Private/MoviePipelineImagePassBase.cpp
--- a/Engine/Plugins/MovieScene/MovieRenderPipeline/Source/MovieRenderPipelineRenderPasses/Private/MoviePipelineImagePassBase.cpp
+++ b/Engine/Plugins/MovieScene/MovieRenderPipeline/Source/MovieRenderPipelineRenderPasses/Private/MoviePipelineImagePassBase.cpp
@@ -1004,6 +1004,13 @@ namespace MoviePipeline
 			// Free the memory in the accumulator.
 			PinnedImageAccumulator->Reset();
 		}
+
+		{
+			// Explicitly free the SamplePixelData (which by now has been copied into the accumulator)
+			// so that we can profile how long freeing the allocation takes.
+			TRACE_CPUPROFILER_EVENT_SCOPE(ReleasePixelDataSample);
+			SamplePixelData.Reset();
+		}
 	}
 }
 
