MoviePipeline: Added a new UI Renderer widget to the graph, as well as its rendering implementation.

Refactored the common logic for the UI renderer and the Burn-In renderer into a new base class, UMovieGraphWidgetRendererBaseNode.

diff --git a/Engine/Plugins/MovieScene/MovieRenderPipeline/Source/MovieRenderPipelineCore/Private/Graph/Nodes/MovieGraphBurnInNode.cpp b/Engine/Plugins/MovieScene/MovieRenderPipeline/Source/MovieRenderPipelineCore/Private/Graph/Nodes/MovieGraphBurnInNode.cpp
--- a/Engine/Plugins/MovieScene/MovieRenderPipeline/Source/MovieRenderPipelineCore/Private/Graph/Nodes/MovieGraphBurnInNode.cpp
+++ b/Engine/Plugins/MovieScene/MovieRenderPipeline/Source/MovieRenderPipelineCore/Private/Graph/Nodes/MovieGraphBurnInNode.cpp
@@ -2,19 +2,14 @@
 
 #include "Graph/Nodes/MovieGraphBurnInNode.h"
 
-#include "Engine/TextureRenderTarget2D.h"
 #include "Framework/Application/SlateApplication.h"
-#include "Graph/MovieGraphBlueprintLibrary.h"
 #include "Graph/MovieGraphBurnInWidget.h"
 #include "Graph/MovieGraphDataTypes.h"
 #include "Graph/MovieGraphDefaultRenderer.h"
 #include "Graph/MovieGraphPipeline.h"
 #include "MovieRenderPipelineCoreModule.h"
-#include "OpenColorIODisplayExtension.h"
-#include "Slate/WidgetRenderer.h"
 #include "Styling/AppStyle.h"
 #include "TextureResource.h"
-#include "Widgets/SVirtualWindow.h"
 
 const FString UMovieGraphBurnInNode::RendererName = FString("BurnIn");
 
@@ -38,79 +33,18 @@ FSlateIcon UMovieGraphBurnInNode::GetIconAndTint(FLinearColor& OutColor) const
 }
 #endif	// WITH_EDITOR
 
-void UMovieGraphBurnInNode::SetupImpl(const FMovieGraphRenderPassSetupData& InSetupData)
+TUniquePtr<UMovieGraphWidgetRendererBaseNode::FMovieGraphWidgetPass> UMovieGraphBurnInNode::GeneratePass()
 {
-	for (const FMovieGraphRenderPassLayerData& LayerData : InSetupData.Layers)
-	{
-		TUniquePtr<FMovieGraphBurnInPass> RendererInstance = MakeUnique<FMovieGraphBurnInPass>();
-		RendererInstance->Setup(InSetupData.Renderer, LayerData);
-		CurrentInstances.Add(MoveTemp(RendererInstance));
-	}
-
-	constexpr bool bApplyGammaCorrection = false;
-	WidgetRenderer = MakeShared<FWidgetRenderer>(bApplyGammaCorrection);
+	return MakeUnique<FMovieGraphBurnInPass>();
 }
 
 void UMovieGraphBurnInNode::TeardownImpl()
 {
-	for (const TUniquePtr<FMovieGraphBurnInPass>& Instance : CurrentInstances)
-	{
-		Instance->Teardown();
-	}
+	Super::TeardownImpl();
 	
-	CurrentInstances.Reset();
-
-	if (FSlateApplication::IsInitialized())
-	{
-		for (const TTuple<const FIntPoint, TSharedPtr<SVirtualWindow>>& Entry : SharedVirtualWindows)
-		{
-			const TSharedPtr<SVirtualWindow>& VirtualWindow = Entry.Value;
-			if (VirtualWindow.IsValid())
-			{
-				FSlateApplication::Get().UnregisterVirtualWindow(VirtualWindow.ToSharedRef());
-			}
-		}
-	}
-
-	SharedVirtualWindows.Empty();
-	WidgetRenderer = nullptr;
 	BurnInWidgetInstances.Empty();
 }
 
-void UMovieGraphBurnInNode::RenderImpl(const FMovieGraphTraversalContext& InFrameTraversalContext, const FMovieGraphTimeStepData& InTimeData)
-{
-	for (const TUniquePtr<FMovieGraphBurnInPass>& Instance : CurrentInstances)
-	{
-		Instance->Render(InFrameTraversalContext, InTimeData);
-	}
-}
-
-void UMovieGraphBurnInNode::GatherOutputPassesImpl(TArray<FMovieGraphRenderDataIdentifier>& OutExpectedPasses) const
-{
-	for (const TUniquePtr<FMovieGraphBurnInPass>& Instance : CurrentInstances)
-    {
-    	Instance->GatherOutputPassesImpl(OutExpectedPasses);
-    }
-}
-
-TSharedPtr<SVirtualWindow> UMovieGraphBurnInNode::GetOrCreateVirtualWindow(const FIntPoint& InResolution)
-{
-	if (const TSharedPtr<SVirtualWindow>* ExistingVirtualWindow = SharedVirtualWindows.Find(InResolution))
-	{
-		return *ExistingVirtualWindow;
-	}
-
-	const TSharedPtr<SVirtualWindow> NewVirtualWindow = SNew(SVirtualWindow).Size(FVector2D(InResolution.X, InResolution.Y));
-	SharedVirtualWindows.Emplace(InResolution, NewVirtualWindow);
-
-	if (FSlateApplication::IsInitialized())
-	{
-		FSlateApplication::Get().RegisterVirtualWindow(NewVirtualWindow.ToSharedRef());
-	}
-
-	return NewVirtualWindow;
-}
-
 TObjectPtr<UMovieGraphBurnInWidget> UMovieGraphBurnInNode::GetOrCreateBurnInWidget(UClass* InWidgetClass, UWorld* InOwner)
 {
 	if (const TObjectPtr<UMovieGraphBurnInWidget>* ExistingBurnInWidget = BurnInWidgetInstances.Find(InWidgetClass))
@@ -129,115 +63,59 @@ TObjectPtr<UMovieGraphBurnInWidget> UMovieGraphBurnInNode::GetOrCreateBurnInWidg
 
 void UMovieGraphBurnInNode::FMovieGraphBurnInPass::Setup(TWeakObjectPtr<UMovieGraphDefaultRenderer> InRenderer, const FMovieGraphRenderPassLayerData& InLayer)
 {
-	LayerData = InLayer;
-	Renderer = InRenderer;
-	RenderPassNode = CastChecked<UMovieGraphBurnInNode>(InLayer.RenderPassNode);
-
-	RenderDataIdentifier.RootBranchName = LayerData.BranchName;
-	RenderDataIdentifier.RendererName = RenderPassNode->GetRendererName();
+	FMovieGraphWidgetPass::Setup(InRenderer, InLayer);
+	
 	RenderDataIdentifier.SubResourceName = TEXT("burnin");
-
-	UE::MovieGraph::DefaultRenderer::FCameraInfo CameraInfo = Renderer->GetCameraInfo(LayerData.CameraIdentifier);
-	RenderDataIdentifier.CameraName = CameraInfo.CameraName;
 }
 
-void UMovieGraphBurnInNode::FMovieGraphBurnInPass::Teardown()
+TSharedPtr<SWidget> UMovieGraphBurnInNode::FMovieGraphBurnInPass::GetWidget()
 {
-	// Nothing to do
+	if (const TObjectPtr<UMovieGraphBurnInWidget> BurnInWidget = GetBurnInWidget())
+	{
+		return BurnInWidget->TakeWidget();
+	}
+	
+	return nullptr;
 }
 
-void UMovieGraphBurnInNode::FMovieGraphBurnInPass::Render(const FMovieGraphTraversalContext& InFrameTraversalContext, const FMovieGraphTimeStepData& InTimeData)
+TObjectPtr<UMovieGraphBurnInWidget> UMovieGraphBurnInNode::FMovieGraphBurnInPass::GetBurnInWidget() const
 {
-	UMovieGraphPipeline* Pipeline = Renderer->GetOwningGraph();
-	
-	const FIntPoint OutputResolution = UMovieGraphBlueprintLibrary::GetEffectiveOutputResolution(InTimeData.EvaluatedConfig, LayerData.BranchName);
-	const int32 MaxResolution = GetMax2DTextureDimension();
-	if ((OutputResolution.X > MaxResolution) || (OutputResolution.Y > MaxResolution))
-	{
-		UE_LOG(LogMovieRenderPipeline, Error, TEXT("Resolution %dx%d exceeds maximum allowed by GPU. Burn-ins do not support high-resolution tiling and thus can't exceed %dx%d."), OutputResolution.X, OutputResolution.Y, MaxResolution, MaxResolution);
-		return;
-	}
+	const UMovieGraphPipeline* Pipeline = Renderer->GetOwningGraph();
 
-	UClass* LoadedBurnInClass = RenderPassNode->BurnInClass.TryLoadClass<UMovieGraphBurnInWidget>();
+	const UMovieGraphBurnInNode* BurnInNode = CastChecked<UMovieGraphBurnInNode>(RenderPassNode);
+	UClass* LoadedBurnInClass = BurnInNode->BurnInClass.TryLoadClass<UMovieGraphBurnInWidget>();
 	if (!LoadedBurnInClass)
 	{
 		UE_LOG(LogMovieRenderPipeline, Error, TEXT("The burn-in widget provided in layer '%s' for renderer '%s' is not valid."), *LayerData.BranchName.ToString(), *Renderer->GetClass()->GetName());
-		return;
+		return nullptr;
 	}
-
-	// Create the render target the widget will be rendered into
-	UE::MovieGraph::DefaultRenderer::FRenderTargetInitParams RenderTargetInitParams;
-	RenderTargetInitParams.Size = OutputResolution;
-	RenderTargetInitParams.TargetGamma = FOpenColorIORendering::DefaultDisplayGamma;
-	RenderTargetInitParams.PixelFormat = PF_B8G8R8A8;
-	UTextureRenderTarget2D* RenderTarget = Renderer->GetOrCreateViewRenderTarget(RenderTargetInitParams);
 	
 	// The CDO contains the resources which are shared with all FMovieGraphBurnInPass instances
 	UMovieGraphBurnInNode* BurnInCDO = RenderPassNode->GetClass()->GetDefaultObject<UMovieGraphBurnInNode>();
-	const TSharedPtr<SVirtualWindow> VirtualWindow = BurnInCDO->GetOrCreateVirtualWindow(OutputResolution);
-	const TSharedPtr<FWidgetRenderer> BurnInWidgetRenderer = BurnInCDO->WidgetRenderer;
-	UMovieGraphBurnInWidget* BurnInWidget = BurnInCDO->GetOrCreateBurnInWidget(LoadedBurnInClass, Pipeline->GetWorld());
+	const TObjectPtr<UMovieGraphBurnInWidget> BurnInWidget = BurnInCDO->GetOrCreateBurnInWidget(LoadedBurnInClass, Pipeline->GetWorld());
 	if (!BurnInWidget)
 	{
-		UE_LOG(LogMovieRenderPipeline, Error, TEXT("Unable to load burn-in widget at path: %s"), *RenderPassNode->BurnInClass.GetAssetPath().ToString());
-		return;
+		UE_LOG(LogMovieRenderPipeline, Error, TEXT("Unable to load burn-in widget at path: %s"), *BurnInNode->BurnInClass.GetAssetPath().ToString());
+		return nullptr;
 	}
 
-	if (InFrameTraversalContext.Time.bIsFirstTemporalSampleForFrame)
-	{
-		// Put the widget in our window
-		VirtualWindow->SetContent(BurnInWidget->TakeWidget());
+	return BurnInWidget;
+}
 
-		// Update the widget with the latest frame information
+void UMovieGraphBurnInNode::FMovieGraphBurnInPass::Render(const FMovieGraphTraversalContext& InFrameTraversalContext, const FMovieGraphTimeStepData& InTimeData)
+{
+	// Update the widget with the latest frame information
+	if (const TObjectPtr<UMovieGraphBurnInWidget> BurnInWidget = GetBurnInWidget())
+	{
+		UMovieGraphPipeline* Pipeline = Renderer->GetOwningGraph();
 		BurnInWidget->UpdateForGraph(Pipeline);
-		
-		// Draw the widget to the render target
-		BurnInWidgetRenderer->DrawWindow(RenderTarget, VirtualWindow->GetHittestGrid(), VirtualWindow.ToSharedRef(),
-			1.f, OutputResolution, InTimeData.FrameDeltaTime);
-
-		FRenderTarget* BackbufferRenderTarget = RenderTarget->GameThread_GetRenderTargetResource();
-		TSharedPtr<UE::MovieGraph::IMovieGraphOutputMerger> OutputMerger = Pipeline->GetOutputMerger();
-
-		ENQUEUE_RENDER_COMMAND(BurnInRenderTargetResolveCommand)(
-			[InTimeData, InFrameTraversalContext, RenderTargetInitParams, RenderDataIdentifier = this->RenderDataIdentifier, bComposite = RenderPassNode->bCompositeOntoFinalImage, BackbufferRenderTarget, OutputMerger, OutputResolution](FRHICommandListImmediate& RHICmdList)
-			{
-				const FIntRect SourceRect = FIntRect(0, 0, BackbufferRenderTarget->GetSizeXY().X, BackbufferRenderTarget->GetSizeXY().Y);
-
-				// Read the data back to the CPU
-				TArray<FColor> RawPixels;
-				RawPixels.SetNumUninitialized(SourceRect.Width() * SourceRect.Height());
-
-				FReadSurfaceDataFlags ReadDataFlags(ERangeCompressionMode::RCM_MinMax);
-				ReadDataFlags.SetLinearToGamma(false);
-
-				{
-					// TODO: The readback is taking ~37ms on a 4k image. This is definitely an area that should be a target for optimization in the future.
-					TRACE_CPUPROFILER_EVENT_SCOPE(MRQ::FMovieGraphBurnInPass::ReadSurfaceData);
-					RHICmdList.ReadSurfaceData(BackbufferRenderTarget->GetRenderTargetTexture(), SourceRect, RawPixels, ReadDataFlags);
-				}
-
-				// Take our per-frame Traversal Context and update it with context specific to this sample.
-				FMovieGraphTraversalContext UpdatedTraversalContext = InFrameTraversalContext;
-				UpdatedTraversalContext.Time = InTimeData;
-				UpdatedTraversalContext.RenderDataIdentifier = RenderDataIdentifier;
-
-				TSharedRef<UE::MovieGraph::FMovieGraphSampleState> SampleStatePayload =
-					MakeShared<UE::MovieGraph::FMovieGraphSampleState>();
-				SampleStatePayload->TraversalContext = MoveTemp(UpdatedTraversalContext);
-				SampleStatePayload->BackbufferResolution = RenderTargetInitParams.Size;
-				SampleStatePayload->bRequiresAccumulator = false;
-				SampleStatePayload->bFetchFromAccumulator = false;
-				SampleStatePayload->bCompositeOnOtherRenders = bComposite;
-
-				TUniquePtr<FImagePixelData> PixelData =
-					MakeUnique<TImagePixelData<FColor>>(OutputResolution, TArray64<FColor>(MoveTemp(RawPixels)), SampleStatePayload);
-
-				OutputMerger->OnCompleteRenderPassDataAvailable_AnyThread(MoveTemp(PixelData));
-			});
 	}
+	
+	FMovieGraphWidgetPass::Render(InFrameTraversalContext, InTimeData);
 }
 
-void UMovieGraphBurnInNode::FMovieGraphBurnInPass::GatherOutputPassesImpl(TArray<FMovieGraphRenderDataIdentifier>& OutExpectedPasses) const
+int32 UMovieGraphBurnInNode::FMovieGraphBurnInPass::GetCompositingSortOrder() const
 {
-	OutExpectedPasses.Add(RenderDataIdentifier);
+	// Burn-ins should always appear over all other passes
+	return 0;
 }
\ No newline at end of file
diff --git a/Engine/Plugins/MovieScene/MovieRenderPipeline/Source/MovieRenderPipelineCore/Private/Graph/Nodes/MovieGraphImageSequenceOutputNode.cpp b/Engine/Plugins/MovieScene/MovieRenderPipeline/Source/MovieRenderPipelineCore/Private/Graph/Nodes/MovieGraphImageSequenceOutputNode.cpp
--- a/Engine/Plugins/MovieScene/MovieRenderPipeline/Source/MovieRenderPipelineCore/Private/Graph/Nodes/MovieGraphImageSequenceOutputNode.cpp
+++ b/Engine/Plugins/MovieScene/MovieRenderPipeline/Source/MovieRenderPipelineCore/Private/Graph/Nodes/MovieGraphImageSequenceOutputNode.cpp
@@ -35,7 +35,7 @@ void UMovieGraphImageSequenceOutputNode::OnReceiveImageDataImpl(UMovieGraphPipel
 	check(InRawFrameData);
 
 	// Gather the passes that need to be composited
-	TArray<TPair<FMovieGraphRenderDataIdentifier, TUniquePtr<FImagePixelData>>> CompositingPasses;
+	TArray<TPair<FMovieGraphRenderDataIdentifier, TUniquePtr<FImagePixelData>>> CompositedPasses;
 	for (TPair<FMovieGraphRenderDataIdentifier, TUniquePtr<FImagePixelData>>& RenderData : InRawFrameData->ImageOutputData)
 	{
 		UE::MovieGraph::FMovieGraphSampleState* Payload = RenderData.Value->GetPayload<UE::MovieGraph::FMovieGraphSampleState>();
@@ -48,9 +48,23 @@ void UMovieGraphImageSequenceOutputNode::OnReceiveImageDataImpl(UMovieGraphPipel
 		TPair<FMovieGraphRenderDataIdentifier, TUniquePtr<FImagePixelData>> CompositePass;
 		CompositePass.Key = RenderData.Key;
 		CompositePass.Value = RenderData.Value->CopyImageData();
-		CompositingPasses.Add(MoveTemp(CompositePass));
+		CompositedPasses.Add(MoveTemp(CompositePass));
 	}
 
+	// Sort composited passes if multiple were found. Passes with a smaller sort order go to the end of the array so they
+	// get composited on top of passes with a higher sort order.
+	CompositedPasses.Sort([](
+		const TPair<FMovieGraphRenderDataIdentifier, TUniquePtr<FImagePixelData>>& PassA,
+		const TPair<FMovieGraphRenderDataIdentifier, TUniquePtr<FImagePixelData>>& PassB)
+	{
+		const UE::MovieGraph::FMovieGraphSampleState* PayloadA = PassA.Value->GetPayload<UE::MovieGraph::FMovieGraphSampleState>();
+		const UE::MovieGraph::FMovieGraphSampleState* PayloadB = PassB.Value->GetPayload<UE::MovieGraph::FMovieGraphSampleState>();
+		check(PayloadA);
+		check(PayloadB);
+
+		return PayloadA->CompositingSortOrder > PayloadB->CompositingSortOrder;
+	});
+
 	// ToDo:
 	// The ImageWriteQueue is set up in a fire-and-forget manner. This means that the data needs to be placed in the WriteQueue
 	// as a TUniquePtr (so it can free the data when its done). Unfortunately we can have multiple output formats at once,
@@ -68,7 +82,7 @@ void UMovieGraphImageSequenceOutputNode::OnReceiveImageDataImpl(UMovieGraphPipel
 	{
 		// If this pass is composited, skip it for now
 		bool bSkip = false;
-		for (TPair<FMovieGraphRenderDataIdentifier, TUniquePtr<FImagePixelData>>& CompositedPass : CompositingPasses)
+		for (TPair<FMovieGraphRenderDataIdentifier, TUniquePtr<FImagePixelData>>& CompositedPass : CompositedPasses)
 		{
 			if (CompositedPass.Key == RenderData.Key)
 			{
@@ -166,10 +180,10 @@ void UMovieGraphImageSequenceOutputNode::OnReceiveImageDataImpl(UMovieGraphPipel
 
 		EImagePixelType PixelType = TileImageTask->PixelData->GetType();
 
-		// Perform compositing if any compositing passes were found earlier
-		for (TPair<FMovieGraphRenderDataIdentifier, TUniquePtr<FImagePixelData>>& CompositedPass : CompositingPasses)
+		// Perform compositing if any composited passes were found earlier
+		for (TPair<FMovieGraphRenderDataIdentifier, TUniquePtr<FImagePixelData>>& CompositedPass : CompositedPasses)
 		{
-			// This compositing pass will only composite on top of renders w/ the same branch and camera
+			// This composited pass will only composite on top of renders w/ the same branch and camera
 			const FMovieGraphRenderDataIdentifier& Id = CompositedPass.Key;
 			if ((Id.CameraName != RenderData.Key.CameraName) || (Id.RootBranchName != RenderData.Key.RootBranchName))
 			{
diff --git a/Engine/Plugins/MovieScene/MovieRenderPipeline/Source/MovieRenderPipelineCore/Private/Graph/Nodes/MovieGraphUIRendererNode.cpp b/Engine/Plugins/MovieScene/MovieRenderPipeline/Source/MovieRenderPipelineCore/Private/Graph/Nodes/MovieGraphUIRendererNode.cpp
new file mode 100644
--- /dev/null
+++ b/Engine/Plugins/MovieScene/MovieRenderPipeline/Source/MovieRenderPipelineCore/Private/Graph/Nodes/MovieGraphUIRendererNode.cpp
@@ -0,0 +1,60 @@
+// Copyright Epic Games, Inc. All Rights Reserved.
+
+#include "Graph/Nodes/MovieGraphUIRendererNode.h"
+
+#include "Engine/GameViewportClient.h"
+#include "Graph/MovieGraphDefaultRenderer.h"
+#include "Graph/MovieGraphPipeline.h"
+#include "Slate/SGameLayerManager.h"
+#include "TextureResource.h"
+#include "UnrealClient.h"
+
+const FString UMovieGraphUIRendererNode::RendererName = FString("UI Renderer");
+
+#if WITH_EDITOR
+FText UMovieGraphUIRendererNode::GetNodeTitle(const bool bGetDescriptive) const
+{
+	return NSLOCTEXT("MovieGraphNodes", "UIRendererGraphNode_Description", "UI Renderer");
+}
+
+FLinearColor UMovieGraphUIRendererNode::GetNodeTitleColor() const
+{
+	return FLinearColor(0.572f, 0.274f, 1.f);
+}
+
+FSlateIcon UMovieGraphUIRendererNode::GetIconAndTint(FLinearColor& OutColor) const
+{
+	static const FSlateIcon DeferredRendererIcon = FSlateIcon(FAppStyle::GetAppStyleSetName(), "SequenceRecorder.TabIcon");
+
+	OutColor = FLinearColor::White;
+	return DeferredRendererIcon;
+}
+#endif	// WITH_EDITOR
+
+void UMovieGraphUIRendererNode::FMovieGraphUIPass::Setup(TWeakObjectPtr<UMovieGraphDefaultRenderer> InRenderer, const FMovieGraphRenderPassLayerData& InLayer)
+{
+	FMovieGraphWidgetPass::Setup(InRenderer, InLayer);
+
+	RenderDataIdentifier.SubResourceName = TEXT("ui");
+}
+
+TSharedPtr<SWidget> UMovieGraphUIRendererNode::FMovieGraphUIPass::GetWidget()
+{
+	const UMovieGraphPipeline* Pipeline = Renderer->GetOwningGraph();
+	const ULocalPlayer* LocalPlayer = Pipeline->GetWorld()->GetFirstLocalPlayerFromController();
+	
+	// Cast the interface to a widget is a little yucky but the implementation is unlikely to change.
+	const TSharedPtr<SGameLayerManager> GameLayerManager = StaticCastSharedPtr<SGameLayerManager>(LocalPlayer->ViewportClient->GetGameLayerManager());
+	
+	return GameLayerManager;
+}
+
+int32 UMovieGraphUIRendererNode::FMovieGraphUIPass::GetCompositingSortOrder() const
+{
+	return 100;
+}
+
+TUniquePtr<UMovieGraphWidgetRendererBaseNode::FMovieGraphWidgetPass> UMovieGraphUIRendererNode::GeneratePass()
+{
+	return MakeUnique<FMovieGraphUIPass>();
+}
diff --git a/Engine/Plugins/MovieScene/MovieRenderPipeline/Source/MovieRenderPipelineCore/Private/Graph/Nodes/MovieGraphWidgetRendererBaseNode.cpp b/Engine/Plugins/MovieScene/MovieRenderPipeline/Source/MovieRenderPipelineCore/Private/Graph/Nodes/MovieGraphWidgetRendererBaseNode.cpp
new file mode 100644
--- /dev/null
+++ b/Engine/Plugins/MovieScene/MovieRenderPipeline/Source/MovieRenderPipelineCore/Private/Graph/Nodes/MovieGraphWidgetRendererBaseNode.cpp
@@ -0,0 +1,197 @@
+// Copyright Epic Games, Inc. All Rights Reserved.
+
+#include "Graph/Nodes/MovieGraphWidgetRendererBaseNode.h"
+
+#include "Engine/TextureRenderTarget2D.h"
+#include "Framework/Application/SlateApplication.h"
+#include "Graph/MovieGraphBlueprintLibrary.h"
+#include "Graph/MovieGraphDataTypes.h"
+#include "Graph/MovieGraphDefaultRenderer.h"
+#include "Graph/MovieGraphPipeline.h"
+#include "MovieRenderPipelineCoreModule.h"
+#include "OpenColorIORendering.h"
+#include "Slate/WidgetRenderer.h"
+#include "Styling/AppStyle.h"
+#include "TextureResource.h"
+#include "Widgets/SVirtualWindow.h"
+#include "RenderingThread.h"
+
+void UMovieGraphWidgetRendererBaseNode::SetupImpl(const FMovieGraphRenderPassSetupData& InSetupData)
+{
+	for (const FMovieGraphRenderPassLayerData& LayerData : InSetupData.Layers)
+	{
+		TUniquePtr<FMovieGraphWidgetPass> RendererInstance = GeneratePass();
+		RendererInstance->Setup(InSetupData.Renderer, LayerData);
+		CurrentInstances.Add(MoveTemp(RendererInstance));
+	}
+
+	constexpr bool bApplyGammaCorrection = false;
+	WidgetRenderer = MakeShared<FWidgetRenderer>(bApplyGammaCorrection);
+}
+
+void UMovieGraphWidgetRendererBaseNode::TeardownImpl()
+{
+	for (const TUniquePtr<FMovieGraphWidgetPass>& Instance : CurrentInstances)
+	{
+		Instance->Teardown();
+	}
+	
+	if (FSlateApplication::IsInitialized())
+	{
+		for (const TTuple<const FIntPoint, TSharedPtr<SVirtualWindow>>& Entry : SharedVirtualWindows)
+		{
+			const TSharedPtr<SVirtualWindow>& VirtualWindow = Entry.Value;
+			if (VirtualWindow.IsValid())
+			{
+				FSlateApplication::Get().UnregisterVirtualWindow(VirtualWindow.ToSharedRef());
+			}
+		}
+	}
+    
+	SharedVirtualWindows.Empty();
+	CurrentInstances.Reset();
+	WidgetRenderer = nullptr;
+}
+
+void UMovieGraphWidgetRendererBaseNode::RenderImpl(const FMovieGraphTraversalContext& InFrameTraversalContext, const FMovieGraphTimeStepData& InTimeData)
+{
+	for (const TUniquePtr<FMovieGraphWidgetPass>& Instance : CurrentInstances)
+	{
+		Instance->Render(InFrameTraversalContext, InTimeData);
+	}
+}
+
+void UMovieGraphWidgetRendererBaseNode::GatherOutputPassesImpl(TArray<FMovieGraphRenderDataIdentifier>& OutExpectedPasses) const
+{
+	for (const TUniquePtr<FMovieGraphWidgetPass>& Instance : CurrentInstances)
+    {
+    	Instance->GatherOutputPasses(OutExpectedPasses);
+    }
+}
+
+void UMovieGraphWidgetRendererBaseNode::FMovieGraphWidgetPass::Setup(TWeakObjectPtr<UMovieGraphDefaultRenderer> InRenderer, const FMovieGraphRenderPassLayerData& InLayer)
+{
+	LayerData = InLayer;
+	Renderer = InRenderer;
+	RenderPassNode = CastChecked<UMovieGraphWidgetRendererBaseNode>(InLayer.RenderPassNode);
+
+	RenderDataIdentifier.RootBranchName = LayerData.BranchName;
+	RenderDataIdentifier.RendererName = RenderPassNode->GetRendererName();
+
+	// NOTE: Subclasses should specify the SubResourceName
+
+	UE::MovieGraph::DefaultRenderer::FCameraInfo CameraInfo = Renderer->GetCameraInfo(LayerData.CameraIdentifier);
+	RenderDataIdentifier.CameraName = CameraInfo.CameraName;
+}
+
+void UMovieGraphWidgetRendererBaseNode::FMovieGraphWidgetPass::Teardown()
+{
+	// Nothing to do
+}
+
+void UMovieGraphWidgetRendererBaseNode::FMovieGraphWidgetPass::Render(const FMovieGraphTraversalContext& InFrameTraversalContext, const FMovieGraphTimeStepData& InTimeData)
+{
+	const UMovieGraphPipeline* Pipeline = Renderer->GetOwningGraph();
+	
+	const FIntPoint OutputResolution = UMovieGraphBlueprintLibrary::GetEffectiveOutputResolution(InTimeData.EvaluatedConfig, LayerData.BranchName);
+	const int32 MaxResolution = GetMax2DTextureDimension();
+	if ((OutputResolution.X > MaxResolution) || (OutputResolution.Y > MaxResolution))
+	{
+		UE_LOG(LogMovieRenderPipeline, Error, TEXT("Resolution %dx%d exceeds maximum allowed by GPU. Widget renders (burn-ins, etc) do not support high-resolution tiling and thus can't exceed %dx%d."), OutputResolution.X, OutputResolution.Y, MaxResolution, MaxResolution);
+		return;
+	}
+
+	// Create the render target the widget will be rendered into
+	UE::MovieGraph::DefaultRenderer::FRenderTargetInitParams RenderTargetInitParams;
+	RenderTargetInitParams.Size = OutputResolution;
+	RenderTargetInitParams.TargetGamma = FOpenColorIORendering::DefaultDisplayGamma;
+	RenderTargetInitParams.PixelFormat = PF_B8G8R8A8;
+	UTextureRenderTarget2D* RenderTarget = Renderer->GetOrCreateViewRenderTarget(RenderTargetInitParams);
+
+	if (InFrameTraversalContext.Time.bIsFirstTemporalSampleForFrame)
+	{
+		FRenderTarget* BackbufferRenderTarget = RenderTarget->GameThread_GetRenderTargetResource();
+		TSharedPtr<UE::MovieGraph::IMovieGraphOutputMerger> OutputMerger = Pipeline->GetOutputMerger();
+		
+		// The CDO contains the resources which are shared with all pass instances
+		UMovieGraphWidgetRendererBaseNode* NodeCDO = RenderPassNode->GetClass()->GetDefaultObject<UMovieGraphWidgetRendererBaseNode>();
+		const TSharedPtr<SVirtualWindow> VirtualWindow = NodeCDO->GetOrCreateVirtualWindow(OutputResolution);
+
+		const TSharedPtr<SWidget> WidgetToRender = GetWidget();
+		if (!WidgetToRender)
+		{
+			// The subclass implementation of GetWidget() is responsible for emitting error messages if getting the widget failed
+			return;
+		}
+		
+		// Put the widget in our window
+		VirtualWindow->SetContent(WidgetToRender.ToSharedRef());
+
+		// Draw the widget to the render target
+		NodeCDO->WidgetRenderer->DrawWindow(RenderTarget, VirtualWindow->GetHittestGrid(), VirtualWindow.ToSharedRef(),
+			1.f, OutputResolution, InTimeData.FrameDeltaTime);
+
+		ENQUEUE_RENDER_COMMAND(WidgetRenderTargetResolveCommand)(
+			[InTimeData, InFrameTraversalContext, RenderTargetInitParams, RenderDataIdentifier = this->RenderDataIdentifier,
+			bComposite = RenderPassNode->bCompositeOntoFinalImage, CompositingSortOrder = GetCompositingSortOrder(),
+			BackbufferRenderTarget, OutputMerger, OutputResolution](FRHICommandListImmediate& RHICmdList)
+			{
+				const FIntRect SourceRect = FIntRect(0, 0, BackbufferRenderTarget->GetSizeXY().X, BackbufferRenderTarget->GetSizeXY().Y);
+
+				// Read the data back to the CPU
+				TArray<FColor> RawPixels;
+				RawPixels.SetNumUninitialized(SourceRect.Width() * SourceRect.Height());
+
+				FReadSurfaceDataFlags ReadDataFlags(ERangeCompressionMode::RCM_MinMax);
+				ReadDataFlags.SetLinearToGamma(false);
+
+				{
+					// TODO: The readback is taking ~37ms on a 4k image. This is definitely an area that should be a target for optimization in the future.
+					TRACE_CPUPROFILER_EVENT_SCOPE(MRQ::FMovieGraphWidgetPass::ReadSurfaceData);
+					RHICmdList.ReadSurfaceData(BackbufferRenderTarget->GetRenderTargetTexture(), SourceRect, RawPixels, ReadDataFlags);
+				}
+
+				// Take our per-frame Traversal Context and update it with context specific to this sample.
+				FMovieGraphTraversalContext UpdatedTraversalContext = InFrameTraversalContext;
+				UpdatedTraversalContext.Time = InTimeData;
+				UpdatedTraversalContext.RenderDataIdentifier = RenderDataIdentifier;
+
+				TSharedRef<UE::MovieGraph::FMovieGraphSampleState> SampleStatePayload =
+					MakeShared<UE::MovieGraph::FMovieGraphSampleState>();
+				SampleStatePayload->TraversalContext = MoveTemp(UpdatedTraversalContext);
+				SampleStatePayload->BackbufferResolution = RenderTargetInitParams.Size;
+				SampleStatePayload->bRequiresAccumulator = false;
+				SampleStatePayload->bFetchFromAccumulator = false;
+				SampleStatePayload->bCompositeOnOtherRenders = bComposite;
+				SampleStatePayload->CompositingSortOrder = CompositingSortOrder;
+
+				TUniquePtr<FImagePixelData> PixelData =
+					MakeUnique<TImagePixelData<FColor>>(OutputResolution, TArray64<FColor>(MoveTemp(RawPixels)), SampleStatePayload);
+
+				OutputMerger->OnCompleteRenderPassDataAvailable_AnyThread(MoveTemp(PixelData));
+			});
+	}
+}
+
+void UMovieGraphWidgetRendererBaseNode::FMovieGraphWidgetPass::GatherOutputPasses(TArray<FMovieGraphRenderDataIdentifier>& OutExpectedPasses) const
+{
+	OutExpectedPasses.Add(RenderDataIdentifier);
+}
+
+TSharedPtr<SVirtualWindow> UMovieGraphWidgetRendererBaseNode::GetOrCreateVirtualWindow(const FIntPoint& InResolution)
+{
+	if (const TSharedPtr<SVirtualWindow>* ExistingVirtualWindow = SharedVirtualWindows.Find(InResolution))
+	{
+		return *ExistingVirtualWindow;
+	}
+
+	const TSharedPtr<SVirtualWindow> NewVirtualWindow = SNew(SVirtualWindow).Size(FVector2D(InResolution.X, InResolution.Y));
+	SharedVirtualWindows.Emplace(InResolution, NewVirtualWindow);
+
+	if (FSlateApplication::IsInitialized())
+	{
+		FSlateApplication::Get().RegisterVirtualWindow(NewVirtualWindow.ToSharedRef());
+	}
+
+	return NewVirtualWindow;
+}
diff --git a/Engine/Plugins/MovieScene/MovieRenderPipeline/Source/MovieRenderPipelineCore/Private/Graph/Nodes/MovieGraphWidgetRendererBaseNode.h b/Engine/Plugins/MovieScene/MovieRenderPipeline/Source/MovieRenderPipelineCore/Private/Graph/Nodes/MovieGraphWidgetRendererBaseNode.h
new file mode 100644
--- /dev/null
+++ b/Engine/Plugins/MovieScene/MovieRenderPipeline/Source/MovieRenderPipelineCore/Private/Graph/Nodes/MovieGraphWidgetRendererBaseNode.h
@@ -0,0 +1,83 @@
+﻿// Copyright Epic Games, Inc. All Rights Reserved.
+
+#pragma once
+
+#include "Graph/Nodes/MovieGraphRenderPassNode.h"
+
+#include "Styling/AppStyle.h"
+#include "MovieGraphWidgetRendererBaseNode.generated.h"
+
+// Forward Declares
+class FWidgetRenderer;
+class SVirtualWindow;
+class SWidget;
+class UMovieGraphDefaultRenderer;
+class UTextureRenderTarget2D;
+struct FMovieGraphRenderPassLayerData;
+struct FMovieGraphRenderPassSetupData;
+
+/** Base node containing common logic for nodes that render widgets. */
+UCLASS(Abstract)
+class MOVIERENDERPIPELINECORE_API UMovieGraphWidgetRendererBaseNode : public UMovieGraphRenderPassNode
+{
+	GENERATED_BODY()
+
+public:
+	UPROPERTY(EditAnywhere, BlueprintReadWrite, Category = Overrides, meta = (InlineEditConditionToggle))
+	uint8 bOverride_bCompositeOntoFinalImage : 1;
+
+	/** If true, the pass will be composited onto each render. Does not apply to multi-layer EXR files. */
+	UPROPERTY(EditAnywhere, BlueprintReadWrite, meta=(EditCondition="bOverride_bCompositeOntoFinalImage"), Category="Widget Settings")
+	bool bCompositeOntoFinalImage;
+
+protected:
+	/** A render pass for a specific render layer, generated by the UMovieGraphWidgetRendererBaseNode CDO. */
+	struct FMovieGraphWidgetPass
+	{
+		virtual ~FMovieGraphWidgetPass() = default;
+
+		virtual void Setup(TWeakObjectPtr<UMovieGraphDefaultRenderer> InRenderer, const FMovieGraphRenderPassLayerData& InLayer);
+		virtual void Teardown();
+		virtual TSharedPtr<SWidget> GetWidget() = 0;
+		virtual void Render(const FMovieGraphTraversalContext& InFrameTraversalContext, const FMovieGraphTimeStepData& InTimeData);
+		virtual void GatherOutputPasses(TArray<FMovieGraphRenderDataIdentifier>& OutExpectedPasses) const;
+		virtual int32 GetCompositingSortOrder() const = 0;
+
+	protected:
+		FMovieGraphRenderPassLayerData LayerData;
+
+		/** Unique identifier passed in GatherOutputPasses and with each render that identifies the data produced by this renderer. */
+		FMovieGraphRenderDataIdentifier RenderDataIdentifier;
+
+		/** The renderer associated with this pass. */
+		TWeakObjectPtr<UMovieGraphDefaultRenderer> Renderer;
+
+		/** The node instance (not CDO) associated with this pass. */
+		TWeakObjectPtr<UMovieGraphWidgetRendererBaseNode> RenderPassNode;
+	};
+
+	/** Widget renderer shared with all FMovieGraphWidgetPass instances. */
+	TSharedPtr<FWidgetRenderer> WidgetRenderer;
+
+	/** Render passes generated by GeneratePass() calls on the CDO. */
+	TArray<TUniquePtr<FMovieGraphWidgetPass>> CurrentInstances;
+
+protected:
+	/** Generate a new pass that the widget renderer will use. */
+	virtual TUniquePtr<FMovieGraphWidgetPass> GeneratePass() PURE_VIRTUAL(UMovieGraphWidgetRendererBaseNode::GeneratePass, return nullptr;);
+	
+	// UMovieGraphRenderPassNode Interface
+	virtual void SetupImpl(const FMovieGraphRenderPassSetupData& InSetupData) override;
+	virtual void TeardownImpl() override;
+	virtual void RenderImpl(const FMovieGraphTraversalContext& InFrameTraversalContext, const FMovieGraphTimeStepData& InTimeData) override;
+	virtual void GatherOutputPassesImpl(TArray<FMovieGraphRenderDataIdentifier>& OutExpectedPasses) const override;
+	// ~UMovieGraphRenderPassNode Interface
+
+private:
+	/** Virtual window instances shared with all FMovieGraphWidgetPass instances, keyed by resolution. */
+	TMap<const FIntPoint, TSharedPtr<SVirtualWindow>> SharedVirtualWindows;
+
+private:
+	/** Gets an existing virtual window with the specified resolution, or creates a new one if one does not exist yet. */
+	TSharedPtr<SVirtualWindow> GetOrCreateVirtualWindow(const FIntPoint& InResolution);
+};
\ No newline at end of file
diff --git a/Engine/Plugins/MovieScene/MovieRenderPipeline/Source/MovieRenderPipelineCore/Public/Graph/MovieGraphDataTypes.h b/Engine/Plugins/MovieScene/MovieRenderPipeline/Source/MovieRenderPipelineCore/Public/Graph/MovieGraphDataTypes.h
--- a/Engine/Plugins/MovieScene/MovieRenderPipeline/Source/MovieRenderPipelineCore/Public/Graph/MovieGraphDataTypes.h
+++ b/Engine/Plugins/MovieScene/MovieRenderPipeline/Source/MovieRenderPipelineCore/Public/Graph/MovieGraphDataTypes.h
@@ -221,6 +221,7 @@ namespace UE::MovieGraph
 			, bRequiresAccumulator(false)
 			, bFetchFromAccumulator(false)
 			, bCompositeOnOtherRenders(false)
+			, CompositingSortOrder(0)
 		{}
 
 		/** The traversal context used to read graph values at the time of submission. */
@@ -246,6 +247,12 @@ namespace UE::MovieGraph
 
 		/** Set this to true if this pass should be composited on top of other renders. */
 		bool bCompositeOnOtherRenders;
+
+		/**
+		* If multiple passes are composited on top of a render, the sort order determines the order in which they're composited.
+		* Passes with a low sort order will composite on top of passes with a higher sort order.
+		*/
+		int32 CompositingSortOrder;
 	};
 
 	/**
diff --git a/Engine/Plugins/MovieScene/MovieRenderPipeline/Source/MovieRenderPipelineCore/Public/Graph/Nodes/MovieGraphBurnInNode.h b/Engine/Plugins/MovieScene/MovieRenderPipeline/Source/MovieRenderPipelineCore/Public/Graph/Nodes/MovieGraphBurnInNode.h
--- a/Engine/Plugins/MovieScene/MovieRenderPipeline/Source/MovieRenderPipelineCore/Public/Graph/Nodes/MovieGraphBurnInNode.h
+++ b/Engine/Plugins/MovieScene/MovieRenderPipeline/Source/MovieRenderPipelineCore/Public/Graph/Nodes/MovieGraphBurnInNode.h
@@ -2,22 +2,19 @@
 
 #pragma once
 
-#include "Graph/Nodes/MovieGraphRenderPassNode.h"
+#include "Graph/Nodes/MovieGraphWidgetRendererBaseNode.h"
 
 #include "MovieGraphBurnInNode.generated.h"
 
 // Forward Declares
-class FWidgetRenderer;
 class SVirtualWindow;
 class UMovieGraphBurnInWidget;
 class UMovieGraphDefaultRenderer;
-class UMoviePipelineBurnInWidget;
 struct FMovieGraphRenderPassLayerData;
-struct FMovieGraphRenderPassSetupData;
 
 /** A node which generates a widget burn-in, rendered to a standalone image or composited on top of a render layer. */
 UCLASS()
-class MOVIERENDERPIPELINECORE_API UMovieGraphBurnInNode : public UMovieGraphRenderPassNode
+class MOVIERENDERPIPELINECORE_API UMovieGraphBurnInNode : public UMovieGraphWidgetRendererBaseNode
 {
 	GENERATED_BODY()
 
@@ -28,9 +25,6 @@ public:
 	virtual FSlateIcon GetIconAndTint(FLinearColor& OutColor) const override;
 #endif
 
-	/** Gets an existing virtual window with the specified resolution, or creates a new one if one does not exist yet. */
-	TSharedPtr<SVirtualWindow> GetOrCreateVirtualWindow(const FIntPoint& InResolution);
-
 	/**
 	 * Gets an existing widget instance of type WidgetClass if one has been created. Otherwise, returns a new instance
 	 * of WidgetClass with owner InOwner.
@@ -43,59 +37,35 @@ public:
 	UPROPERTY(EditAnywhere, BlueprintReadWrite, Category = Overrides, meta = (InlineEditConditionToggle))
 	uint8 bOverride_BurnInClass : 1;
 
-	UPROPERTY(EditAnywhere, BlueprintReadWrite, Category = Overrides, meta = (InlineEditConditionToggle))
-	uint8 bOverride_bCompositeOntoFinalImage : 1;
-
 	/** The widget that the burn-in should use. */
 	UPROPERTY(EditAnywhere, BlueprintReadWrite, meta=(MetaClass="/Script/MovieRenderPipelineCore.MovieGraphBurnInWidget", EditCondition="bOverride_BurnInClass"), Category="Widget Settings")
 	FSoftClassPath BurnInClass;
 
-	/** If true, the burn-in image will be composited into the final image pass. Does not apply to multi-layer EXR files. */
-	UPROPERTY(EditAnywhere, BlueprintReadWrite, meta=(EditCondition="bOverride_bCompositeOntoFinalImage"), Category="Widget Settings")
-	bool bCompositeOntoFinalImage;
-
 protected:
-	/** A render pass for a specific render layer, generated by the UMovieGraphBurnInNode CDO. */
-	struct FMovieGraphBurnInPass
+	/** A burn-in pass for a specific render layer. Instances are stored on the UMovieGraphWidgetRendererBaseNode CDO. */
+	struct FMovieGraphBurnInPass final : public FMovieGraphWidgetPass
 	{
-		void Setup(TWeakObjectPtr<UMovieGraphDefaultRenderer> InRenderer, const FMovieGraphRenderPassLayerData& InLayer);
-		void Teardown();
-		void Render(const FMovieGraphTraversalContext& InFrameTraversalContext, const FMovieGraphTimeStepData& InTimeData);
-		void GatherOutputPassesImpl(TArray<FMovieGraphRenderDataIdentifier>& OutExpectedPasses) const;
-
-	protected:
-		FMovieGraphRenderPassLayerData LayerData;
-
-		/** Unique identifier passed in GatherOutputPasses and with each render that identifies the data produced by this renderer. */
-		FMovieGraphRenderDataIdentifier RenderDataIdentifier;
+		virtual void Setup(TWeakObjectPtr<UMovieGraphDefaultRenderer> InRenderer, const FMovieGraphRenderPassLayerData& InLayer) override;
+		virtual TSharedPtr<SWidget> GetWidget() override;
+		virtual void Render(const FMovieGraphTraversalContext& InFrameTraversalContext, const FMovieGraphTimeStepData& InTimeData) override;
+		virtual int32 GetCompositingSortOrder() const override;
 
-		/** The renderer associated with this pass. */
-		TWeakObjectPtr<UMovieGraphDefaultRenderer> Renderer;
-
-		/** The node instance (not CDO) associated with this pass. */
-		TWeakObjectPtr<UMovieGraphBurnInNode> RenderPassNode;
+	private:
+		TObjectPtr<UMovieGraphBurnInWidget> GetBurnInWidget() const;
 	};
 
 private:
-	/** Widget renderer shared with all FMovieGraphBurnInPass instances. */
-	TSharedPtr<FWidgetRenderer> WidgetRenderer;
-
-	/** Virtual window instances shared with all FMovieGraphBurnInPass instances, keyed by resolution. */
-	TMap<const FIntPoint, TSharedPtr<SVirtualWindow>> SharedVirtualWindows;
-
 	/** Burn-in widget instances shared with all FMovieGraphBurnInPass instances, keyed by burn-in class. */
 	UPROPERTY(Transient)
 	TMap<const UClass*, TObjectPtr<UMovieGraphBurnInWidget>> BurnInWidgetInstances;
 
-	/** Render passes generated by the UMovieGraphBurnInNode CDO. */
-	TArray<TUniquePtr<FMovieGraphBurnInPass>> CurrentInstances;
-
 protected:
+	// UMovieGraphWidgetRendererBaseNode Interface
+	virtual TUniquePtr<FMovieGraphWidgetPass> GeneratePass() override;
+	// ~UMovieGraphWidgetRendererBaseNode Interface
+	
 	// UMovieGraphRenderPassNode Interface
 	virtual FString GetRendererNameImpl() const override { return RendererName; }
-	virtual void SetupImpl(const FMovieGraphRenderPassSetupData& InSetupData) override;
 	virtual void TeardownImpl() override;
-	virtual void RenderImpl(const FMovieGraphTraversalContext& InFrameTraversalContext, const FMovieGraphTimeStepData& InTimeData) override;
-	virtual void GatherOutputPassesImpl(TArray<FMovieGraphRenderDataIdentifier>& OutExpectedPasses) const override;
 	// ~UMovieGraphRenderPassNode Interface
 };
\ No newline at end of file
diff --git a/Engine/Plugins/MovieScene/MovieRenderPipeline/Source/MovieRenderPipelineCore/Public/Graph/Nodes/MovieGraphUIRendererNode.h b/Engine/Plugins/MovieScene/MovieRenderPipeline/Source/MovieRenderPipelineCore/Public/Graph/Nodes/MovieGraphUIRendererNode.h
new file mode 100644
--- /dev/null
+++ b/Engine/Plugins/MovieScene/MovieRenderPipeline/Source/MovieRenderPipelineCore/Public/Graph/Nodes/MovieGraphUIRendererNode.h
@@ -0,0 +1,47 @@
+// Copyright Epic Games, Inc. All Rights Reserved.
+
+#pragma once
+
+#include "Graph/Nodes/MovieGraphWidgetRendererBaseNode.h"
+#include "Styling/AppStyle.h"
+
+#include "MovieGraphUIRendererNode.generated.h"
+
+// Forward Declares
+class UMovieGraphDefaultRenderer;
+class UWidget;
+struct FMovieGraphRenderPassLayerData;
+
+/** A node which renders the viewport's UMG widget to a standalone image, or composited on top of a render layer. */
+UCLASS()
+class MOVIERENDERPIPELINECORE_API UMovieGraphUIRendererNode : public UMovieGraphWidgetRendererBaseNode
+{
+	GENERATED_BODY()
+
+public:
+#if WITH_EDITOR
+	virtual FText GetNodeTitle(const bool bGetDescriptive = false) const override;
+	virtual FLinearColor GetNodeTitleColor() const override;
+	virtual FSlateIcon GetIconAndTint(FLinearColor& OutColor) const override;
+#endif
+
+public:
+	static const FString RendererName;
+
+protected:
+	/** A UI pass for a specific render layer. Instances are stored on the UMovieGraphWidgetRendererBaseNode CDO. */
+	struct FMovieGraphUIPass final : public FMovieGraphWidgetPass
+	{
+		virtual void Setup(TWeakObjectPtr<UMovieGraphDefaultRenderer> InRenderer, const FMovieGraphRenderPassLayerData& InLayer) override;
+		virtual TSharedPtr<SWidget> GetWidget() override;
+		virtual int32 GetCompositingSortOrder() const override;
+	};
+	
+	// UMovieGraphWidgetRendererBaseNode Interface
+	virtual TUniquePtr<FMovieGraphWidgetPass> GeneratePass() override;
+	// ~UMovieGraphWidgetRendererBaseNode Interface
+	
+	// UMovieGraphRenderPassNode Interface
+	virtual FString GetRendererNameImpl() const override { return RendererName; }
+	// ~UMovieGraphRenderPassNode Interface
+};
\ No newline at end of file
