[MRQ] Status bar completes (on success or failure) in UI

diff --git a/Engine/Plugins/MovieScene/MovieRenderPipeline/Source/MovieRenderPipelineEditor/Private/MoviePipelinePIEExecutor.cpp b/Engine/Plugins/MovieScene/MovieRenderPipeline/Source/MovieRenderPipelineEditor/Private/MoviePipelinePIEExecutor.cpp
--- a/Engine/Plugins/MovieScene/MovieRenderPipeline/Source/MovieRenderPipelineEditor/Private/MoviePipelinePIEExecutor.cpp
+++ b/Engine/Plugins/MovieScene/MovieRenderPipeline/Source/MovieRenderPipelineEditor/Private/MoviePipelinePIEExecutor.cpp
@@ -326,6 +326,9 @@ void UMoviePipelinePIEExecutor::OnPIEMoviePipelineFinished(FMoviePipelineOutputD
 	// Unsubscribe to the EndPIE event so we don't think the user canceled it.
 	FCoreDelegates::OnBeginFrame.RemoveAll(this);
 
+	// Reset progress for UI (regardless of success)
+	Queue->GetJobs()[CurrentPipelineIndex]->SetStatusProgress(0.0f);
+
 	if (ActiveMoviePipeline)
 	{
 		// Unsubscribe in the event that it gets called twice we don't have issues.
