MoviePipeline: Fixed an issue where canceling a job would report as a successful job for some callbacks.

Fixed OnExecutorErroredImpl invoking OnExecutorFinishedDelegate/OnExecutorFinishedDelegateNative with an inverted success param (when canceling an active job).

Fixed OnIndividualJobFinishedImpl invoking the deprecated OnIndividualJobFinishedDelegate/OnIndividualJobFinishedDelegateNative delegates an inverted success state.

diff --git a/Engine/Plugins/MovieScene/MovieRenderPipeline/Source/MovieRenderPipelineCore/Public/MoviePipelineExecutor.h b/Engine/Plugins/MovieScene/MovieRenderPipeline/Source/MovieRenderPipelineCore/Public/MoviePipelineExecutor.h
--- a/Engine/Plugins/MovieScene/MovieRenderPipeline/Source/MovieRenderPipelineCore/Public/MoviePipelineExecutor.h
+++ b/Engine/Plugins/MovieScene/MovieRenderPipeline/Source/MovieRenderPipelineCore/Public/MoviePipelineExecutor.h
@@ -229,8 +229,8 @@ protected:
 		}
 
 		// Broadcast to both Native and Python/BP
-		OnExecutorFinishedDelegateNative.Broadcast(this, bFatal);
-		OnExecutorFinishedDelegate.Broadcast(this, bFatal);
+		OnExecutorFinishedDelegateNative.Broadcast(this, !bAnyJobHadFatalError);
+		OnExecutorFinishedDelegate.Broadcast(this, !bAnyJobHadFatalError);
 	}
 
 	bool IsAnyJobErrored() const { return bAnyJobHadFatalError; }
diff --git a/Engine/Plugins/MovieScene/MovieRenderPipeline/Source/MovieRenderPipelineEditor/Private/MoviePipelinePIEExecutor.cpp b/Engine/Plugins/MovieScene/MovieRenderPipeline/Source/MovieRenderPipelineEditor/Private/MoviePipelinePIEExecutor.cpp
--- a/Engine/Plugins/MovieScene/MovieRenderPipeline/Source/MovieRenderPipelineEditor/Private/MoviePipelinePIEExecutor.cpp
+++ b/Engine/Plugins/MovieScene/MovieRenderPipeline/Source/MovieRenderPipelineEditor/Private/MoviePipelinePIEExecutor.cpp
@@ -379,7 +379,6 @@ void UMoviePipelinePIEExecutor::OnPIEEnded(bool)
 	// We need to null out this reference this frame, otherwise we try to hold onto a PIE
 	// object after PIE finishes which causes a GC leak.
 	ActiveMoviePipeline = nullptr;
-	// ToDo: bAnyJobHadFatalError
 
 	// Delay for one frame so that PIE can finish shut down. It's not a huge fan of us starting up on the same frame.
 	GEditor->GetTimerManager()->SetTimerForNextTick(FTimerDelegate::CreateUObject(this, &UMoviePipelinePIEExecutor::DelayedFinishNotification));
@@ -412,8 +411,8 @@ void UMoviePipelinePIEExecutor::OnIndividualJobFinishedImpl(FMoviePipelineOutput
 {
 	// Broadcast to both Native and Python/BP
 	PRAGMA_DISABLE_DEPRECATION_WARNINGS
-	OnIndividualJobFinishedDelegateNative.Broadcast(InOutputData.Job, IsAnyJobErrored());
-	OnIndividualJobFinishedDelegate.Broadcast(InOutputData.Job, IsAnyJobErrored());
+	OnIndividualJobFinishedDelegateNative.Broadcast(InOutputData.Job, !IsAnyJobErrored());
+	OnIndividualJobFinishedDelegate.Broadcast(InOutputData.Job, !IsAnyJobErrored());
 	PRAGMA_ENABLE_DEPRECATION_WARNINGS
 
 	OnIndividualJobWorkFinishedDelegateNative.Broadcast(InOutputData);
