MoviePipeline: Added a new EXR output node to the graph. There are a few things that need to be tweaked in the future in order for it to be fully functional and match existing functionality, however (needs OCIO support when an OCIO node exists, overscan support needs to be added when that's available in the graph, in addition to multi-camera layer naming).

To re-use as much existing EXR logic as possible, the new node was created in the "MovieRenderPipelineRenderPasses" module. Existing output nodes were also moved here for consistency.

There are additional features that would be nice-to-have on the node (multi-part, compositing with metadata), but this at least has parity with the old system (with the exceptions noted before).

#rnx

diff --git a/Engine/Plugins/MovieScene/MovieRenderPipeline/Source/MovieRenderPipelineCore/Private/Graph/MovieGraphBlueprintLibrary.cpp b/Engine/Plugins/MovieScene/MovieRenderPipeline/Source/MovieRenderPipelineCore/Private/Graph/MovieGraphBlueprintLibrary.cpp
--- a/Engine/Plugins/MovieScene/MovieRenderPipeline/Source/MovieRenderPipelineCore/Private/Graph/MovieGraphBlueprintLibrary.cpp
+++ b/Engine/Plugins/MovieScene/MovieRenderPipeline/Source/MovieRenderPipelineCore/Private/Graph/MovieGraphBlueprintLibrary.cpp
@@ -134,7 +134,10 @@ FString UMovieGraphBlueprintLibrary::ResolveFilenameFormatArguments(const FStrin
 	}
 
 	// Copy the metadata in the parameters into the merged output version
-	OutMergedFormatArgs.FileMetadata = InParams.FileMetadata;
+	for (const TPair<FString, FString>& FileMetadata : InParams.FileMetadata)
+	{
+		OutMergedFormatArgs.FileMetadata.Add(FileMetadata.Key, FileMetadata.Value);
+	}
 
 	// We expect the incoming string to have a {file_dup} token where they want file duplication numbers to be handled
 	// but we don't actually want a value there (unless there's a collision). So we override it with an empty string by default.
diff --git a/Engine/Plugins/MovieScene/MovieRenderPipeline/Source/MovieRenderPipelineCore/Private/Graph/Nodes/MovieGraphImageSequenceOutputNode.cpp b/Engine/Plugins/MovieScene/MovieRenderPipeline/Source/MovieRenderPipelineCore/Private/Graph/Nodes/MovieGraphImageSequenceOutputNode.cpp
deleted file mode 100644
diff --git a/Engine/Plugins/MovieScene/MovieRenderPipeline/Source/MovieRenderPipelineCore/Public/Graph/MovieGraphPlaceholderNodes.cpp b/Engine/Plugins/MovieScene/MovieRenderPipeline/Source/MovieRenderPipelineCore/Public/Graph/MovieGraphPlaceholderNodes.cpp
--- a/Engine/Plugins/MovieScene/MovieRenderPipeline/Source/MovieRenderPipelineCore/Public/Graph/MovieGraphPlaceholderNodes.cpp
+++ b/Engine/Plugins/MovieScene/MovieRenderPipeline/Source/MovieRenderPipelineCore/Public/Graph/MovieGraphPlaceholderNodes.cpp
@@ -8,7 +8,6 @@
 #define LOCTEXT_NAMESPACE "MovieGraphNode"
 
 static const FText NodeCategory_Rendering = LOCTEXT("NodeCategory_Rendering", "Rendering");
-static const FText NodeCategory_OutputType = LOCTEXT("NodeCategory_OutputType", "Output Type");
 static const FText NodeCategory_Settings = LOCTEXT("NodeCategory_Settings", "Settings");
 
 #if WITH_EDITOR
@@ -38,33 +37,6 @@ FSlateIcon UMovieGraphPathTracedRendererNode::GetIconAndTint(FLinearColor& OutCo
 }
 #endif // WITH_EDITOR
 
-#if WITH_EDITOR
-FText UMovieGraphEXRSequenceNode::GetNodeTitle(const bool bGetDescriptive) const
-{
-	static const FText EXRSequenceNodeName = LOCTEXT("NodeName_EXRSequence", ".exr Sequence");
-	return EXRSequenceNodeName;
-}
-
-FText UMovieGraphEXRSequenceNode::GetMenuCategory() const
-{
-	return NodeCategory_OutputType;
-}
-
-FLinearColor UMovieGraphEXRSequenceNode::GetNodeTitleColor() const
-{
-	static const FLinearColor ImageSequenceNodeColor = FLinearColor(0.047f, 0.654f, 0.537f);
-	return ImageSequenceNodeColor;
-}
-
-FSlateIcon UMovieGraphEXRSequenceNode::GetIconAndTint(FLinearColor& OutColor) const
-{
-	static const FSlateIcon ImageSequenceIcon = FSlateIcon(FAppStyle::GetAppStyleSetName(), "ClassIcon.Texture2D");
-
-	OutColor = FLinearColor::White;
-	return ImageSequenceIcon;
-}
-#endif // WITH_EDITOR
-
 #if WITH_EDITOR
 FText UMovieGraphAntiAliasingNode::GetNodeTitle(const bool bGetDescriptive) const
 {
diff --git a/Engine/Plugins/MovieScene/MovieRenderPipeline/Source/MovieRenderPipelineCore/Public/Graph/MovieGraphPlaceholderNodes.h b/Engine/Plugins/MovieScene/MovieRenderPipeline/Source/MovieRenderPipelineCore/Public/Graph/MovieGraphPlaceholderNodes.h
--- a/Engine/Plugins/MovieScene/MovieRenderPipeline/Source/MovieRenderPipelineCore/Public/Graph/MovieGraphPlaceholderNodes.h
+++ b/Engine/Plugins/MovieScene/MovieRenderPipeline/Source/MovieRenderPipelineCore/Public/Graph/MovieGraphPlaceholderNodes.h
@@ -23,23 +23,6 @@ public:
 #endif
 };
 
-/** A node which generates an EXR image sequence. */
-UCLASS()
-class MOVIERENDERPIPELINECORE_API UMovieGraphEXRSequenceNode : public UMovieGraphSettingNode
-{
-	GENERATED_BODY()
-
-public:
-	UMovieGraphEXRSequenceNode() = default;
-
-#if WITH_EDITOR
-	virtual FText GetNodeTitle(const bool bGetDescriptive = false) const override;
-	virtual FText GetMenuCategory() const override;
-	virtual FLinearColor GetNodeTitleColor() const override;
-	virtual FSlateIcon GetIconAndTint(FLinearColor& OutColor) const override;
-#endif
-};
-
 /** A node which configures anti-aliasing settings. */
 UCLASS()
 class MOVIERENDERPIPELINECORE_API UMovieGraphAntiAliasingNode : public UMovieGraphSettingNode
diff --git a/Engine/Plugins/MovieScene/MovieRenderPipeline/Source/MovieRenderPipelineCore/Public/Graph/Nodes/MovieGraphFileOutputNode.h b/Engine/Plugins/MovieScene/MovieRenderPipeline/Source/MovieRenderPipelineCore/Public/Graph/Nodes/MovieGraphFileOutputNode.h
--- a/Engine/Plugins/MovieScene/MovieRenderPipeline/Source/MovieRenderPipelineCore/Public/Graph/Nodes/MovieGraphFileOutputNode.h
+++ b/Engine/Plugins/MovieScene/MovieRenderPipeline/Source/MovieRenderPipelineCore/Public/Graph/Nodes/MovieGraphFileOutputNode.h
@@ -14,7 +14,7 @@
 * multiple output formats at once (such as exr and jpeg).
 */
 UCLASS(Abstract)
-class UMovieGraphFileOutputNode : public UMovieGraphSettingNode
+class MOVIERENDERPIPELINECORE_API UMovieGraphFileOutputNode : public UMovieGraphSettingNode
 {
 	GENERATED_BODY()
 public:
diff --git a/Engine/Plugins/MovieScene/MovieRenderPipeline/Source/MovieRenderPipelineCore/Public/Graph/Nodes/MovieGraphImageSequenceOutputNode.h b/Engine/Plugins/MovieScene/MovieRenderPipeline/Source/MovieRenderPipelineCore/Public/Graph/Nodes/MovieGraphImageSequenceOutputNode.h
deleted file mode 100644
diff --git a/Engine/Plugins/MovieScene/MovieRenderPipeline/Source/MovieRenderPipelineRenderPasses/MovieRenderPipelineRenderPasses.Build.cs b/Engine/Plugins/MovieScene/MovieRenderPipeline/Source/MovieRenderPipelineRenderPasses/MovieRenderPipelineRenderPasses.Build.cs
--- a/Engine/Plugins/MovieScene/MovieRenderPipeline/Source/MovieRenderPipelineRenderPasses/MovieRenderPipelineRenderPasses.Build.cs
+++ b/Engine/Plugins/MovieScene/MovieRenderPipeline/Source/MovieRenderPipelineRenderPasses/MovieRenderPipelineRenderPasses.Build.cs
@@ -22,7 +22,8 @@ public class MovieRenderPipelineRenderPasses : ModuleRules
 				"ImageWrapper",				
 				"CinematicCamera", // For metadata
 				"MovieRenderPipelineSettings", // For settings
-				"ColorManagement"
+				"ColorManagement",
+				"SlateCore",
 			}
 		);
 
diff --git a/Engine/Plugins/MovieScene/MovieRenderPipeline/Source/MovieRenderPipelineRenderPasses/Private/MovieGraphImageSequenceOutputNode.cpp b/Engine/Plugins/MovieScene/MovieRenderPipeline/Source/MovieRenderPipelineRenderPasses/Private/MovieGraphImageSequenceOutputNode.cpp
new file mode 100644
--- /dev/null
+++ b/Engine/Plugins/MovieScene/MovieRenderPipeline/Source/MovieRenderPipelineRenderPasses/Private/MovieGraphImageSequenceOutputNode.cpp
@@ -0,0 +1,485 @@
+// Copyright Epic Games, Inc. All Rights Reserved.
+
+#include "MovieGraphImageSequenceOutputNode.h"
+
+#include "Graph/Nodes/MovieGraphOutputSettingNode.h"
+#include "Graph/Nodes/MovieGraphRenderLayerNode.h"
+#include "Graph/MovieGraphDataTypes.h"
+#include "Graph/MovieGraphPipeline.h"
+#include "Graph/MovieGraphConfig.h"
+#include "Graph/MovieGraphFilenameResolveParams.h"
+#include "Graph/MovieGraphBlueprintLibrary.h"
+#include "Modules/ModuleManager.h"
+#include "MoviePipelineUtils.h"
+#include "ImageWriteQueue.h"
+#include "Misc/Paths.h"
+#include "Async/TaskGraphInterfaces.h"
+
+UMovieGraphImageSequenceOutputNode::UMovieGraphImageSequenceOutputNode()
+{
+	ImageWriteQueue = &FModuleManager::Get().LoadModuleChecked<IImageWriteQueueModule>("ImageWriteQueue").GetWriteQueue();
+}
+
+void UMovieGraphImageSequenceOutputNode::OnAllFramesSubmittedImpl()
+{
+	FinalizeFence = ImageWriteQueue->CreateFence();
+}
+
+bool UMovieGraphImageSequenceOutputNode::IsFinishedWritingToDiskImpl() const
+{
+	// Wait until the finalization fence is reached meaning we've written everything to disk.
+	return Super::IsFinishedWritingToDiskImpl() && (!FinalizeFence.IsValid() || FinalizeFence.WaitFor(0));
+}
+
+void UMovieGraphImageSequenceOutputNode::OnReceiveImageDataImpl(UMovieGraphPipeline* InPipeline, UE::MovieGraph::FMovieGraphOutputMergerFrame* InRawFrameData, const TSet<FMovieGraphRenderDataIdentifier>& InMask)
+{
+	check(InRawFrameData);
+
+	// Gather the passes that need to be composited
+	TArray<TPair<FMovieGraphRenderDataIdentifier, TUniquePtr<FImagePixelData>>> CompositedPasses;
+	for (TPair<FMovieGraphRenderDataIdentifier, TUniquePtr<FImagePixelData>>& RenderData : InRawFrameData->ImageOutputData)
+	{
+		UE::MovieGraph::FMovieGraphSampleState* Payload = RenderData.Value->GetPayload<UE::MovieGraph::FMovieGraphSampleState>();
+		check(Payload);
+		if (!Payload->bCompositeOnOtherRenders)
+		{
+			continue;
+		}
+
+		TPair<FMovieGraphRenderDataIdentifier, TUniquePtr<FImagePixelData>> CompositePass;
+		CompositePass.Key = RenderData.Key;
+		CompositePass.Value = RenderData.Value->CopyImageData();
+		CompositedPasses.Add(MoveTemp(CompositePass));
+	}
+
+	// Sort composited passes if multiple were found. Passes with a smaller sort order go to the end of the array so they
+	// get composited on top of passes with a higher sort order.
+	CompositedPasses.Sort([](
+		const TPair<FMovieGraphRenderDataIdentifier, TUniquePtr<FImagePixelData>>& PassA,
+		const TPair<FMovieGraphRenderDataIdentifier, TUniquePtr<FImagePixelData>>& PassB)
+	{
+		const UE::MovieGraph::FMovieGraphSampleState* PayloadA = PassA.Value->GetPayload<UE::MovieGraph::FMovieGraphSampleState>();
+		const UE::MovieGraph::FMovieGraphSampleState* PayloadB = PassB.Value->GetPayload<UE::MovieGraph::FMovieGraphSampleState>();
+		check(PayloadA);
+		check(PayloadB);
+
+		return PayloadA->CompositingSortOrder > PayloadB->CompositingSortOrder;
+	});
+
+	// ToDo:
+	// The ImageWriteQueue is set up in a fire-and-forget manner. This means that the data needs to be placed in the WriteQueue
+	// as a TUniquePtr (so it can free the data when its done). Unfortunately we can have multiple output formats at once,
+	// so we can't MoveTemp the data into it, we need to make a copy (though we could optimize for the common case where there is
+	// only one output format).
+	// Copying can be expensive (3ms @ 1080p, 12ms at 4k for a single layer image) so ideally we'd like to do it on the task graph
+	// but this isn't really compatible with the ImageWriteQueue API as we need the future returned by the ImageWriteQueue to happen
+	// in order, so that we push our futures to the main Movie Pipeline in order, otherwise when we encode files to videos they'll
+	// end up with frames out of order. A workaround for this would be to chain all of the send-to-imagewritequeue tasks to each
+	// other with dependencies, but I'm not sure that's going to scale to the potentialy high data volume going wide MRQ will eventually
+	// need.
+
+	// The base ImageSequenceOutputNode doesn't support any multilayer formats, so we write out each render pass separately.
+	for (TPair<FMovieGraphRenderDataIdentifier, TUniquePtr<FImagePixelData>>& RenderData : InRawFrameData->ImageOutputData)
+	{
+		// If this pass is composited, skip it for now
+		bool bSkip = false;
+		for (TPair<FMovieGraphRenderDataIdentifier, TUniquePtr<FImagePixelData>>& CompositedPass : CompositedPasses)
+		{
+			if (CompositedPass.Key == RenderData.Key)
+			{
+				bSkip = true;
+				break;
+			}
+		}
+
+		if (bSkip)
+		{
+			continue;
+		}
+		
+		// A layer within this output data may have chosen to not be written to disk by this CDO node
+		if (!InMask.Contains(RenderData.Key))
+		{
+			continue;
+		}
+
+		UE::MovieGraph::FMovieGraphSampleState* Payload = RenderData.Value->GetPayload<UE::MovieGraph::FMovieGraphSampleState>();
+
+		const bool bIncludeCDOs = true;
+		UMovieGraphOutputSettingNode* OutputSettingNode = InRawFrameData->EvaluatedConfig->GetSettingForBranch<UMovieGraphOutputSettingNode>(RenderData.Key.RootBranchName, bIncludeCDOs);
+		if (!ensure(OutputSettingNode))
+		{
+			continue;
+		}
+
+		FString RenderLayerName = RenderData.Key.RootBranchName.ToString();
+		UMovieGraphRenderLayerNode* RenderLayerNode = InRawFrameData->EvaluatedConfig->GetSettingForBranch<UMovieGraphRenderLayerNode>(RenderData.Key.RootBranchName, bIncludeCDOs);
+		if (RenderLayerNode)
+		{
+			RenderLayerName = RenderLayerNode->GetRenderLayerName();
+		}
+		// ToDo: Certain images may require transparency, at which point
+		// we write out a .png instead of a .jpeg.
+		EImageFormat PreferredOutputFormat = OutputFormat;
+
+		const TCHAR* Extension = TEXT("");
+		switch (PreferredOutputFormat)
+		{
+		case EImageFormat::PNG: Extension = TEXT("png"); break;
+		case EImageFormat::JPEG: Extension = TEXT("jpeg"); break;
+		case EImageFormat::BMP: Extension = TEXT("bmp"); break;
+		case EImageFormat::EXR: Extension = TEXT("exr"); break;
+		}
+
+		// Generate one string that puts the directory combined with the filename format.
+		FString FileNameFormatString = OutputSettingNode->OutputDirectory.Path / OutputSettingNode->FileNameFormat;
+
+		// ToDo: Validate the string, ie: ensure it has {render_pass} in there somewhere there are multiple render passes
+		// in the output data, include {camera_name} if there are multiple cameras for that render pass, etc. Validation
+		// should insert {file_dup} tokens so it can put them at a logical place (ie: before frame numbers?)
+		FileNameFormatString += TEXT(".{ext}");
+
+		// Map the .ext to be specific to our output data.
+		TMap<FString, FString> AdditionalFormatArgs;
+		AdditionalFormatArgs.Add(TEXT("ext"), Extension);
+
+		FMovieGraphFilenameResolveParams Params = FMovieGraphFilenameResolveParams();
+		Params.RenderDataIdentifier = RenderData.Key;
+		//Params.RootFrameNumber = Payload->TraversalContext.Time.RootFrameNumber;
+		//Params.ShotFrameNumber = Payload->TraversalContext.Time.ShotFrameNumber;
+		Params.RootFrameNumberRel = Payload->TraversalContext.Time.OutputFrameNumber;
+		//Params.ShotFrameNumberRel = Payload->TraversalCOntext.Time.ShotFrameNumberRel
+		//Params.FileMetadata = ToDo: Track File Metadata
+		Params.Version = 1; // ToDo: Track versions
+		Params.ZeroPadFrameNumberCount = OutputSettingNode->ZeroPadFrameNumbers;
+		Params.FrameNumberOffset = OutputSettingNode->FrameNumberOffset;
+
+		// If time dilation is in effect, RootFrameNumber and ShotFrameNumber will contain duplicates and the files will overwrite each other, 
+		// so we force them into relative mode and then warn users we did that (as their numbers will jump from say 1001 -> 0000).
+		bool bForceRelativeFrameNumbers = true; // TODO: Use relative frame numbers until we track Root vs. Shot frame numbers. (Previously false);
+		//if (FileNameFormatString.Contains(TEXT("{frame")) Payload->TraversalContext.Time.IsTimeDilated() && !FileNameFormatString.Contains(TEXT("_rel}")))
+		//{
+		//	UE_LOG(LogMovieRenderPipeline, Warning, TEXT("Time Dilation was used but output format does not use relative time, forcing relative numbers. Change {frame_number} to {frame_number_rel} (or shot version) to remove this message."));
+		//	bForceRelativeFrameNumbers = true;
+		//}
+		Params.bForceRelativeFrameNumbers = bForceRelativeFrameNumbers;
+		Params.bEnsureAbsolutePath = true;
+		Params.FileNameFormatOverrides = AdditionalFormatArgs;
+		Params.InitializationTime = InPipeline->GetInitializationTime();
+		Params.Job = InPipeline->GetCurrentJob();
+		Params.EvaluatedConfig = InRawFrameData->EvaluatedConfig.Get();
+
+		// Take our string path from the Output Setting and resolve it.
+		FMovieGraphResolveArgs FinalResolvedKVPs;
+		const FString FileName = UMovieGraphBlueprintLibrary::ResolveFilenameFormatArguments(FileNameFormatString, Params, FinalResolvedKVPs);
+
+		TUniquePtr<FImageWriteTask> TileImageTask = MakeUnique<FImageWriteTask>();
+		TileImageTask->Format = OutputFormat;
+		TileImageTask->CompressionQuality = 100;
+		TileImageTask->Filename = FileName;
+		TileImageTask->PixelData = RenderData.Value->CopyImageData();
+
+		EImagePixelType PixelType = TileImageTask->PixelData->GetType();
+
+		// Perform compositing if any composited passes were found earlier
+		for (TPair<FMovieGraphRenderDataIdentifier, TUniquePtr<FImagePixelData>>& CompositedPass : CompositedPasses)
+		{
+			// This composited pass will only composite on top of renders w/ the same branch and camera
+			const FMovieGraphRenderDataIdentifier& Id = CompositedPass.Key;
+			if ((Id.CameraName != RenderData.Key.CameraName) || (Id.RootBranchName != RenderData.Key.RootBranchName))
+			{
+				continue;
+			}
+
+			// There could be multiple renders within this branch using the composited pass, so we have to copy the image data
+			switch (PixelType)
+			{
+			case EImagePixelType::Color:
+				TileImageTask->PixelPreProcessors.Add(TAsyncCompositeImage<FColor>(CompositedPass.Value->CopyImageData()));
+				break;
+			case EImagePixelType::Float16:
+				TileImageTask->PixelPreProcessors.Add(TAsyncCompositeImage<FFloat16Color>(CompositedPass.Value->CopyImageData()));
+				break;
+			case EImagePixelType::Float32:
+				TileImageTask->PixelPreProcessors.Add(TAsyncCompositeImage<FLinearColor>(CompositedPass.Value->CopyImageData()));
+				break;
+			}
+		}
+
+		UE::MovieGraph::FMovieGraphOutputFutureData OutputData;
+		OutputData.Shot = nullptr;
+		OutputData.FilePath = FileName;
+		OutputData.DataIdentifier = RenderData.Key;
+
+		TFuture<bool> Future = ImageWriteQueue->Enqueue(MoveTemp(TileImageTask));
+
+		InPipeline->AddOutputFuture(MoveTemp(Future), OutputData);
+	}
+
+}
+
+void UMovieGraphImageSequenceOutputNode_EXR::OnReceiveImageDataImpl(UMovieGraphPipeline* InPipeline, UE::MovieGraph::FMovieGraphOutputMergerFrame* InRawFrameData, const TSet<FMovieGraphRenderDataIdentifier>& InMask)
+{
+	if (!bMultilayer)
+	{
+		// Some software doesn't support multi-layer, so in that case we fall back to the single-layer-multiple-file
+		// codepath of our parent.
+		Super::OnReceiveImageDataImpl(InPipeline, InRawFrameData, InMask);
+		return;
+	}
+	
+	check(InRawFrameData);
+
+	// Ensure our OpenExrRTTI module gets loaded. This needs to happen from the main thread, if it's not loaded then metadata silently fails when writing.
+	static const FName RTTIExtensionModuleName("UEOpenExrRTTI");
+	FModuleManager::Get().LoadModule(RTTIExtensionModuleName);
+
+	// Generate a mapping of resolved filename -> RenderIDs, and filename -> resolve args. The generated EXRs can only
+	// store layers with a common resolution, and this takes care of ensuring that only renderIDs with a common resolution
+	// map to the same filename.
+	TMap<FString, TArray<FMovieGraphRenderDataIdentifier>> FilenameToRenderIDs;
+	TMap<FString, FMovieGraphResolveArgs> FilenameToResolveArgs;
+	GetFilenameToRenderIDMappings(InPipeline, InRawFrameData, FilenameToRenderIDs, FilenameToResolveArgs);
+
+	// Write an EXR for each filename, which potentially contains multiple passes (render IDs).
+	for (const TPair<FString, TArray<FMovieGraphRenderDataIdentifier>>& RenderIDsForFilename : FilenameToRenderIDs)
+	{
+		const FString& Filename = RenderIDsForFilename.Key;
+		const TArray<FMovieGraphRenderDataIdentifier>& RenderIDs = RenderIDsForFilename.Value;
+		
+		TUniquePtr<FEXRImageWriteTask> MultiLayerImageTask = MakeUnique<FEXRImageWriteTask>();
+		
+		MultiLayerImageTask->Filename = Filename;
+		MultiLayerImageTask->Compression = Compression;
+		// MultiLayerImageTask->CompressionLevel is intentionally skipped because it doesn't seem to make any practical difference
+		// so we don't expose it to the user because that will just cause confusion where the setting doesn't seem to do anything.
+
+		// FileMetadata has been generated by ResolveFilenameFormatArgs, but we need to convert from FString, FString
+		// (needed for BP/Python purposes) to a FStringFormatArg as we need to preserve numeric metadata types later in
+		// the image writing process (for compression level).
+		TMap<FString, FStringFormatArg> NewFileMetadataMap;
+		for (const TPair<FString, FString>& Metadata : FilenameToResolveArgs[Filename].FileMetadata)
+		{
+			NewFileMetadataMap.Add(Metadata.Key, Metadata.Value);
+		}
+		MultiLayerImageTask->FileMetadata = NewFileMetadataMap;
+
+		// Add color space metadata to the output: xy chromaticity coordinates and/or the color space source/dest names.
+		// TODO: Support is also needed for regular exrs via the image wrapper module.
+		// TODO: No OCIO node yet
+		// {
+		// 	UMoviePipelineColorSetting* ColorSetting = InPipeline->GetPipelinePrimaryConfig()->FindSetting<UMoviePipelineColorSetting>();
+		//
+		// 	FColorSpaceMetadata ColorSpaceMetadata = GetColorSpaceMetadata(ColorSetting);
+		//
+		// 	if (!ColorSpaceMetadata.SourceName.IsEmpty())
+		// 	{
+		// 		MultiLayerImageTask->FileMetadata.Add("unreal/colorSpace/source", ColorSpaceMetadata.SourceName);
+		// 	}
+		// 	if (!ColorSpaceMetadata.DestinationName.IsEmpty())
+		// 	{
+		// 		MultiLayerImageTask->FileMetadata.Add("unreal/colorSpace/destination", ColorSpaceMetadata.DestinationName);
+		// 	}
+		//
+		// 	MultiLayerImageTask->ColorSpaceChromaticities = ColorSpaceMetadata.Chromaticities;
+		// }
+
+		// Add each render pass as a layer to the EXR
+		int32 LayerIndex = 0;
+		int32 ShotIndex = 0;
+		for (const FMovieGraphRenderDataIdentifier& RenderID : RenderIDs)
+		{
+			const TUniquePtr<FImagePixelData>& ImageData = InRawFrameData->ImageOutputData[RenderID];
+
+			// No quantization required, just copy the data as we will move it into the image write task.
+			TUniquePtr<FImagePixelData> PixelData = ImageData->CopyImageData();
+			const UE::MovieGraph::FMovieGraphSampleState* Payload = ImageData->GetPayload<UE::MovieGraph::FMovieGraphSampleState>();
+			ShotIndex = Payload->TraversalContext.ShotIndex;
+
+			// If there is more than one layer, then we will prefix the layer. The first layer is not prefixed (and gets inserted as RGBA)
+			// as most programs that handle EXRs expect the main image data to be in an unnamed layer.
+			if (LayerIndex == 0)
+			{
+				// Add task information that is common to all layers
+				MultiLayerImageTask->FileMetadata.Add("owner", UE::MoviePipeline::GetJobAuthor(Payload->TraversalContext.Job));
+				MultiLayerImageTask->FileMetadata.Add("comments", Payload->TraversalContext.Job->Comment);
+
+				const FIntPoint& Resolution = ImageData->GetSize();
+				MultiLayerImageTask->Width = Resolution.X;
+				MultiLayerImageTask->Height = Resolution.Y;
+				
+				// TODO: Overscan not available in the graph yet
+				// MultiLayerImageTask->OverscanPercentage = Payload->SampleState.OverscanPercentage;
+			}
+			else
+			{
+				// If there is more than one layer, then we will prefix the layer. The first layer is not prefixed (and gets inserted as RGBA)
+				// as most programs that handle EXRs expect the main image data to be in an unnamed layer. We only postfix with cameraname
+				// if there's multiple cameras, as pipelines may be already be built around the generic "one camera" support.
+				// TODO: The number of cameras may be inaccurate -- no camera setting in the graph yet
+				UMoviePipelineExecutorShot* CurrentShot = InPipeline->GetActiveShotList()[ShotIndex];
+				int32 NumCameras = CurrentShot->SidecarCameras.Num();
+				
+				FString CombinedName = FString::Printf(TEXT("%s_%s"), *RenderID.RootBranchName.ToString(), *RenderID.RendererName);
+				if (NumCameras > 1)
+				{
+					CombinedName = FString::Printf(TEXT("%s_%s"), *CombinedName, *RenderID.CameraName);
+				}
+				
+				MultiLayerImageTask->LayerNames.FindOrAdd(PixelData.Get(), CombinedName);
+			}
+
+			MultiLayerImageTask->Layers.Add(MoveTemp(PixelData));
+			LayerIndex++;
+		}
+		
+		UE::MovieGraph::FMovieGraphOutputFutureData OutputFutureData;
+		OutputFutureData.Shot = InPipeline->GetActiveShotList()[ShotIndex];
+		OutputFutureData.FilePath = Filename;
+		OutputFutureData.DataIdentifier = FMovieGraphRenderDataIdentifier(); // EXRs put all the render passes internally so this resolves to a ""
+
+		InPipeline->AddOutputFuture(ImageWriteQueue->Enqueue(MoveTemp(MultiLayerImageTask)), OutputFutureData);
+	}
+}
+
+void UMovieGraphImageSequenceOutputNode_EXR::GetFilenameToRenderIDMappings(
+	UMovieGraphPipeline* InPipeline, UE::MovieGraph::FMovieGraphOutputMergerFrame* InRawFrameData,
+	TMap<FString, TArray<FMovieGraphRenderDataIdentifier>>& OutFilenameToRenderIDs,
+	TMap<FString, FMovieGraphResolveArgs>& OutFilenameToResolveArgs) const
+{
+	TMap<FString, TArray<FIntPoint>> FilenameToResolutions;
+
+	// First, generate filename -> renderID mapping, and filename -> resolution mapping.
+	// This assumes that all render passes will have the same resolution, so we use 0 as the resolution index.
+	// Once we know the resolutions of all the render passes, they can be binned together into groups with the same
+	// resolution, and the filenames can be regenerated to ensure that passes of differing resolutions go to
+	// different files.
+	//
+	// This two-step process is necessary due to the flexibility in file naming, and the multi-layer nature of EXRs.
+	// For example, if the file name format is "{sequence_name}.{frame_number}", and the second of two branches in the
+	// graph has a differing resolution, only after resolving the output filenames for all outputs is a problem found;
+	// layers of differing resolutions will be written to the same file. Using "{render_layer}.{sequence_name}.{frame_number}"
+	// as the file name format would prevent the issue, but the two-step process is a generic way of approaching the
+	// problem.
+	for (const TPair<FMovieGraphRenderDataIdentifier, TUniquePtr<FImagePixelData>>& RenderPassData : InRawFrameData->ImageOutputData)
+	{
+		constexpr int32 ResolutionIndex = 0;
+		FMovieGraphResolveArgs ResolveArgs;
+		const FString PreliminaryFileName = ResolveOutputFilename(InPipeline, ResolutionIndex, InRawFrameData, RenderPassData.Key.RootBranchName, ResolveArgs);
+		
+		TArray<FMovieGraphRenderDataIdentifier>& RenderIDs = OutFilenameToRenderIDs.FindOrAdd(PreliminaryFileName);
+		RenderIDs.Add(RenderPassData.Key);
+
+		TArray<FIntPoint>& Resolutions = FilenameToResolutions.FindOrAdd(PreliminaryFileName);
+		Resolutions.AddUnique(RenderPassData.Value->GetSize());
+
+		OutFilenameToResolveArgs.Add(PreliminaryFileName, ResolveArgs);
+	}
+
+	// Second, re-generate filenames if any render passes of differing resolutions map to the same file.
+	for (const TPair<FString, TArray<FIntPoint>>& FilenameAndResolutions : FilenameToResolutions)
+	{
+		const FString& PreliminaryFilename = FilenameAndResolutions.Key;
+		const TArray<FIntPoint>& Resolutions = FilenameAndResolutions.Value;
+
+		// If there's only one resolution for this filename, there's nothing to disambiguate
+		if (Resolutions.Num() == 1)
+		{
+			continue;
+		}
+
+		// If there IS more than one resolution for this filename, then disambiguate
+		const TArray<FMovieGraphRenderDataIdentifier> OldRenderIDs = OutFilenameToRenderIDs.FindAndRemoveChecked(PreliminaryFilename);
+		OutFilenameToResolveArgs.Remove(PreliminaryFilename);
+		for (const FMovieGraphRenderDataIdentifier& RenderID : OldRenderIDs)
+		{
+			TUniquePtr<FImagePixelData>& ImageData = InRawFrameData->ImageOutputData[RenderID];
+			const int32 ResolutionIndex = Resolutions.IndexOfByKey(ImageData->GetSize());
+
+			// Re-resolve the filename, this time using the resolution index to generate a filename that will only contain
+			// passes with this particular resolution
+			FMovieGraphResolveArgs ResolveArgs;
+			const FString FinalFilename = ResolveOutputFilename(InPipeline, ResolutionIndex, InRawFrameData, RenderID.RootBranchName, ResolveArgs);
+
+			TArray<FMovieGraphRenderDataIdentifier>& RenderIDs = OutFilenameToRenderIDs.FindOrAdd(FinalFilename);
+			RenderIDs.Add(RenderID);
+
+			OutFilenameToResolveArgs.Add(FinalFilename, ResolveArgs);
+		}
+	}
+}
+
+FString UMovieGraphImageSequenceOutputNode_EXR::ResolveOutputFilename(
+	const UMovieGraphPipeline* InPipeline,
+	const int32 ResolutionIndex, const UE::MovieGraph::FMovieGraphOutputMergerFrame* InRawFrameData,
+	const FName& InBranchName, FMovieGraphResolveArgs& OutResolveArgs) const
+{
+	const TCHAR* Extension = TEXT("exr");
+
+	constexpr bool bIncludeCDOs = true;
+	const UMovieGraphOutputSettingNode* OutputSettings = InRawFrameData->EvaluatedConfig->GetSettingForBranch<UMovieGraphOutputSettingNode>(InBranchName, bIncludeCDOs);
+	if (!ensure(OutputSettings))
+	{
+		return FString();
+	}
+	
+	// If we have more than one resolution we'll store it as "_Add" / "_Add(1)" etc via {ExtraTag}.
+	FString FileNameFormatString = OutputSettings->FileNameFormat + "{ExtraTag}";
+	FileNameFormatString += TEXT(".{ext}");
+
+	const FString FilePathFormatString = OutputSettings->OutputDirectory.Path / FileNameFormatString;
+	
+	// If we're writing more than one render pass out, we need to ensure the file name has the format string in it so we don't
+	// overwrite the same file multiple times. Burn In overlays don't count because they get composited on top of an existing file.
+	constexpr bool bIncludeRenderPass = false;
+	constexpr bool bTestFrameNumber = true;
+	UE::MoviePipeline::ValidateOutputFormatString(FileNameFormatString, bIncludeRenderPass, bTestFrameNumber);
+
+	// Create specific data that needs to override 
+	TMap<FString, FString> FormatOverrides;
+	FormatOverrides.Add(TEXT("render_pass"), TEXT("")); // Render Passes are included inside the exr file by named layers.
+	FormatOverrides.Add(TEXT("ext"), Extension);
+
+	// The logic for the ExtraTag is a little complicated. If there's only one layer (ideal situation) then it's empty.
+	if (ResolutionIndex == 0)
+	{
+		FormatOverrides.Add(TEXT("ExtraTag"), TEXT(""));
+	}
+	else if (ResolutionIndex == 1)
+	{
+		// This is our most common case when we have a second file (the only expected one really)
+		FormatOverrides.Add(TEXT("ExtraTag"), TEXT("_Add"));
+	}
+	else
+	{
+		// Finally a fallback in the event we have three or more unique resolutions.
+		FormatOverrides.Add(TEXT("ExtraTag"), FString::Printf(TEXT("_Add(%d)"), ResolutionIndex));
+	}
+
+	// Since a multi-layer EXR can store renders from multiple cameras, renderers, etc, the render data identifier isn't
+	// very useful. However, we still need to provide the branch name -- this is important for resolving the correct output path.
+	FMovieGraphRenderDataIdentifier TempRenderDataIdentifier;
+	TempRenderDataIdentifier.RootBranchName = InBranchName;
+
+	// This resolves the filename format and gathers metadata from the settings at the same time.
+	FMovieGraphFilenameResolveParams Params = FMovieGraphFilenameResolveParams();
+	Params.RenderDataIdentifier = TempRenderDataIdentifier;
+	Params.RootFrameNumberRel = InRawFrameData->TraversalContext.Time.OutputFrameNumber;
+	Params.ZeroPadFrameNumberCount = OutputSettings->ZeroPadFrameNumbers;
+	Params.FrameNumberOffset = OutputSettings->FrameNumberOffset;
+	Params.bForceRelativeFrameNumbers = true;	// TODO: Use the right value
+	Params.bEnsureAbsolutePath = true;
+	Params.FileNameFormatOverrides = FormatOverrides;
+	Params.InitializationTime = InPipeline->GetInitializationTime();
+	Params.Job = InPipeline->GetCurrentJob();
+	Params.EvaluatedConfig = InRawFrameData->EvaluatedConfig.Get();
+	
+	FString FinalFilePath = UMovieGraphBlueprintLibrary::ResolveFilenameFormatArguments(FilePathFormatString, Params, OutResolveArgs);
+
+	if (FPaths::IsRelative(FinalFilePath))
+	{
+		FinalFilePath = FPaths::ConvertRelativePathToFull(FinalFilePath);
+	}
+
+	return FinalFilePath;
+}
diff --git a/Engine/Plugins/MovieScene/MovieRenderPipeline/Source/MovieRenderPipelineRenderPasses/Public/MovieGraphImageSequenceOutputNode.h b/Engine/Plugins/MovieScene/MovieRenderPipeline/Source/MovieRenderPipelineRenderPasses/Public/MovieGraphImageSequenceOutputNode.h
new file mode 100644
--- /dev/null
+++ b/Engine/Plugins/MovieScene/MovieRenderPipeline/Source/MovieRenderPipelineRenderPasses/Public/MovieGraphImageSequenceOutputNode.h
@@ -0,0 +1,249 @@
+// Copyright Epic Games, Inc. All Rights Reserved.
+
+#pragma once
+
+#include "Graph/Nodes/MovieGraphFileOutputNode.h"
+
+#include "IImageWrapper.h"
+#include "MoviePipelineEXROutput.h"
+#include "Styling/AppStyle.h"
+#include "Async/Future.h"
+
+#include "MovieGraphImageSequenceOutputNode.generated.h"
+
+// Forward Declare
+class IImageWriteQueue;
+
+/**
+* The UMovieGraphImageSequenceOutputNode node is the base class for all image sequence outputs, such as 
+* a series of jpeg, png, bmp, or .exr images. Create an instance of the appropriate class (such as 
+* UMovieGraphImageSequenceOutputNode_JPG) instead of this abstract base class.
+*/
+UCLASS(Abstract)
+class UMovieGraphImageSequenceOutputNode : public UMovieGraphFileOutputNode
+{
+	GENERATED_BODY()
+public:
+	UMovieGraphImageSequenceOutputNode();
+	
+	// UMovieGraphFileOutputNode Interface
+	virtual void OnReceiveImageDataImpl(UMovieGraphPipeline* InPipeline, UE::MovieGraph::FMovieGraphOutputMergerFrame* InRawFrameData, const TSet<FMovieGraphRenderDataIdentifier>& InMask) override;
+	virtual void OnAllFramesSubmittedImpl() override;
+	virtual bool IsFinishedWritingToDiskImpl() const override;
+	// ~UMovieGraphFileOutputNode Interface
+
+protected:
+	/** The output format (as known used by the ImageWriteQueue) to output into. */
+	EImageFormat OutputFormat;
+
+	/** A pointer to the image write queue used for asynchronously writing images */
+	IImageWriteQueue* ImageWriteQueue;
+
+	/** A fence to keep track of when the Image Write queue has fully flushed. */
+	TFuture<void> FinalizeFence;
+};
+
+/**
+ * Image sequence output node that can write EXR files, which are optionally multi-layer.
+ */
+UCLASS()
+class UMovieGraphImageSequenceOutputNode_EXR : public UMovieGraphImageSequenceOutputNode
+{
+	GENERATED_BODY()
+
+public:
+	UMovieGraphImageSequenceOutputNode_EXR()
+	{
+		OutputFormat = EImageFormat::EXR;
+		Compression = EEXRCompressionFormat::PIZ;
+		bMultilayer = true;
+	}
+
+	virtual void OnReceiveImageDataImpl(UMovieGraphPipeline* InPipeline, UE::MovieGraph::FMovieGraphOutputMergerFrame* InRawFrameData, const TSet<FMovieGraphRenderDataIdentifier>& InMask) override;
+
+#if WITH_EDITOR
+	virtual FText GetNodeTitle(const bool bGetDescriptive = false) const override 
+	{ 
+		static const FText EXRSequenceNodeName = NSLOCTEXT("MovieGraph", "NodeName_EXRSequence", ".exr Sequence");
+		return EXRSequenceNodeName;
+	}
+	
+	virtual FLinearColor GetNodeTitleColor() const override
+	{
+		return FLinearColor(0.047f, 0.654f, 0.537f);
+	}
+	
+	virtual FSlateIcon GetIconAndTint(FLinearColor& OutColor) const override
+	{
+		static const FSlateIcon ImageSequenceIcon = FSlateIcon(FAppStyle::GetAppStyleSetName(), "ClassIcon.Texture2D");
+
+		OutColor = FLinearColor::White;
+		return ImageSequenceIcon;
+	}
+#endif
+	
+public:
+	UPROPERTY(EditAnywhere, BlueprintReadWrite, Category = Overrides, meta = (InlineEditConditionToggle))
+	uint8 bOverride_Compression : 1;
+
+	UPROPERTY(EditAnywhere, BlueprintReadWrite, Category = Overrides, meta = (InlineEditConditionToggle))
+	uint8 bOverride_bMultilayer : 1;
+	
+	/**
+	 * Which compression method should the resulting EXR file be compressed with.
+	 */
+	UPROPERTY(EditAnywhere, BlueprintReadWrite, meta=(EditCondition="bOverride_Compression"), Category = "EXR")
+	EEXRCompressionFormat Compression;
+
+	/**
+	 * Whether all renders should be written to the same EXR file via multi-layer support. Not all software
+	 * supports reading multi-layer EXR files.
+	 */
+	UPROPERTY(EditAnywhere, BlueprintReadWrite, meta=(EditCondition="bOverride_bMultilayer"), Category = "EXR")
+	bool bMultilayer;
+
+private:
+	/**
+	 * Generates filenames for each render pass (renderID), placed in OutFilenameToRenderIDs. Also provides the resolve
+	 * args that were created when resolving the filename in OutFilenameToResolveArgs. The generated mapping ensures
+	 * that each filename only points to renderIDs with a common resolution, since currently EXRs can only contain layers
+	 * with the same resolution.
+	 */
+	void GetFilenameToRenderIDMappings(
+		UMovieGraphPipeline* InPipeline, UE::MovieGraph::FMovieGraphOutputMergerFrame* InRawFrameData,
+		TMap<FString, TArray<FMovieGraphRenderDataIdentifier>>& OutFilenameToRenderIDs,
+		TMap<FString, FMovieGraphResolveArgs>& OutFilenameToResolveArgs) const;
+	
+	/**
+	 * Generates the filename that the EXR will be written to, as well as the resulting resolve args via OutResolveArgs.
+	 * Use GetFilenameToRenderIDMappings() to guarantee that the filename respects EXR limitations.
+	 */
+	FString ResolveOutputFilename(
+		const UMovieGraphPipeline* InPipeline, const int32 ResolutionIndex, const UE::MovieGraph::FMovieGraphOutputMergerFrame* InRawFrameData,
+		const FName& InBranchName, FMovieGraphResolveArgs& OutResolveArgs) const;
+};
+
+/**
+* Save the images generated by the Movie Graph Pipeline as an lossless 8 bit bmp format. This can
+* be useful in rare occasions (bmp files are uncompressed but larger). sRGB is applied.
+* No metadata is supported.
+*/
+UCLASS()
+class UMovieGraphImageSequenceOutputNode_BMP : public UMovieGraphImageSequenceOutputNode
+{
+	GENERATED_BODY()
+
+public:
+	UMovieGraphImageSequenceOutputNode_BMP()
+	{
+		OutputFormat = EImageFormat::BMP;
+	}
+
+#if WITH_EDITOR
+	virtual FText GetNodeTitle(const bool bGetDescriptive = false) const override 
+	{ 
+		if(bGetDescriptive)
+		{
+			return NSLOCTEXT("MovieGraph", "ImgSequenceBMPSetting_NodeTitleFull", ".bmp Sequence\n[8bit]"); 
+		}
+		return NSLOCTEXT("MovieGraph", "ImgSequenceBMPSetting_NodeTitleShort", ".bmp Sequence");
+	}
+	
+	virtual FLinearColor GetNodeTitleColor() const override
+	{
+		return FLinearColor(0.047f, 0.654f, 0.537f);
+	}
+	
+	virtual FSlateIcon GetIconAndTint(FLinearColor& OutColor) const override
+	{
+		static const FSlateIcon ImageSequenceIcon = FSlateIcon(FAppStyle::GetAppStyleSetName(), "ClassIcon.Texture2D");
+
+		OutColor = FLinearColor::White;
+		return ImageSequenceIcon;
+	}
+#endif
+};
+
+
+/**
+* Save the images generated by the Movie Graph Pipeline as an 8 bit jpg format. JPEG image files
+* are lossy, but a good balance between compression speed and final filesize. sRGB is applied.
+* No metadata is supported.
+*/
+UCLASS()
+class UMovieGraphImageSequenceOutputNode_JPG : public UMovieGraphImageSequenceOutputNode
+{
+	GENERATED_BODY()
+
+public:
+	UMovieGraphImageSequenceOutputNode_JPG()
+	{
+		OutputFormat = EImageFormat::JPEG;
+	}
+
+#if WITH_EDITOR
+	virtual FText GetNodeTitle(const bool bGetDescriptive = false) const override 
+	{ 
+		if(bGetDescriptive)
+		{
+			return NSLOCTEXT("MovieGraph", "ImgSequenceJPGSetting_NodeTitleFull", ".jpg Sequence\n[8bit]"); 
+		}
+		return NSLOCTEXT("MovieGraph", "ImgSequenceJPGSetting_NodeTitleShort", ".jpg Sequence");
+	}
+	
+	virtual FLinearColor GetNodeTitleColor() const override
+	{
+		return FLinearColor(0.047f, 0.654f, 0.537f);
+	}
+	
+	virtual FSlateIcon GetIconAndTint(FLinearColor& OutColor) const override
+	{
+		static const FSlateIcon ImageSequenceIcon = FSlateIcon(FAppStyle::GetAppStyleSetName(), "ClassIcon.Texture2D");
+
+		OutColor = FLinearColor::White;
+		return ImageSequenceIcon;
+	}
+#endif
+};
+
+/**
+* Save the images generated by the Movie Graph Pipeline as an 8 bit png format. PNG image files
+* are lossless but slow to compress and have a larger final filesize than JPEG. sRGB is applied.
+* No metadata is supported.
+*/
+UCLASS()
+class UMovieGraphImageSequenceOutputNode_PNG : public UMovieGraphImageSequenceOutputNode
+{
+	GENERATED_BODY()
+
+public:
+	UMovieGraphImageSequenceOutputNode_PNG()
+	{
+		OutputFormat = EImageFormat::PNG;
+	}
+
+#if WITH_EDITOR
+	virtual FText GetNodeTitle(const bool bGetDescriptive = false) const override 
+	{ 
+		if(bGetDescriptive)
+		{
+			return NSLOCTEXT("MovieGraph", "ImgSequencePNGSetting_NodeTitleFull", ".png Sequence\n[8bit]"); 
+		}
+		return NSLOCTEXT("MovieGraph", "ImgSequencePNGSetting_NodeTitleShort", ".png Sequence");
+	}
+	
+	virtual FLinearColor GetNodeTitleColor() const override
+	{
+		return FLinearColor(0.047f, 0.654f, 0.537f);
+	}
+	
+	virtual FSlateIcon GetIconAndTint(FLinearColor& OutColor) const override
+	{
+		static const FSlateIcon ImageSequenceIcon = FSlateIcon(FAppStyle::GetAppStyleSetName(), "ClassIcon.Texture2D");
+
+		OutColor = FLinearColor::White;
+		return ImageSequenceIcon;
+	}
+#endif
+};
+
