MoviePipeline: Added "Apply CVar Preset" node to the graph, and integrated cvars into the graph pipeline.

Because the graph pipeline pulls in cvars from multiple node types (currently the "Apply CVar Preset" and "Set CVar Value" nodes), an individual node should not be responsible for applying cvars. A new class that deals with applying cvars in the pipeline, FMovieGraphCVarManager, was created to deal with this.

The cvar nodes were also locked down to the Globals branch. In this version of the graph, cvars will only be applied once before a render to avoid issues with setting cvars to different values on various branches, or setting cvars in the middle of a render (thread synchronization problems).

#rnx

diff --git a/Engine/Plugins/MovieScene/MovieRenderPipeline/Source/MovieRenderPipelineCore/Private/Graph/MovieGraphCVarManager.cpp b/Engine/Plugins/MovieScene/MovieRenderPipeline/Source/MovieRenderPipelineCore/Private/Graph/MovieGraphCVarManager.cpp
new file mode 100644
--- /dev/null
+++ b/Engine/Plugins/MovieScene/MovieRenderPipeline/Source/MovieRenderPipelineCore/Private/Graph/MovieGraphCVarManager.cpp
@@ -0,0 +1,155 @@
+﻿// Copyright Epic Games, Inc. All Rights Reserved.
+
+#include "Graph/MovieGraphCVarManager.h"
+
+#include "Graph/MovieGraphConfig.h"
+#include "Graph/Nodes/MovieGraphApplyCVarPresetNode.h"
+#include "Graph/Nodes/MovieGraphSetCVarValueNode.h"
+#include "HAL/IConsoleManager.h"
+#include "Misc/DefaultValueHelper.h"
+#include "MovieRenderPipelineCoreModule.h"
+#include "Sections/MovieSceneConsoleVariableTrackInterface.h"
+
+namespace UE::MovieGraph::Private
+{
+
+void FMovieGraphCVarManager::AddCVar(const FString& InName, const float Value)
+{
+	FMovieGraphCVarOverride* CVarOverride = CVars.FindByPredicate([&InName](const FMovieGraphCVarOverride& Override)
+	{
+		return Override.Name == InName;
+	});
+
+	if (CVarOverride)
+	{
+		CVarOverride->Value = Value;
+	}
+	else
+	{
+		CVars.Add(FMovieGraphCVarOverride(InName, Value));
+	}
+}
+
+void FMovieGraphCVarManager::AddPreset(const TScriptInterface<IMovieSceneConsoleVariableTrackInterface>& InPreset)
+{
+	if (!InPreset)
+	{
+		UE_LOG(LogMovieRenderPipeline, Warning, TEXT("Invalid CVar preset specified. Ignoring."));
+		return;
+	}
+
+	constexpr bool bOnlyIncludeChecked = true;
+	TArray<TTuple<FString, FString>> PresetCVars;
+	InPreset->GetConsoleVariablesForTrack(bOnlyIncludeChecked, PresetCVars);
+
+	for (const TTuple<FString, FString>& CVarPair : PresetCVars)
+	{
+		float CVarFloatValue = 0.0f;
+		if (FDefaultValueHelper::ParseFloat(CVarPair.Value, CVarFloatValue))
+		{
+			AddCVar(CVarPair.Key, CVarFloatValue);
+		}
+		else
+		{
+			UE_LOG(LogMovieRenderPipeline, Warning, TEXT("Failed to apply CVar \"%s\" (from preset \"%s\") because value could not be parsed into a float. Ignoring."),
+				*CVarPair.Key, *InPreset.GetObject()->GetName());
+		}
+	}
+}
+
+void FMovieGraphCVarManager::AddEvaluatedGraph(const UMovieGraphEvaluatedConfig* InEvaluatedGraph)
+{
+	const FMovieGraphEvaluatedBranchConfig* GlobalsBranchConfig = InEvaluatedGraph->BranchConfigMapping.Find(UMovieGraphNode::GlobalsPinName);
+	if (!GlobalsBranchConfig)
+	{
+		return;
+	}
+
+	// Nodes which set cvars can only be created in the Globals branch, so we only need to look for nodes in that branch
+	const TArray<TObjectPtr<UMovieGraphNode>> GlobalsNodes = GlobalsBranchConfig->GetNodes();
+
+	// Iterate the evaluated graph in reverse. Nodes which come first in the array appear furthest downstream in the
+	// graph, so they should be applied last.
+	for (int32 NodeIndex = GlobalsNodes.Num() - 1; NodeIndex >= 0; --NodeIndex)
+	{
+		const TObjectPtr<UMovieGraphNode>& Node = GlobalsNodes[NodeIndex];
+
+		if (const UMovieGraphSetCVarValueNode* CVarValueNode = Cast<UMovieGraphSetCVarValueNode>(Node))
+		{
+			AddCVar(CVarValueNode->Name, CVarValueNode->Value);
+		}
+		else if (const UMovieGraphApplyCVarPresetNode* ApplyPresetNode = Cast<UMovieGraphApplyCVarPresetNode>(Node))
+		{
+			AddPreset(ApplyPresetNode->ConsoleVariablePreset);
+		}
+	}
+}
+
+void FMovieGraphCVarManager::ApplyAllCVars()
+{
+	PreviousConsoleVariableValues.Reset();
+	PreviousConsoleVariableValues.SetNumZeroed(CVars.Num());
+
+	for (int32 CVarIndex = 0; CVarIndex < CVars.Num(); ++CVarIndex)
+	{
+		const FMovieGraphCVarOverride& CVarEntry = CVars[CVarIndex];
+		const FString TrimmedCvar = CVarEntry.Name.TrimStartAndEnd();
+	
+		if (IConsoleVariable* CVar = IConsoleManager::Get().FindConsoleVariable(*TrimmedCvar))
+		{
+			PreviousConsoleVariableValues[CVarIndex] = CVar->GetFloat();
+			ApplyCVar(CVar, CVarEntry.Value);
+			UE_LOG(LogMovieRenderPipeline, Log, TEXT("Applying CVar \"%s\", PreviousValue: %f NewValue: %f"), *CVarEntry.Name, PreviousConsoleVariableValues[CVarIndex], CVarEntry.Value);
+		}
+		else
+		{
+			UE_LOG(LogMovieRenderPipeline, Warning, TEXT("Failed to apply CVar \"%s\" due to no cvar by that name. Ignoring."), *CVarEntry.Name);
+		}
+	}
+}
+
+void FMovieGraphCVarManager::RevertAllCVars()
+{
+	for (int32 CVarIndex = 0; CVarIndex < CVars.Num(); ++CVarIndex)
+	{
+		const FMovieGraphCVarOverride& CVarEntry = CVars[CVarIndex];
+		const FString TrimmedCvar = CVarEntry.Name.TrimStartAndEnd();
+	
+		if (IConsoleVariable* CVar = IConsoleManager::Get().FindConsoleVariable(*TrimmedCvar))
+		{
+			ApplyCVar(CVar, PreviousConsoleVariableValues[CVarIndex]);
+			UE_LOG(LogMovieRenderPipeline, Log, TEXT("Restoring CVar \"%s\", PreviousValue: %f NewValue: %f"), *CVarEntry.Name, CVarEntry.Value, PreviousConsoleVariableValues[CVarIndex]);
+		}
+		else
+		{
+			UE_LOG(LogMovieRenderPipeline, Warning, TEXT("Failed to revert CVar \"%s\" due to no cvar by that name. Ignoring."), *CVarEntry.Name);
+		}
+	}
+
+	CVars.Reset();
+	PreviousConsoleVariableValues.Reset();
+}
+
+void FMovieGraphCVarManager::ApplyCVar(IConsoleVariable* InCVar, float InValue)
+{
+	check(InCVar);
+
+	// When Set is called on a cvar the value is turned into a string. With very large
+	// floats this is turned into scientific notation. If the cvar is later retrieved as
+	// an integer, the scientific notation doesn't parse into integer correctly. We'll
+	// cast to integer first (to avoid scientific notation) if we know the cvar is an integer.
+	if (InCVar->IsVariableInt())
+	{
+		InCVar->SetWithCurrentPriority(static_cast<int32>(InValue));
+	}
+	else if (InCVar->IsVariableBool())
+	{
+		InCVar->SetWithCurrentPriority(InValue != 0.f ? true : false);
+	}
+	else
+	{
+		InCVar->SetWithCurrentPriority(InValue);
+	}
+}
+
+} // namespace UE::MovieGraph::Private
\ No newline at end of file
diff --git a/Engine/Plugins/MovieScene/MovieRenderPipeline/Source/MovieRenderPipelineCore/Private/Graph/MovieGraphCVarManager.h b/Engine/Plugins/MovieScene/MovieRenderPipeline/Source/MovieRenderPipelineCore/Private/Graph/MovieGraphCVarManager.h
new file mode 100644
--- /dev/null
+++ b/Engine/Plugins/MovieScene/MovieRenderPipeline/Source/MovieRenderPipelineCore/Private/Graph/MovieGraphCVarManager.h
@@ -0,0 +1,75 @@
+﻿// Copyright Epic Games, Inc. All Rights Reserved.
+
+#pragma once
+
+#include "Containers/Array.h"
+#include "Containers/UnrealString.h"
+#include "UObject/ScriptInterface.h"
+
+class IMovieSceneConsoleVariableTrackInterface;
+class UMovieGraphEvaluatedConfig;
+class UMovieScene;
+
+namespace UE::MovieGraph::Private
+{
+	/**
+	 * Responsible for applying individual cvars and cvar presets. CVars can be added to the manager with the Add*() methods,
+	 * then applied via ApplyAllCVars().
+	 */
+	class FMovieGraphCVarManager final
+	{
+	public:
+		FMovieGraphCVarManager() = default;
+
+		/**
+		 * Adds a cvar that should be applied with the given Value. If the cvar with InName already exists, its value
+		 * will be updated.
+		 */
+		void AddCVar(const FString& InName, const float Value);
+
+		/**
+		 * Adds a preset that should be applied. If the preset contains cvars that have already been added via other Add*()
+		 * calls, the values in the preset will override values already added.
+		 */
+		void AddPreset(const TScriptInterface<IMovieSceneConsoleVariableTrackInterface>& InPreset);
+
+		/** For all cvar and cvar preset nodes in InEvaluatedGraph, calls either AddCVar() or AddPreset(). */
+		void AddEvaluatedGraph(const UMovieGraphEvaluatedConfig* InEvaluatedGraph);
+
+		/** Applies all cvars that have been gathered via the Add*() methods. */
+		void ApplyAllCVars();
+
+		/**
+		 * For all cvars that were gathered via the Add*() methods and applied, reverts their values to what they were before
+		 * being applied. After calling this, ApplyAllCVars() is a no-op, and the Add*() methods need to be called again.
+		 */
+		void RevertAllCVars();
+
+	private:
+		/** Sets the given cvar, InCVar, to InValue. */
+		static void ApplyCVar(IConsoleVariable* InCVar, float InValue);
+
+		/** Represents a console variable override. */
+		struct FMovieGraphCVarOverride
+		{
+			FMovieGraphCVarOverride(const FString& InName, const float InValue)
+				: Name(InName)
+				, Value(InValue)
+			{
+			}
+
+			/* The name of the console variable. */
+			FString Name;
+
+			/* The value of the console variable. */
+			float Value;
+		};
+
+	private:
+		/** The merged result of both individual CVars and CVar presets. */
+		TArray<FMovieGraphCVarOverride> CVars;
+
+		/** The values of the gathered cvars before the the manager sets their values. */
+		TArray<float> PreviousConsoleVariableValues;
+	};
+} // namespace UE::MovieGraph::Private
\ No newline at end of file
diff --git a/Engine/Plugins/MovieScene/MovieRenderPipeline/Source/MovieRenderPipelineCore/Private/Graph/MovieGraphPipeline.cpp b/Engine/Plugins/MovieScene/MovieRenderPipeline/Source/MovieRenderPipelineCore/Private/Graph/MovieGraphPipeline.cpp
--- a/Engine/Plugins/MovieScene/MovieRenderPipeline/Source/MovieRenderPipelineCore/Private/Graph/MovieGraphPipeline.cpp
+++ b/Engine/Plugins/MovieScene/MovieRenderPipeline/Source/MovieRenderPipelineCore/Private/Graph/MovieGraphPipeline.cpp
@@ -1,6 +1,8 @@
 // Copyright Epic Games, Inc. All Rights Reserved.
 
 #include "Graph/MovieGraphPipeline.h"
+
+#include "Graph/MovieGraphCVarManager.h"
 #include "Graph/MovieGraphDataTypes.h"
 #include "Graph/MovieGraphOutputMerger.h"
 #include "Graph/Nodes/MovieGraphFileOutputNode.h"
@@ -76,6 +78,7 @@ void UMovieGraphPipeline::Initialize(UMoviePipelineExecutorJob* InJob, const FMo
 	CurrentJob = InJob;
 	CurrentShotIndex = 0;
 	GraphInitializationTime = FDateTime::UtcNow();
+	CVarManager = MakeShared<UE::MovieGraph::Private::FMovieGraphCVarManager>();
 
 	// Now that we've created our various systems, we will start using them. First thing we do is cache data about
 	// the world, job, player viewport, etc, before we make any modifications. These will be restored at the end
@@ -217,7 +220,6 @@ void UMovieGraphPipeline::TickProducingFrames()
 	GraphTimeStepInstance->TickProducingFrames();
 }
 
-
 void UMovieGraphPipeline::TickFinalizeOutputContainers(const bool bInForceFinish)
 {
 	// Tick all containers until they all report that they have finalized.
@@ -344,6 +346,13 @@ void UMovieGraphPipeline::SetupShot(UMoviePipelineExecutorShot* InShot)
 	//	}
 	//}
 
+	const FMovieGraphTimeStepData& TimeStepData = GetTimeStepInstance()->GetCalculatedTimeData();
+	const UMovieGraphEvaluatedConfig* EvaluatedConfig = TimeStepData.EvaluatedConfig;
+
+	// Apply cvars for the shot
+	CVarManager->AddEvaluatedGraph(EvaluatedConfig);
+	CVarManager->ApplyAllCVars();
+
 	// Setup required rendering architecture for all passes in this shot.
 	GraphRendererInstance->SetupRenderingPipelineForShot(InShot);
 }
@@ -357,6 +366,9 @@ void UMovieGraphPipeline::TeardownShot(UMoviePipelineExecutorShot* InShot)
 
 	CurrentShotIndex++;
 
+	// Revert the cvar values that were initially applied for the shot
+	CVarManager->RevertAllCVars();
+
 	// Check to see if this was the last shot in the Pipeline, otherwise on the next
 	// tick the new shot will be initialized and processed.
 	if (CurrentShotIndex >= ActiveShotList.Num())
diff --git a/Engine/Plugins/MovieScene/MovieRenderPipeline/Source/MovieRenderPipelineCore/Private/Graph/Nodes/MovieGraphApplyCVarPresetNode.cpp b/Engine/Plugins/MovieScene/MovieRenderPipeline/Source/MovieRenderPipelineCore/Private/Graph/Nodes/MovieGraphApplyCVarPresetNode.cpp
new file mode 100644
--- /dev/null
+++ b/Engine/Plugins/MovieScene/MovieRenderPipeline/Source/MovieRenderPipelineCore/Private/Graph/Nodes/MovieGraphApplyCVarPresetNode.cpp
@@ -0,0 +1,66 @@
+﻿// Copyright Epic Games, Inc. All Rights Reserved.
+
+#include "Graph/Nodes/MovieGraphApplyCVarPresetNode.h"
+
+#include "Graph/MovieGraphConfig.h"
+#include "Styling/AppStyle.h"
+
+FString UMovieGraphApplyCVarPresetNode::GetNodeInstanceName() const
+{
+	if (ConsoleVariablePreset)
+	{
+		return ConsoleVariablePreset.GetObject()->GetName();
+	}
+
+	return FString();
+}
+
+EMovieGraphBranchRestriction UMovieGraphApplyCVarPresetNode::GetBranchRestriction() const
+{
+	return EMovieGraphBranchRestriction::Globals;
+}
+
+#if WITH_EDITOR
+FText UMovieGraphApplyCVarPresetNode::GetNodeTitle(const bool bGetDescriptive) const
+{
+	static const FText ApplyCVarPresetNodeName = NSLOCTEXT("MovieGraphNodes", "NodeName_ApplyCVarPreset", "Apply CVar Preset");
+	static const FText ApplyCVarPresetNodeDescription = NSLOCTEXT("MovieGraphNodes", "NodeDescription_ApplyCVarPreset", "Apply CVar Preset\n{0}");
+
+	const FString InstanceName = GetNodeInstanceName();
+	if (bGetDescriptive && !InstanceName.IsEmpty())
+	{
+		return FText::Format(ApplyCVarPresetNodeDescription, FText::FromString(InstanceName));
+	}
+	
+	return ApplyCVarPresetNodeName;
+}
+
+FText UMovieGraphApplyCVarPresetNode::GetMenuCategory() const
+{
+	return NSLOCTEXT("MovieGraphNodes", "ApplyCVarPresetGraphNode_Category", "Utility");
+}
+
+FLinearColor UMovieGraphApplyCVarPresetNode::GetNodeTitleColor() const
+{
+	static const FLinearColor ApplyCVarPresetNodeColor = FLinearColor(0.04f, 0.22f, 0.36f);
+	return ApplyCVarPresetNodeColor;
+}
+
+FSlateIcon UMovieGraphApplyCVarPresetNode::GetIconAndTint(FLinearColor& OutColor) const
+{
+	static const FSlateIcon ApplyCVarPresetIcon = FSlateIcon(FAppStyle::GetAppStyleSetName(), "MainFrame.BrowseCVars");
+
+	OutColor = FLinearColor::White;
+	return ApplyCVarPresetIcon;
+}
+
+void UMovieGraphApplyCVarPresetNode::PostEditChangeProperty(FPropertyChangedEvent& PropertyChangedEvent)
+{
+	Super::PostEditChangeProperty(PropertyChangedEvent);
+	
+	if (PropertyChangedEvent.GetPropertyName() == GET_MEMBER_NAME_CHECKED(UMovieGraphApplyCVarPresetNode, ConsoleVariablePreset))
+	{
+		OnNodeChangedDelegate.Broadcast(this);
+	}
+}
+#endif // WITH_EDITOR
\ No newline at end of file
diff --git a/Engine/Plugins/MovieScene/MovieRenderPipeline/Source/MovieRenderPipelineCore/Private/Graph/Nodes/MovieGraphSetCVarValueNode.cpp b/Engine/Plugins/MovieScene/MovieRenderPipeline/Source/MovieRenderPipelineCore/Private/Graph/Nodes/MovieGraphSetCVarValueNode.cpp
--- a/Engine/Plugins/MovieScene/MovieRenderPipeline/Source/MovieRenderPipelineCore/Private/Graph/Nodes/MovieGraphSetCVarValueNode.cpp
+++ b/Engine/Plugins/MovieScene/MovieRenderPipeline/Source/MovieRenderPipelineCore/Private/Graph/Nodes/MovieGraphSetCVarValueNode.cpp
@@ -5,6 +5,16 @@
 #include "Graph/MovieGraphConfig.h"
 #include "Styling/AppStyle.h"
 
+FString UMovieGraphSetCVarValueNode::GetNodeInstanceName() const
+{
+	return Name;
+}
+
+EMovieGraphBranchRestriction UMovieGraphSetCVarValueNode::GetBranchRestriction() const
+{
+	return EMovieGraphBranchRestriction::Globals;
+}
+
 #if WITH_EDITOR
 FText UMovieGraphSetCVarValueNode::GetNodeTitle(const bool bGetDescriptive) const
 {
diff --git a/Engine/Plugins/MovieScene/MovieRenderPipeline/Source/MovieRenderPipelineCore/Public/Graph/MovieGraphPipeline.h b/Engine/Plugins/MovieScene/MovieRenderPipeline/Source/MovieRenderPipelineCore/Public/Graph/MovieGraphPipeline.h
--- a/Engine/Plugins/MovieScene/MovieRenderPipeline/Source/MovieRenderPipelineCore/Public/Graph/MovieGraphPipeline.h
+++ b/Engine/Plugins/MovieScene/MovieRenderPipeline/Source/MovieRenderPipelineCore/Public/Graph/MovieGraphPipeline.h
@@ -18,6 +18,11 @@ class IImageWriteQueue;
 namespace UE::MovieGraph
 {
 	typedef TTuple<TFuture<bool>, UE::MovieGraph::FMovieGraphOutputFutureData> FMovieGraphOutputFuture;
+
+	namespace Private
+	{
+		class FMovieGraphCVarManager;
+	}
 }
 
 UCLASS(BlueprintType)
@@ -160,4 +165,7 @@ protected:
 	* rendering (and is only pushed to once the file is actually on disk), and does not store pixel data.
 	*/
 	TArray<FMovieGraphRenderOutputData> GeneratedOutputData;
+
+	/** Responsible for managing cvars throughout the lifetime of the pipeline. */
+	TSharedPtr<UE::MovieGraph::Private::FMovieGraphCVarManager> CVarManager;
 };
\ No newline at end of file
diff --git a/Engine/Plugins/MovieScene/MovieRenderPipeline/Source/MovieRenderPipelineCore/Public/Graph/Nodes/MovieGraphApplyCVarPresetNode.h b/Engine/Plugins/MovieScene/MovieRenderPipeline/Source/MovieRenderPipelineCore/Public/Graph/Nodes/MovieGraphApplyCVarPresetNode.h
new file mode 100644
--- /dev/null
+++ b/Engine/Plugins/MovieScene/MovieRenderPipeline/Source/MovieRenderPipelineCore/Public/Graph/Nodes/MovieGraphApplyCVarPresetNode.h
@@ -0,0 +1,41 @@
+﻿// Copyright Epic Games, Inc. All Rights Reserved.
+
+#pragma once
+
+#include "Graph/MovieGraphNode.h"
+
+#include "MovieGraphApplyCVarPresetNode.generated.h"
+
+class IMovieSceneConsoleVariableTrackInterface;
+
+/** A node which can apply a console variable preset. */
+UCLASS()
+class MOVIERENDERPIPELINECORE_API UMovieGraphApplyCVarPresetNode : public UMovieGraphSettingNode
+{
+	GENERATED_BODY()
+	
+public:
+	UMovieGraphApplyCVarPresetNode() = default;
+
+	virtual FString GetNodeInstanceName() const override;
+	virtual EMovieGraphBranchRestriction GetBranchRestriction() const override;
+
+#if WITH_EDITOR
+	virtual FText GetNodeTitle(const bool bGetDescriptive = false) const override;
+	virtual FText GetMenuCategory() const override;
+	virtual FLinearColor GetNodeTitleColor() const override;
+	virtual FSlateIcon GetIconAndTint(FLinearColor& OutColor) const override;
+
+	//~ Begin UObject interface
+	virtual void PostEditChangeProperty(FPropertyChangedEvent& PropertyChangedEvent) override;
+	//~ End UObject interface
+#endif
+
+public:
+	UPROPERTY(EditAnywhere, BlueprintReadWrite, Category = Overrides, meta = (InlineEditConditionToggle))
+	uint8 bOverride_ConsoleVariablePreset : 1;
+
+	/** The console variable preset that should be applied. */
+	UPROPERTY(BlueprintReadWrite, EditAnywhere, Category = "Settings", meta = (EditCondition = "bOverride_ConsoleVariablePreset"))
+	TScriptInterface<IMovieSceneConsoleVariableTrackInterface> ConsoleVariablePreset;
+};
\ No newline at end of file
diff --git a/Engine/Plugins/MovieScene/MovieRenderPipeline/Source/MovieRenderPipelineCore/Public/Graph/Nodes/MovieGraphSetCVarValueNode.h b/Engine/Plugins/MovieScene/MovieRenderPipeline/Source/MovieRenderPipelineCore/Public/Graph/Nodes/MovieGraphSetCVarValueNode.h
--- a/Engine/Plugins/MovieScene/MovieRenderPipeline/Source/MovieRenderPipelineCore/Public/Graph/Nodes/MovieGraphSetCVarValueNode.h
+++ b/Engine/Plugins/MovieScene/MovieRenderPipeline/Source/MovieRenderPipelineCore/Public/Graph/Nodes/MovieGraphSetCVarValueNode.h
@@ -14,7 +14,8 @@ class MOVIERENDERPIPELINECORE_API UMovieGraphSetCVarValueNode : public UMovieGra
 public:
 	UMovieGraphSetCVarValueNode() = default;
 
-	virtual FString GetNodeInstanceName() const override { return Name; }
+	virtual FString GetNodeInstanceName() const override;
+	virtual EMovieGraphBranchRestriction GetBranchRestriction() const override;
 
 #if WITH_EDITOR
 	virtual FText GetNodeTitle(const bool bGetDescriptive = false) const override;
@@ -40,5 +41,5 @@ public:
 
 	/** The new value of the CVar. */
 	UPROPERTY(BlueprintReadWrite, EditAnywhere, Category = "Settings", meta = (EditCondition = "bOverride_Value"))
-	FString Value;
+	float Value;
 };
\ No newline at end of file
