MoviePipeline: Allow for disabling nodes in the graph, exposed via "Enable/Disable Nodes" in the graph's context menu. Some nodes can opt-out of being disabled (eg, the Inputs/Outputs nodes).

#rnx

diff --git a/Engine/Plugins/MovieScene/MovieRenderPipeline/Source/MovieRenderPipelineCore/Private/Graph/MovieGraphConfig.cpp b/Engine/Plugins/MovieScene/MovieRenderPipeline/Source/MovieRenderPipelineCore/Private/Graph/MovieGraphConfig.cpp
--- a/Engine/Plugins/MovieScene/MovieRenderPipeline/Source/MovieRenderPipelineCore/Private/Graph/MovieGraphConfig.cpp
+++ b/Engine/Plugins/MovieScene/MovieRenderPipeline/Source/MovieRenderPipelineCore/Private/Graph/MovieGraphConfig.cpp
@@ -1009,9 +1009,9 @@ void UMovieGraphConfig::CreateFlattenedGraph_Recursive(UMovieGraphEvaluatedConfi
 
 	// Check to see if our flattened evaluation graph already has a copy of this node.
 	InEvaluationContext.VisitedNodes.Add(Node);
-	const bool bShouldIncludeNode = Node->IsA<UMovieGraphSettingNode>();
+	const bool bShouldIncludeNode = Node->IsA<UMovieGraphSettingNode>() && !Node->IsDisabled();
 
-	if(bShouldIncludeNode)
+	if (bShouldIncludeNode)
 	{
 		const UMovieGraphSettingNode* NodeAsSetting = CastChecked<UMovieGraphSettingNode>(Node);
 		const FString& NodeInstanceName = NodeAsSetting->GetNodeInstanceName();
diff --git a/Engine/Plugins/MovieScene/MovieRenderPipeline/Source/MovieRenderPipelineCore/Private/Graph/MovieGraphNode.cpp b/Engine/Plugins/MovieScene/MovieRenderPipeline/Source/MovieRenderPipelineCore/Private/Graph/MovieGraphNode.cpp
--- a/Engine/Plugins/MovieScene/MovieRenderPipeline/Source/MovieRenderPipelineCore/Private/Graph/MovieGraphNode.cpp
+++ b/Engine/Plugins/MovieScene/MovieRenderPipeline/Source/MovieRenderPipelineCore/Private/Graph/MovieGraphNode.cpp
@@ -312,6 +312,35 @@ UMovieGraphPin* UMovieGraphNode::GetOutputPin(const FName& Label) const
 	return nullptr;
 }
 
+UMovieGraphPin* UMovieGraphNode::GetFirstConnectedInputPin() const
+{
+	for (const TObjectPtr<UMovieGraphPin>& InputPin : InputPins)
+	{
+		if (InputPin->IsConnected())
+		{
+			return InputPin.Get();
+		}
+	}
+
+	return nullptr;
+}
+
+bool UMovieGraphNode::CanBeDisabled() const
+{
+	// By default, all nodes can be disabled
+	return true;
+}
+
+void UMovieGraphNode::SetDisabled(const bool bNewDisableState)
+{
+	bIsDisabled = bNewDisableState;
+}
+
+bool UMovieGraphNode::IsDisabled() const
+{
+	return bIsDisabled;
+}
+
 #if WITH_EDITOR
 FLinearColor UMovieGraphNode::GetNodeTitleColor() const
 {
@@ -422,6 +451,17 @@ TArray<UMovieGraphPin*> UMovieGraphNode::EvaluatePinsToFollow(FMovieGraphEvaluat
 {
 	TArray<UMovieGraphPin*> PinsToFollow;
 
+	// If the node is disabled, only follow the first connected pin.
+	if (IsDisabled())
+	{
+		if (UMovieGraphPin* GraphPin = GetFirstConnectedInputPin())
+		{
+			PinsToFollow.Add(GraphPin);
+		}
+
+		return PinsToFollow;
+	}
+
 	// By default we provide all Input Pins to this node that are the Branch Type.
 	// You should override this in downstream nodes that need custom logic, such
 	// as branch or switch nodes.
diff --git a/Engine/Plugins/MovieScene/MovieRenderPipeline/Source/MovieRenderPipelineCore/Private/Graph/Nodes/MovieGraphBranchNode.cpp b/Engine/Plugins/MovieScene/MovieRenderPipeline/Source/MovieRenderPipelineCore/Private/Graph/Nodes/MovieGraphBranchNode.cpp
--- a/Engine/Plugins/MovieScene/MovieRenderPipeline/Source/MovieRenderPipelineCore/Private/Graph/Nodes/MovieGraphBranchNode.cpp
+++ b/Engine/Plugins/MovieScene/MovieRenderPipeline/Source/MovieRenderPipelineCore/Private/Graph/Nodes/MovieGraphBranchNode.cpp
@@ -21,6 +21,17 @@ TArray<UMovieGraphPin*> UMovieGraphBranchNode::EvaluatePinsToFollow(FMovieGraphE
 {
 	TArray<UMovieGraphPin*> PinsToFollow;
 
+	// If the node is disabled, follow the first connected pin.
+	if (IsDisabled())
+	{
+		if (UMovieGraphPin* GraphPin = GetFirstConnectedInputPin())
+		{
+			PinsToFollow.Add(GraphPin);
+		}
+		
+		return PinsToFollow;
+	}
+
 	// The branch node has two branches that could be followed, True or False. To figure out which one we're actually going
 	// to follow, we need to evaluate the Conditional pin. 
 	UMovieGraphPin* ConditionalPin = GetInputPin(UE::MovieGraph::BranchNode::Condition);
diff --git a/Engine/Plugins/MovieScene/MovieRenderPipeline/Source/MovieRenderPipelineCore/Private/Graph/Nodes/MovieGraphInputNode.cpp b/Engine/Plugins/MovieScene/MovieRenderPipeline/Source/MovieRenderPipelineCore/Private/Graph/Nodes/MovieGraphInputNode.cpp
--- a/Engine/Plugins/MovieScene/MovieRenderPipeline/Source/MovieRenderPipelineCore/Private/Graph/Nodes/MovieGraphInputNode.cpp
+++ b/Engine/Plugins/MovieScene/MovieRenderPipeline/Source/MovieRenderPipelineCore/Private/Graph/Nodes/MovieGraphInputNode.cpp
@@ -80,6 +80,11 @@ TArray<UMovieGraphPin*> UMovieGraphInputNode::EvaluatePinsToFollow(FMovieGraphEv
 	return PinsToFollow;
 }
 
+bool UMovieGraphInputNode::CanBeDisabled() const
+{
+	return false;
+}
+
 #if WITH_EDITOR
 FText UMovieGraphInputNode::GetNodeTitle(const bool bGetDescriptive) const
 {
diff --git a/Engine/Plugins/MovieScene/MovieRenderPipeline/Source/MovieRenderPipelineCore/Private/Graph/Nodes/MovieGraphOutputNode.cpp b/Engine/Plugins/MovieScene/MovieRenderPipeline/Source/MovieRenderPipelineCore/Private/Graph/Nodes/MovieGraphOutputNode.cpp
--- a/Engine/Plugins/MovieScene/MovieRenderPipeline/Source/MovieRenderPipelineCore/Private/Graph/Nodes/MovieGraphOutputNode.cpp
+++ b/Engine/Plugins/MovieScene/MovieRenderPipeline/Source/MovieRenderPipelineCore/Private/Graph/Nodes/MovieGraphOutputNode.cpp
@@ -36,6 +36,11 @@ TArray<FMovieGraphPinProperties> UMovieGraphOutputNode::GetInputPinProperties()
 	return Properties;
 }
 
+bool UMovieGraphOutputNode::CanBeDisabled() const
+{
+	return false;
+}
+
 #if WITH_EDITOR
 FText UMovieGraphOutputNode::GetNodeTitle(const bool bGetDescriptive) const
 {
diff --git a/Engine/Plugins/MovieScene/MovieRenderPipeline/Source/MovieRenderPipelineCore/Private/Graph/Nodes/MovieGraphSelectNode.cpp b/Engine/Plugins/MovieScene/MovieRenderPipeline/Source/MovieRenderPipelineCore/Private/Graph/Nodes/MovieGraphSelectNode.cpp
--- a/Engine/Plugins/MovieScene/MovieRenderPipeline/Source/MovieRenderPipelineCore/Private/Graph/Nodes/MovieGraphSelectNode.cpp
+++ b/Engine/Plugins/MovieScene/MovieRenderPipeline/Source/MovieRenderPipelineCore/Private/Graph/Nodes/MovieGraphSelectNode.cpp
@@ -36,6 +36,17 @@ TArray<UMovieGraphPin*> UMovieGraphSelectNode::EvaluatePinsToFollow(FMovieGraphE
 {
 	TArray<UMovieGraphPin*> PinsToFollow;
 
+	// If the node is disabled, follow the pin for the first option available.
+	if (IsDisabled())
+	{
+		if (UMovieGraphPin* GraphPin = GetFirstConnectedInputPin())
+		{
+			PinsToFollow.Add(GraphPin);
+		}
+		
+		return PinsToFollow;
+	}
+
 	// The resolved value of the "Selected Option" property. May come from a connection or a value specified on the node.
 	FString ResolvedSelectValue;
 
diff --git a/Engine/Plugins/MovieScene/MovieRenderPipeline/Source/MovieRenderPipelineCore/Private/Graph/Nodes/MovieGraphSubgraphNode.cpp b/Engine/Plugins/MovieScene/MovieRenderPipeline/Source/MovieRenderPipelineCore/Private/Graph/Nodes/MovieGraphSubgraphNode.cpp
--- a/Engine/Plugins/MovieScene/MovieRenderPipeline/Source/MovieRenderPipelineCore/Private/Graph/Nodes/MovieGraphSubgraphNode.cpp
+++ b/Engine/Plugins/MovieScene/MovieRenderPipeline/Source/MovieRenderPipelineCore/Private/Graph/Nodes/MovieGraphSubgraphNode.cpp
@@ -68,6 +68,17 @@ TArray<FMovieGraphPinProperties> UMovieGraphSubgraphNode::GetOutputPinProperties
 TArray<UMovieGraphPin*> UMovieGraphSubgraphNode::EvaluatePinsToFollow(FMovieGraphEvaluationContext& InContext) const
 {
 	TArray<UMovieGraphPin*> PinsToFollow;
+
+	// If the node is disabled, follow the pin for the first option available.
+	if (IsDisabled())
+	{
+		if (UMovieGraphPin* GraphPin = GetFirstConnectedInputPin())
+		{
+			PinsToFollow.Add(GraphPin);
+		}
+		
+		return PinsToFollow;
+	}
 	
 	if (!ensure(InContext.PinBeingFollowed))
 	{
diff --git a/Engine/Plugins/MovieScene/MovieRenderPipeline/Source/MovieRenderPipelineCore/Public/Graph/MovieGraphNode.h b/Engine/Plugins/MovieScene/MovieRenderPipeline/Source/MovieRenderPipelineCore/Public/Graph/MovieGraphNode.h
--- a/Engine/Plugins/MovieScene/MovieRenderPipeline/Source/MovieRenderPipelineCore/Public/Graph/MovieGraphNode.h
+++ b/Engine/Plugins/MovieScene/MovieRenderPipeline/Source/MovieRenderPipelineCore/Public/Graph/MovieGraphNode.h
@@ -172,6 +172,9 @@ public:
 	* By default we will follow any input pin (with Branch type) on the node, but override this in
 	* inherited classes and change that if you need custom logic, such as boolean nodes that want 
 	* to choose one or the other based on the results of a conditional property.
+	*
+	* Note that if custom logic is provided, the case where the node is disabled should be handled
+	* as well.
 	*/
 	virtual TArray<UMovieGraphPin*> EvaluatePinsToFollow(FMovieGraphEvaluationContext& InContext) const;
 
@@ -202,12 +205,24 @@ public:
 	UMovieGraphPin* GetInputPin(const FName& InPinLabel) const;
 	UMovieGraphPin* GetOutputPin(const FName& InPinLabel) const;
 
+	/** Gets the first input pin on the node which has a connection, or nullptr if no pins are connected. */
+	UMovieGraphPin* GetFirstConnectedInputPin() const;
+
 	/** Gets the GUID which uniquely identifies this node. */
 	const FGuid& GetGuid() const { return Guid; }
 	
 	/** Determines which types of branches the node can be created in. */
 	virtual EMovieGraphBranchRestriction GetBranchRestriction() const { return EMovieGraphBranchRestriction::Any; }
 
+	/** Determines if this node can be disabled. */
+	virtual bool CanBeDisabled() const;
+
+	/** Set whether this node is currently disabled. */
+	void SetDisabled(const bool bNewDisableState);
+
+	/** Determines if this node is currently disabled. */
+	bool IsDisabled() const;
+
 #if WITH_EDITOR
 	int32 GetNodePosX() const { return NodePosX; }
 	int32 GetNodePosY() const { return NodePosY; }
@@ -296,6 +311,10 @@ protected:
 	uint8 bIsCommentBubbleVisible : 1;
 #endif
 
+	/** Whether this node is currently disabled in the graph. */
+	UPROPERTY()
+	uint8 bIsDisabled : 1;
+
 	/** A GUID which uniquely identifies this node. */
 	UPROPERTY()
 	FGuid Guid;
diff --git a/Engine/Plugins/MovieScene/MovieRenderPipeline/Source/MovieRenderPipelineCore/Public/Graph/Nodes/MovieGraphInputNode.h b/Engine/Plugins/MovieScene/MovieRenderPipeline/Source/MovieRenderPipelineCore/Public/Graph/Nodes/MovieGraphInputNode.h
--- a/Engine/Plugins/MovieScene/MovieRenderPipeline/Source/MovieRenderPipelineCore/Public/Graph/Nodes/MovieGraphInputNode.h
+++ b/Engine/Plugins/MovieScene/MovieRenderPipeline/Source/MovieRenderPipelineCore/Public/Graph/Nodes/MovieGraphInputNode.h
@@ -18,6 +18,7 @@ public:
 	//~ Begin UMovieGraphNode interface
 	virtual TArray<FMovieGraphPinProperties> GetOutputPinProperties() const override;
 	virtual TArray<UMovieGraphPin*> EvaluatePinsToFollow(FMovieGraphEvaluationContext& InContext) const override;
+	virtual bool CanBeDisabled() const override;
 
 #if WITH_EDITOR
 	virtual FText GetNodeTitle(const bool bGetDescriptive = false) const override;
diff --git a/Engine/Plugins/MovieScene/MovieRenderPipeline/Source/MovieRenderPipelineCore/Public/Graph/Nodes/MovieGraphOutputNode.h b/Engine/Plugins/MovieScene/MovieRenderPipeline/Source/MovieRenderPipelineCore/Public/Graph/Nodes/MovieGraphOutputNode.h
--- a/Engine/Plugins/MovieScene/MovieRenderPipeline/Source/MovieRenderPipelineCore/Public/Graph/Nodes/MovieGraphOutputNode.h
+++ b/Engine/Plugins/MovieScene/MovieRenderPipeline/Source/MovieRenderPipelineCore/Public/Graph/Nodes/MovieGraphOutputNode.h
@@ -16,6 +16,7 @@ public:
 	UMovieGraphOutputNode();
 	
 	virtual TArray<FMovieGraphPinProperties> GetInputPinProperties() const override;
+	virtual bool CanBeDisabled() const override;
 
 #if WITH_EDITOR
 	virtual FText GetNodeTitle(const bool bGetDescriptive = false) const override;
diff --git a/Engine/Plugins/MovieScene/MovieRenderPipeline/Source/MovieRenderPipelineEditor/Private/Graph/MovieEdGraphNode.cpp b/Engine/Plugins/MovieScene/MovieRenderPipeline/Source/MovieRenderPipelineEditor/Private/Graph/MovieEdGraphNode.cpp
--- a/Engine/Plugins/MovieScene/MovieRenderPipeline/Source/MovieRenderPipelineEditor/Private/Graph/MovieEdGraphNode.cpp
+++ b/Engine/Plugins/MovieScene/MovieRenderPipeline/Source/MovieRenderPipelineEditor/Private/Graph/MovieEdGraphNode.cpp
@@ -30,6 +30,8 @@ void UMoviePipelineEdGraphNodeBase::Construct(UMovieGraphNode* InRuntimeNode)
 	NodeComment = InRuntimeNode->GetNodeComment();
 	bCommentBubblePinned = InRuntimeNode->IsCommentBubblePinned();
 	bCommentBubbleVisible = InRuntimeNode->IsCommentBubbleVisible();
+	
+	SetEnabledState(InRuntimeNode->IsDisabled() ? ENodeEnabledState::Disabled : ENodeEnabledState::Enabled);
 }
 
 void UMoviePipelineEdGraphNodeBase::PostTransacted(const FTransactionObjectEvent& TransactionEvent)
@@ -46,6 +48,11 @@ void UMoviePipelineEdGraphNodeBase::PostTransacted(const FTransactionObjectEvent
 	{
 		UpdateCommentBubblePinned();
 	}
+
+	if (ChangedProperties.Contains(TEXT("EnabledState")))
+	{
+		UpdateEnableState();
+	}
 }
 
 FEdGraphPinType UMoviePipelineEdGraphNodeBase::GetPinType(EMovieGraphValueType ValueType, bool bIsBranch)
@@ -142,6 +149,15 @@ void UMoviePipelineEdGraphNodeBase::UpdateCommentBubblePinned() const
 	}
 }
 
+void UMoviePipelineEdGraphNodeBase::UpdateEnableState() const
+{
+	if (RuntimeNode)
+	{
+		RuntimeNode->Modify();
+		RuntimeNode->SetDisabled(GetDesiredEnabledState() == ENodeEnabledState::Disabled);
+	}
+}
+
 void UMoviePipelineEdGraphNode::AllocateDefaultPins()
 {
 	if(RuntimeNode)
@@ -514,6 +530,8 @@ void UMoviePipelineEdGraphNodeBase::GetNodeContextMenuActions(UToolMenu* Menu, U
 		Section.AddMenuEntry(FGenericCommands::Get().Cut);
 		Section.AddMenuEntry(FGenericCommands::Get().Copy);
 		Section.AddMenuEntry(FGenericCommands::Get().Duplicate);
+		Section.AddMenuEntry(FGraphEditorCommands::Get().EnableNodes);
+		Section.AddMenuEntry(FGraphEditorCommands::Get().DisableNodes);
 	}
 
 	{
diff --git a/Engine/Plugins/MovieScene/MovieRenderPipeline/Source/MovieRenderPipelineEditor/Private/Graph/MovieEdGraphNode.h b/Engine/Plugins/MovieScene/MovieRenderPipeline/Source/MovieRenderPipelineEditor/Private/Graph/MovieEdGraphNode.h
--- a/Engine/Plugins/MovieScene/MovieRenderPipeline/Source/MovieRenderPipelineEditor/Private/Graph/MovieEdGraphNode.h
+++ b/Engine/Plugins/MovieScene/MovieRenderPipeline/Source/MovieRenderPipelineEditor/Private/Graph/MovieEdGraphNode.h
@@ -62,6 +62,11 @@ protected:
 	 */
 	void UpdateCommentBubblePinned() const;
 
+	/**
+	 * Update the enable state of the underlying runtime node to match the editor node.
+	 */
+	void UpdateEnableState() const;
+
 protected:
 	/** The runtime node that this editor node represents. */
 	UPROPERTY()
diff --git a/Engine/Plugins/MovieScene/MovieRenderPipeline/Source/MovieRenderPipelineEditor/Private/Widgets/Graph/SMovieGraphConfigPanel.cpp b/Engine/Plugins/MovieScene/MovieRenderPipeline/Source/MovieRenderPipelineEditor/Private/Widgets/Graph/SMovieGraphConfigPanel.cpp
--- a/Engine/Plugins/MovieScene/MovieRenderPipeline/Source/MovieRenderPipelineEditor/Private/Widgets/Graph/SMovieGraphConfigPanel.cpp
+++ b/Engine/Plugins/MovieScene/MovieRenderPipeline/Source/MovieRenderPipelineEditor/Private/Widgets/Graph/SMovieGraphConfigPanel.cpp
@@ -166,6 +166,18 @@ void SMoviePipelineGraphPanel::MakeEditorCommands()
 			FCanExecuteAction::CreateSP(this, &SMoviePipelineGraphPanel::CanDuplicateNodes)
 		);
 
+		GraphEditorCommands->MapAction(FGraphEditorCommands::Get().EnableNodes,
+			FExecuteAction::CreateSP(this, &SMoviePipelineGraphPanel::SetEnabledStateForSelectedNodes, ENodeEnabledState::Enabled),
+			FCanExecuteAction::CreateSP(this, &SMoviePipelineGraphPanel::CanDisableSelectedNodes),
+			FGetActionCheckState::CreateSP(this, &SMoviePipelineGraphPanel::CheckEnabledStateForSelectedNodes, ENodeEnabledState::Enabled)
+		);
+
+		GraphEditorCommands->MapAction(FGraphEditorCommands::Get().DisableNodes,
+			FExecuteAction::CreateSP(this, &SMoviePipelineGraphPanel::SetEnabledStateForSelectedNodes, ENodeEnabledState::Disabled),
+			FCanExecuteAction::CreateSP(this, &SMoviePipelineGraphPanel::CanDisableSelectedNodes),
+			FGetActionCheckState::CreateSP(this, &SMoviePipelineGraphPanel::CheckEnabledStateForSelectedNodes, ENodeEnabledState::Disabled)
+		);
+
 		// Alignment Commands
 		GraphEditorCommands->MapAction(FGraphEditorCommands::Get().AlignNodesTop,
 			FExecuteAction::CreateSP(this, &SMoviePipelineGraphPanel::OnAlignTop)
@@ -448,6 +460,77 @@ bool SMoviePipelineGraphPanel::CanDuplicateNodes() const
 	return CanCopySelectedNodes();
 }
 
+void SMoviePipelineGraphPanel::SetEnabledStateForSelectedNodes(const ENodeEnabledState NewState) const
+{
+	if (!GraphEditorWidget.IsValid())
+	{
+		return;
+	}
+
+	const FScopedTransaction Transaction(LOCTEXT("SetNodeEnabledState", "Set Node Enabled State"));
+	
+	for (UObject* SelectedNode : GraphEditorWidget->GetSelectedNodes())
+	{
+		if (UEdGraphNode* SelectedGraphNode = Cast<UEdGraphNode>(SelectedNode))
+		{
+			SelectedGraphNode->Modify();
+			SelectedGraphNode->SetEnabledState(NewState);
+		}
+	}
+	
+	GraphEditorWidget->NotifyGraphChanged();
+}
+
+ECheckBoxState SMoviePipelineGraphPanel::CheckEnabledStateForSelectedNodes(const ENodeEnabledState EnabledStateToCheck) const
+{
+	if (!GraphEditorWidget.IsValid())
+	{
+		return ECheckBoxState::Unchecked;
+	}
+	
+	ECheckBoxState CheckBoxState = ECheckBoxState::Undetermined;
+	
+	for (UObject* SelectedNode : GraphEditorWidget->GetSelectedNodes())
+	{
+		if (const UEdGraphNode* SelectedGraphNode = Cast<UEdGraphNode>(SelectedNode))
+		{
+			const ENodeEnabledState NodeEnabledState = SelectedGraphNode->GetDesiredEnabledState();
+			const ECheckBoxState NewCheckBoxState = (NodeEnabledState == EnabledStateToCheck) ? ECheckBoxState::Checked : ECheckBoxState::Unchecked;
+
+			// Initialize CheckBoxState to either Checked/Unchecked; this only happens when the state is Undetermined
+			if (CheckBoxState == ECheckBoxState::Undetermined)
+			{
+				CheckBoxState = NewCheckBoxState;
+				continue;
+			}
+
+			// If the new/old checkbox states don't match, revert to an undetermined state
+			if (NewCheckBoxState != CheckBoxState)
+			{
+				CheckBoxState = ECheckBoxState::Undetermined;
+				break;
+			}
+			
+			CheckBoxState = NewCheckBoxState;
+		}
+	}
+
+	return CheckBoxState;
+}
+
+bool SMoviePipelineGraphPanel::CanDisableSelectedNodes() const
+{
+	for (const UMovieGraphNode* SelectedModelNode : GetSelectedModelNodes())
+	{
+		if (!SelectedModelNode->CanBeDisabled())
+		{
+			return false;
+		}
+	}
+
+	return true;
+}
+
 void SMoviePipelineGraphPanel::OnAlignTop()
 {
 	if (GraphEditorWidget.IsValid())
@@ -805,6 +888,24 @@ void SMoviePipelineGraphPanel::OnSelectionChanged(const TArray<UMoviePipelineExe
 	NumSelectedJobs = InSelectedJobs.Num();
 }
 
+TArray<UMovieGraphNode*> SMoviePipelineGraphPanel::GetSelectedModelNodes() const
+{
+	TArray<UMovieGraphNode*> SelectedModelNodes;
+	
+	for (UObject* SelectedNode : GraphEditorWidget->GetSelectedNodes())
+	{
+		if (const UMoviePipelineEdGraphNodeBase* SelectedGraphNode = Cast<UMoviePipelineEdGraphNodeBase>(SelectedNode))
+		{
+			if (UMovieGraphNode* SelectedModelNode = SelectedGraphNode->GetRuntimeNode())
+			{
+				SelectedModelNodes.Add(SelectedModelNode);
+			}
+		}
+	}
+
+	return SelectedModelNodes;
+}
+
 TSharedRef<SWidget> SMoviePipelineGraphPanel::OnGenerateSavedQueuesMenu()
 {
 	FMenuBuilder MenuBuilder(true, nullptr);
diff --git a/Engine/Plugins/MovieScene/MovieRenderPipeline/Source/MovieRenderPipelineEditor/Private/Widgets/Graph/SMovieGraphConfigPanel.h b/Engine/Plugins/MovieScene/MovieRenderPipeline/Source/MovieRenderPipelineEditor/Private/Widgets/Graph/SMovieGraphConfigPanel.h
--- a/Engine/Plugins/MovieScene/MovieRenderPipeline/Source/MovieRenderPipelineEditor/Private/Widgets/Graph/SMovieGraphConfigPanel.h
+++ b/Engine/Plugins/MovieScene/MovieRenderPipeline/Source/MovieRenderPipelineEditor/Private/Widgets/Graph/SMovieGraphConfigPanel.h
@@ -16,6 +16,8 @@ class IDetailsView;
 struct FAssetData;
 class FUICommandList;
 class UMovieGraphNode;
+enum class ECheckBoxState : uint8;
+enum class ENodeEnabledState : uint8;
 
 /**
  * Outermost widget that is used for adding and removing jobs from the Movie Pipeline Queue Subsystem.
@@ -97,6 +99,19 @@ private:
 	bool CanPasteNodes() const;
 	void DuplicateNodes();
 	bool CanDuplicateNodes() const;
+	
+	/** Set the enabled state for the currently selected nodes. */
+	void SetEnabledStateForSelectedNodes(const ENodeEnabledState NewState) const;
+
+	/**
+	 * Attempt to match the given enable state for currently-selected nodes. If StateToCheck matches all selected
+	 * nodes, returns ECheckBoxState::Checked or ECheckBoxState::Unchecked; if there is a mismatch, returns
+	 * ECheckBoxState::Undetermined.
+	 */
+	ECheckBoxState CheckEnabledStateForSelectedNodes(const ENodeEnabledState EnabledStateToCheck) const;
+
+	/** Determines if all of the selected nodes can have their enable/disable state toggled. */
+	bool CanDisableSelectedNodes() const;
 
 	void OnAlignTop();
 	void OnAlignMiddle();
