MoviePipeline: (Graph) Stubbed in support for the in-game debug view to see what the renderer is doing.

#rnx

diff --git a/Engine/Plugins/MovieScene/MovieRenderPipeline/Source/MovieRenderPipelineCore/Private/Graph/MovieGraphConfig.cpp b/Engine/Plugins/MovieScene/MovieRenderPipeline/Source/MovieRenderPipelineCore/Private/Graph/MovieGraphConfig.cpp
--- a/Engine/Plugins/MovieScene/MovieRenderPipeline/Source/MovieRenderPipelineCore/Private/Graph/MovieGraphConfig.cpp
+++ b/Engine/Plugins/MovieScene/MovieRenderPipeline/Source/MovieRenderPipelineCore/Private/Graph/MovieGraphConfig.cpp
@@ -1179,14 +1179,10 @@ UMovieGraphEvaluatedConfig* UMovieGraphConfig::CreateFlattenedGraph(const FMovie
 		}		
 	}
 
-	UE_LOG(LogTemp, Warning, TEXT("Traversed Graph:"));
 	for (const TPair<FName, FMovieGraphEvaluatedBranchConfig>& Pair : NewContext->BranchConfigMapping)
 	{
-		UE_LOG(LogTemp, Warning, TEXT("\t Branch: %s"), *Pair.Key.ToString());
-
 		for (UMovieGraphNode* Node : Pair.Value.GetNodes())
 		{
-			UE_LOG(LogTemp, Warning, TEXT("\t\t %s Class:"), *Node->GetClass()->GetName());
 			for (TFieldIterator<FProperty> PropertyIterator(Node->GetClass()); PropertyIterator; ++PropertyIterator)
 			{
 				FProperty* CheckProperty = *PropertyIterator;
@@ -1195,7 +1191,6 @@ UMovieGraphEvaluatedConfig* UMovieGraphConfig::CreateFlattenedGraph(const FMovie
 				{
 					FString ExportText;
 					CheckProperty->ExportText_InContainer(0, ExportText, Node, Node, Node, 0);
-					UE_LOG(LogTemp, Warning, TEXT("\t\t\t %s : %s"), *CheckProperty->GetName(), *ExportText);
 				}
 			}
 		}
diff --git a/Engine/Plugins/MovieScene/MovieRenderPipeline/Source/MovieRenderPipelineCore/Private/Graph/MovieGraphDefaultRenderer.cpp b/Engine/Plugins/MovieScene/MovieRenderPipeline/Source/MovieRenderPipelineCore/Private/Graph/MovieGraphDefaultRenderer.cpp
--- a/Engine/Plugins/MovieScene/MovieRenderPipeline/Source/MovieRenderPipelineCore/Private/Graph/MovieGraphDefaultRenderer.cpp
+++ b/Engine/Plugins/MovieScene/MovieRenderPipeline/Source/MovieRenderPipelineCore/Private/Graph/MovieGraphDefaultRenderer.cpp
@@ -69,7 +69,7 @@ void UMovieGraphDefaultRenderer::SetupRenderingPipelineForShot(UMoviePipelineExe
 		}
 	}
 
-	UE_LOG(LogMovieRenderPipeline, Warning, TEXT("Found: %d Render Passes:"), OutputPasses.Num());
+	//UE_LOG(LogMovieRenderPipeline, Warning, TEXT("Found: %d Render Passes:"), OutputPasses.Num());
 	int32 TotalLayerCount = 0;
 	for (const FMovieGraphPass& Pass : OutputPasses)
 	{
@@ -79,14 +79,14 @@ void UMovieGraphDefaultRenderer::SetupRenderingPipelineForShot(UMoviePipelineExe
 		
 		FMovieGraphRenderPassSetupData SetupData;
 		SetupData.Renderer = this;
-		UE_LOG(LogMovieRenderPipeline, Warning, TEXT("\tRenderer Class: %s"), *Pass.ClassType->GetName());
+		//UE_LOG(LogMovieRenderPipeline, Warning, TEXT("\tRenderer Class: %s"), *Pass.ClassType->GetName());
 		for (const TTuple<FName, TWeakObjectPtr<UMovieGraphRenderPassNode>>& BranchRenderer : Pass.BranchRenderers)
 		{
 			FMovieGraphRenderPassLayerData& LayerData = SetupData.Layers.AddDefaulted_GetRef();
 			LayerData.BranchName = BranchRenderer.Key;
 			LayerData.RenderPassNode = BranchRenderer.Value;
 
-			UE_LOG(LogMovieRenderPipeline, Warning, TEXT("\t\tBranch Name: %s"), *BranchRenderer.Key.ToString());
+			// UE_LOG(LogMovieRenderPipeline, Warning, TEXT("\t\tBranch Name: %s"), *LayerBranchName.ToString());
 		}
 
 		UMovieGraphRenderPassNode* RenderPassCDO = Pass.ClassType->GetDefaultObject<UMovieGraphRenderPassNode>();
@@ -186,6 +186,9 @@ void UMovieGraphDefaultRenderer::Render(const FMovieGraphTimeStepData& InTimeSte
 		}
 	}
 
+	// Hide the progress widget before we render anything. This allows widget captures to not include the progress bar.
+	GetOwningGraph()->SetPreviewWidgetVisible(false);
+
 	if (InTimeStepData.bIsFirstTemporalSampleForFrame)
 	{
 		// If this is the first sample for this output frame, then we need to 
@@ -218,6 +221,9 @@ void UMovieGraphDefaultRenderer::Render(const FMovieGraphTimeStepData& InTimeSte
 		UE::MovieGraph::FMovieGraphOutputMergerFrame& OutputFrame = GetOwningGraph()->GetOutputMerger()->GetOutputFrame_GameThread(InTimeStepData.OutputFrameNumber);
 		RenderPass->Render(OutputFrame.TraversalContext, InTimeStepData);
 	}
+
+	// Re-enable the progress widget so when the player viewport is drawn to the preview window, it shows.
+	GetOwningGraph()->SetPreviewWidgetVisible(true);
 }
 
 void UMovieGraphDefaultRenderer::AddOutstandingRenderTask_AnyThread(UE::Tasks::FTask InTask)
@@ -300,6 +306,20 @@ UE::MovieGraph::FRenderTimeStatistics* UMovieGraphDefaultRenderer::GetRenderTime
 	return &RenderTimeStatistics.FindOrAdd(InFrameNumber);
 }
 
+UTexture* UMovieGraphDefaultRenderer::GetPreviewTexture() const
+{
+	TArray<TObjectPtr<UTextureRenderTarget2D>> PooledTargets;
+	PooledViewRenderTargets.GenerateValueArray(PooledTargets);
+
+	if (PooledTargets.Num() > 0)
+	{
+		return PooledTargets[0];
+	}
+
+	return nullptr;
+}
+
+
 namespace UE::MovieGraph::DefaultRenderer
 {
 FSurfaceAccumulatorPool::FInstancePtr FSurfaceAccumulatorPool::BlockAndGetAccumulator_GameThread(int32 InFrameNumber, const FMovieGraphRenderDataIdentifier& InPassIdentifier)
diff --git a/Engine/Plugins/MovieScene/MovieRenderPipeline/Source/MovieRenderPipelineCore/Private/Graph/MovieGraphLinearTimeStep.cpp b/Engine/Plugins/MovieScene/MovieRenderPipeline/Source/MovieRenderPipelineCore/Private/Graph/MovieGraphLinearTimeStep.cpp
--- a/Engine/Plugins/MovieScene/MovieRenderPipeline/Source/MovieRenderPipelineCore/Private/Graph/MovieGraphLinearTimeStep.cpp
+++ b/Engine/Plugins/MovieScene/MovieRenderPipeline/Source/MovieRenderPipelineCore/Private/Graph/MovieGraphLinearTimeStep.cpp
@@ -76,6 +76,7 @@ void UMovieGraphLinearTimeStep::TickProducingFrames()
 
 		// Ensure we've set it in the CurrentTimeStepData so things can fetch from it below.
 		CurrentTimeStepData.EvaluatedConfig = TObjectPtr<UMovieGraphEvaluatedConfig>(CurrentFrameData.EvaluatedConfig.Get());
+		CurrentTimeStepData.bExpandShotForTemporalSubSample = CurrentFrameData.TemporalSampleCount > 1;
 
 
 		// Sets up the render state, etc.
@@ -233,9 +234,9 @@ void UMovieGraphLinearTimeStep::TickProducingFrames()
 		
 		CurrentTimeStepData.EvaluatedConfig = TObjectPtr<UMovieGraphEvaluatedConfig>(CurrentFrameData.EvaluatedConfig.Get());
 
-		UE_LOG(LogTemp, Warning, TEXT("F# %d bFirst: %d bLast: %d bReqAc: %d"),
-			CurrentTimeStepData.OutputFrameNumber, CurrentTimeStepData.bIsFirstTemporalSampleForFrame,
-			CurrentTimeStepData.bIsLastTemporalSampleForFrame, CurrentTimeStepData.bRequiresAccumulator);
+		//UE_LOG(LogTemp, Warning, TEXT("F# %d bFirst: %d bLast: %d bReqAc: %d"),
+		//	CurrentTimeStepData.OutputFrameNumber, CurrentTimeStepData.bIsFirstTemporalSampleForFrame,
+		//	CurrentTimeStepData.bIsLastTemporalSampleForFrame, CurrentTimeStepData.bRequiresAccumulator);
 
 		// Set our time step for the next frame. We use the undilated delta time for the Custom Timestep as the engine will
 		// apply the time dilation to the world tick for us, so we don't want to double up time dilation.
@@ -245,7 +246,7 @@ void UMovieGraphLinearTimeStep::TickProducingFrames()
 
 		// ToDo: This should be converted back to an 'effective' frame number (source frame in external data asset)
 		// so you can line up profiling with the acutal content on screen.
-		TRACE_BOOKMARK(TEXT("MoviePipeline - Rendering Frame %d [TS: %d]"), CurrentTimeStepData.OutputFrameNumber, CurrentFrameData.TemporalSampleIndex);
+		TRACE_BOOKMARK(TEXT("MRQ Frame %d [TS: %d]"), CurrentTimeStepData.OutputFrameNumber, CurrentFrameData.TemporalSampleIndex);
 
 		// Increment various post-frame counters to set them up for the next frame. This is okay
 		// because the rest of the Movie Graph Pipeline system is based on CurrentTimeStepData which accurately
diff --git a/Engine/Plugins/MovieScene/MovieRenderPipeline/Source/MovieRenderPipelineCore/Private/Graph/MovieGraphPipeline.cpp b/Engine/Plugins/MovieScene/MovieRenderPipeline/Source/MovieRenderPipelineCore/Private/Graph/MovieGraphPipeline.cpp
--- a/Engine/Plugins/MovieScene/MovieRenderPipeline/Source/MovieRenderPipelineCore/Private/Graph/MovieGraphPipeline.cpp
+++ b/Engine/Plugins/MovieScene/MovieRenderPipeline/Source/MovieRenderPipelineCore/Private/Graph/MovieGraphPipeline.cpp
@@ -6,6 +6,8 @@
 #include "Graph/MovieGraphDataTypes.h"
 #include "Graph/MovieGraphOutputMerger.h"
 #include "Graph/Nodes/MovieGraphFileOutputNode.h"
+#include "Graph/Nodes/MovieGraphOutputSettingNode.h"
+#include "Graph/MovieGraphBlueprintLibrary.h"
 #include "MovieRenderPipelineCoreModule.h"
 #include "Misc/CoreDelegates.h"
 #include "MoviePipelineQueue.h"
@@ -17,6 +19,8 @@
 // Temp
 #include "LevelSequence.h"
 
+FString UMovieGraphPipeline::DefaultPreviewWidgetAsset = TEXT("/MovieRenderPipeline/Blueprints/UI_MovieGraphPipelineScreenOverlay.UI_MovieGraphPipelineScreenOverlay_C");
+
 UMovieGraphPipeline::UMovieGraphPipeline()
 	: CurrentShotIndex(-1)
 	, bIsTransitioningState(false)
@@ -85,8 +89,8 @@ void UMovieGraphPipeline::Initialize(UMoviePipelineExecutorJob* InJob, const FMo
 	// of the render.
 	GraphDataSourceInstance->CacheDataPreJob(InitConfig);
 
-	// Construct a debug UI and bind it to this instance.
-	// LoadDebugWidget();
+	// Construct the viewport preview UI and bind it to this instance.
+	LoadPreviewWidget();
 	// SetupAudioRendering();
 
 	// Update our list of shots from our data source, and then
@@ -125,6 +129,41 @@ void UMovieGraphPipeline::Initialize(UMoviePipelineExecutorJob* InJob, const FMo
 	}
 }
 
+void UMovieGraphPipeline::LoadPreviewWidget()
+{
+	// ToDo: Allow overriding this widget so that users can style it to match their project.
+	if (PreviewWidgetClassToUse.Get() == nullptr)
+	{
+		PreviewWidgetClassToUse = LoadClass<UMovieGraphRenderPreviewWidget>(nullptr, *DefaultPreviewWidgetAsset, nullptr, LOAD_None, nullptr);
+	}
+
+	if (PreviewWidgetClassToUse.Get() != nullptr)
+	{
+		PreviewWidget = CreateWidget<UMovieGraphRenderPreviewWidget>(GetWorld(), PreviewWidgetClassToUse.Get());
+		if (PreviewWidget)
+		{
+			PreviewWidget->OnInitializedForPipeline(this);
+			PreviewWidget->AddToViewport();
+		}
+		else
+		{
+			UE_LOG(LogMovieRenderPipeline, Warning, TEXT("Failed to create Preview Screen UMG Widget. No in-game overlay available."));
+		}
+	}
+	else
+	{
+		UE_LOG(LogMovieRenderPipeline, Warning, TEXT("Failed to find Preview Screen UMG Widget class. No in-game overlay available."));
+	}
+}
+
+void UMovieGraphPipeline::SetPreviewWidgetVisibleImpl(bool bInIsVisible)
+{
+	if (PreviewWidget)
+	{
+		PreviewWidget->SetVisibility(bInIsVisible ? ESlateVisibility::Visible : ESlateVisibility::Collapsed);
+	}
+}
+
 void UMovieGraphPipeline::BuildShotListFromDataSource()
 {
 	// Synchronize our shot list with our target data source. New shots will be added and outdated shots removed.
@@ -587,11 +626,11 @@ void UMovieGraphPipeline::TransitionToState(const EMovieRenderPipelineState InNe
 			FCoreDelegates::OnBeginFrame.RemoveAll(this);
 			FCoreDelegates::OnEndFrame.RemoveAll(this);
 
-			//if (DebugWidget)
-			//{
-			//	DebugWidget->RemoveFromParent();
-			//	DebugWidget = nullptr;
-			//}
+			if (PreviewWidget)
+			{
+				PreviewWidget->RemoveFromParent();
+				PreviewWidget = nullptr;
+			}
 
 			//TArray<UMoviePipelineOutputBase*> ContainerSettings = GetPipelinePrimaryConfig()->GetOutputContainers();
 			//Algo::SortBy(ContainerSettings, [](const UMoviePipelineOutputBase* Setting) { return Setting->GetPriority(); });
@@ -737,13 +776,13 @@ void UMovieGraphPipeline::ProcessOutstandingFinishedFrames()
 		// Now that we've looped through the above, we have the total list of which output formats are being used by the graph for
 		// all of the render layers given. We also have a list of which identifiers should go into each one. So we can loop through
 		// the CDO instances and pass the data to them.
-		UE_LOG(LogMovieRenderPipeline, Warning, TEXT("File Outputs:"));
+		//UE_LOG(LogMovieRenderPipeline, Warning, TEXT("File Outputs:"));
 		for (const TPair<UMovieGraphFileOutputNode*, TSet<FMovieGraphRenderDataIdentifier>>& Pair : MaskData)
 		{
-			UE_LOG(LogMovieRenderPipeline, Warning, TEXT("\tNode: %s"), *Pair.Key->GetClass()->GetName());
+			//UE_LOG(LogMovieRenderPipeline, Warning, TEXT("\tNode: %s"), *Pair.Key->GetClass()->GetName());
 			for (const FMovieGraphRenderDataIdentifier& ID : Pair.Value)
 			{
-				UE_LOG(LogMovieRenderPipeline, Warning, TEXT("\t\tBranch: %s:"), *ID.RootBranchName.ToString());
+				//UE_LOG(LogMovieRenderPipeline, Warning, TEXT("\t\tBranch: %s:"), *ID.RootBranchName.ToString());
 			}
 
 			Pair.Key->OnReceiveImageData(this, &OutputFrame, Pair.Value);
diff --git a/Engine/Plugins/MovieScene/MovieRenderPipeline/Source/MovieRenderPipelineCore/Private/MovieRenderPipelineCoreModule.cpp b/Engine/Plugins/MovieScene/MovieRenderPipeline/Source/MovieRenderPipelineCore/Private/MovieRenderPipelineCoreModule.cpp
--- a/Engine/Plugins/MovieScene/MovieRenderPipeline/Source/MovieRenderPipelineCore/Private/MovieRenderPipelineCoreModule.cpp
+++ b/Engine/Plugins/MovieScene/MovieRenderPipeline/Source/MovieRenderPipelineCore/Private/MovieRenderPipelineCoreModule.cpp
@@ -8,6 +8,7 @@
 #include "Engine/World.h"
 #include "TimerManager.h"
 #include "MoviePipeline.h"
+#include "Graph/MovieGraphPipeline.h"
 
 FName IMoviePipelineBurnInExtension::ModularFeatureName = "ModularFeature_MoviePipelineBurnInExt";
 
@@ -18,6 +19,7 @@ void FMovieRenderPipelineCoreModule::StartupModule()
 	{
 		TArray<FString> Assets;
 		Assets.Add(UMoviePipeline::DefaultDebugWidgetAsset);
+		Assets.Add(UMovieGraphPipeline::DefaultPreviewWidgetAsset);
 
 		for (const FString& Asset : Assets)
 		{
diff --git a/Engine/Plugins/MovieScene/MovieRenderPipeline/Source/MovieRenderPipelineCore/Public/Graph/MovieGraphDataTypes.h b/Engine/Plugins/MovieScene/MovieRenderPipeline/Source/MovieRenderPipelineCore/Public/Graph/MovieGraphDataTypes.h
--- a/Engine/Plugins/MovieScene/MovieRenderPipeline/Source/MovieRenderPipelineCore/Public/Graph/MovieGraphDataTypes.h
+++ b/Engine/Plugins/MovieScene/MovieRenderPipeline/Source/MovieRenderPipelineCore/Public/Graph/MovieGraphDataTypes.h
@@ -119,6 +119,10 @@ class MOVIERENDERPIPELINECORE_API UMovieGraphRendererBase : public UObject
 {
 	GENERATED_BODY()
 public:
+	/** Get an texture that can be used to see what the Renderer is producing for the in-game overlay. */
+	UFUNCTION(BlueprintPure, Category = "Movie Graph")
+	virtual UTexture* GetPreviewTexture() const { return nullptr; }
+
 	virtual void Render(const FMovieGraphTimeStepData& InTimeData) {}
 	virtual void SetupRenderingPipelineForShot(UMoviePipelineExecutorShot* InShot) {}
 	virtual void TeardownRenderingPipelineForShot(UMoviePipelineExecutorShot* InShot) {}
diff --git a/Engine/Plugins/MovieScene/MovieRenderPipeline/Source/MovieRenderPipelineCore/Public/Graph/MovieGraphDefaultRenderer.h b/Engine/Plugins/MovieScene/MovieRenderPipeline/Source/MovieRenderPipelineCore/Public/Graph/MovieGraphDefaultRenderer.h
--- a/Engine/Plugins/MovieScene/MovieRenderPipeline/Source/MovieRenderPipelineCore/Public/Graph/MovieGraphDefaultRenderer.h
+++ b/Engine/Plugins/MovieScene/MovieRenderPipeline/Source/MovieRenderPipelineCore/Public/Graph/MovieGraphDefaultRenderer.h
@@ -140,6 +140,7 @@ class MOVIERENDERPIPELINECORE_API UMovieGraphDefaultRenderer : public UMovieGrap
 
 public:
 	// UMovieGraphRendererBase Interface
+	virtual UTexture* GetPreviewTexture() const override;
 	virtual void Render(const FMovieGraphTimeStepData& InTimeData) override;
 	virtual void SetupRenderingPipelineForShot(UMoviePipelineExecutorShot* InShot) override;
 	virtual void TeardownRenderingPipelineForShot(UMoviePipelineExecutorShot* InShot) override;
diff --git a/Engine/Plugins/MovieScene/MovieRenderPipeline/Source/MovieRenderPipelineCore/Public/Graph/MovieGraphPipeline.h b/Engine/Plugins/MovieScene/MovieRenderPipeline/Source/MovieRenderPipelineCore/Public/Graph/MovieGraphPipeline.h
--- a/Engine/Plugins/MovieScene/MovieRenderPipeline/Source/MovieRenderPipelineCore/Public/Graph/MovieGraphPipeline.h
+++ b/Engine/Plugins/MovieScene/MovieRenderPipeline/Source/MovieRenderPipelineCore/Public/Graph/MovieGraphPipeline.h
@@ -71,6 +71,7 @@ public:
 	/** Used occasionally to cross-reference other components. Don't call this unless you know what you're doing. */
 	UMovieGraphTimeStepBase* GetTimeStepInstance() const { return GraphTimeStepInstance; }
 	/** Used occasionally to cross-reference other components. Don't call this unless you know what you're doing. */
+	UFUNCTION(BlueprintCallable, Category = "Movie Render Pipeline")
 	UMovieGraphRendererBase* GetRendererInstance() const { return GraphRendererInstance; }
 	/** Used occasionally to cross-reference other components. Don't call this unless you know what you're doing. */
 	UMovieGraphDataSourceBase* GetDataSourceInstance() const { return GraphDataSourceInstance; }
@@ -78,6 +79,8 @@ public:
 	TSharedPtr<UE::MovieGraph::IMovieGraphOutputMerger> GetOutputMerger() const { return OutputMerger; }
 	/** Writing images to disk is an async process. When you start writing, declare a future with the filename you will eventually write to, and complete the future once it is on disk. */
 	void AddOutputFuture(TFuture<bool>&& InOutputFuture, const UE::MovieGraph::FMovieGraphOutputFutureData& InData);
+	/** Used by the Renderer Instance to disable the preview widget before rendering so it isn't baked into the UI Renderer. */
+	void SetPreviewWidgetVisible(bool bInIsVisible) { SetPreviewWidgetVisibleImpl(bInIsVisible); }
 
 protected:
 	UFUNCTION(BlueprintCallable, Category = "Movie Render Pipeline")
@@ -97,6 +100,8 @@ protected:
 	virtual void TransitionToState(const EMovieRenderPipelineState InNewState);
 	virtual const TSet<TObjectPtr<UMovieGraphFileOutputNode>> GetOutputNodesUsed() const;
 	virtual void BeginFinalize();
+	virtual void LoadPreviewWidget();
+	virtual void SetPreviewWidgetVisibleImpl(bool bInIsVisible);
 
 	// UMoviePipelineBase Interface
 	virtual void RequestShutdownImpl(bool bIsError) override;
@@ -123,6 +128,12 @@ protected:
 	UPROPERTY(Transient)
 	TArray<TObjectPtr<UMoviePipelineExecutorShot>> ActiveShotList;
 
+	UPROPERTY(Transient)
+	TSubclassOf<UMovieGraphRenderPreviewWidget> PreviewWidgetClassToUse;
+
+	UPROPERTY(Transient)
+	TObjectPtr<UMovieGraphRenderPreviewWidget> PreviewWidget;
+
 	/**
 	* An array of Node CDOs that we sent data through to write data to disk.
 	* This is the list of output nodes that will have OnAllFramesSubmitted/IsFinishedWritingToDisk
@@ -168,4 +179,7 @@ protected:
 
 	/** Responsible for managing cvars throughout the lifetime of the pipeline. */
 	TSharedPtr<UE::MovieGraph::Private::FMovieGraphCVarManager> CVarManager;
+	
+public:
+	static FString DefaultPreviewWidgetAsset;
 };
\ No newline at end of file
diff --git a/Engine/Plugins/MovieScene/MovieRenderPipeline/Source/MovieRenderPipelineCore/Public/Graph/MovieGraphTimeStepData.h b/Engine/Plugins/MovieScene/MovieRenderPipeline/Source/MovieRenderPipelineCore/Public/Graph/MovieGraphTimeStepData.h
--- a/Engine/Plugins/MovieScene/MovieRenderPipeline/Source/MovieRenderPipelineCore/Public/Graph/MovieGraphTimeStepData.h
+++ b/Engine/Plugins/MovieScene/MovieRenderPipeline/Source/MovieRenderPipelineCore/Public/Graph/MovieGraphTimeStepData.h
@@ -25,6 +25,7 @@ public:
 		, MotionBlurFraction(0.f)
 		, bIsFirstTemporalSampleForFrame(false)
 		, bIsLastTemporalSampleForFrame(false)
+		, bExpandShotForTemporalSubSample(false)
 		, bRequiresAccumulator(false)
 	{}
 
@@ -58,6 +59,14 @@ public:
 	UPROPERTY(BlueprintReadWrite, EditAnywhere, Category = "Movie Graph")
 	bool bIsLastTemporalSampleForFrame;
 
+	/**
+	* If this time struct is being used when setting up a shot, defines if the shot 
+	* should be expanded by 1 frame to account for temporal sub-sampling going 
+	* out of bounds on the first frame.
+	*/
+	UPROPERTY(BlueprintReadWrite, EditAnywhere, Category = "Movie Graph")
+	bool bExpandShotForTemporalSubSample;
+
 	/**
 	* Should be set to true for every sample if there is more than one temporal sample
 	* making up this render. This will cause the renderer to allocate accumulators
diff --git a/Engine/Plugins/MovieScene/MovieRenderPipeline/Source/MovieRenderPipelineCore/Public/MovieRenderDebugWidget.h b/Engine/Plugins/MovieScene/MovieRenderPipeline/Source/MovieRenderPipelineCore/Public/MovieRenderDebugWidget.h
--- a/Engine/Plugins/MovieScene/MovieRenderPipeline/Source/MovieRenderPipelineCore/Public/MovieRenderDebugWidget.h
+++ b/Engine/Plugins/MovieScene/MovieRenderPipeline/Source/MovieRenderPipelineCore/Public/MovieRenderDebugWidget.h
@@ -18,4 +18,18 @@ class UMovieRenderDebugWidget : public UUserWidget
 public:
 	UFUNCTION(BlueprintImplementableEvent)
 	void OnInitializedForPipeline(UMoviePipeline* ForPipeline);
+};
+
+/**
+* C++ Base Class for the preview widget that is drawn onto the game viewport
+* (but not burned into the output files) that allow us to easily visualize
+* the current state of the pipeline. Used for Graph Based pipelines.
+*/
+UCLASS(Blueprintable, Abstract)
+class UMovieGraphRenderPreviewWidget : public UUserWidget
+{
+	GENERATED_BODY()
+public:
+	UFUNCTION(BlueprintImplementableEvent)
+	void OnInitializedForPipeline(UMovieGraphPipeline* InPipeline);
 };
\ No newline at end of file
