MoviePipeline: Added a "Remove Render Setting" node. This node takes a type; all nodes that are upstream of this node that match the type will be removed from the graph's evaluated config.

#rnx

diff --git a/Engine/Plugins/MovieScene/MovieRenderPipeline/Source/MovieRenderPipelineCore/Private/Graph/MovieGraphConfig.cpp b/Engine/Plugins/MovieScene/MovieRenderPipeline/Source/MovieRenderPipelineCore/Private/Graph/MovieGraphConfig.cpp
--- a/Engine/Plugins/MovieScene/MovieRenderPipeline/Source/MovieRenderPipelineCore/Private/Graph/MovieGraphConfig.cpp
+++ b/Engine/Plugins/MovieScene/MovieRenderPipeline/Source/MovieRenderPipelineCore/Private/Graph/MovieGraphConfig.cpp
@@ -6,6 +6,7 @@
 #include "Graph/MovieGraphEdge.h"
 #include "Graph/Nodes/MovieGraphInputNode.h"
 #include "Graph/Nodes/MovieGraphOutputNode.h"
+#include "Graph/Nodes/MovieGraphRemoveRenderSettingNode.h"
 #include "Graph/Nodes/MovieGraphVariableNode.h"
 #include "MovieGraphUtils.h"
 #include "MovieRenderPipelineCoreModule.h"
@@ -1009,7 +1010,11 @@ void UMovieGraphConfig::CreateFlattenedGraph_Recursive(UMovieGraphEvaluatedConfi
 
 	// Check to see if our flattened evaluation graph already has a copy of this node.
 	InEvaluationContext.VisitedNodes.Add(Node);
-	const bool bShouldIncludeNode = Node->IsA<UMovieGraphSettingNode>() && !Node->IsDisabled();
+	
+	const bool bShouldIncludeNode =
+		Node->IsA<UMovieGraphSettingNode>() &&
+		!Node->IsDisabled() &&
+		!InEvaluationContext.NodeTypesToRemoveStack.Contains(Node->GetClass());
 
 	if (bShouldIncludeNode)
 	{
@@ -1030,18 +1035,25 @@ void UMovieGraphConfig::CreateFlattenedGraph_Recursive(UMovieGraphEvaluatedConfi
 
 		// Now do a property-copy from this node onto our flattened one. We don't use the generic property
 		// copy routines in the engine because we have special handling (we want to check if the property
-		// is actually marked for override, and also skip if this has already been overridden.
-
-		// ToDo: Handle "Disable" ndoes that can disable upstream types of nodes, etc. Push disable types into
-		// the currently evaluating context
+		// is actually marked for override, and also skip if this has already been overridden).
 		CopyOverriddenProperties(Node, ExistingNode, &InEvaluationContext.UserContext);
 	}
+
+	// If this is a special "removal" node, keep track of the type that should be removed. Since this method is recursive,
+	// a stack is used to keep track of the types. The graph is iterated starting from the Outputs node, so all matching
+	// nodes that are *upstream* of the removal node will be removed.
+	const UMovieGraphRemoveRenderSettingNode* RemovalNode = Cast<UMovieGraphRemoveRenderSettingNode>(Node);
+	const bool bIsARemovalNode = RemovalNode && !Node->IsDisabled() && (RemovalNode->NodeType.Get() != nullptr);
+	if (bIsARemovalNode)
+	{
+		InEvaluationContext.NodeTypesToRemoveStack.Push(RemovalNode->NodeType);
+	}
 	
 	// Now that we've potentially resolved the values on this node, continue to travel up-stream along any execution pins,
 	// potentially following re-route nodes, sub-graph nodes, through branches, etc.
 	TArray<UMovieGraphPin*> NewPinsToFollow = Node->EvaluatePinsToFollow(InEvaluationContext);
 	
-	for(UMovieGraphPin* Pin : NewPinsToFollow)
+	for (UMovieGraphPin* Pin : NewPinsToFollow)
 	{
 		for (UMovieGraphEdge* Edge : Pin->Edges)
 		{
@@ -1061,6 +1073,12 @@ void UMovieGraphConfig::CreateFlattenedGraph_Recursive(UMovieGraphEvaluatedConfi
 			}
 		}
 	}
+
+	// Done with this removal node now; pop it off the stack so it doesn't affect other branches
+	if (bIsARemovalNode)
+	{
+		InEvaluationContext.NodeTypesToRemoveStack.Pop();
+	}
 }
 
 void UMovieGraphConfig::VisitUpstreamNodes_Recursive(UMovieGraphNode* FromNode,	const FVisitNodesCallback& VisitCallback, TSet<UMovieGraphNode*>& VisitedNodes) const
diff --git a/Engine/Plugins/MovieScene/MovieRenderPipeline/Source/MovieRenderPipelineCore/Private/Graph/Nodes/MovieGraphRemoveRenderSettingNode.cpp b/Engine/Plugins/MovieScene/MovieRenderPipeline/Source/MovieRenderPipelineCore/Private/Graph/Nodes/MovieGraphRemoveRenderSettingNode.cpp
new file mode 100644
--- /dev/null
+++ b/Engine/Plugins/MovieScene/MovieRenderPipeline/Source/MovieRenderPipelineCore/Private/Graph/Nodes/MovieGraphRemoveRenderSettingNode.cpp
@@ -0,0 +1,65 @@
+// Copyright Epic Games, Inc. All Rights Reserved.
+
+#include "Graph/Nodes/MovieGraphRemoveRenderSettingNode.h"
+
+#include "Graph/MovieGraphConfig.h"
+#include "Styling/AppStyle.h"
+
+TArray<FMovieGraphPinProperties> UMovieGraphRemoveRenderSettingNode::GetInputPinProperties() const
+{
+	TArray<FMovieGraphPinProperties> Properties;
+	Properties.Add(FMovieGraphPinProperties::MakeBranchProperties());
+	return Properties;
+}
+
+TArray<FMovieGraphPinProperties> UMovieGraphRemoveRenderSettingNode::GetOutputPinProperties() const
+{
+	TArray<FMovieGraphPinProperties> Properties;
+	Properties.Add(FMovieGraphPinProperties::MakeBranchProperties());
+	return Properties;
+}
+
+#if WITH_EDITOR
+FText UMovieGraphRemoveRenderSettingNode::GetNodeTitle(const bool bGetDescriptive) const
+{
+	static const FText RemoveNode = NSLOCTEXT("MovieGraphNodes", "NodeName_RemoveNode", "Remove Render Setting");
+	static const FText RemoveNodeDescription = NSLOCTEXT("MovieGraphNodes", "NodeDescription_RemoveNode", "Remove Render Setting\nType: {0}");
+
+	if (bGetDescriptive && (NodeType.Get() != nullptr))
+	{
+		const UMovieGraphSettingNode* SettingNodeCDO = Cast<UMovieGraphSettingNode>(NodeType->GetDefaultObject());
+		return FText::Format(RemoveNodeDescription, SettingNodeCDO->GetNodeTitle());
+	}
+
+	return RemoveNode;
+}
+
+FText UMovieGraphRemoveRenderSettingNode::GetMenuCategory() const
+{
+	return NSLOCTEXT("MovieGraphNodes", "RemoveNode_Category", "Utility");
+}
+
+FLinearColor UMovieGraphRemoveRenderSettingNode::GetNodeTitleColor() const
+{
+	static const FLinearColor RemoveNodeColor = FLinearColor(0.f, 0.f, 0.f);
+	return RemoveNodeColor;
+}
+
+FSlateIcon UMovieGraphRemoveRenderSettingNode::GetIconAndTint(FLinearColor& OutColor) const
+{
+	static const FSlateIcon ApplyCVarPresetIcon = FSlateIcon(FAppStyle::GetAppStyleSetName(), "GenericCommands.Delete");
+
+	OutColor = FLinearColor::White;
+	return ApplyCVarPresetIcon;
+}
+
+void UMovieGraphRemoveRenderSettingNode::PostEditChangeProperty(FPropertyChangedEvent& PropertyChangedEvent)
+{
+	Super::PostEditChangeProperty(PropertyChangedEvent);
+	
+	if (PropertyChangedEvent.GetPropertyName() == GET_MEMBER_NAME_CHECKED(UMovieGraphRemoveRenderSettingNode, NodeType))
+	{
+		OnNodeChangedDelegate.Broadcast(this);
+	}
+}
+#endif // WITH_EDITOR
\ No newline at end of file
diff --git a/Engine/Plugins/MovieScene/MovieRenderPipeline/Source/MovieRenderPipelineCore/Public/Graph/MovieGraphConfig.h b/Engine/Plugins/MovieScene/MovieRenderPipeline/Source/MovieRenderPipelineCore/Public/Graph/MovieGraphConfig.h
--- a/Engine/Plugins/MovieScene/MovieRenderPipeline/Source/MovieRenderPipelineCore/Public/Graph/MovieGraphConfig.h
+++ b/Engine/Plugins/MovieScene/MovieRenderPipeline/Source/MovieRenderPipelineCore/Public/Graph/MovieGraphConfig.h
@@ -257,6 +257,32 @@ struct MOVIERENDERPIPELINECORE_API FMovieGraphEvaluatedBranchConfig
 		return AllNodeInstances;
 	}
 
+	/** Removes all nodes that are subclasses of the given type from the evaluated config. */
+	void RemoveNodesOfType(const TSubclassOf<UMovieGraphNode>& InClass)
+	{
+		// Keep track of the instance names (keys) in the map to remove if all node instances under the key are removed
+		TArray<FString> InstanceNamesToRemove;
+
+		for (TTuple<FString, FMovieGraphEvaluatedSettingsStack>& SettingsPair : NamedNodes)
+		{
+			SettingsPair.Value.NodeInstances.RemoveAll([InClass](const TObjectPtr<UMovieGraphNode>& NodeInstance)
+			{
+				return NodeInstance && (NodeInstance->GetClass() == InClass);
+			});
+
+			// Remove this entry in the map if all node instances were removed
+			if (SettingsPair.Value.NodeInstances.IsEmpty())
+			{
+				InstanceNamesToRemove.Add(SettingsPair.Key);
+			}
+		}
+
+		for (const FString& KeyToRemove : InstanceNamesToRemove)
+		{
+			NamedNodes.Remove(KeyToRemove);
+		}
+	}
+
 private:
 	// Allow the config to add nodes to this, but otherwise we don't want the public adding nodes to them
 	// without going through the graph resolving.
@@ -306,6 +332,12 @@ public:
 	*/
 	UPROPERTY()
 	TArray<TObjectPtr<const UMovieGraphSubgraphNode>> SubgraphStack;
+
+	/**
+	 * The stack of node types (exact match) that should be removed from the graph while it is being traversed. Each node
+	 * which specifies a type to removed adds to the stack.
+	 */
+	TArray<TSubclassOf<UMovieGraphSettingNode>> NodeTypesToRemoveStack;
 };
 
 /**
@@ -514,7 +546,8 @@ protected:
 	void InitializeFlattenedNode(UMovieGraphNode* InNode);
 
 	/** Traverse the graph, generating a combined "flatten" graph as it goes. */
-	void CreateFlattenedGraph_Recursive(UMovieGraphEvaluatedConfig* InOwningConfig, FMovieGraphEvaluatedBranchConfig& OutBranchConfig, FMovieGraphEvaluationContext& InEvaluationContext, UMovieGraphPin* InPinToFollow);
+	void CreateFlattenedGraph_Recursive(UMovieGraphEvaluatedConfig* InOwningConfig, FMovieGraphEvaluatedBranchConfig& OutBranchConfig,
+		FMovieGraphEvaluationContext& InEvaluationContext, UMovieGraphPin* InPinToFollow);
 
 	/** Recursive helper for VisitUpstreamNodes(). */
 	void VisitUpstreamNodes_Recursive(UMovieGraphNode* FromNode, const FVisitNodesCallback& VisitCallback, TSet<UMovieGraphNode*>& VisitedNodes) const;
diff --git a/Engine/Plugins/MovieScene/MovieRenderPipeline/Source/MovieRenderPipelineCore/Public/Graph/Nodes/MovieGraphRemoveRenderSettingNode.h b/Engine/Plugins/MovieScene/MovieRenderPipeline/Source/MovieRenderPipelineCore/Public/Graph/Nodes/MovieGraphRemoveRenderSettingNode.h
new file mode 100644
--- /dev/null
+++ b/Engine/Plugins/MovieScene/MovieRenderPipeline/Source/MovieRenderPipelineCore/Public/Graph/Nodes/MovieGraphRemoveRenderSettingNode.h
@@ -0,0 +1,37 @@
+﻿// Copyright Epic Games, Inc. All Rights Reserved.
+
+#pragma once
+
+#include "Graph/MovieGraphNode.h"
+#include "Templates/SubclassOf.h"
+
+#include "MovieGraphRemoveRenderSettingNode.generated.h"
+
+/** A node which can remove other nodes in the graph. */
+UCLASS()
+class MOVIERENDERPIPELINECORE_API UMovieGraphRemoveRenderSettingNode : public UMovieGraphNode
+{
+	GENERATED_BODY()
+
+public:
+	UMovieGraphRemoveRenderSettingNode() = default;
+
+	virtual TArray<FMovieGraphPinProperties> GetInputPinProperties() const override;
+	virtual TArray<FMovieGraphPinProperties> GetOutputPinProperties() const override;
+
+#if WITH_EDITOR
+	virtual FText GetNodeTitle(const bool bGetDescriptive = false) const override;
+	virtual FText GetMenuCategory() const override;
+	virtual FLinearColor GetNodeTitleColor() const override;
+	virtual FSlateIcon GetIconAndTint(FLinearColor& OutColor) const override;
+
+	//~ Begin UObject interface
+	virtual void PostEditChangeProperty(FPropertyChangedEvent& PropertyChangedEvent) override;
+	//~ End UObject interface
+#endif
+
+public:
+	/** The type of node (exact match) that should be removed. */
+	UPROPERTY(BlueprintReadWrite, EditAnywhere, Category = "Settings")
+	TSubclassOf<UMovieGraphSettingNode> NodeType;
+};
\ No newline at end of file
