MoviePipeline: Updated the preview widget and burn-in widget to include more information on render progress.

There are some things about this change that are not quite right yet, but it's most of the way there. Eg, some bits of info (like frame rate) need to take the Output Setting node into account, rather than relying solely on the sequence.

#rnx

diff --git a/Engine/Plugins/MovieScene/MovieRenderPipeline/Source/MovieRenderPipelineCore/Private/Graph/MovieGraphBlueprintLibrary.cpp b/Engine/Plugins/MovieScene/MovieRenderPipeline/Source/MovieRenderPipelineCore/Private/Graph/MovieGraphBlueprintLibrary.cpp
--- a/Engine/Plugins/MovieScene/MovieRenderPipeline/Source/MovieRenderPipelineCore/Private/Graph/MovieGraphBlueprintLibrary.cpp
+++ b/Engine/Plugins/MovieScene/MovieRenderPipeline/Source/MovieRenderPipelineCore/Private/Graph/MovieGraphBlueprintLibrary.cpp
@@ -2,6 +2,7 @@
 
 #include "Graph/MovieGraphBlueprintLibrary.h"
 
+#include "Graph/MovieGraphPipeline.h"
 #include "Graph/Nodes/MovieGraphOutputSettingNode.h"
 #include "Graph/Nodes/MovieGraphRenderLayerNode.h"
 #include "MoviePipelineBlueprintLibrary.h"
@@ -220,4 +221,191 @@ FIntPoint UMovieGraphBlueprintLibrary::GetEffectiveOutputResolution(UMovieGraphE
 	}
 	
 	return UMoviePipelineBlueprintLibrary::Utility_GetEffectiveOutputResolution(0 /* TODO: Overscan percentage needs to be provided */, OutputSetting->OutputResolution);
+}
+
+FText UMovieGraphBlueprintLibrary::GetJobName(const UMovieGraphPipeline* InMovieGraphPipeline)
+{
+	return InMovieGraphPipeline ? FText::FromString(InMovieGraphPipeline->GetCurrentJob()->JobName) : FText();
+}
+
+FText UMovieGraphBlueprintLibrary::GetJobAuthor(const UMovieGraphPipeline* InMovieGraphPipeline)
+{
+	return InMovieGraphPipeline ? FText::FromString(UE::MoviePipeline::GetJobAuthor(InMovieGraphPipeline->GetCurrentJob())) : FText();
+}
+
+void UMovieGraphBlueprintLibrary::GetOverallOutputFrames(const UMovieGraphPipeline* InMovieGraphPipeline, int32& OutCurrentIndex, int32& OutTotalCount)
+{
+	OutCurrentIndex = 0;
+	OutTotalCount = 0;
+	
+	if (InMovieGraphPipeline)
+	{
+		OutCurrentIndex = InMovieGraphPipeline->GetTimeStepInstance()->GetCalculatedTimeData().OutputFrameNumber;
+
+		for (const TObjectPtr<UMoviePipelineExecutorShot>& Shot : InMovieGraphPipeline->GetActiveShotList())
+		{
+			OutTotalCount += Shot->ShotInfo.WorkMetrics.TotalOutputFrameCount;
+		}
+	}
+}
+
+FDateTime UMovieGraphBlueprintLibrary::GetJobInitializationTime(const UMovieGraphPipeline* InMovieGraphPipeline)
+{
+	return InMovieGraphPipeline ? InMovieGraphPipeline->GetInitializationTime() : FDateTime();
+}
+
+bool UMovieGraphBlueprintLibrary::GetEstimatedTimeRemaining(const UMovieGraphPipeline* InMovieGraphPipeline, FTimespan& OutEstimate)
+{
+	if (!InMovieGraphPipeline)
+	{
+		OutEstimate = FTimespan();
+		return false;
+	}
+
+	// If they haven't produced a single frame yet, we can't give an estimate.
+	int32 OutputFrames;
+	int32 TotalOutputFrames;
+	GetOverallOutputFrames(InMovieGraphPipeline, OutputFrames, TotalOutputFrames);
+
+	if (OutputFrames <= 0 || TotalOutputFrames <= 0)
+	{
+		OutEstimate = FTimespan();
+		return false;
+	}
+
+	const float CompletionPercentage = OutputFrames / static_cast<float>(TotalOutputFrames);
+	const FTimespan CurrentDuration = FDateTime::UtcNow() - InMovieGraphPipeline->GetInitializationTime();
+
+	// If it has taken us CurrentDuration to process CompletionPercentage samples, then we can get a total duration
+	// estimate by taking (CurrentDuration/CompletionPercentage) and then take that total estimate minus elapsed
+	// to get remaining.
+	const FTimespan EstimatedTotalDuration = CurrentDuration / CompletionPercentage;
+	OutEstimate = EstimatedTotalDuration - CurrentDuration;
+
+	return true;
+}
+
+EMovieRenderPipelineState UMovieGraphBlueprintLibrary::GetPipelineState(const UMovieGraphPipeline* InMovieGraphPipeline)
+{
+	return InMovieGraphPipeline ? InMovieGraphPipeline->GetPipelineState() : EMovieRenderPipelineState::Uninitialized;
+}
+
+EMovieRenderShotState UMovieGraphBlueprintLibrary::GetCurrentSegmentState(const UMovieGraphPipeline* InMovieGraphPipeline)
+{
+	if (InMovieGraphPipeline)
+	{
+		const int32 ShotIndex = InMovieGraphPipeline->GetCurrentShotIndex();
+		
+		if (ShotIndex < InMovieGraphPipeline->GetActiveShotList().Num())
+		{
+			return InMovieGraphPipeline->GetActiveShotList()[ShotIndex]->ShotInfo.State;
+		}
+	}
+
+	return EMovieRenderShotState::Uninitialized;
+}
+
+void UMovieGraphBlueprintLibrary::GetCurrentSegmentName(const UMovieGraphPipeline* InMovieGraphPipeline, FText& OutOuterName, FText& OutInnerName)
+{
+	if (InMovieGraphPipeline)
+    {
+    	const int32 ShotIndex = InMovieGraphPipeline->GetCurrentShotIndex();
+		
+    	if (ShotIndex < InMovieGraphPipeline->GetActiveShotList().Num())
+    	{
+    		OutOuterName = FText::FromString(InMovieGraphPipeline->GetActiveShotList()[ShotIndex]->OuterName);
+    		OutInnerName = FText::FromString(InMovieGraphPipeline->GetActiveShotList()[ShotIndex]->InnerName);
+    	}
+    }
+}
+
+void UMovieGraphBlueprintLibrary::GetOverallSegmentCounts(const UMovieGraphPipeline* InMovieGraphPipeline, int32& OutCurrentIndex, int32& OutTotalCount)
+{
+	OutCurrentIndex = InMovieGraphPipeline->GetCurrentShotIndex();
+	OutTotalCount = InMovieGraphPipeline->GetActiveShotList().Num();
+}
+
+FMoviePipelineSegmentWorkMetrics UMovieGraphBlueprintLibrary::GetCurrentSegmentWorkMetrics(const UMovieGraphPipeline* InMovieGraphPipeline)
+{
+	if (InMovieGraphPipeline)
+	{
+		const int32 ShotIndex = InMovieGraphPipeline->GetCurrentShotIndex();
+		
+		if (ShotIndex < InMovieGraphPipeline->GetActiveShotList().Num())
+		{
+			return InMovieGraphPipeline->GetActiveShotList()[ShotIndex]->ShotInfo.WorkMetrics;
+		}
+	}
+
+	return FMoviePipelineSegmentWorkMetrics();
+}
+
+FTimecode UMovieGraphBlueprintLibrary::GetRootTimecode(const UMovieGraphPipeline* InMovieGraphPipeline)
+{
+	if (InMovieGraphPipeline)
+	{
+		return InMovieGraphPipeline->GetTimeStepInstance()->GetCalculatedTimeData().RootTimeCode;
+	}
+
+	return FTimecode();
+}
+
+FFrameNumber UMovieGraphBlueprintLibrary::GetRootFrameNumber(const UMovieGraphPipeline* InMovieGraphPipeline)
+{
+	if (InMovieGraphPipeline)
+	{
+		return InMovieGraphPipeline->GetTimeStepInstance()->GetCalculatedTimeData().RootFrameNumber;
+	}
+
+	return FFrameNumber(-1);
+}
+
+FTimecode UMovieGraphBlueprintLibrary::GetCurrentShotTimecode(const UMovieGraphPipeline* InMovieGraphPipeline)
+{
+	if (InMovieGraphPipeline)
+	{
+		return InMovieGraphPipeline->GetTimeStepInstance()->GetCalculatedTimeData().ShotTimeCode;
+	}
+
+	return FTimecode();
+}
+
+FFrameNumber UMovieGraphBlueprintLibrary::GetCurrentShotFrameNumber(const UMovieGraphPipeline* InMovieGraphPipeline)
+{
+	if (InMovieGraphPipeline)
+	{
+		return InMovieGraphPipeline->GetTimeStepInstance()->GetCalculatedTimeData().ShotFrameNumber;
+	}
+
+	return FFrameNumber(-1);
+}
+
+float UMovieGraphBlueprintLibrary::GetCurrentFocusDistance(const UMovieGraphPipeline* InMovieGraphPipeline)
+{
+	if (const UCineCameraComponent* CineCameraComponent = UMoviePipelineBlueprintLibrary::Utility_GetCurrentCineCamera(InMovieGraphPipeline->GetWorld()))
+	{
+		return CineCameraComponent->CurrentFocusDistance;
+	}
+
+	return -1.f;
+}
+
+float UMovieGraphBlueprintLibrary::GetCurrentFocalLength(const UMovieGraphPipeline* InMovieGraphPipeline)
+{
+	if (const UCineCameraComponent* CineCameraComponent = UMoviePipelineBlueprintLibrary::Utility_GetCurrentCineCamera(InMovieGraphPipeline->GetWorld()))
+	{
+		return CineCameraComponent->CurrentFocalLength;
+	}
+
+	return -1.f;
+}
+
+float UMovieGraphBlueprintLibrary::GetCurrentAperture(const UMovieGraphPipeline* InMovieGraphPipeline)
+{
+	if (const UCineCameraComponent* CineCameraComponent = UMoviePipelineBlueprintLibrary::Utility_GetCurrentCineCamera(InMovieGraphPipeline->GetWorld()))
+	{
+		return CineCameraComponent->CurrentAperture;
+	}
+
+	return 0.f;
 }
\ No newline at end of file
diff --git a/Engine/Plugins/MovieScene/MovieRenderPipeline/Source/MovieRenderPipelineCore/Private/Graph/MovieGraphLinearTimeStep.cpp b/Engine/Plugins/MovieScene/MovieRenderPipeline/Source/MovieRenderPipelineCore/Private/Graph/MovieGraphLinearTimeStep.cpp
--- a/Engine/Plugins/MovieScene/MovieRenderPipeline/Source/MovieRenderPipelineCore/Private/Graph/MovieGraphLinearTimeStep.cpp
+++ b/Engine/Plugins/MovieScene/MovieRenderPipeline/Source/MovieRenderPipelineCore/Private/Graph/MovieGraphLinearTimeStep.cpp
@@ -78,7 +78,6 @@ void UMovieGraphLinearTimeStep::TickProducingFrames()
 		CurrentTimeStepData.EvaluatedConfig = TObjectPtr<UMovieGraphEvaluatedConfig>(CurrentFrameData.EvaluatedConfig.Get());
 		CurrentTimeStepData.bExpandShotForTemporalSubSample = CurrentFrameData.TemporalSampleCount > 1;
 
-
 		// Sets up the render state, etc.
 		GetOwningGraph()->SetupShot(CurrentCameraCut);
 
@@ -107,13 +106,12 @@ void UMovieGraphLinearTimeStep::TickProducingFrames()
 		CurrentCameraCut->ShotInfo.State = EMovieRenderShotState::Rendering;
 	}
 
-
-
 	if (CurrentCameraCut->ShotInfo.State == EMovieRenderShotState::Rendering)
 	{
+		FMovieGraphTraversalContext Context = GetOwningGraph()->GetCurrentTraversalContext();
+		
 		if (IsFirstTemporalSample())
 		{
-			FMovieGraphTraversalContext Context = GetOwningGraph()->GetCurrentTraversalContext();
 			UMovieGraphConfig* Config = GetOwningGraph()->GetRootGraphForShot(CurrentCameraCut);
 			CurrentFrameData.EvaluatedConfig = TStrongObjectPtr<UMovieGraphEvaluatedConfig>(Config->CreateFlattenedGraph(Context));
 
@@ -124,6 +122,8 @@ void UMovieGraphLinearTimeStep::TickProducingFrames()
 			// Re-calculate some timing statistics about the given TS count
 			UpdateFrameMetrics();
 
+			CurrentCameraCut->ShotInfo.WorkMetrics.OutputFrameIndex++;
+
 			// We create a range of time that we want to represent with this frame in absolute root sequence
 			// space. This is because each frame can have a different temporal sample count, and when we jump
 			// over the time period the shutter is closed, it's no longer an easy fixed number. So instead,
@@ -217,7 +217,6 @@ void UMovieGraphLinearTimeStep::TickProducingFrames()
 			*LexToString(CurrentFrameMetrics.ShutterOffsetFrameTime),
 			*LexToString(CurrentFrameMetrics.MotionBlurCenteringOffsetTime));
 
-
 		// Now we need to fill out some of our current timestep data for the renderer portion to use.
 		double FrameDeltaTimeAsSeconds = CurrentFrameMetrics.TickResolution.AsSeconds(FrameDeltaTime);
 		CurrentTimeStepData.FrameDeltaTime = FrameDeltaTimeAsSeconds;
@@ -231,13 +230,36 @@ void UMovieGraphLinearTimeStep::TickProducingFrames()
 		CurrentTimeStepData.bIsLastTemporalSampleForFrame = IsLastTemporalSample();
 		CurrentTimeStepData.bRequiresAccumulator = CurrentFrameData.TemporalSampleCount > 1;
 		CurrentTimeStepData.OutputFrameNumber = CurrentFrameData.OutputFrameNumber;
-		
 		CurrentTimeStepData.EvaluatedConfig = TObjectPtr<UMovieGraphEvaluatedConfig>(CurrentFrameData.EvaluatedConfig.Get());
 
 		//UE_LOG(LogTemp, Warning, TEXT("F# %d bFirst: %d bLast: %d bReqAc: %d"),
 		//	CurrentTimeStepData.OutputFrameNumber, CurrentTimeStepData.bIsFirstTemporalSampleForFrame,
 		//	CurrentTimeStepData.bIsLastTemporalSampleForFrame, CurrentTimeStepData.bRequiresAccumulator);
 
+		// Calculate frame numbers and timecodes for the current sequence (root) and shot
+		{
+			constexpr bool bIncludeCDOs = true;
+			UMovieGraphOutputSettingNode* OutputSetting = CurrentFrameData.EvaluatedConfig->GetSettingForBranch<UMovieGraphOutputSettingNode>(UMovieGraphNode::GlobalsPinName, bIncludeCDOs);
+
+			// "Closest" isn't straightforward when using temporal sub-sampling, ie: A large enough shutter angle pushes a sample over the half way point and it rounds to
+			// the wrong one. Because temporal sub-sampling isn't centered around a frame (the centering is done via the final eval time) we can just subtract TSI*TPS to get our centered value.
+			const FFrameTime CenteringOffset = CurrentFrameData.TemporalSampleIndex * CurrentFrameMetrics.FrameTimePerTemporalSample;
+			FFrameTime CenteredFrameTime = CurrentCameraCut->ShotInfo.CurrentTimeInRoot - CenteringOffset;
+	
+			const FFrameRate SourceFrameRate = GetOwningGraph()->GetDataSourceInstance()->GetDisplayRate();
+			const FFrameRate EffectiveFrameRate = UMovieGraphBlueprintLibrary::GetEffectiveFrameRate(OutputSetting, SourceFrameRate);
+			const FFrameRate TickResolution = GetOwningGraph()->GetDataSourceInstance()->GetTickResolution();
+
+			constexpr bool bDropFrame = false;
+			CurrentTimeStepData.RootFrameNumber = FFrameRate::TransformTime(CenteredFrameTime, TickResolution, EffectiveFrameRate).RoundToFrame();
+			CurrentTimeStepData.RootTimeCode = FTimecode::FromFrameNumber(CurrentTimeStepData.RootFrameNumber, EffectiveFrameRate, bDropFrame);
+
+			// Calculate metrics for the shot as well
+			CenteredFrameTime = CenteredFrameTime * CurrentCameraCut->ShotInfo.OuterToInnerTransform;
+			CurrentTimeStepData.ShotFrameNumber = FFrameRate::TransformTime(CenteredFrameTime, TickResolution, EffectiveFrameRate).RoundToFrame();
+			CurrentTimeStepData.ShotTimeCode = FTimecode::FromFrameNumber(CurrentTimeStepData.ShotFrameNumber, EffectiveFrameRate, bDropFrame);
+		}
+
 		// Set our time step for the next frame. We use the undilated delta time for the Custom Timestep as the engine will
 		// apply the time dilation to the world tick for us, so we don't want to double up time dilation.
 		double UndilatedDeltaTime = CurrentFrameMetrics.TickResolution.AsSeconds(FrameDeltaTime);
@@ -260,7 +282,6 @@ void UMovieGraphLinearTimeStep::TickProducingFrames()
 				
 				// Increment the output frame number only on the first temporal sample.
 				CurrentFrameData.OutputFrameNumber++;
-
 			}
 
 			if (CurrentFrameData.TemporalSampleIndex >= CurrentFrameData.TemporalSampleCount - 1)
@@ -276,7 +297,6 @@ void UMovieGraphLinearTimeStep::TickProducingFrames()
 			}
 		}
 
-
 		if (!ensure(CurrentCameraCut->ShotInfo.CurrentTickInRoot < CurrentCameraCut->ShotInfo.TotalOutputRangeRoot.GetUpperBoundValue()))
 		{
 			UE_LOG(LogMovieRenderPipeline, Error, TEXT("Shot ran past evaluation range, this shouldn't be possible."));
@@ -331,7 +351,15 @@ void UMovieGraphLinearTimeStep::UpdateFrameMetrics()
 	FrameData.MotionBlurAmount = GetBlendedMotionBlurAmount();
 
 	// Calculate how long of a duration we want to represent where the camera shutter is open.
-	FrameData.FrameTimeWhileShutterOpen = FrameData.FrameTimePerOutputFrame * FrameData.MotionBlurAmount;
+	// If the shutter angle is effectively zero, lie about how long a frame is to prevent divide by zero
+	if (FrameData.MotionBlurAmount < (1.0 / 360.0))
+	{
+		FrameData.FrameTimeWhileShutterOpen = FrameData.FrameTimePerOutputFrame * (1.0 / 360.0);
+	}
+	else
+	{
+		FrameData.FrameTimeWhileShutterOpen = FrameData.FrameTimePerOutputFrame * FrameData.MotionBlurAmount;
+	}
 
 	// Now that we know how long the shutter is open, figure out how long each temporal sub-sample gets.
 	FrameData.FrameTimePerTemporalSample = FrameData.FrameTimeWhileShutterOpen / CurrentFrameData.TemporalSampleCount;
diff --git a/Engine/Plugins/MovieScene/MovieRenderPipeline/Source/MovieRenderPipelineCore/Private/Graph/MovieGraphPipeline.cpp b/Engine/Plugins/MovieScene/MovieRenderPipeline/Source/MovieRenderPipelineCore/Private/Graph/MovieGraphPipeline.cpp
--- a/Engine/Plugins/MovieScene/MovieRenderPipeline/Source/MovieRenderPipelineCore/Private/Graph/MovieGraphPipeline.cpp
+++ b/Engine/Plugins/MovieScene/MovieRenderPipeline/Source/MovieRenderPipelineCore/Private/Graph/MovieGraphPipeline.cpp
@@ -172,10 +172,10 @@ void UMovieGraphPipeline::BuildShotListFromDataSource()
 
 	for (UMoviePipelineExecutorShot* Shot : GetCurrentJob()->ShotInfo)
 	{
-		Shot->ShotInfo.CurrentTickInRoot = Shot->ShotInfo.TotalOutputRangeRoot.GetLowerBoundValue();
+		Shot->ShotInfo.CurrentTimeInRoot = FFrameTime(Shot->ShotInfo.CurrentTickInRoot);
+		Shot->ShotInfo.CalculateWorkMetrics(GetDataSourceInstance());
 	}
 
-
 	// The active shot-list is a subset of the whole shot-list; The ShotInfo contains information about every range it detected to render
 	// but if the user has turned the shot off in the UI then we don't want to render it.
 	ActiveShotList.Empty();
diff --git a/Engine/Plugins/MovieScene/MovieRenderPipeline/Source/MovieRenderPipelineCore/Private/MoviePipelineBlueprintLibrary.cpp b/Engine/Plugins/MovieScene/MovieRenderPipeline/Source/MovieRenderPipelineCore/Private/MoviePipelineBlueprintLibrary.cpp
--- a/Engine/Plugins/MovieScene/MovieRenderPipeline/Source/MovieRenderPipelineCore/Private/MoviePipelineBlueprintLibrary.cpp
+++ b/Engine/Plugins/MovieScene/MovieRenderPipeline/Source/MovieRenderPipelineCore/Private/MoviePipelineBlueprintLibrary.cpp
@@ -282,21 +282,9 @@ FFrameNumber UMoviePipelineBlueprintLibrary::GetCurrentShotFrameNumber(const UMo
 
 float UMoviePipelineBlueprintLibrary::GetCurrentFocusDistance(const UMoviePipeline* InMoviePipeline)
 {
-	if (InMoviePipeline)
+	if (const UCineCameraComponent* CineCameraComponent = Utility_GetCurrentCineCamera(InMoviePipeline->GetWorld()))
 	{
-		if (InMoviePipeline->GetWorld()->GetFirstPlayerController()->PlayerCameraManager)
-		{
-			// This only works if you use a Cine Camera (which is almost guranteed with Sequencer) and it's easier (and less human error prone) than re-deriving the information
-			ACineCameraActor* CineCameraActor = Cast<ACineCameraActor>(InMoviePipeline->GetWorld()->GetFirstPlayerController()->PlayerCameraManager->GetViewTarget());
-			if (CineCameraActor)
-			{
-				UCineCameraComponent* CineCameraComponent = CineCameraActor->GetCineCameraComponent();
-				if (CineCameraComponent)
-				{
-					return CineCameraComponent->CurrentFocusDistance;
-				}
-			}
-		}
+		return CineCameraComponent->CurrentFocusDistance;
 	}
 
 	return -1.0f;
@@ -304,21 +292,9 @@ float UMoviePipelineBlueprintLibrary::GetCurrentFocusDistance(const UMoviePipeli
 
 float UMoviePipelineBlueprintLibrary::GetCurrentFocalLength(const UMoviePipeline* InMoviePipeline)
 {
-	if (InMoviePipeline)
+	if (const UCineCameraComponent* CineCameraComponent = Utility_GetCurrentCineCamera(InMoviePipeline->GetWorld()))
 	{
-		if (InMoviePipeline->GetWorld()->GetFirstPlayerController()->PlayerCameraManager)
-		{
-			// This only works if you use a Cine Camera (which is almost guranteed with Sequencer) and it's easier (and less human error prone) than re-deriving the information
-			ACineCameraActor* CineCameraActor = Cast<ACineCameraActor>(InMoviePipeline->GetWorld()->GetFirstPlayerController()->PlayerCameraManager->GetViewTarget());
-			if (CineCameraActor)
-			{
-				UCineCameraComponent* CineCameraComponent = CineCameraActor->GetCineCameraComponent();
-				if (CineCameraComponent)
-				{
-					return CineCameraComponent->CurrentFocalLength;
-				}
-			}
-		}
+		return CineCameraComponent->CurrentFocalLength;
 	}
 
 	return -1.0f;
@@ -326,21 +302,9 @@ float UMoviePipelineBlueprintLibrary::GetCurrentFocalLength(const UMoviePipeline
 
 float UMoviePipelineBlueprintLibrary::GetCurrentAperture(const UMoviePipeline* InMoviePipeline)
 {
-	if (InMoviePipeline)
+	if (const UCineCameraComponent* CineCameraComponent = Utility_GetCurrentCineCamera(InMoviePipeline->GetWorld()))
 	{
-		if (InMoviePipeline->GetWorld()->GetFirstPlayerController()->PlayerCameraManager)
-		{
-			// This only works if you use a Cine Camera (which is almost guranteed with Sequencer) and it's easier (and less human error prone) than re-deriving the information
-			ACineCameraActor* CineCameraActor = Cast<ACineCameraActor>(InMoviePipeline->GetWorld()->GetFirstPlayerController()->PlayerCameraManager->GetViewTarget());
-			if (CineCameraActor)
-			{
-				UCineCameraComponent* CineCameraComponent = CineCameraActor->GetCineCameraComponent();
-				if (CineCameraComponent)
-				{
-					return CineCameraComponent->CurrentAperture;
-				}
-			}
-		}
+		return CineCameraComponent->CurrentAperture;
 	}
 
 	return 0.0f;
@@ -924,6 +888,23 @@ FIntPoint UMoviePipelineBlueprintLibrary::Utility_GetEffectiveOutputResolution(c
 	return EffectiveResolution;
 }
 
+UCineCameraComponent* UMoviePipelineBlueprintLibrary::Utility_GetCurrentCineCamera(const UWorld* InWorld)
+{
+	if (InWorld)
+	{
+		if (const TObjectPtr<APlayerCameraManager> PlayerCameraManager = InWorld->GetFirstPlayerController()->PlayerCameraManager)
+		{
+			// This only works if you use a Cine Camera (which is almost guaranteed with Sequencer) and it's easier (and less human error prone) than re-deriving the information
+			if (const ACineCameraActor* CineCameraActor = Cast<ACineCameraActor>(PlayerCameraManager->GetViewTarget()))
+			{
+				return CineCameraActor->GetCineCameraComponent();
+			}
+		}
+	}
+
+	return nullptr;
+}
+
 UMoviePipelineSetting* UMoviePipelineBlueprintLibrary::FindOrGetDefaultSettingForShot(TSubclassOf<UMoviePipelineSetting> InSettingType, const UMoviePipelinePrimaryConfig* InPrimaryConfig, const UMoviePipelineExecutorShot* InShot)
 {
 	// Check to see if this setting is in the shot override, if it is we'll use the shot version of that.
diff --git a/Engine/Plugins/MovieScene/MovieRenderPipeline/Source/MovieRenderPipelineCore/Public/Graph/MovieGraphBlueprintLibrary.h b/Engine/Plugins/MovieScene/MovieRenderPipeline/Source/MovieRenderPipelineCore/Public/Graph/MovieGraphBlueprintLibrary.h
--- a/Engine/Plugins/MovieScene/MovieRenderPipeline/Source/MovieRenderPipelineCore/Public/Graph/MovieGraphBlueprintLibrary.h
+++ b/Engine/Plugins/MovieScene/MovieRenderPipeline/Source/MovieRenderPipelineCore/Public/Graph/MovieGraphBlueprintLibrary.h
@@ -10,6 +10,7 @@
 
 // Forward Declare
 class UMovieGraphOutputSettingNode;
+class UMovieGraphPipeline;
 
 UCLASS(meta = (ScriptName = "MovieGraphLibrary"))
 class MOVIERENDERPIPELINECORE_API UMovieGraphBlueprintLibrary : public UBlueprintFunctionLibrary
@@ -50,4 +51,130 @@ public:
 	*/
 	UFUNCTION(BlueprintCallable, Category = "Movie Graph")
 	static FIntPoint GetEffectiveOutputResolution(UMovieGraphEvaluatedConfig* InEvaluatedGraph, const FName& InBranchName);
+
+	/**
+	* Gets the name of the current job.
+	* @param	InMovieGraphPipeline	 - The pipeline to get the job name from.
+	*/
+	UFUNCTION(BlueprintPure, Category = "Movie Graph")
+	static FText GetJobName(const UMovieGraphPipeline* InMovieGraphPipeline);
+
+	/**
+	* Gets the author of the current job, or the logged in user's username if the job has no specified author.
+	* @param	InMovieGraphPipeline	 - The pipeline to get the job author from.
+	*/
+	UFUNCTION(BlueprintPure, Category = "Movie Graph")
+	static FText GetJobAuthor(const UMovieGraphPipeline* InMovieGraphPipeline);
+
+	/**
+	* Determines the overall current frame number and total number of frames.
+	* @param	InMovieGraphPipeline	- The pipeline to get the frame information from.
+	* @param	OutCurrentIndex			- The current frame number.
+	* @param	OutTotalCount			- The total number of frames.
+	*/
+	UFUNCTION(BlueprintPure, Category = "Movie Graph")
+	static void GetOverallOutputFrames(const UMovieGraphPipeline* InMovieGraphPipeline, int32& OutCurrentIndex, int32& OutTotalCount);
+
+	/**
+	* Gets the time the job was initialized.
+	* @param	InMovieGraphPipeline	- The pipeline to get the job initialization time from.
+	*/
+	UFUNCTION(BlueprintPure, Category = "Movie Graph")
+	static FDateTime GetJobInitializationTime(const UMovieGraphPipeline* InMovieGraphPipeline);
+
+	/**
+	* Get the estimated amount of time remaining for the current pipeline. Based on looking at the total
+	* amount of samples to render vs. how many have been completed so far. Inaccurate when Time Dilation
+	* is used, and gets more accurate over the course of the render.
+	* @param	InMovieGraphPipeline	- The pipeline to get the time estimate from.
+	* @param	OutEstimate				- The resulting estimate, or FTimespan() if estimate is not valid.
+	* @return							- True if a valid estimate can be calculated, or false if it is not ready yet (ie: not enough samples rendered)
+	*/
+	UFUNCTION(BlueprintPure, Category = "Movie Graph")
+	static bool GetEstimatedTimeRemaining(const UMovieGraphPipeline* InMovieGraphPipeline, FTimespan& OutEstimate);
+
+	/**
+	* Get the current state of the specified pipeline. See EMovieRenderPipelineState for more detail about each state.
+	* @param	InMovieGraphPipeline	- The pipeline to get the state for.
+	* @return							- The current state.
+	*/
+	UFUNCTION(BlueprintPure, Category = "Movie Graph")
+	static EMovieRenderPipelineState GetPipelineState(const UMovieGraphPipeline* InMovieGraphPipeline);
+
+	/**
+	* Gets the state of the segment (shot) currently being rendered.
+	* @param	InMovieGraphPipeline	- The pipeline to get segment information from.
+	*/
+	UFUNCTION(BlueprintPure, Category = "Movie Graph")
+	static EMovieRenderShotState GetCurrentSegmentState(const UMovieGraphPipeline* InMovieGraphPipeline);
+
+	/**
+	* Gets the name of the segment (shot) currently being rendered.
+	* @param	InMovieGraphPipeline	- The pipeline to get segment information from.
+	*/
+	UFUNCTION(BlueprintPure, Category = "Movie Graph")
+	static void GetCurrentSegmentName(const UMovieGraphPipeline* InMovieGraphPipeline, FText& OutOuterName, FText& OutInnerName);
+
+	/**
+	* Gets the number of segments (shots) that will be rendered.
+	* @param	InMovieGraphPipeline	- The pipeline to get segment information from.
+	*/
+	UFUNCTION(BlueprintPure, Category = "Movie Graph")
+	static void GetOverallSegmentCounts(const UMovieGraphPipeline* InMovieGraphPipeline, int32& OutCurrentIndex, int32& OutTotalCount);
+
+	/**
+	* Gets the work metrics for the segment (shot) that is currently being rendered.
+	* @param	InMovieGraphPipeline	- The pipeline to get segment information from.
+	*/
+	UFUNCTION(BlueprintPure, Category = "Movie Graph")
+	static FMoviePipelineSegmentWorkMetrics GetCurrentSegmentWorkMetrics(const UMovieGraphPipeline* InMovieGraphPipeline);
+
+	/**
+	* Gets the timecode of the current render at the root (sequence) level.
+	* @param	InMovieGraphPipeline	- The pipeline to get timecode information from.
+	*/
+	UFUNCTION(BlueprintPure, Category = "Movie Graph")
+	static FTimecode GetRootTimecode(const UMovieGraphPipeline* InMovieGraphPipeline);
+
+	/**
+	* Gets the frame number of the current render at the root (sequence) level.
+	* @param	InMovieGraphPipeline	- The pipeline to get frame number information from.
+	*/
+	UFUNCTION(BlueprintPure, Category = "Movie Graph")
+	static FFrameNumber GetRootFrameNumber(const UMovieGraphPipeline* InMovieGraphPipeline);
+
+	/**
+	* Gets the timecode of the current render at the shot level.
+	* @param	InMovieGraphPipeline	- The pipeline to get timecode information from.
+	*/
+	UFUNCTION(BlueprintPure, Category = "Movie Graph")
+	static FTimecode GetCurrentShotTimecode(const UMovieGraphPipeline* InMovieGraphPipeline);
+
+	/**
+	* Gets the frame number of the current render at the shot level.
+	* @param	InMovieGraphPipeline	- The pipeline to get frame number information from.
+	*/
+	UFUNCTION(BlueprintPure, Category = "Movie Graph")
+	static FFrameNumber GetCurrentShotFrameNumber(const UMovieGraphPipeline* InMovieGraphPipeline);
+
+	/**
+	* Gets the focus distance for the camera currently in use.
+	* @param	InMovieGraphPipeline	- The pipeline to get the camera information from.
+	*/
+	UFUNCTION(BlueprintPure, Category = "Movie Graph")
+	static float GetCurrentFocusDistance(const UMovieGraphPipeline* InMovieGraphPipeline);
+
+	/**
+	* Gets the focal length for the camera currently in use.
+	* @param	InMovieGraphPipeline	- The pipeline to get the camera information from.
+	*/
+	UFUNCTION(BlueprintPure, Category = "Movie Graph")
+	static float GetCurrentFocalLength(const UMovieGraphPipeline* InMovieGraphPipeline);
+	
+	/**
+	* Gets the aperture for the camera currently in use.
+	* @param	InMovieGraphPipeline	- The pipeline to get the camera information from.
+	*/
+	UFUNCTION(BlueprintPure, Category = "Movie Graph")
+	static float GetCurrentAperture(const UMovieGraphPipeline* InMovieGraphPipeline);
 };
\ No newline at end of file
diff --git a/Engine/Plugins/MovieScene/MovieRenderPipeline/Source/MovieRenderPipelineCore/Public/Graph/MovieGraphPipeline.h b/Engine/Plugins/MovieScene/MovieRenderPipeline/Source/MovieRenderPipelineCore/Public/Graph/MovieGraphPipeline.h
--- a/Engine/Plugins/MovieScene/MovieRenderPipeline/Source/MovieRenderPipelineCore/Public/Graph/MovieGraphPipeline.h
+++ b/Engine/Plugins/MovieScene/MovieRenderPipeline/Source/MovieRenderPipelineCore/Public/Graph/MovieGraphPipeline.h
@@ -53,11 +53,12 @@ public:
 	*
 	* @param InDateTime - The DateTime object to return for GetInitializationTime.
 	*/
-	UFUNCTION(BlueprintCallable, Category = "Movie Render Pipeline")
+	UFUNCTION(BlueprintCallable, Category = "Movie Graph")
 	void SetInitializationTime(const FDateTime& InDateTime) { GraphInitializationTime = InDateTime; }
 
-public:
 	UMovieGraphConfig* GetRootGraphForShot(UMoviePipelineExecutorShot* InShot) const;
+
+	UFUNCTION(BlueprintCallable, Category = "Movie Graph")
 	FMovieGraphTraversalContext GetCurrentTraversalContext() const;
 
 	/** Get the Active Shot list, which is the full shot list generated from the external data source, with disabled shots removed. */
@@ -71,7 +72,7 @@ public:
 	/** Used occasionally to cross-reference other components. Don't call this unless you know what you're doing. */
 	UMovieGraphTimeStepBase* GetTimeStepInstance() const { return GraphTimeStepInstance; }
 	/** Used occasionally to cross-reference other components. Don't call this unless you know what you're doing. */
-	UFUNCTION(BlueprintCallable, Category = "Movie Render Pipeline")
+	UFUNCTION(BlueprintCallable, Category = "Movie Graph")
 	UMovieGraphRendererBase* GetRendererInstance() const { return GraphRendererInstance; }
 	/** Used occasionally to cross-reference other components. Don't call this unless you know what you're doing. */
 	UMovieGraphDataSourceBase* GetDataSourceInstance() const { return GraphDataSourceInstance; }
diff --git a/Engine/Plugins/MovieScene/MovieRenderPipeline/Source/MovieRenderPipelineCore/Public/Graph/MovieGraphTimeStepData.h b/Engine/Plugins/MovieScene/MovieRenderPipeline/Source/MovieRenderPipelineCore/Public/Graph/MovieGraphTimeStepData.h
--- a/Engine/Plugins/MovieScene/MovieRenderPipeline/Source/MovieRenderPipelineCore/Public/Graph/MovieGraphTimeStepData.h
+++ b/Engine/Plugins/MovieScene/MovieRenderPipeline/Source/MovieRenderPipelineCore/Public/Graph/MovieGraphTimeStepData.h
@@ -2,6 +2,7 @@
 #pragma once
 
 #include "CoreMinimal.h"
+#include "Misc/Timecode.h"
 #include "MovieGraphTimeStepData.generated.h"
 
 class UMovieGraphEvaluatedConfig;
@@ -84,4 +85,20 @@ public:
 	*/
 	UPROPERTY(BlueprintReadWrite, EditAnywhere, Category = "Movie Graph")
 	TObjectPtr<UMovieGraphEvaluatedConfig> EvaluatedConfig;
+
+	/** The current timecode at the root (sequence) level. */
+	UPROPERTY(BlueprintReadWrite, EditAnywhere, Category = "Movie Graph")
+	FTimecode RootTimeCode;
+
+	/** The current frame number at the root (sequence) level. */
+	UPROPERTY(BlueprintReadWrite, EditAnywhere, Category = "Movie Graph")
+	FFrameNumber RootFrameNumber;
+
+	/** The current timecode at the shot level. */
+	UPROPERTY(BlueprintReadWrite, EditAnywhere, Category = "Movie Graph")
+	FTimecode ShotTimeCode;
+
+	/** The current frame number at the shot level. */
+	UPROPERTY(BlueprintReadWrite, EditAnywhere, Category = "Movie Graph")
+	FFrameNumber ShotFrameNumber;
 };
diff --git a/Engine/Plugins/MovieScene/MovieRenderPipeline/Source/MovieRenderPipelineCore/Public/Graph/MovieGraphTraversalContext.h b/Engine/Plugins/MovieScene/MovieRenderPipeline/Source/MovieRenderPipelineCore/Public/Graph/MovieGraphTraversalContext.h
--- a/Engine/Plugins/MovieScene/MovieRenderPipeline/Source/MovieRenderPipelineCore/Public/Graph/MovieGraphTraversalContext.h
+++ b/Engine/Plugins/MovieScene/MovieRenderPipeline/Source/MovieRenderPipelineCore/Public/Graph/MovieGraphTraversalContext.h
@@ -33,9 +33,6 @@ public:
 		, ShotCount(0)
 	{}
 
-	UPROPERTY(BlueprintReadWrite, EditAnywhere, Category = "MovieGraph")
-	FName RootBranch;
-
 	/** Which shot (out of ShotCount) is this time step for? */
 	UPROPERTY(BlueprintReadWrite, EditAnywhere, Category = "Movie Graph")
 	int32 ShotIndex;
diff --git a/Engine/Plugins/MovieScene/MovieRenderPipeline/Source/MovieRenderPipelineCore/Public/MoviePipelineBlueprintLibrary.h b/Engine/Plugins/MovieScene/MovieRenderPipeline/Source/MovieRenderPipelineCore/Public/MoviePipelineBlueprintLibrary.h
--- a/Engine/Plugins/MovieScene/MovieRenderPipeline/Source/MovieRenderPipelineCore/Public/MoviePipelineBlueprintLibrary.h
+++ b/Engine/Plugins/MovieScene/MovieRenderPipeline/Source/MovieRenderPipelineCore/Public/MoviePipelineBlueprintLibrary.h
@@ -8,6 +8,7 @@
 #include "MoviePipelineBlueprintLibrary.generated.h"
 
 // Forward Declare
+class UCineCameraComponent;
 class UMoviePipeline;
 class UMovieSceneSequence;
 class UMoviePipelineSetting;
@@ -169,4 +170,7 @@ private:
 
 	/** Get the effective output resolution, taking into account overscan. */
 	static FIntPoint Utility_GetEffectiveOutputResolution(const float OverscanPercentage, const FIntPoint& InOutputResolution);
+
+	/** Get the current cine camera in use, or nullptr if there is none. */
+	static UCineCameraComponent* Utility_GetCurrentCineCamera(const UWorld* InWorld);
 };
\ No newline at end of file
diff --git a/Engine/Plugins/MovieScene/MovieRenderPipeline/Source/MovieRenderPipelineCore/Public/MovieRenderPipelineDataTypes.cpp b/Engine/Plugins/MovieScene/MovieRenderPipeline/Source/MovieRenderPipelineCore/Public/MovieRenderPipelineDataTypes.cpp
--- a/Engine/Plugins/MovieScene/MovieRenderPipeline/Source/MovieRenderPipelineCore/Public/MovieRenderPipelineDataTypes.cpp
+++ b/Engine/Plugins/MovieScene/MovieRenderPipeline/Source/MovieRenderPipelineCore/Public/MovieRenderPipelineDataTypes.cpp
@@ -4,14 +4,19 @@
 #include "MoviePipeline.h"
 #include "MovieRenderPipelineCoreModule.h"
 #include "RHI.h"
+#include "Graph/MovieGraphDataTypes.h"
 #include "Sections/MovieSceneCinematicShotSection.h"
 
 #include UE_INLINE_GENERATED_CPP_BY_NAME(MovieRenderPipelineDataTypes)
 
-FFrameNumber FMoviePipelineCameraCutInfo::GetOutputFrameCountEstimate() const
+FFrameNumber FMoviePipelineCameraCutInfo::GetOutputFrameCountEstimate(const UMovieGraphDataSourceBase* InDataSource) const
 {
+	// The graph uses the data source; the old version uses members in this struct
+	const FFrameRate SourceRate = InDataSource ? InDataSource->GetDisplayRate() : CachedFrameRate;
+	const FFrameRate DestinationRate = InDataSource ? InDataSource->GetTickResolution() : CachedTickResolution;
+	
 	// TotalRange is stored in Tick Resolution, so we convert 1 frame of Frame Rate to the number of ticks.
-	FFrameNumber OneFrameInTicks = FFrameRate::TransformTime(FFrameTime(FFrameNumber(1)), CachedFrameRate, CachedTickResolution).FloorToFrame();
+	FFrameNumber OneFrameInTicks = FFrameRate::TransformTime(FFrameTime(FFrameNumber(1)), SourceRate, DestinationRate).FloorToFrame();
 
 	// Find out how many ticks long our total output range is.
 	FFrameNumber TotalOutputRangeTicks = TotalOutputRangeRoot.Size<FFrameNumber>();
@@ -20,10 +25,10 @@ FFrameNumber FMoviePipelineCameraCutInfo::GetOutputFrameCountEstimate() const
 	return FFrameNumber(NumFrames);
 }
 
-void FMoviePipelineCameraCutInfo::CalculateWorkMetrics()
+void FMoviePipelineCameraCutInfo::CalculateWorkMetrics(const UMovieGraphDataSourceBase* InDataSource)
 {
 	// Initial Range + Handle Frames
-	FFrameNumber OutputFrameCount = GetOutputFrameCountEstimate();
+	FFrameNumber OutputFrameCount = GetOutputFrameCountEstimate(InDataSource);
 	WorkMetrics.TotalOutputFrameCount = OutputFrameCount.Value;
 	WorkMetrics.TotalSubSampleCount = NumSpatialSamples * NumTemporalSamples * NumTiles.X * NumTiles.Y; // Samples to generate an output frame.
 	WorkMetrics.TotalEngineWarmUpFrameCount = NumEngineWarmUpFramesRemaining;
diff --git a/Engine/Plugins/MovieScene/MovieRenderPipeline/Source/MovieRenderPipelineCore/Public/MovieRenderPipelineDataTypes.h b/Engine/Plugins/MovieScene/MovieRenderPipeline/Source/MovieRenderPipelineCore/Public/MovieRenderPipelineDataTypes.h
--- a/Engine/Plugins/MovieScene/MovieRenderPipeline/Source/MovieRenderPipelineCore/Public/MovieRenderPipelineDataTypes.h
+++ b/Engine/Plugins/MovieScene/MovieRenderPipeline/Source/MovieRenderPipelineCore/Public/MovieRenderPipelineDataTypes.h
@@ -34,6 +34,7 @@ class FRenderTarget;
 class UMoviePipeline;
 struct FMoviePipelineFormatArgs;
 class UMoviePipelineExecutorShot;
+class UMovieGraphDataSourceBase;
 
 namespace Audio { class FMixerSubmix; }
 
@@ -489,10 +490,12 @@ public:
 
 	bool IsInitialized() const { return State != EMovieRenderShotState::Uninitialized; }
 	void SetNextStateAfter(const EMovieRenderShotState InCurrentState);
-	void CalculateWorkMetrics();
+
+	/** The data source is required for calculating some metrics (graph only). */
+	void CalculateWorkMetrics(const UMovieGraphDataSourceBase* InDataSource = nullptr);
 
 private:
-	FFrameNumber GetOutputFrameCountEstimate() const;
+	FFrameNumber GetOutputFrameCountEstimate(const UMovieGraphDataSourceBase* InDataSource) const;
 
 public:
 	
