OCIO: After reconsideration, we also leave the MRQ bIsEnabled default value to false.

#rnx

diff --git a/Engine/Plugins/MovieScene/MovieRenderPipeline/Source/MovieRenderPipelineCore/Private/MoviePipelineColorSetting.cpp b/Engine/Plugins/MovieScene/MovieRenderPipeline/Source/MovieRenderPipelineCore/Private/MoviePipelineColorSetting.cpp
--- a/Engine/Plugins/MovieScene/MovieRenderPipeline/Source/MovieRenderPipelineCore/Private/MoviePipelineColorSetting.cpp
+++ b/Engine/Plugins/MovieScene/MovieRenderPipeline/Source/MovieRenderPipelineCore/Private/MoviePipelineColorSetting.cpp
@@ -6,9 +6,6 @@
 
 UMoviePipelineColorSetting::UMoviePipelineColorSetting() : bDisableToneCurve(false)
 {
-	// Note: Since MRQ settings already have enabled/disabled state control, this is slighty redundant.
-	// As such, we retain the old behavior and always default to enabled == true.
-	OCIOConfiguration.bIsEnabled = true;
 }
 
 #if WITH_EDITOR
