Movie Pipeline: Check for valid sidecar view locations and rotations. This can happen if the camera is not spawned for certain frames, which is undefined behavior anyways.

#rnx

diff --git a/Engine/Plugins/MovieScene/MovieRenderPipeline/Source/MovieRenderPipelineCore/Private/MoviePipeline.cpp b/Engine/Plugins/MovieScene/MovieRenderPipeline/Source/MovieRenderPipelineCore/Private/MoviePipeline.cpp
--- a/Engine/Plugins/MovieScene/MovieRenderPipeline/Source/MovieRenderPipelineCore/Private/MoviePipeline.cpp
+++ b/Engine/Plugins/MovieScene/MovieRenderPipeline/Source/MovieRenderPipelineCore/Private/MoviePipeline.cpp
@@ -1788,9 +1788,13 @@ void UMoviePipeline::GetSidecarCameraData(UMoviePipelineExecutorShot* InShot, in
 		BoundCamera->GetCameraView(GetWorld()->GetDeltaSeconds(), OutViewInfo);
 				
 		// We override the current/previous transform based on cached data though to ensure we accurately handle cur/next frame positions
-		OutViewInfo.Location = FrameInfo.CurrSidecarViewLocations[InCameraIndex];
-		OutViewInfo.Rotation = FrameInfo.CurrSidecarViewRotations[InCameraIndex];
-		OutViewInfo.PreviousViewTransform = FTransform(FrameInfo.PrevSidecarViewRotations[InCameraIndex], FrameInfo.PrevSidecarViewLocations[InCameraIndex]);
+		if (FrameInfo.CurrSidecarViewLocations.IsValidIndex(InCameraIndex) && FrameInfo.CurrSidecarViewRotations.IsValidIndex(InCameraIndex) &&
+			FrameInfo.PrevSidecarViewLocations.IsValidIndex(InCameraIndex) && FrameInfo.PrevSidecarViewRotations.IsValidIndex(InCameraIndex))
+		{
+			OutViewInfo.Location = FrameInfo.CurrSidecarViewLocations[InCameraIndex];
+			OutViewInfo.Rotation = FrameInfo.CurrSidecarViewRotations[InCameraIndex];
+			OutViewInfo.PreviousViewTransform = FTransform(FrameInfo.PrevSidecarViewRotations[InCameraIndex], FrameInfo.PrevSidecarViewLocations[InCameraIndex]);
+		}
 	}
 }
 
diff --git a/Engine/Plugins/MovieScene/MovieRenderPipeline/Source/MovieRenderPipelineRenderPasses/Private/MoviePipelineDeferredPasses.cpp b/Engine/Plugins/MovieScene/MovieRenderPipeline/Source/MovieRenderPipelineRenderPasses/Private/MoviePipelineDeferredPasses.cpp
--- a/Engine/Plugins/MovieScene/MovieRenderPipeline/Source/MovieRenderPipelineRenderPasses/Private/MoviePipelineDeferredPasses.cpp
+++ b/Engine/Plugins/MovieScene/MovieRenderPipeline/Source/MovieRenderPipelineRenderPasses/Private/MoviePipelineDeferredPasses.cpp
@@ -816,7 +816,15 @@ UE::MoviePipeline::FImagePassCameraViewData UMoviePipelineDeferredPassBase::GetC
 				UE::MoviePipeline::GetMetadataFromCineCamera(CineCameraComponent, PassIdentifier.CameraName, PassIdentifier.Name, OutCameraData.FileMetadata);
 
 				// We only do this in the multi-camera case because the single camera case is covered by the main Rendering loop.
-				UE::MoviePipeline::GetMetadataFromCameraLocRot(PassIdentifier.CameraName, PassIdentifier.Name, OutCameraData.ViewInfo.Location, OutCameraData.ViewInfo.Rotation, OutCameraData.ViewInfo.PreviousViewTransform->GetLocation(), FRotator(OutCameraData.ViewInfo.PreviousViewTransform->GetRotation()), OutCameraData.FileMetadata);
+				FVector PreviousLocation = OutCameraData.ViewInfo.Location;
+				FRotator PreviousRotation = OutCameraData.ViewInfo.Rotation;
+				if (OutCameraData.ViewInfo.PreviousViewTransform.IsSet())
+				{
+					PreviousLocation = OutCameraData.ViewInfo.PreviousViewTransform->GetLocation();
+					PreviousRotation = FRotator(OutCameraData.ViewInfo.PreviousViewTransform->GetRotation());
+				}
+
+				UE::MoviePipeline::GetMetadataFromCameraLocRot(PassIdentifier.CameraName, PassIdentifier.Name, OutCameraData.ViewInfo.Location, OutCameraData.ViewInfo.Rotation, PreviousLocation, PreviousRotation, OutCameraData.FileMetadata);
 			}
 		}
 		else
