TObjectPtr compilation fixes ahead of enabling GC barrier

diff --git a/Engine/Plugins/Animation/ControlRig/Source/ControlRig/Private/ControlRigComponent.cpp b/Engine/Plugins/Animation/ControlRig/Source/ControlRig/Private/ControlRigComponent.cpp
--- a/Engine/Plugins/Animation/ControlRig/Source/ControlRig/Private/ControlRigComponent.cpp
+++ b/Engine/Plugins/Animation/ControlRig/Source/ControlRig/Private/ControlRigComponent.cpp
@@ -427,7 +427,7 @@ void UControlRigComponent::Initialize()
 		}
 	}
 
-	for (const TObjectPtr<USceneComponent> ErrorComponent : MappedComponentsWithErrors)
+	for (const TObjectPtr<USceneComponent>& ErrorComponent : MappedComponentsWithErrors)
 	{
 		FString Message = FString::Printf(
 				TEXT("Elements from the same component (%s) should not be mapped to both input and output,"
diff --git a/Engine/Plugins/Animation/DeformerGraph/Source/OptimusCore/Private/OptimusComputeDataInterface.cpp b/Engine/Plugins/Animation/DeformerGraph/Source/OptimusCore/Private/OptimusComputeDataInterface.cpp
--- a/Engine/Plugins/Animation/DeformerGraph/Source/OptimusCore/Private/OptimusComputeDataInterface.cpp
+++ b/Engine/Plugins/Animation/DeformerGraph/Source/OptimusCore/Private/OptimusComputeDataInterface.cpp
@@ -60,7 +60,7 @@ TArray<TSubclassOf<UOptimusComputeDataInterface>> UOptimusComputeDataInterface::
 TSet<TArray<FName>> UOptimusComputeDataInterface::GetUniqueDomainDimensions()
 {
 	TSet<TArray<FName>> UniqueNestedContextNames;
-	for (const TSubclassOf<UOptimusComputeDataInterface> DataInterfaceClass: GetAllComputeDataInterfaceClasses())
+	for (const TSubclassOf<UOptimusComputeDataInterface>& DataInterfaceClass: GetAllComputeDataInterfaceClasses())
 	{
 		UOptimusComputeDataInterface* DataInterface = Cast<UOptimusComputeDataInterface>(DataInterfaceClass->GetDefaultObject());
 		if (DataInterface)
@@ -73,7 +73,7 @@ TSet<TArray<FName>> UOptimusComputeDataInterface::GetUniqueDomainDimensions()
 
 void UOptimusComputeDataInterface::RegisterAllTypes()
 {
-	for (const TSubclassOf<UOptimusComputeDataInterface> DataInterfaceClass : GetAllComputeDataInterfaceClasses())
+	for (const TSubclassOf<UOptimusComputeDataInterface>& DataInterfaceClass : GetAllComputeDataInterfaceClasses())
 	{
 		UOptimusComputeDataInterface* DataInterface = Cast<UOptimusComputeDataInterface>(DataInterfaceClass->GetDefaultObject());
 		if (DataInterface != nullptr)
diff --git a/Engine/Plugins/Animation/IKRig/Source/IKRigEditor/Private/RigEditor/IKRigController.cpp b/Engine/Plugins/Animation/IKRig/Source/IKRigEditor/Private/RigEditor/IKRigController.cpp
--- a/Engine/Plugins/Animation/IKRig/Source/IKRigEditor/Private/RigEditor/IKRigController.cpp
+++ b/Engine/Plugins/Animation/IKRig/Source/IKRigEditor/Private/RigEditor/IKRigController.cpp
@@ -1185,7 +1185,7 @@ bool UIKRigController::IsGoalConnectedToSolver(const FName GoalName, int32 Solve
 
 bool UIKRigController::IsGoalConnectedToAnySolver(const FName GoalName) const
 {
-	for (const TObjectPtr<UIKRigSolver> Solver : Asset->Solvers)
+	for (const TObjectPtr<UIKRigSolver>& Solver : Asset->Solvers)
 	{
 		if (Solver->IsGoalConnected(GoalName))
 		{
diff --git a/Engine/Plugins/Developer/PerforceSourceControl/Source/PerforceSourceControl/PerforceSourceControl.Build.cs b/Engine/Plugins/Developer/PerforceSourceControl/Source/PerforceSourceControl/PerforceSourceControl.Build.cs
--- a/Engine/Plugins/Developer/PerforceSourceControl/Source/PerforceSourceControl/PerforceSourceControl.Build.cs
+++ b/Engine/Plugins/Developer/PerforceSourceControl/Source/PerforceSourceControl/PerforceSourceControl.Build.cs
@@ -11,6 +11,7 @@ public class PerforceSourceControl : ModuleRules
 		PrivateDependencyModuleNames.AddRange(
 			new string[] {
 				"Core",
+				"CoreUObject",
 				"InputCore",
 				"Slate",
 				"SlateCore",
diff --git a/Engine/Plugins/Editor/ObjectMixer/ObjectMixer/Source/ObjectMixer/Private/ObjectMixerEditorModule.cpp b/Engine/Plugins/Editor/ObjectMixer/ObjectMixer/Source/ObjectMixer/Private/ObjectMixerEditorModule.cpp
--- a/Engine/Plugins/Editor/ObjectMixer/ObjectMixer/Source/ObjectMixer/Private/ObjectMixerEditorModule.cpp
+++ b/Engine/Plugins/Editor/ObjectMixer/ObjectMixer/Source/ObjectMixer/Private/ObjectMixerEditorModule.cpp
@@ -394,7 +394,7 @@ TSet<FName> FObjectMixerEditorModule::GetPropertiesThatRequireRefresh() const
 
 	if (ListModel.IsValid())
 	{
-		for (const TObjectPtr<UObjectMixerObjectFilter> Instance : ListModel->GetObjectFilterInstances())
+		for (const TObjectPtr<UObjectMixerObjectFilter>& Instance : ListModel->GetObjectFilterInstances())
 		{
 			ReturnValue.Append(Instance->GetPropertiesThatRequireListRefresh());
 		}
diff --git a/Engine/Plugins/Editor/ObjectMixer/ObjectMixer/Source/ObjectMixer/Private/Views/List/ObjectMixerEditorList.cpp b/Engine/Plugins/Editor/ObjectMixer/ObjectMixer/Source/ObjectMixer/Private/Views/List/ObjectMixerEditorList.cpp
--- a/Engine/Plugins/Editor/ObjectMixer/ObjectMixer/Source/ObjectMixer/Private/Views/List/ObjectMixerEditorList.cpp
+++ b/Engine/Plugins/Editor/ObjectMixer/ObjectMixer/Source/ObjectMixer/Private/Views/List/ObjectMixerEditorList.cpp
@@ -89,7 +89,7 @@ void FObjectMixerEditorList::CacheObjectFilterInstances()
 
 	if (ObjectFilterClasses.Num() > 0)
 	{
-		for (const TSubclassOf<UObjectMixerObjectFilter> Class : ObjectFilterClasses)
+		for (const TSubclassOf<UObjectMixerObjectFilter>& Class : ObjectFilterClasses)
 		{
 			UObjectMixerObjectFilter* NewInstance = NewObject<UObjectMixerObjectFilter>(GetTransientPackage(), Class);
 			ObjectFilterInstances.Add(TObjectPtr<UObjectMixerObjectFilter>(NewInstance));
diff --git a/Engine/Plugins/Editor/WaveformEditor/Source/WaveformEditor/Private/WaveformEditor.cpp b/Engine/Plugins/Editor/WaveformEditor/Source/WaveformEditor/Private/WaveformEditor.cpp
--- a/Engine/Plugins/Editor/WaveformEditor/Source/WaveformEditor/Private/WaveformEditor.cpp
+++ b/Engine/Plugins/Editor/WaveformEditor/Source/WaveformEditor/Private/WaveformEditor.cpp
@@ -103,7 +103,7 @@ void FWaveformEditor::AddDefaultTransformations()
 	if (SoundWave->Transformations.Num() == 0)
 	{
 		const UWaveformEditorTransformationsSettings* TransformationsSettings = GetWaveformEditorTransformationsSettings();
-		for (const TSubclassOf<UWaveformTransformationBase> TransformationClass : TransformationsSettings->LaunchTransformations)
+		for (const TSubclassOf<UWaveformTransformationBase>& TransformationClass : TransformationsSettings->LaunchTransformations)
 		{
 			if (TransformationClass)
 			{
diff --git a/Engine/Plugins/Editor/WaveformEditor/Source/WaveformEditorWidgets/Private/WaveformEditorSequenceDataProvider.cpp b/Engine/Plugins/Editor/WaveformEditor/Source/WaveformEditorWidgets/Private/WaveformEditorSequenceDataProvider.cpp
--- a/Engine/Plugins/Editor/WaveformEditor/Source/WaveformEditorWidgets/Private/WaveformEditorSequenceDataProvider.cpp
+++ b/Engine/Plugins/Editor/WaveformEditor/Source/WaveformEditorWidgets/Private/WaveformEditorSequenceDataProvider.cpp
@@ -62,7 +62,7 @@ void FWaveformEditorSequenceDataProvider::GenerateLayersChain()
 	{
 		FTransformationsToPropertiesArray TransformPropertiesMap;
 
-		for (const TObjectPtr<UWaveformTransformationBase> Transformation : SoundWaveToRender->Transformations)
+		for (const TObjectPtr<UWaveformTransformationBase>& Transformation : SoundWaveToRender->Transformations)
 		{
 			if (Transformation != nullptr)
 			{
@@ -212,7 +212,7 @@ TArray<Audio::FTransformationPtr> FWaveformEditorSequenceDataProvider::CreateTra
 {
 	TArray<Audio::FTransformationPtr> TransformationPtrs;
 
-	for (const TObjectPtr<UWaveformTransformationBase> TransformationBase : TransformationsToRender)
+	for (const TObjectPtr<UWaveformTransformationBase>& TransformationBase : TransformationsToRender)
 	{
 		TransformationPtrs.Add(TransformationBase->CreateTransformation());
 	}
diff --git a/Engine/Plugins/EnhancedInput/Source/EnhancedInput/Private/EnhancedInputPlatformSettings.cpp b/Engine/Plugins/EnhancedInput/Source/EnhancedInput/Private/EnhancedInputPlatformSettings.cpp
--- a/Engine/Plugins/EnhancedInput/Source/EnhancedInput/Private/EnhancedInputPlatformSettings.cpp
+++ b/Engine/Plugins/EnhancedInput/Source/EnhancedInput/Private/EnhancedInputPlatformSettings.cpp
@@ -19,7 +19,7 @@ EDataValidationResult UEnhancedInputPlatformData::IsDataValid(FDataValidationCon
 {
 	EDataValidationResult Result = Super::IsDataValid(Context);
 
-	for (const TPair<TObjectPtr<const UInputMappingContext>, TObjectPtr<const UInputMappingContext>> Pair : MappingContextRedirects)
+	for (const TPair<TObjectPtr<const UInputMappingContext>, TObjectPtr<const UInputMappingContext>>& Pair : MappingContextRedirects)
 	{
 		if (!Pair.Key)
 		{
diff --git a/Engine/Plugins/EnhancedInput/Source/EnhancedInput/Private/EnhancedInputSubsystemInterface.cpp b/Engine/Plugins/EnhancedInput/Source/EnhancedInput/Private/EnhancedInputSubsystemInterface.cpp
--- a/Engine/Plugins/EnhancedInput/Source/EnhancedInput/Private/EnhancedInputSubsystemInterface.cpp
+++ b/Engine/Plugins/EnhancedInput/Source/EnhancedInput/Private/EnhancedInputSubsystemInterface.cpp
@@ -854,7 +854,7 @@ void IEnhancedInputSubsystemInterface::RebuildControlMappings()
 	{
 		TMap<TObjectPtr<const UInputMappingContext>, TObjectPtr<const UInputMappingContext>> ContextRedirects;
 		PlatformSettings->GetAllMappingContextRedirects(ContextRedirects);
-		for (const TPair<TObjectPtr<const UInputMappingContext>, TObjectPtr<const UInputMappingContext>> Pair : ContextRedirects)
+		for (const TPair<TObjectPtr<const UInputMappingContext>, TObjectPtr<const UInputMappingContext>>& Pair : ContextRedirects)
 		{
 			if (!Pair.Key || !Pair.Value)
 			{
diff --git a/Engine/Plugins/Experimental/Chooser/Source/ProxyTable/Private/ProxyTable.cpp b/Engine/Plugins/Experimental/Chooser/Source/ProxyTable/Private/ProxyTable.cpp
--- a/Engine/Plugins/Experimental/Chooser/Source/ProxyTable/Private/ProxyTable.cpp
+++ b/Engine/Plugins/Experimental/Chooser/Source/ProxyTable/Private/ProxyTable.cpp
@@ -91,7 +91,7 @@ static void BuildRuntimeDataRecursive(UProxyTable* RootTable, UProxyTable* Table
 		}
 	}
 
-	for (const TObjectPtr<UProxyTable> ParentTable : Table->InheritEntriesFrom)
+	for (const TObjectPtr<UProxyTable>& ParentTable : Table->InheritEntriesFrom)
 	{
 		if (!OutDependencies.Contains(ParentTable))
 		{
diff --git a/Engine/Plugins/Experimental/Mutable/Source/CustomizableObjectEditor/Private/MuCOE/Nodes/CustomizableObjectNodeMaterial.cpp b/Engine/Plugins/Experimental/Mutable/Source/CustomizableObjectEditor/Private/MuCOE/Nodes/CustomizableObjectNodeMaterial.cpp
--- a/Engine/Plugins/Experimental/Mutable/Source/CustomizableObjectEditor/Private/MuCOE/Nodes/CustomizableObjectNodeMaterial.cpp
+++ b/Engine/Plugins/Experimental/Mutable/Source/CustomizableObjectEditor/Private/MuCOE/Nodes/CustomizableObjectNodeMaterial.cpp
@@ -499,7 +499,7 @@ void UCustomizableObjectNodeMaterial::BackwardsCompatibleFixup()
 					FName TextureName = TextureParameterInfo[TextureIndex].Name;
 
 					// Checking if the pin's texture has been modified in the material instance
-					for (const FTextureParameterValue Texture : DefaultPinValue->TextureParameterValues)
+					for (const FTextureParameterValue& Texture : DefaultPinValue->TextureParameterValues)
 					{
 						if (TextureName == Texture.ParameterInfo.Name)
 						{
diff --git a/Engine/Plugins/Experimental/PythonScriptPlugin/Source/PythonScriptPlugin/Private/PyWrapperObject.cpp b/Engine/Plugins/Experimental/PythonScriptPlugin/Source/PythonScriptPlugin/Private/PyWrapperObject.cpp
--- a/Engine/Plugins/Experimental/PythonScriptPlugin/Source/PythonScriptPlugin/Private/PyWrapperObject.cpp
+++ b/Engine/Plugins/Experimental/PythonScriptPlugin/Source/PythonScriptPlugin/Private/PyWrapperObject.cpp
@@ -503,7 +503,7 @@ PyTypeObject InitializePyWrapperObjectType()
 				return nullptr;
 			}
 
-			return PyUnicode_FromFormat("<Object '%s' (%p) Class '%s'>", TCHAR_TO_UTF8(*InSelf->ObjectInstance->GetPathName()), InSelf->ObjectInstance, TCHAR_TO_UTF8(*InSelf->ObjectInstance->GetClass()->GetName()));
+			return PyUnicode_FromFormat("<Object '%s' (%p) Class '%s'>", TCHAR_TO_UTF8(*InSelf->ObjectInstance->GetPathName()), InSelf->ObjectInstance.Get(), TCHAR_TO_UTF8(*InSelf->ObjectInstance->GetClass()->GetName()));
 		}
 
 		static PyUtil::FPyHashType Hash(FPyWrapperObject* InSelf)
diff --git a/Engine/Plugins/MovieScene/MovieRenderPipeline/Source/MovieRenderPipelineCore/Private/Graph/MovieGraphConfig.cpp b/Engine/Plugins/MovieScene/MovieRenderPipeline/Source/MovieRenderPipelineCore/Private/Graph/MovieGraphConfig.cpp
--- a/Engine/Plugins/MovieScene/MovieRenderPipeline/Source/MovieRenderPipelineCore/Private/Graph/MovieGraphConfig.cpp
+++ b/Engine/Plugins/MovieScene/MovieRenderPipeline/Source/MovieRenderPipelineCore/Private/Graph/MovieGraphConfig.cpp
@@ -611,7 +611,7 @@ UMovieGraphOutput* UMovieGraphConfig::AddOutput()
 
 UMovieGraphVariable* UMovieGraphConfig::GetVariableByGuid(const FGuid& InGuid) const
 {
-	for (const TObjectPtr<UMovieGraphVariable> Variable : Variables)
+	for (const TObjectPtr<UMovieGraphVariable>& Variable : Variables)
 	{
 		if (Variable->GetGuid() == InGuid)
 		{
diff --git a/Engine/Plugins/MovieScene/MovieRenderPipeline/Source/MovieRenderPipelineEditor/Private/Graph/MovieEdGraph.cpp b/Engine/Plugins/MovieScene/MovieRenderPipeline/Source/MovieRenderPipelineEditor/Private/Graph/MovieEdGraph.cpp
--- a/Engine/Plugins/MovieScene/MovieRenderPipeline/Source/MovieRenderPipelineEditor/Private/Graph/MovieEdGraph.cpp
+++ b/Engine/Plugins/MovieScene/MovieRenderPipeline/Source/MovieRenderPipelineEditor/Private/Graph/MovieEdGraph.cpp
@@ -42,7 +42,7 @@ void UMoviePipelineEdGraph::InitFromRuntimeGraph(UMovieGraphConfig* InGraph)
 	}
 
 	// Create the rest of the nodes in the graph
-	for (const TObjectPtr<UMovieGraphNode> RuntimeNode : InGraph->GetNodes())
+	for (const TObjectPtr<UMovieGraphNode>& RuntimeNode : InGraph->GetNodes())
 	{
 		if (RuntimeNode->IsA<UMovieGraphVariableNode>())
 		{
@@ -64,7 +64,7 @@ void UMoviePipelineEdGraph::InitFromRuntimeGraph(UMovieGraphConfig* InGraph)
 	}
 
 	// Restore editor-only nodes, which have no runtime node equivalent
-	for (const TObjectPtr<UObject> EditorOnlyNodeObject : InGraph->GetEditorOnlyNodes())
+	for (const TObjectPtr<UObject>& EditorOnlyNodeObject : InGraph->GetEditorOnlyNodes())
 	{
 		if (const UEdGraphNode* EdGraphNode = Cast<UEdGraphNode>(EditorOnlyNodeObject))
 		{
@@ -164,7 +164,7 @@ void UMoviePipelineEdGraph::OnGraphConfigChanged()
 {
 	// TODO: Optimize this. Ideally we can target specific changes to the graph rather than rebuilding everything.
 	// However, this isn't strictly necessary unless rebuilding becomes a performance bottleneck.
-	for (const TObjectPtr<UEdGraphNode> Node: Nodes)
+	for (const TObjectPtr<UEdGraphNode>& Node: Nodes)
 	{
 		Node->ReconstructNode();
 	}
@@ -174,7 +174,7 @@ void UMoviePipelineEdGraph::OnGraphNodesDeleted(TArray<UMovieGraphNode*> Deleted
 {
 	TArray<TObjectPtr<UEdGraphNode>> EdGraphNodesToDelete;
 	
-	for (const TObjectPtr<UEdGraphNode> Node : Nodes)
+	for (const TObjectPtr<UEdGraphNode>& Node : Nodes)
 	{
 		if (!Node)
 		{
diff --git a/Engine/Plugins/Runtime/CommonUI/Source/CommonUI/Private/CommonActionWidget.cpp b/Engine/Plugins/Runtime/CommonUI/Source/CommonUI/Private/CommonActionWidget.cpp
--- a/Engine/Plugins/Runtime/CommonUI/Source/CommonUI/Private/CommonActionWidget.cpp
+++ b/Engine/Plugins/Runtime/CommonUI/Source/CommonUI/Private/CommonActionWidget.cpp
@@ -170,7 +170,7 @@ bool UCommonActionWidget::IsHeldAction() const
 {
 	if (EnhancedInputAction && CommonUI::IsEnhancedInputSupportEnabled())
 	{
-		for (const TObjectPtr<UInputTrigger> Trigger : EnhancedInputAction->Triggers)
+		for (const TObjectPtr<UInputTrigger>& Trigger : EnhancedInputAction->Triggers)
 		{
 			if (EnumHasAnyFlags(Trigger->GetSupportedTriggerEvents(), ETriggerEventsSupported::Ongoing))
 			{
diff --git a/Engine/Plugins/Runtime/GameplayAbilities/Source/GameplayAbilities/Private/GameplayCueNotifyTypes.cpp b/Engine/Plugins/Runtime/GameplayAbilities/Source/GameplayAbilities/Private/GameplayCueNotifyTypes.cpp
--- a/Engine/Plugins/Runtime/GameplayAbilities/Source/GameplayAbilities/Private/GameplayCueNotifyTypes.cpp
+++ b/Engine/Plugins/Runtime/GameplayAbilities/Source/GameplayAbilities/Private/GameplayCueNotifyTypes.cpp
@@ -882,7 +882,7 @@ bool FGameplayCueNotify_InputDevicePropertyInfo::SetDeviceProperties(const FGame
 void FGameplayCueNotify_InputDevicePropertyInfo::ValidateBurstAssets(const UObject* ContainingAsset, const FString& Context, FDataValidationContext& ValidationContext) const
 {
 #if WITH_EDITORONLY_DATA
-	for (const TSubclassOf<UInputDeviceProperty> PropClass : DeviceProperties)
+	for (const TSubclassOf<UInputDeviceProperty>& PropClass : DeviceProperties)
 	{
 		if (!PropClass)
 		{
diff --git a/Engine/Plugins/Runtime/MassGameplay/Source/MassSpawner/Private/MassEntityConfigAsset.cpp b/Engine/Plugins/Runtime/MassGameplay/Source/MassSpawner/Private/MassEntityConfigAsset.cpp
--- a/Engine/Plugins/Runtime/MassGameplay/Source/MassSpawner/Private/MassEntityConfigAsset.cpp
+++ b/Engine/Plugins/Runtime/MassGameplay/Source/MassSpawner/Private/MassEntityConfigAsset.cpp
@@ -34,7 +34,7 @@ FMassEntityConfig::FMassEntityConfig(UObject& InOwner)
 
 const UMassEntityTraitBase* FMassEntityConfig::FindTrait(TSubclassOf<UMassEntityTraitBase> TraitClass, const bool bExactMatch) const
 {
-	for (const TObjectPtr<UMassEntityTraitBase> Trait : Traits)
+	for (const TObjectPtr<UMassEntityTraitBase>& Trait : Traits)
 	{
 		if (Trait && (bExactMatch ? Trait->GetClass() == TraitClass : Trait->IsA(TraitClass)))
 		{
diff --git a/Engine/Plugins/Runtime/MassGameplay/Source/MassSpawner/Private/MassSpawner.cpp b/Engine/Plugins/Runtime/MassGameplay/Source/MassSpawner/Private/MassSpawner.cpp
--- a/Engine/Plugins/Runtime/MassGameplay/Source/MassSpawner/Private/MassSpawner.cpp
+++ b/Engine/Plugins/Runtime/MassGameplay/Source/MassSpawner/Private/MassSpawner.cpp
@@ -459,7 +459,7 @@ void AMassSpawner::SpawnGeneratedEntities(TConstArrayView<FMassEntitySpawnDataGe
 
 	for (const FMassEntitySpawnDataGeneratorResult& Result : Results)
 	{
-		for (const TSubclassOf<UMassProcessor> ProcessorClass : Result.PostSpawnProcessors)
+		for (const TSubclassOf<UMassProcessor>& ProcessorClass : Result.PostSpawnProcessors)
 		{
 			if (AddedProcessors.Contains(ProcessorClass) == false)
 			{
diff --git a/Engine/Plugins/Runtime/Metasound/Source/MetasoundGenerator/MetasoundGenerator.Build.cs b/Engine/Plugins/Runtime/Metasound/Source/MetasoundGenerator/MetasoundGenerator.Build.cs
--- a/Engine/Plugins/Runtime/Metasound/Source/MetasoundGenerator/MetasoundGenerator.Build.cs
+++ b/Engine/Plugins/Runtime/Metasound/Source/MetasoundGenerator/MetasoundGenerator.Build.cs
@@ -20,6 +20,7 @@ namespace UnrealBuildTool.Rules
             PublicDependencyModuleNames.AddRange(
 				new string[] {
                     "Core",
+					"CoreUObject",
 					"AudioMixer",
 					"MetasoundGraphCore",
 					"MetasoundStandardNodes",
diff --git a/Engine/Plugins/Runtime/NetworkPredictionInsights/Source/NetworkPredictionInsights/NetworkPredictionInsights.Build.cs b/Engine/Plugins/Runtime/NetworkPredictionInsights/Source/NetworkPredictionInsights/NetworkPredictionInsights.Build.cs
--- a/Engine/Plugins/Runtime/NetworkPredictionInsights/Source/NetworkPredictionInsights/NetworkPredictionInsights.Build.cs
+++ b/Engine/Plugins/Runtime/NetworkPredictionInsights/Source/NetworkPredictionInsights/NetworkPredictionInsights.Build.cs
@@ -10,6 +10,7 @@ namespace UnrealBuildTool.Rules
 			new string[]
 			{
 				"Core",
+				"CoreUObject",
 				"InputCore",
 				"SlateCore",
 				"Slate",
diff --git a/Engine/Source/Editor/Persona/Private/AnimTimeline/AnimTimelineClipboard.cpp b/Engine/Source/Editor/Persona/Private/AnimTimeline/AnimTimelineClipboard.cpp
--- a/Engine/Source/Editor/Persona/Private/AnimTimeline/AnimTimelineClipboard.cpp
+++ b/Engine/Source/Editor/Persona/Private/AnimTimeline/AnimTimelineClipboard.cpp
@@ -25,7 +25,7 @@ FAnimationCurveIdentifier UAnimCurveBaseCopyObject::GetAnimationCurveIdentifier(
 
 bool UAnimTimelineClipboardContent::AreAllCurvesOfTrackType(const ERawCurveTrackTypes Type) const
 {
-	for (const TObjectPtr<UAnimCurveBaseCopyObject> CurveCopyObj : Curves)
+	for (const TObjectPtr<UAnimCurveBaseCopyObject>& CurveCopyObj : Curves)
 	{
 		if (CurveCopyObj->CurveType != Type)
 		{
diff --git a/Engine/Source/Editor/Sequencer/Private/MVVM/ViewModels/ChannelModel.cpp b/Engine/Source/Editor/Sequencer/Private/MVVM/ViewModels/ChannelModel.cpp
--- a/Engine/Source/Editor/Sequencer/Private/MVVM/ViewModels/ChannelModel.cpp
+++ b/Engine/Source/Editor/Sequencer/Private/MVVM/ViewModels/ChannelModel.cpp
@@ -551,7 +551,7 @@ void FChannelGroupOverrideHelper::BuildChannelOverrideMenu(FMenuBuilder& MenuBui
 	}
 
 	UMovieSceneChannelOverrideContainer::FOverrideCandidates CommonCandidateOverrides;
-	for (const TTuple<TSubclassOf<UMovieSceneChannelOverrideContainer>, int32> CandidateCount : CandidateCounts)
+	for (const TTuple<TSubclassOf<UMovieSceneChannelOverrideContainer>, int32>& CandidateCount : CandidateCounts)
 	{
 		if (CandidateCount.Value == NumOverridableChannels)
 		{
diff --git a/Engine/Source/Editor/UnrealEd/Private/Toolkits/AssetEditorToolkit.cpp b/Engine/Source/Editor/UnrealEd/Private/Toolkits/AssetEditorToolkit.cpp
--- a/Engine/Source/Editor/UnrealEd/Private/Toolkits/AssetEditorToolkit.cpp
+++ b/Engine/Source/Editor/UnrealEd/Private/Toolkits/AssetEditorToolkit.cpp
@@ -570,7 +570,7 @@ TArray<TObjectPtr<UObject>>& FAssetEditorToolkit::GetEditingObjectPtrs()
 
 void FAssetEditorToolkit::GetSaveableObjects(TArray<UObject*>& OutObjects) const
 {
-	for (const auto Object : EditingObjects)
+	for (const TObjectPtr<UObject>& Object : EditingObjects)
 	{
 		// If we are editing a subobject of asset (e.g., a level script blueprint which is contained in a map asset), still provide the
 		// option to work with it but treat save operations/etc... as working on the top level asset itself
diff --git a/Engine/Source/Editor/UnrealEd/Private/WorldPartition/WorldPartitionNavigationDataBuilder.cpp b/Engine/Source/Editor/UnrealEd/Private/WorldPartition/WorldPartitionNavigationDataBuilder.cpp
--- a/Engine/Source/Editor/UnrealEd/Private/WorldPartition/WorldPartitionNavigationDataBuilder.cpp
+++ b/Engine/Source/Editor/UnrealEd/Private/WorldPartition/WorldPartitionNavigationDataBuilder.cpp
@@ -30,7 +30,7 @@ bool UWorldPartitionNavigationDataBuilder::PreRun(UWorld* World, FPackageSourceC
 {
 	// Set runtime data layer to be included in the base navmesh generation.
 	UDataLayerManager* DataLayerManager = UDataLayerManager::GetDataLayerManager(World);
-	for (const TObjectPtr<UDataLayerAsset> DataLayer : World->GetWorldSettings()->BaseNavmeshDataLayers)
+	for (const TObjectPtr<UDataLayerAsset>& DataLayer : World->GetWorldSettings()->BaseNavmeshDataLayers)
 	{
 		if (DataLayer != nullptr)
 		{
diff --git a/Engine/Source/Runtime/Engine/Private/Components/StaticMeshComponent.cpp b/Engine/Source/Runtime/Engine/Private/Components/StaticMeshComponent.cpp
--- a/Engine/Source/Runtime/Engine/Private/Components/StaticMeshComponent.cpp
+++ b/Engine/Source/Runtime/Engine/Private/Components/StaticMeshComponent.cpp
@@ -693,7 +693,7 @@ void UStaticMeshComponent::OutdatedKnownStaticMeshDetected() const
 		TEXT("StaticMesh property overwritten for component %s without a call to NotifyIfStaticMeshChanged(). KnownStaticMesh (%p) != StaticMesh (%p - %s)"),
 		*GetFullName(),
 		KnownStaticMesh,
-		StaticMesh,
+		StaticMesh.Get(),
 		StaticMesh ? *StaticMesh->GetFullName() : TEXT("nullptr")
 		);
 
diff --git a/Engine/Source/Runtime/Engine/Private/Texture.cpp b/Engine/Source/Runtime/Engine/Private/Texture.cpp
--- a/Engine/Source/Runtime/Engine/Private/Texture.cpp
+++ b/Engine/Source/Runtime/Engine/Private/Texture.cpp
@@ -1063,7 +1063,7 @@ PRAGMA_DISABLE_DEPRECATION_WARNINGS
 		TEXT("CompositeTexture property overwritten for texture %s without a call to NotifyIfCompositeTextureChanged(). KnownCompositeTexture (%p) != CompositeTexture (%p - %s)"),
 		*GetFullName(),
 		KnownCompositeTexture,
-		CompositeTexture,
+		CompositeTexture.Get(),
 		CompositeTexture ? *CompositeTexture->GetFullName() : TEXT("nullptr")
 	);
 
diff --git a/Engine/Source/Runtime/MovieScene/Private/MovieScene.cpp b/Engine/Source/Runtime/MovieScene/Private/MovieScene.cpp
--- a/Engine/Source/Runtime/MovieScene/Private/MovieScene.cpp
+++ b/Engine/Source/Runtime/MovieScene/Private/MovieScene.cpp
@@ -1159,7 +1159,7 @@ UMovieSceneTrack* UMovieScene::FindTrack( TSubclassOf<UMovieSceneTrack> TrackCla
 {
 	UMovieSceneTrack* FoundTrack = nullptr;
 
-	for (const auto Track : Tracks)
+	for (const TObjectPtr<UMovieSceneTrack>& Track : Tracks)
 	{
 		if (Track->GetClass()->IsChildOf(TrackClass))
 		{
diff --git a/Engine/Source/Runtime/NavigationSystem/Private/NavModifierComponent.cpp b/Engine/Source/Runtime/NavigationSystem/Private/NavModifierComponent.cpp
--- a/Engine/Source/Runtime/NavigationSystem/Private/NavModifierComponent.cpp
+++ b/Engine/Source/Runtime/NavigationSystem/Private/NavModifierComponent.cpp
@@ -162,7 +162,7 @@ void UNavModifierComponent::CalcAndCacheBounds() const
 					{
 						Bounds += GeometryCollection->Bounds.GetBox();
 
-						for (const TObjectPtr<UStaticMesh> ProxyMesh : RestCollection->RootProxyData.ProxyMeshes)
+						for (const TObjectPtr<UStaticMesh>& ProxyMesh : RestCollection->RootProxyData.ProxyMeshes)
 						{
 							if (ProxyMesh != nullptr)
 							{
diff --git a/Engine/Source/Runtime/Slate/Public/Framework/Views/TableViewTypeTraits.h b/Engine/Source/Runtime/Slate/Public/Framework/Views/TableViewTypeTraits.h
--- a/Engine/Source/Runtime/Slate/Public/Framework/Views/TableViewTypeTraits.h
+++ b/Engine/Source/Runtime/Slate/Public/Framework/Views/TableViewTypeTraits.h
@@ -504,7 +504,7 @@ public:
 
 	static FString DebugDump(TObjectPtr<T> InPtr)
 	{
-		return InPtr ? FString::Printf(TEXT("0x%08x [%s]"), InPtr, *InPtr->GetName()) : FString(TEXT("nullptr"));
+		return InPtr ? FString::Printf(TEXT("0x%08x [%s]"), InPtr.Get(), *InPtr->GetName()) : FString(TEXT("nullptr"));
 	}
 
 	typedef FGCObject SerializerType;
