MoviePipeline: Added collections and modifiers to the graph runtime.

There will need to be updates to the graph UX, and most likely address performance concerns, but this paves the way to making functional render layers.

#rnx

diff --git a/Engine/Plugins/MovieScene/MovieRenderPipeline/Source/MovieRenderPipelineCore/Private/Graph/MovieGraphLinearTimeStep.cpp b/Engine/Plugins/MovieScene/MovieRenderPipeline/Source/MovieRenderPipelineCore/Private/Graph/MovieGraphLinearTimeStep.cpp
--- a/Engine/Plugins/MovieScene/MovieRenderPipeline/Source/MovieRenderPipelineCore/Private/Graph/MovieGraphLinearTimeStep.cpp
+++ b/Engine/Plugins/MovieScene/MovieRenderPipeline/Source/MovieRenderPipelineCore/Private/Graph/MovieGraphLinearTimeStep.cpp
@@ -84,6 +84,13 @@ void UMovieGraphLinearTimeStep::TickProducingFrames()
 		// Seeks the external data source to match
 		GetOwningGraph()->GetDataSourceInstance()->InitializeShot(CurrentCameraCut);
 
+		// Generate render layers from the evaluated graph
+		GetOwningGraph()->CreateLayersInRenderLayerSubsystem(CurrentTimeStepData.EvaluatedConfig);
+
+		// Update the contents of the render layers
+		// TODO: This should probably be done per-frame
+		GetOwningGraph()->UpdateLayerContentsInRenderLayerSubsystem(CurrentTimeStepData.EvaluatedConfig);
+
 		const bool bIncludeCDO = false;
 		UMovieGraphOutputSettingNode* OutputNode = CurrentFrameData.EvaluatedConfig->GetSettingForBranch<UMovieGraphOutputSettingNode>(UMovieGraphNode::GlobalsPinName, bIncludeCDO);
 
diff --git a/Engine/Plugins/MovieScene/MovieRenderPipeline/Source/MovieRenderPipelineCore/Private/Graph/MovieGraphPipeline.cpp b/Engine/Plugins/MovieScene/MovieRenderPipeline/Source/MovieRenderPipelineCore/Private/Graph/MovieGraphPipeline.cpp
--- a/Engine/Plugins/MovieScene/MovieRenderPipeline/Source/MovieRenderPipelineCore/Private/Graph/MovieGraphPipeline.cpp
+++ b/Engine/Plugins/MovieScene/MovieRenderPipeline/Source/MovieRenderPipelineCore/Private/Graph/MovieGraphPipeline.cpp
@@ -5,7 +5,10 @@
 #include "Graph/MovieGraphCVarManager.h"
 #include "Graph/MovieGraphDataTypes.h"
 #include "Graph/MovieGraphOutputMerger.h"
+#include "Graph/MoviePipelineRenderLayerSubsystem.h"
+#include "Graph/Nodes/MovieGraphCollectionNode.h"
 #include "Graph/Nodes/MovieGraphFileOutputNode.h"
+#include "Graph/Nodes/MovieGraphModifierNode.h"
 #include "Graph/Nodes/MovieGraphOutputSettingNode.h"
 #include "Graph/MovieGraphBlueprintLibrary.h"
 #include "MovieRenderPipelineCoreModule.h"
@@ -164,6 +167,78 @@ void UMovieGraphPipeline::SetPreviewWidgetVisibleImpl(bool bInIsVisible)
 	}
 }
 
+void UMovieGraphPipeline::CreateLayersInRenderLayerSubsystem(const UMovieGraphEvaluatedConfig* EvaluatedConfig) const
+{
+	UMoviePipelineRenderLayerSubsystem* LayerSubsystem = GetWorld()->GetSubsystem<UMoviePipelineRenderLayerSubsystem>();
+	if (!LayerSubsystem)
+	{
+		return;
+	}
+	
+	LayerSubsystem->Reset();
+
+	// One render layer is generated per branch
+	for (const FName& BranchName : EvaluatedConfig->GetBranchNames())
+	{
+		UMoviePipelineRenderLayer* RenderLayer = NewObject<UMoviePipelineRenderLayer>();
+		RenderLayer->SetRenderLayerName(BranchName);
+		LayerSubsystem->AddRenderLayer(RenderLayer);
+	}
+}
+
+void UMovieGraphPipeline::UpdateLayerContentsInRenderLayerSubsystem(const UMovieGraphEvaluatedConfig* EvaluatedConfig) const
+{
+	UMoviePipelineRenderLayerSubsystem* LayerSubsystem = GetWorld()->GetSubsystem<UMoviePipelineRenderLayerSubsystem>();
+	if (!LayerSubsystem)
+	{
+		return;
+	}
+
+	for (UMoviePipelineRenderLayer* RenderLayer : LayerSubsystem->GetRenderLayers())
+	{
+		const FName& LayerName = RenderLayer->GetRenderLayerName();
+		
+		// Gather all collection and modifier nodes from the evaluated graph
+		constexpr bool bIncludeCDOs = false;
+		constexpr bool bExactMatch = true;
+		const TArray<UMovieGraphCollectionNode*> CollectionNodes =
+			EvaluatedConfig->GetSettingsForBranch<UMovieGraphCollectionNode>(LayerName, bIncludeCDOs, bExactMatch);
+		const TArray<UMovieGraphModifierNode*> ModifierNodes =
+			EvaluatedConfig->GetSettingsForBranch<UMovieGraphModifierNode>(LayerName, bIncludeCDOs, bExactMatch);
+		
+		for (const UMovieGraphModifierNode* ModifierNode : ModifierNodes)
+		{
+			ModifierNode->ModifierClass->SetCollections({});
+			
+			// Find the collection that this modifier is modifying
+			for (const UMovieGraphCollectionNode* CollectionNode : CollectionNodes)
+			{
+				if (CollectionNode && (CollectionNode->CollectionName == ModifierNode->ModifiedCollectionName))
+				{
+					UMoviePipelineCollection* NewCollection = NewObject<UMoviePipelineCollection>();
+					NewCollection->AddQuery(CollectionNode->QueryClass);
+					NewCollection->SetCollectionName(CollectionNode->CollectionName);
+					
+					ModifierNode->ModifierClass->SetCollections({NewCollection});
+					break;
+				}
+			}
+
+			// If the modifier had a valid collection added to it, add the modifier to the render layer. The modifier
+			// won't have any effect without a collection to act on.
+			if (ModifierNode->ModifierClass->GetCollections().IsEmpty())
+			{
+				UE_LOG(LogMovieRenderPipeline, Warning, TEXT("The modifier '%s' specified a collection '%s', but the collection couldn't be found."),
+					*ModifierNode->ModifierName, *ModifierNode->ModifiedCollectionName);
+			}
+			else
+			{
+				RenderLayer->AddModifier(ModifierNode->ModifierClass);
+			}
+		}
+	}
+}
+
 void UMovieGraphPipeline::BuildShotListFromDataSource()
 {
 	// Synchronize our shot list with our target data source. New shots will be added and outdated shots removed.
diff --git a/Engine/Plugins/MovieScene/MovieRenderPipeline/Source/MovieRenderPipelineCore/Private/Graph/MoviePipelineRenderLayerSubsystem.cpp b/Engine/Plugins/MovieScene/MovieRenderPipeline/Source/MovieRenderPipelineCore/Private/Graph/MoviePipelineRenderLayerSubsystem.cpp
--- a/Engine/Plugins/MovieScene/MovieRenderPipeline/Source/MovieRenderPipelineCore/Private/Graph/MoviePipelineRenderLayerSubsystem.cpp
+++ b/Engine/Plugins/MovieScene/MovieRenderPipeline/Source/MovieRenderPipelineCore/Private/Graph/MoviePipelineRenderLayerSubsystem.cpp
@@ -3,6 +3,9 @@
 #include "Graph/MoviePipelineRenderLayerSubsystem.h"
 
 #include "Components/PrimitiveComponent.h"
+#include "Components/SkyAtmosphereComponent.h"
+#include "Components/SkyLightComponent.h"
+#include "Components/SphereReflectionCaptureComponent.h"
 #include "EngineUtils.h"
 #include "Materials/MaterialInterface.h"
 #include "MovieRenderPipelineCoreModule.h"
@@ -174,8 +177,37 @@ bool UMoviePipelineCollectionCommonQuery::DoesActorMatchQuery(const AActor* Acto
 	return bMatchesActorNames || bMatchesTags || bMatchesComponentTypes;
 }
 
+bool UMoviePipelineCollectionLightingQuery::DoesActorMatchQuery(const AActor* Actor) const
+{
+	const TArray<UClass*> LightingComponentTypes = {
+		ULightComponentBase::StaticClass(),
+		UReflectionCaptureComponent::StaticClass(),
+		USkyAtmosphereComponent::StaticClass(),
+		USkyLightComponent::StaticClass()
+	};
+	
+	TArray<UActorComponent*> ActorComponents;
+	constexpr bool bIncludeFromChildActors = true;
+	Actor->GetComponents(ActorComponents, bIncludeFromChildActors);
+	
+	for (const UActorComponent* Component : ActorComponents)
+	{
+		for (const UClass* ComponentType : LightingComponentTypes)
+		{
+			if (Component->IsA(ComponentType))
+			{
+				return true;
+			}
+		}
+	}
+
+	return false;
+}
+
 TArray<AActor*> UMoviePipelineCollection::GetMatchingActors(const UWorld* World, const bool bInvertResult) const
 {
+	TRACE_CPUPROFILER_EVENT_SCOPE(MRQ::Collection::GetMatchingActors);
+	
 	TArray<AActor*> MatchingActors;
 
 	for (TActorIterator<AActor> ActorItr(World); ActorItr; ++ActorItr)
@@ -321,7 +353,7 @@ void UMoviePipelineRenderLayerSubsystem::Initialize(FSubsystemCollectionBase& Co
 	VisualizationModifier_HideWorld = NewObject<UMoviePipelineVisibilityModifier>(GetTransientPackage(), NAME_None, RF_Transient);
 	VisualizationModifier_HideWorld->AddCollection(VisualizationEmptyCollection);
 	VisualizationModifier_HideWorld->SetHidden(true);
-	VisualizationModifier_HideWorld->SetIsInverted(true);
+	VisualizationModifier_HideWorld->bUseInvertedActors = true;
 
 	// Selectively show collections in the visualization
 	VisualizationModifier_VisibleCollections = NewObject<UMoviePipelineVisibilityModifier>(GetTransientPackage(), NAME_None, RF_Transient);
@@ -357,7 +389,7 @@ bool UMoviePipelineRenderLayerSubsystem::AddRenderLayer(UMoviePipelineRenderLaye
 
 	if (bRenderLayerExists)
 	{
-		UE_LOG(LogMovieRenderPipeline, Warning, TEXT("Render layer '%s' already exists in the render layer subsystem; it will not be added again."), *RenderLayer->GetRenderLayerName());
+		UE_LOG(LogMovieRenderPipeline, Warning, TEXT("Render layer '%s' already exists in the render layer subsystem; it will not be added again."), *RenderLayer->GetRenderLayerName().ToString());
 		return false;
 	}
 
@@ -394,7 +426,7 @@ void UMoviePipelineRenderLayerSubsystem::SetActiveRenderLayerByObj(UMoviePipelin
 	SetAndPreviewRenderLayer(RenderLayer);
 }
 
-void UMoviePipelineRenderLayerSubsystem::SetActiveRenderLayerByName(const FString& RenderLayerName)
+void UMoviePipelineRenderLayerSubsystem::SetActiveRenderLayerByName(const FName& RenderLayerName)
 {
 	const uint32 Index = RenderLayers.IndexOfByPredicate([&RenderLayerName](const UMoviePipelineRenderLayer* RenderLayer)
 	{
diff --git a/Engine/Plugins/MovieScene/MovieRenderPipeline/Source/MovieRenderPipelineCore/Private/Graph/Nodes/MovieGraphDeferredRenderPassNode.cpp b/Engine/Plugins/MovieScene/MovieRenderPipeline/Source/MovieRenderPipelineCore/Private/Graph/Nodes/MovieGraphDeferredRenderPassNode.cpp
--- a/Engine/Plugins/MovieScene/MovieRenderPipeline/Source/MovieRenderPipelineCore/Private/Graph/Nodes/MovieGraphDeferredRenderPassNode.cpp
+++ b/Engine/Plugins/MovieScene/MovieRenderPipeline/Source/MovieRenderPipelineCore/Private/Graph/Nodes/MovieGraphDeferredRenderPassNode.cpp
@@ -5,6 +5,7 @@
 #include "Graph/MovieGraphDataTypes.h"
 #include "Graph/MovieGraphDefaultRenderer.h"
 #include "Graph/MovieGraphPipeline.h"
+#include "Graph/MoviePipelineRenderLayerSubsystem.h"
 #include "MovieRenderPipelineCoreModule.h"
 #include "MovieRenderOverlappedImage.h"
 #include "MoviePipelineSurfaceReader.h"
@@ -49,12 +50,28 @@ void UMovieGraphDeferredRenderPassNode::TeardownImpl()
 	CurrentInstances.Reset();
 }
 
-
 void UMovieGraphDeferredRenderPassNode::RenderImpl(const FMovieGraphTraversalContext& InFrameTraversalContext, const FMovieGraphTimeStepData& InTimeData)
 {
-	for (TUniquePtr<FMovieGraphDeferredRenderPass>& Instance : CurrentInstances)
+	for (const TUniquePtr<FMovieGraphDeferredRenderPass>& Instance : CurrentInstances)
 	{
+		UMoviePipelineRenderLayerSubsystem* LayerSubsystem =
+			Instance->GetRenderer()->GetWorld()->GetSubsystem<UMoviePipelineRenderLayerSubsystem>();
+		
+		// Apply all modifiers in the evaluated graph
+		if (LayerSubsystem)
+		{
+			TRACE_CPUPROFILER_EVENT_SCOPE(MRQ::DeferredRender::ActivateRenderLayer);
+			LayerSubsystem->SetActiveRenderLayerByName(Instance->GetBranchName());
+		}
+		
 		Instance->Render(InFrameTraversalContext, InTimeData);
+
+		// Revert all modifiers
+		if (LayerSubsystem)
+		{
+			TRACE_CPUPROFILER_EVENT_SCOPE(MRQ::DeferredRender::RevertRenderLayer);
+			LayerSubsystem->ClearActiveRenderLayer();
+		}
 	}
 }
 
@@ -129,6 +146,16 @@ void UMovieGraphDeferredRenderPassNode::FMovieGraphDeferredRenderPass::AddRefere
 	}
 }
 
+FName UMovieGraphDeferredRenderPassNode::FMovieGraphDeferredRenderPass::GetBranchName() const
+{
+	return LayerData.BranchName;
+}
+
+TWeakObjectPtr<UMovieGraphDefaultRenderer> UMovieGraphDeferredRenderPassNode::FMovieGraphDeferredRenderPass::GetRenderer() const
+{
+	return Renderer;
+}
+
 void UMovieGraphDeferredRenderPassNode::FMovieGraphDeferredRenderPass::Render(const FMovieGraphTraversalContext& InFrameTraversalContext, const FMovieGraphTimeStepData& InTimeData)
 {
 	const bool bIncludeCDOs = true;
diff --git a/Engine/Plugins/MovieScene/MovieRenderPipeline/Source/MovieRenderPipelineCore/Public/Graph/MovieGraphConfig.h b/Engine/Plugins/MovieScene/MovieRenderPipeline/Source/MovieRenderPipelineCore/Public/Graph/MovieGraphConfig.h
--- a/Engine/Plugins/MovieScene/MovieRenderPipeline/Source/MovieRenderPipelineCore/Public/Graph/MovieGraphConfig.h
+++ b/Engine/Plugins/MovieScene/MovieRenderPipeline/Source/MovieRenderPipelineCore/Public/Graph/MovieGraphConfig.h
@@ -359,7 +359,7 @@ public:
  		return OutKeys;
 	}
 
-	UMovieGraphSettingNode* GetSettingForBranch(UClass* InClass, const FName InBranchName, bool bIncludeCDOs = true, bool bExactMatch = false)
+	UMovieGraphSettingNode* GetSettingForBranch(UClass* InClass, const FName InBranchName, bool bIncludeCDOs = true, bool bExactMatch = false) const
 	{
 		TArray<UMovieGraphSettingNode*> AllSettings = GetSettingsForBranch(InClass, InBranchName, bIncludeCDOs, bExactMatch);
 		if (AllSettings.Num() > 0)
@@ -370,9 +370,9 @@ public:
 		return nullptr;
 	}
 
-	TArray<UMovieGraphSettingNode*> GetSettingsForBranch(UClass* InClass, const FName InBranchName, bool bIncludeCDOs = true, bool bExactMatch = false)
+	TArray<UMovieGraphSettingNode*> GetSettingsForBranch(UClass* InClass, const FName InBranchName, bool bIncludeCDOs = true, bool bExactMatch = false) const
 	{
-		FMovieGraphEvaluatedBranchConfig* BranchConfig = BranchConfigMapping.Find(InBranchName);
+		const FMovieGraphEvaluatedBranchConfig* BranchConfig = BranchConfigMapping.Find(InBranchName);
 		ensureMsgf(BranchConfig, TEXT("Failed to find branch mapping for Branch: %s"), *InBranchName.ToString());
 
 		TArray<UMovieGraphSettingNode*> ResultNodes;
@@ -400,13 +400,13 @@ public:
 	}
 
 	template<typename NodeType>
-	NodeType* GetSettingForBranch(const FName InBranchName, bool bIncludeCDOs = true, bool bExactMatch = false)
+	NodeType* GetSettingForBranch(const FName InBranchName, bool bIncludeCDOs = true, bool bExactMatch = false) const
 	{
 		return Cast<NodeType>(GetSettingForBranch(NodeType::StaticClass(), InBranchName, bIncludeCDOs, bExactMatch));
 	}
 
 	template<typename NodeType>
-	TArray<NodeType*> GetSettingsForBranch(const FName InBranchName, bool bIncludeCDOs = true, bool bExactMatch = false)
+	TArray<NodeType*> GetSettingsForBranch(const FName InBranchName, bool bIncludeCDOs = true, bool bExactMatch = false) const
 	{
 		TArray<UMovieGraphSettingNode*> UntypedResults = GetSettingsForBranch(NodeType::StaticClass(), InBranchName, bIncludeCDOs, bExactMatch);
 
diff --git a/Engine/Plugins/MovieScene/MovieRenderPipeline/Source/MovieRenderPipelineCore/Public/Graph/MovieGraphPipeline.h b/Engine/Plugins/MovieScene/MovieRenderPipeline/Source/MovieRenderPipelineCore/Public/Graph/MovieGraphPipeline.h
--- a/Engine/Plugins/MovieScene/MovieRenderPipeline/Source/MovieRenderPipelineCore/Public/Graph/MovieGraphPipeline.h
+++ b/Engine/Plugins/MovieScene/MovieRenderPipeline/Source/MovieRenderPipelineCore/Public/Graph/MovieGraphPipeline.h
@@ -83,11 +83,23 @@ public:
 	/** Used by the Renderer Instance to disable the preview widget before rendering so it isn't baked into the UI Renderer. */
 	void SetPreviewWidgetVisible(bool bInIsVisible) { SetPreviewWidgetVisibleImpl(bInIsVisible); }
 
+	/**
+	 * Resets the render layer subsystem and updates it to reflect the render layers that are present in the evaluated
+	 * graph. Does not update the contents of the layers, however. See UpdateLayerContentsInRenderLayerSubsystem().
+	 */
+	void CreateLayersInRenderLayerSubsystem(const UMovieGraphEvaluatedConfig* EvaluatedConfig) const;
+
+	/**
+	 * Updates the contents of the layers that are currently present in the render layer subsystem by evaluating all
+	 * available collections and modifiers. To update the layers that are present in the subsystem, see
+	 * CreateLayersInRenderLayerSubsystem().
+	 */
+	void UpdateLayerContentsInRenderLayerSubsystem(const UMovieGraphEvaluatedConfig* EvaluatedConfig) const;
+
 protected:
 	UFUNCTION(BlueprintCallable, Category = "Movie Render Pipeline")
 	virtual void OnMoviePipelineFinishedImpl();
 
-protected:
 	virtual void OnEngineTickBeginFrame();
 	virtual void OnEngineTickEndFrame();
 	virtual void RenderFrame();
diff --git a/Engine/Plugins/MovieScene/MovieRenderPipeline/Source/MovieRenderPipelineCore/Public/Graph/MoviePipelineRenderLayerSubsystem.h b/Engine/Plugins/MovieScene/MovieRenderPipeline/Source/MovieRenderPipelineCore/Public/Graph/MoviePipelineRenderLayerSubsystem.h
--- a/Engine/Plugins/MovieScene/MovieRenderPipeline/Source/MovieRenderPipelineCore/Public/Graph/MoviePipelineRenderLayerSubsystem.h
+++ b/Engine/Plugins/MovieScene/MovieRenderPipeline/Source/MovieRenderPipelineCore/Public/Graph/MoviePipelineRenderLayerSubsystem.h
@@ -20,16 +20,16 @@ class UMoviePipelineCollection : public UObject
 	GENERATED_BODY()
 
 public:
-	UFUNCTION(BlueprintCallable, Category = "Render Layers")
+	UFUNCTION(BlueprintCallable, Category = "Settings")
 	TArray<AActor*> GetMatchingActors(const UWorld* World, const bool bInvertResult = false) const;
 
-	UFUNCTION(BlueprintCallable, Category = "Render Layers")
+	UFUNCTION(BlueprintCallable, Category = "Settings")
 	void AddQuery(UMoviePipelineCollectionQuery* Query);
 
-	UFUNCTION(BlueprintCallable, Category = "Render Layers")
+	UFUNCTION(BlueprintCallable, Category = "Settings")
 	void SetCollectionName(const FString& InName) { CollectionName = InName; }
 
-	UFUNCTION(BlueprintCallable, Category = "Render Layers")
+	UFUNCTION(BlueprintCallable, Category = "Settings")
 	FString GetCollectionName() const { return CollectionName; }
 
 private:
@@ -50,33 +50,28 @@ class UMoviePipelineCollectionModifier : public UObject
 
 public:
 	/** Adds a collection to the existing set of collections in this modifier. */
-	UFUNCTION(BlueprintCallable, Category = "Render Layers")
+	UFUNCTION(BlueprintCallable, Category = "Settings")
 	void AddCollection(UMoviePipelineCollection* Collection);
 
 	/** Overwrites the existing collections with the provided array of collections. */
-	UFUNCTION(BlueprintCallable, Category = "Render Layers")
+	UFUNCTION(BlueprintCallable, Category = "Settings")
 	void SetCollections(const TArray<UMoviePipelineCollection*> InCollections) { Collections = InCollections; }
 	
-	UFUNCTION(BlueprintCallable, Category = "Render Layers")
+	UFUNCTION(BlueprintCallable, Category = "Settings")
 	TArray<UMoviePipelineCollection*> GetCollections() const { return Collections; }
 
-	UFUNCTION(BlueprintCallable, Category = "Render Layers")
-	void SetIsInverted(const bool bIsInverted) { bUseInvertedActors = bIsInverted; }
-
-	UFUNCTION(BlueprintCallable, Category = "Render Layers")
-	bool IsInverted() const { return bUseInvertedActors; }
-	
 	virtual void ApplyModifier(const UWorld* World) PURE_VIRTUAL(UMoviePipelineCollectionModifier::ApplyModifier, );
 	virtual void UndoModifier() PURE_VIRTUAL(UMoviePipelineCollectionModifier::UndoModifier, );
 
+public:
+	/** Whether an inverted collection of actors should be used instead. */
+	UPROPERTY(BlueprintReadWrite, EditAnywhere, Category = "Settings")
+	bool bUseInvertedActors = false;
+
 protected:
 	/** The collections which this modifier will operate on. */
 	UPROPERTY()
 	TArray<TObjectPtr<UMoviePipelineCollection>> Collections;
-
-	/** Whether an inverted collection of actors should be used. */
-	UPROPERTY()
-	bool bUseInvertedActors = false;
 };
 
 /**
@@ -88,13 +83,13 @@ class UMoviePipelineMaterialModifier : public UMoviePipelineCollectionModifier
 	GENERATED_BODY()
 
 public:
-	UFUNCTION(BlueprintCallable, Category = "Render Layers")
+	UFUNCTION(BlueprintCallable, Category = "Settings")
 	void SetMaterial(TSoftObjectPtr<UMaterialInterface> InMaterial) { MaterialToApply = InMaterial; }
 
-	UFUNCTION(BlueprintCallable, Category = "Render Layers")
+	UFUNCTION(BlueprintCallable, Category = "Settings")
 	virtual void ApplyModifier(const UWorld* World) override;
 
-	UFUNCTION(BlueprintCallable, Category = "Render Layers")
+	UFUNCTION(BlueprintCallable, Category = "Settings")
 	virtual void UndoModifier() override;
 
 private:
@@ -105,7 +100,7 @@ private:
 	FComponentToMaterialMap ModifiedComponents;
 
 public:
-	UPROPERTY(BlueprintReadWrite, EditAnywhere, Category = "Render Layers")
+	UPROPERTY(BlueprintReadWrite, EditAnywhere, Category = "Settings")
 	TSoftObjectPtr<UMaterialInterface> MaterialToApply;
 };
 
@@ -118,16 +113,16 @@ class UMoviePipelineVisibilityModifier : public UMoviePipelineCollectionModifier
 	GENERATED_BODY()
 
 public:
-	UFUNCTION(BlueprintCallable, Category = "Render Layers")
+	UFUNCTION(BlueprintCallable, Category = "Settings")
 	void SetHidden(const bool bInIsHidden) { bIsHidden = bInIsHidden; }
 
-	UFUNCTION(BlueprintCallable, Category = "Render Layers")
+	UFUNCTION(BlueprintCallable, Category = "Settings")
 	bool IsHidden() const { return bIsHidden; }
 
-	UFUNCTION(BlueprintCallable, Category = "Render Layers")
+	UFUNCTION(BlueprintCallable, Category = "Settings")
 	virtual void ApplyModifier(const UWorld* World) override;
 
-	UFUNCTION(BlueprintCallable, Category = "Render Layers")
+	UFUNCTION(BlueprintCallable, Category = "Settings")
 	virtual void UndoModifier() override;
 
 private:
@@ -139,7 +134,7 @@ private:
 	TMap<TSoftObjectPtr<AActor>, bool> ModifiedActors;
 
 public:
-	UPROPERTY(BlueprintReadWrite, EditAnywhere, Category = "Render Layers")
+	UPROPERTY(BlueprintReadWrite, EditAnywhere, Category = "Settings")
 	bool bIsHidden = false;
 };
 
@@ -173,19 +168,19 @@ class UMoviePipelineCollectionCommonQuery : public UMoviePipelineCollectionQuery
 
 public:
 	// TODO: Add other common query operations (level, etc)
-	UFUNCTION(BlueprintCallable, Category = "Render Layers")
+	UFUNCTION(BlueprintCallable, Category = "Settings")
 	void SetActorNames(const TArray<FString>& InActorNames) { ActorNames = InActorNames; }
 
-	UFUNCTION(BlueprintCallable, Category = "Render Layers")
+	UFUNCTION(BlueprintCallable, Category = "Settings")
 	void SetTags(const TArray<FName>& InTags) { Tags = InTags; }
 
-	UFUNCTION(BlueprintCallable, Category = "Render Layers")
+	UFUNCTION(BlueprintCallable, Category = "Settings")
 	void SetComponentTypes(TArray<UClass*> InComponentTypes) { ComponentTypes = InComponentTypes; }
 
-	UFUNCTION(BlueprintCallable, Category = "Render Layers")
+	UFUNCTION(BlueprintCallable, Category = "Settings")
 	void SetQueryMode(const EMoviePipelineCollectionCommonQueryMode InQueryMode) { QueryMode = InQueryMode; }
 
-	UFUNCTION(BlueprintCallable, Category = "Render Layers")
+	UFUNCTION(BlueprintCallable, Category = "Settings")
 	virtual bool DoesActorMatchQuery(const AActor* Actor) const override;
 
 public:
@@ -202,6 +197,17 @@ public:
 	TArray<FName> Tags;
 };
 
+/** Generates a query that matches all lighting components. */
+UCLASS(BlueprintType)
+class UMoviePipelineCollectionLightingQuery : public UMoviePipelineCollectionQuery
+{
+	GENERATED_BODY()
+
+public:
+	UFUNCTION(BlueprintCallable, Category = "Settings")
+	virtual bool DoesActorMatchQuery(const AActor* Actor) const override;
+};
+
 /**
  * Provides a means of assembling modifiers together to generate a desired view of a scene. 
  */
@@ -213,34 +219,34 @@ class UMoviePipelineRenderLayer : public UObject
 public:
 	UMoviePipelineRenderLayer() = default;
 
-	UFUNCTION(BlueprintCallable, Category = "Render Layers")
-	FString GetRenderLayerName() const { return RenderLayerName; };
+	UFUNCTION(BlueprintCallable, Category = "Settings")
+	FName GetRenderLayerName() const { return RenderLayerName; };
 
-	UFUNCTION(BlueprintCallable, Category = "Render Layers")
-	void SetRenderLayerName(const FString& NewName) { RenderLayerName = NewName; }
+	UFUNCTION(BlueprintCallable, Category = "Settings")
+	void SetRenderLayerName(const FName& NewName) { RenderLayerName = NewName; }
 
-	UFUNCTION(BlueprintCallable, Category = "Render Layers")
+	UFUNCTION(BlueprintCallable, Category = "Settings")
 	UMoviePipelineCollection* GetCollectionByName(const FString& Name) const;
 
-	UFUNCTION(BlueprintCallable, Category = "Render Layers")
+	UFUNCTION(BlueprintCallable, Category = "Settings")
 	void AddModifier(UMoviePipelineCollectionModifier* Modifier);
 
-	UFUNCTION(BlueprintCallable, Category = "Render Layers")
+	UFUNCTION(BlueprintCallable, Category = "Settings")
 	TArray<UMoviePipelineCollectionModifier*> GetModifiers() const { return Modifiers; };
 
-	UFUNCTION(BlueprintCallable, Category = "Render Layers")
+	UFUNCTION(BlueprintCallable, Category = "Settings")
 	void RemoveModifier(UMoviePipelineCollectionModifier* Modifier);
 
-	UFUNCTION(BlueprintCallable, Category = "Render Layers")
+	UFUNCTION(BlueprintCallable, Category = "Settings")
 	void Preview(const UWorld* World);
 
-	UFUNCTION(BlueprintCallable, Category = "Render Layers")
+	UFUNCTION(BlueprintCallable, Category = "Settings")
 	void UndoPreview(const UWorld* World);
 
 private:
 	/** The name of this render layer. */
 	UPROPERTY()
-	FString RenderLayerName;
+	FName RenderLayerName;
 
 	/** The modifiers that are active when this render layer is active. */
 	UPROPERTY()
@@ -261,7 +267,7 @@ public:
 	UMoviePipelineRenderLayerSubsystem() = default;
 
 	/* Get this subsystem for a specific world. Handy for use from Python. */
-	UFUNCTION(BlueprintCallable, Category = "Render Layers")
+	UFUNCTION(BlueprintCallable, Category = "Settings")
 	static UMoviePipelineRenderLayerSubsystem* GetFromWorld(const UWorld* World);
 
 	//~ Begin USubsystem interface
@@ -271,46 +277,46 @@ public:
 	//~ End USubsystem interface
 
 	/** Clear out all tracked render layers and collections. */
-	UFUNCTION(BlueprintCallable, Category = "Render Layers")
+	UFUNCTION(BlueprintCallable, Category = "Settings")
 	void Reset();
 
 	/**
 	 * Adds a render layer to the system, which can later be made active by SetActiveRenderLayer*(). Returns true
 	 * if the layer was added successfully, else false.
 	 */
-	UFUNCTION(BlueprintCallable, Category = "Render Layers")
+	UFUNCTION(BlueprintCallable, Category = "Settings")
 	bool AddRenderLayer(UMoviePipelineRenderLayer* RenderLayer);
 
-	UFUNCTION(BlueprintCallable, Category = "Render Layers")
+	UFUNCTION(BlueprintCallable, Category = "Settings")
 	const TArray<UMoviePipelineRenderLayer*>& GetRenderLayers() { return RenderLayers; }
 
-	UFUNCTION(BlueprintCallable, Category = "Render Layers")
+	UFUNCTION(BlueprintCallable, Category = "Settings")
 	void RemoveRenderLayer(const FString& RenderLayerName);
 
-	UFUNCTION(BlueprintCallable, Category = "Render Layers")
+	UFUNCTION(BlueprintCallable, Category = "Settings")
 	UMoviePipelineRenderLayer* GetActiveRenderLayer() const { return ActiveRenderLayer; }
 
 	/** Previews the layer with the given name. The layer needs to have been registered with AddRenderLayer(). */
-	UFUNCTION(BlueprintCallable, Category = "Render Layers")
-	void SetActiveRenderLayerByName(const FString& RenderLayerName);
+	UFUNCTION(BlueprintCallable, Category = "Settings")
+	void SetActiveRenderLayerByName(const FName& RenderLayerName);
 
 	/** Previews the given layer. The layer does not need to have been registered with AddRenderLayer(). */
-	UFUNCTION(BlueprintCallable, Category = "Render Layers")
+	UFUNCTION(BlueprintCallable, Category = "Settings")
 	void SetActiveRenderLayerByObj(UMoviePipelineRenderLayer* RenderLayer);
 
-	UFUNCTION(BlueprintCallable, Category = "Render Layers")
+	UFUNCTION(BlueprintCallable, Category = "Settings")
 	void ClearActiveRenderLayer();
 
-	UFUNCTION(BlueprintCallable, Category = "Render Layers")
+	UFUNCTION(BlueprintCallable, Category = "Settings")
 	void PreviewCollection(UMoviePipelineCollection* Collection);
 
-	UFUNCTION(BlueprintCallable, Category = "Render Layers")
+	UFUNCTION(BlueprintCallable, Category = "Settings")
 	void ClearCollectionPreview();
 
-	UFUNCTION(BlueprintCallable, Category = "Render Layers")
+	UFUNCTION(BlueprintCallable, Category = "Settings")
 	void PreviewModifier(UMoviePipelineCollectionModifier* Modifier);
 
-	UFUNCTION(BlueprintCallable, Category = "Render Layers")
+	UFUNCTION(BlueprintCallable, Category = "Settings")
 	void ClearModifierPreview();
 
 private:
diff --git a/Engine/Plugins/MovieScene/MovieRenderPipeline/Source/MovieRenderPipelineCore/Public/Graph/Nodes/MovieGraphDeferredRenderPassNode.h b/Engine/Plugins/MovieScene/MovieRenderPipeline/Source/MovieRenderPipelineCore/Public/Graph/Nodes/MovieGraphDeferredRenderPassNode.h
--- a/Engine/Plugins/MovieScene/MovieRenderPipeline/Source/MovieRenderPipelineCore/Public/Graph/Nodes/MovieGraphDeferredRenderPassNode.h
+++ b/Engine/Plugins/MovieScene/MovieRenderPipeline/Source/MovieRenderPipelineCore/Public/Graph/Nodes/MovieGraphDeferredRenderPassNode.h
@@ -130,6 +130,8 @@ protected:
 		void Render(const FMovieGraphTraversalContext& InFrameTraversalContext, const FMovieGraphTimeStepData& InTimeData);
 		void GatherOutputPassesImpl(TArray<FMovieGraphRenderDataIdentifier>& OutExpectedPasses) const;
 		void AddReferencedObjects(FReferenceCollector& Collector);
+		FName GetBranchName() const;
+		TWeakObjectPtr<UMovieGraphDefaultRenderer> GetRenderer() const;
 
 	protected:
 		TSharedRef<FSceneViewFamilyContext> AllocateSceneViewFamilyContext(const FViewFamilyContextInitData& InInitData);
