MoviePipeline: Fixed a crash when the Output Resolution was set to zero.

diff --git a/Engine/Plugins/MovieScene/MovieRenderPipeline/Source/MovieRenderPipelineCore/Private/Graph/Nodes/MovieGraphWidgetRendererBaseNode.cpp b/Engine/Plugins/MovieScene/MovieRenderPipeline/Source/MovieRenderPipelineCore/Private/Graph/Nodes/MovieGraphWidgetRendererBaseNode.cpp
--- a/Engine/Plugins/MovieScene/MovieRenderPipeline/Source/MovieRenderPipelineCore/Private/Graph/Nodes/MovieGraphWidgetRendererBaseNode.cpp
+++ b/Engine/Plugins/MovieScene/MovieRenderPipeline/Source/MovieRenderPipelineCore/Private/Graph/Nodes/MovieGraphWidgetRendererBaseNode.cpp
@@ -100,6 +100,11 @@ void UMovieGraphWidgetRendererBaseNode::FMovieGraphWidgetPass::Render(const FMov
 		UE_LOG(LogMovieRenderPipeline, Error, TEXT("Resolution %dx%d exceeds maximum allowed by GPU. Widget renders (burn-ins, etc) do not support high-resolution tiling and thus can't exceed %dx%d."), OutputResolution.X, OutputResolution.Y, MaxResolution, MaxResolution);
 		return;
 	}
+	if (OutputResolution.X == 0 || OutputResolution.Y == 0)
+	{
+		UE_LOG(LogMovieRenderPipeline, Error, TEXT("Resolution %dx%d must be greater than zero in both dimensions."), OutputResolution.X, OutputResolution.Y);
+		return;
+	}
 
 	// Create the render target the widget will be rendered into
 	UE::MovieGraph::DefaultRenderer::FRenderTargetInitParams RenderTargetInitParams;
diff --git a/Engine/Plugins/MovieScene/MovieRenderPipeline/Source/MovieRenderPipelineCore/Private/MoviePipelineRendering.cpp b/Engine/Plugins/MovieScene/MovieRenderPipeline/Source/MovieRenderPipelineCore/Private/MoviePipelineRendering.cpp
--- a/Engine/Plugins/MovieScene/MovieRenderPipeline/Source/MovieRenderPipelineCore/Private/MoviePipelineRendering.cpp
+++ b/Engine/Plugins/MovieScene/MovieRenderPipeline/Source/MovieRenderPipelineCore/Private/MoviePipelineRendering.cpp
@@ -82,6 +82,12 @@ void UMoviePipeline::SetupRenderingPipelineForShot(UMoviePipelineExecutorShot* I
 			Shutdown(true);
 			return;
 		}
+		if (BackbufferResolution.X == 0 || BackbufferResolution.Y == 0)
+		{
+			UE_LOG(LogMovieRenderPipeline, Error, TEXT("Resolution %dx%d must be greater than zero in both dimensions."), BackbufferResolution.X, BackbufferResolution.Y);
+			Shutdown(true);
+			return;
+		}
 	}
 
 	// Note how many tiles we wish to render with.
