MoviePipeline: Resolve all currently available global variables in the graph.

Global variables have been refactored into their own class type in order to cleanly separate them from normal variables.

#rnx

diff --git a/Engine/Plugins/MovieScene/MovieRenderPipeline/Source/MovieRenderPipelineCore/Private/Graph/MovieGraphBlueprintLibrary.cpp b/Engine/Plugins/MovieScene/MovieRenderPipeline/Source/MovieRenderPipelineCore/Private/Graph/MovieGraphBlueprintLibrary.cpp
--- a/Engine/Plugins/MovieScene/MovieRenderPipeline/Source/MovieRenderPipelineCore/Private/Graph/MovieGraphBlueprintLibrary.cpp
+++ b/Engine/Plugins/MovieScene/MovieRenderPipeline/Source/MovieRenderPipelineCore/Private/Graph/MovieGraphBlueprintLibrary.cpp
@@ -408,4 +408,14 @@ float UMovieGraphBlueprintLibrary::GetCurrentAperture(const UMovieGraphPipeline*
 	}
 
 	return 0.f;
+}
+
+UCineCameraComponent* UMovieGraphBlueprintLibrary::GetCurrentCineCamera(const UMovieGraphPipeline* InMovieGraphPipeline)
+{
+	if (UCineCameraComponent* CineCameraComponent = UMoviePipelineBlueprintLibrary::Utility_GetCurrentCineCamera(InMovieGraphPipeline->GetWorld()))
+	{
+		return CineCameraComponent;
+	}
+
+	return nullptr;
 }
\ No newline at end of file
diff --git a/Engine/Plugins/MovieScene/MovieRenderPipeline/Source/MovieRenderPipelineCore/Private/Graph/MovieGraphConfig.cpp b/Engine/Plugins/MovieScene/MovieRenderPipeline/Source/MovieRenderPipelineCore/Private/Graph/MovieGraphConfig.cpp
--- a/Engine/Plugins/MovieScene/MovieRenderPipeline/Source/MovieRenderPipelineCore/Private/Graph/MovieGraphConfig.cpp
+++ b/Engine/Plugins/MovieScene/MovieRenderPipeline/Source/MovieRenderPipelineCore/Private/Graph/MovieGraphConfig.cpp
@@ -3,12 +3,16 @@
 #include "Graph/MovieGraphConfig.h"
 
 #include "Algo/Transform.h"
+#include "CineCameraComponent.h"
+#include "Graph/MovieGraphBlueprintLibrary.h"
 #include "Graph/MovieGraphEdge.h"
+#include "Graph/MovieGraphPipeline.h"
 #include "Graph/Nodes/MovieGraphInputNode.h"
 #include "Graph/Nodes/MovieGraphOutputNode.h"
 #include "Graph/Nodes/MovieGraphRemoveRenderSettingNode.h"
 #include "Graph/Nodes/MovieGraphVariableNode.h"
 #include "MovieGraphUtils.h"
+#include "MoviePipelineQueue.h"
 #include "MovieRenderPipelineCoreModule.h"
 
 #define LOCTEXT_NAMESPACE "MovieGraphConfig"
@@ -55,6 +59,16 @@ bool UMovieGraphMember::CanRename(const FText& InNewName, FText& OutError) const
 	return true;
 }
 
+bool UMovieGraphVariable::IsGlobal() const
+{
+	return IsA<UMovieGraphGlobalVariable>();
+}
+
+bool UMovieGraphVariable::IsDeletable() const
+{
+	return true;
+}
+
 bool UMovieGraphVariable::CanRename(const FText& InNewName, FText& OutError) const
 {
 	if (!Super::CanRename(InNewName, OutError))
@@ -93,6 +107,81 @@ void UMovieGraphVariable::PostEditChangeProperty(FPropertyChangedEvent& Property
 }
 #endif // WITH_EDITOR
 
+UMovieGraphGlobalVariable::UMovieGraphGlobalVariable()
+{
+	bIsEditable = false;
+}
+
+bool UMovieGraphGlobalVariable::IsDeletable() const
+{
+	return false;
+}
+
+bool UMovieGraphGlobalVariable::CanRename(const FText& InNewName, FText& OutError) const
+{
+	return false;
+}
+
+UMovieGraphGlobalVariable_ShotName::UMovieGraphGlobalVariable_ShotName()
+{
+	Name = FString(TEXT("shot_name"));
+	SetValueType(EMovieGraphValueType::String);
+}
+
+UMovieGraphGlobalVariable_SequenceName::UMovieGraphGlobalVariable_SequenceName()
+{
+	Name = FString(TEXT("seq_name"));
+	SetValueType(EMovieGraphValueType::String);
+}
+
+UMovieGraphGlobalVariable_FrameNumber::UMovieGraphGlobalVariable_FrameNumber()
+{
+	Name = FString(TEXT("frame_num"));
+	SetValueType(EMovieGraphValueType::Int32);
+}
+
+UMovieGraphGlobalVariable_CameraName::UMovieGraphGlobalVariable_CameraName()
+{
+	Name = FString(TEXT("camera_name"));
+	SetValueType(EMovieGraphValueType::String);
+}
+
+void UMovieGraphGlobalVariable_ShotName::UpdateValue(const FMovieGraphTraversalContext* InTraversalContext, const UMovieGraphPipeline* InPipeline)
+{
+	const TArray<TObjectPtr<UMoviePipelineExecutorShot>>& ShotList = InPipeline->GetActiveShotList();
+	
+	if (ShotList.IsValidIndex(InTraversalContext->ShotIndex))
+	{
+		if (const TObjectPtr<UMoviePipelineExecutorShot>& Shot = ShotList[InTraversalContext->ShotIndex])
+		{
+			SetValueString(Shot->OuterName);
+		}
+	}
+}
+
+void UMovieGraphGlobalVariable_SequenceName::UpdateValue(const FMovieGraphTraversalContext* InTraversalContext, const UMovieGraphPipeline* InPipeline)
+{
+	SetValueString(InTraversalContext->Job->Sequence.GetAssetName());
+}
+
+void UMovieGraphGlobalVariable_FrameNumber::UpdateValue(const FMovieGraphTraversalContext* InTraversalContext, const UMovieGraphPipeline* InPipeline)
+{
+	SetValueInt32(InTraversalContext->Time.ShotFrameNumber.Value);
+}
+
+void UMovieGraphGlobalVariable_CameraName::UpdateValue(const FMovieGraphTraversalContext* InTraversalContext, const UMovieGraphPipeline* InPipeline)
+{
+	const TArray<TObjectPtr<UMoviePipelineExecutorShot>>& ShotList = InPipeline->GetActiveShotList();
+	
+	if (ShotList.IsValidIndex(InTraversalContext->ShotIndex))
+	{
+		if (const TObjectPtr<UMoviePipelineExecutorShot>& Shot = ShotList[InTraversalContext->ShotIndex])
+		{
+			SetValueString(Shot->InnerName);
+		}
+	}
+}
+
 bool UMovieGraphInput::IsDeletable() const
 {
 	// The input is deletable as long as it's not the Globals input
@@ -179,12 +268,6 @@ void UMovieGraphOutput::PostEditChangeProperty(FPropertyChangedEvent& PropertyCh
 }
 #endif // WITH_EDITOR
 
-FName UMovieGraphConfig::GlobalVariable_ShotName = "shot_name";
-FName UMovieGraphConfig::GlobalVariable_SequenceName = "seq_name";
-FName UMovieGraphConfig::GlobalVariable_FrameNumber = "frame_num";
-FName UMovieGraphConfig::GlobalVariable_CameraName = "camera_name";
-FName UMovieGraphConfig::GlobalVariable_RenderLayerName = "render_layer_name";
-
 UMovieGraphConfig::UMovieGraphConfig()
 {
 	InputNode = CreateDefaultSubobject<UMovieGraphInputNode>(TEXT("DefaultInputNode"));
@@ -231,28 +314,24 @@ void UMovieGraphConfig::PostLoad()
 	}
 }
 
-UMovieGraphVariable* UMovieGraphConfig::AddGlobalVariable(const FName& InName, EMovieGraphValueType ValueType)
+template<typename T>
+T* UMovieGraphConfig::AddGlobalVariable()
 {
 	// Don't add duplicate global variables
-	const bool VariableExists = Variables.ContainsByPredicate([&InName](const TObjectPtr<UMovieGraphVariable>& Variable)
+	const bool VariableExists = GlobalVariables.ContainsByPredicate([](const TObjectPtr<UMovieGraphVariable>& Variable)
 	{
-		return Variable && (Variable->GetMemberName() == InName);
+		return Variable && (Variable->GetClass() == T::StaticClass());
 	});
 
 	if (VariableExists)
 	{
+		// Don't log here; graphs will typically try to add all available global variables on start-up, even if they
+		// already exist in the current graph
 		return nullptr;
 	}
-	
-	if (UMovieGraphVariable* NewVariable = AddVariable(InName))
-	{
-		NewVariable->bIsGlobal = true;
-		NewVariable->bIsEditable = false;
-		NewVariable->SetValueType(ValueType);
-		return NewVariable;
-	}
 
-	return nullptr;
+	// Pass an empty name to AddMember() since globals set their name upon construction
+	return AddMember<T>(GlobalVariables, FName());
 }
 
 void UMovieGraphConfig::AddDefaultMembers()
@@ -289,20 +368,10 @@ void UMovieGraphConfig::AddDefaultMembers()
 		OutputNode->UpdatePins();
 	}
 
-	static const TMap<FName, EMovieGraphValueType> GlobalVariableNamesAndTypes =
-	{
-		{GlobalVariable_ShotName, EMovieGraphValueType::String},
-		{GlobalVariable_SequenceName, EMovieGraphValueType::String},
-		{GlobalVariable_FrameNumber, EMovieGraphValueType::Int32},
-		{GlobalVariable_CameraName, EMovieGraphValueType::String},
-		{GlobalVariable_RenderLayerName, EMovieGraphValueType::String}
-	};
-
-	// Add all of the global variables that should be available in the graph
-	for (const TTuple<FName, EMovieGraphValueType>& GlobalVariableInfo : GlobalVariableNamesAndTypes)
-	{
-		AddGlobalVariable(GlobalVariableInfo.Key, GlobalVariableInfo.Value);
-	}
+	AddGlobalVariable<UMovieGraphGlobalVariable_CameraName>();
+	AddGlobalVariable<UMovieGraphGlobalVariable_FrameNumber>();
+	AddGlobalVariable<UMovieGraphGlobalVariable_SequenceName>();
+	AddGlobalVariable<UMovieGraphGlobalVariable_ShotName>();
 }
 
 bool UMovieGraphConfig::AddLabeledEdge(UMovieGraphNode* FromNode, const FName& FromPinLabel, UMovieGraphNode* ToNode, const FName& ToPinLabel)
@@ -520,10 +589,10 @@ bool UMovieGraphConfig::RemoveNode(UMovieGraphNode* InNode)
 	return AllNodes.RemoveSingle(InNode) == 1;
 }
 
-template<typename T>
-T* UMovieGraphConfig::AddMember(TArray<TObjectPtr<T>>& InMemberArray, const FName& InBaseName)
+template<typename RetType, typename ArrType>
+RetType* UMovieGraphConfig::AddMember(TArray<TObjectPtr<ArrType>>& InMemberArray, const FName& InBaseName)
 {
-	static_assert(std::is_base_of_v<UMovieGraphMember, T>, "T is not derived from UMovieGraphMember");
+	static_assert(std::is_base_of_v<UMovieGraphMember, RetType>, "RetType is not derived from UMovieGraphMember");
 	
 	using namespace UE::MoviePipeline::RenderGraph;
 
@@ -534,9 +603,9 @@ T* UMovieGraphConfig::AddMember(TArray<TObjectPtr<T>>& InMemberArray, const FNam
 	// when the constructor is running, RF_NeedInitialization will be set. CreateDefaultSubobject() needs to be called
 	// in this scenario instead of NewObject().
 	const bool bIsNewObject = HasAnyFlags(RF_NeedInitialization);
-	T* NewMember = bIsNewObject
-		? CreateDefaultSubobject<T>(MakeUniqueObjectName(this, T::StaticClass()))
-		: NewObject<T>(this, NAME_None);
+	RetType* NewMember = bIsNewObject
+		? CreateDefaultSubobject<RetType>(MakeUniqueObjectName(this, RetType::StaticClass()))
+		: NewObject<RetType>(this, NAME_None);
 	
 	if (!NewMember)
 	{
@@ -548,10 +617,13 @@ T* UMovieGraphConfig::AddMember(TArray<TObjectPtr<T>>& InMemberArray, const FNam
 	NewMember->SetFlags(RF_Transactional);
 	NewMember->SetGuid(FGuid::NewGuid());
 
-	// Generate and set a unique name
-	TArray<FString> ExistingMemberNames;
-	Algo::Transform(InMemberArray, ExistingMemberNames, [](const T* Member) { return Member->GetMemberName(); });
-	NewMember->SetMemberName(GetUniqueName(ExistingMemberNames, InBaseName.ToString()));
+	// Generate and set a unique name. Globals set their name at construction time, so no need to set their name.
+	if (!NewMember->template IsA<UMovieGraphGlobalVariable>())
+	{
+		TArray<FString> ExistingMemberNames;
+		Algo::Transform(InMemberArray, ExistingMemberNames, [](const ArrType* Member) { return Member->GetMemberName(); });
+		NewMember->SetMemberName(GetUniqueName(ExistingMemberNames, InBaseName.ToString()));
+	}
 
 	return NewMember;
 }
@@ -560,7 +632,7 @@ UMovieGraphVariable* UMovieGraphConfig::AddVariable(const FName InCustomBaseName
 {
 	static const FText VariableBaseName = LOCTEXT("VariableBaseName", "Variable");
 	
-	UMovieGraphVariable* NewVariable = AddMember(
+	UMovieGraphVariable* NewVariable = AddMember<UMovieGraphVariable>(
 		Variables, !InCustomBaseName.IsNone() ? InCustomBaseName : FName(*VariableBaseName.ToString()));
 
 	if (NewVariable)
@@ -585,7 +657,7 @@ UMovieGraphInput* UMovieGraphConfig::AddInput()
 {
 	static const FText InputBaseName = LOCTEXT("InputBaseName", "Input");
 
-	UMovieGraphInput* NewInput = AddMember(Inputs, FName(*InputBaseName.ToString()));
+	UMovieGraphInput* NewInput = AddMember<UMovieGraphInput>(Inputs, FName(*InputBaseName.ToString()));
 	InputNode->UpdatePins();
 	
 #if WITH_EDITOR
@@ -599,7 +671,7 @@ UMovieGraphOutput* UMovieGraphConfig::AddOutput()
 {
 	static const FText OutputBaseName = LOCTEXT("OutputBaseName", "Output");
 	
-	UMovieGraphOutput* NewOutput = AddMember(Outputs, FName(*OutputBaseName.ToString()));
+	UMovieGraphOutput* NewOutput = AddMember<UMovieGraphOutput>(Outputs, FName(*OutputBaseName.ToString()));
 	OutputNode->UpdatePins();
 
 #if WITH_EDITOR
@@ -611,7 +683,8 @@ UMovieGraphOutput* UMovieGraphConfig::AddOutput()
 
 UMovieGraphVariable* UMovieGraphConfig::GetVariableByGuid(const FGuid& InGuid) const
 {
-	for (const TObjectPtr<UMovieGraphVariable>& Variable : Variables)
+	constexpr bool bIncludeGlobal = true;
+	for (UMovieGraphVariable* Variable : GetVariables(bIncludeGlobal))
 	{
 		if (Variable->GetGuid() == InGuid)
 		{
@@ -624,12 +697,27 @@ UMovieGraphVariable* UMovieGraphConfig::GetVariableByGuid(const FGuid& InGuid) c
 
 TArray<UMovieGraphVariable*> UMovieGraphConfig::GetVariables(const bool bIncludeGlobal) const
 {
-	if (bIncludeGlobal)
+	if (!bIncludeGlobal)
 	{
 		return Variables;
 	}
 
-	return Variables.FilterByPredicate([](const UMovieGraphVariable* Var) { return Var && !Var->IsGlobal(); });
+	TArray<UMovieGraphVariable*> AllVariables = Variables;
+	AllVariables.Append(GlobalVariables);
+
+	return AllVariables;
+}
+
+void UMovieGraphConfig::UpdateGlobalVariableValues(const UMovieGraphPipeline* InPipeline)
+{
+	// Note: Although UpdateValue could get the traversal context from the pipeline itself, we fetch it once here
+	// to prevent re-creating the context constantly.
+	const FMovieGraphTraversalContext TraversalContext = InPipeline->GetCurrentTraversalContext();
+	
+	for (const TObjectPtr<UMovieGraphGlobalVariable>& GlobalVariable : GlobalVariables)
+	{
+		GlobalVariable->UpdateValue(&TraversalContext, InPipeline);
+	}
 }
 
 TArray<UMovieGraphInput*> UMovieGraphConfig::GetInputs() const
diff --git a/Engine/Plugins/MovieScene/MovieRenderPipeline/Source/MovieRenderPipelineCore/Private/Graph/MovieGraphLinearTimeStep.cpp b/Engine/Plugins/MovieScene/MovieRenderPipeline/Source/MovieRenderPipelineCore/Private/Graph/MovieGraphLinearTimeStep.cpp
--- a/Engine/Plugins/MovieScene/MovieRenderPipeline/Source/MovieRenderPipelineCore/Private/Graph/MovieGraphLinearTimeStep.cpp
+++ b/Engine/Plugins/MovieScene/MovieRenderPipeline/Source/MovieRenderPipelineCore/Private/Graph/MovieGraphLinearTimeStep.cpp
@@ -68,9 +68,13 @@ void UMovieGraphLinearTimeStep::TickProducingFrames()
 		UE_LOG(LogMovieRenderPipeline, Log, TEXT("MovieGraph Initializing Camera Cut [%d/%d] in [%s] %s."),
 			CurrentShotIndex + 1, ActiveShotList.Num(), *CurrentCameraCut->OuterName, *CurrentCameraCut->InnerName);
 
+		FMovieGraphTraversalContext Context = GetOwningGraph()->GetCurrentTraversalContext();
+		
+		// Update global variables before evaluating the graph
+		Context.RootGraph->UpdateGlobalVariableValues(GetOwningGraph());
+
 		// Evaluate the graph so we can fetch values for this shot.
 		UMovieGraphConfig* Config = GetOwningGraph()->GetRootGraphForShot(CurrentCameraCut);
-		FMovieGraphTraversalContext Context = GetOwningGraph()->GetCurrentTraversalContext();
 		CurrentFrameData.EvaluatedConfig = TStrongObjectPtr<UMovieGraphEvaluatedConfig>(Config->CreateFlattenedGraph(Context));
 		CurrentFrameData.TemporalSampleIndex = 0;
 
@@ -116,6 +120,9 @@ void UMovieGraphLinearTimeStep::TickProducingFrames()
 	if (CurrentCameraCut->ShotInfo.State == EMovieRenderShotState::Rendering)
 	{
 		FMovieGraphTraversalContext Context = GetOwningGraph()->GetCurrentTraversalContext();
+
+		// Update global variables before evaluating the graph
+		Context.RootGraph->UpdateGlobalVariableValues(GetOwningGraph());
 		
 		if (IsFirstTemporalSample())
 		{
diff --git a/Engine/Plugins/MovieScene/MovieRenderPipeline/Source/MovieRenderPipelineCore/Public/Graph/MovieGraphBlueprintLibrary.h b/Engine/Plugins/MovieScene/MovieRenderPipeline/Source/MovieRenderPipelineCore/Public/Graph/MovieGraphBlueprintLibrary.h
--- a/Engine/Plugins/MovieScene/MovieRenderPipeline/Source/MovieRenderPipelineCore/Public/Graph/MovieGraphBlueprintLibrary.h
+++ b/Engine/Plugins/MovieScene/MovieRenderPipeline/Source/MovieRenderPipelineCore/Public/Graph/MovieGraphBlueprintLibrary.h
@@ -177,4 +177,11 @@ public:
 	*/
 	UFUNCTION(BlueprintPure, Category = "Movie Graph")
 	static float GetCurrentAperture(const UMovieGraphPipeline* InMovieGraphPipeline);
+
+	/**
+	* Gets the currently active cine camera, or nullptr if one was not found.
+	* @param	InMovieGraphPipeline	- The pipeline to get the camera from.
+	*/
+	UFUNCTION(BlueprintPure, Category = "Movie Graph")
+	static UCineCameraComponent* GetCurrentCineCamera(const UMovieGraphPipeline* InMovieGraphPipeline);
 };
\ No newline at end of file
diff --git a/Engine/Plugins/MovieScene/MovieRenderPipeline/Source/MovieRenderPipelineCore/Public/Graph/MovieGraphConfig.h b/Engine/Plugins/MovieScene/MovieRenderPipeline/Source/MovieRenderPipelineCore/Public/Graph/MovieGraphConfig.h
--- a/Engine/Plugins/MovieScene/MovieRenderPipeline/Source/MovieRenderPipelineCore/Public/Graph/MovieGraphConfig.h
+++ b/Engine/Plugins/MovieScene/MovieRenderPipeline/Source/MovieRenderPipelineCore/Public/Graph/MovieGraphConfig.h
@@ -112,10 +112,10 @@ public:
 	UMovieGraphVariable() = default;
 
 	/** Returns true if this variable is a global variable. */
-	bool IsGlobal() const { return bIsGlobal; }
+	bool IsGlobal() const;
 
 	//~ Begin UMovieGraphMember interface
-	virtual bool IsDeletable() const override { return !bIsGlobal; }
+	virtual bool IsDeletable() const override;
 	virtual bool CanRename(const FText& InNewName, FText& OutError) const override;
 	virtual bool SetMemberName(const FString& InNewName) override;
 	//~ End UMovieGraphMember interface
@@ -128,11 +128,67 @@ public:
 	virtual void PostEditChangeProperty(FPropertyChangedEvent& PropertyChangedEvent) override;
 	//~ End UObject overrides
 #endif // WITH_EDITOR
+};
 
-private:
-	/** Whether this variable represents a global variable. */
-	UPROPERTY()
-	bool bIsGlobal = false;
+/**
+ * Similar to normal UMovieGraphVariable instances. However, their values are provided by the graph, they cannot be
+ * edited/deleted, and they cannot be overridden at the job level.
+ */
+UCLASS(Abstract)
+class MOVIERENDERPIPELINECORE_API UMovieGraphGlobalVariable : public UMovieGraphVariable
+{
+	GENERATED_BODY()
+
+public:
+	UMovieGraphGlobalVariable();
+
+	/** Update the internal value of the global variable. */
+	virtual void UpdateValue(const FMovieGraphTraversalContext* InTraversalContext, const UMovieGraphPipeline* InPipeline) PURE_VIRTUAL(UMovieGraphGlobalVariable::UpdateValue, );
+
+	//~ Begin UMovieGraphMember interface
+	virtual bool IsDeletable() const override;
+	virtual bool CanRename(const FText& InNewName, FText& OutError) const override;
+	//~ End UMovieGraphMember interface
+};
+
+UCLASS()
+class UMovieGraphGlobalVariable_ShotName final : public UMovieGraphGlobalVariable
+{
+	GENERATED_BODY()
+
+public:
+	UMovieGraphGlobalVariable_ShotName();
+	virtual void UpdateValue(const FMovieGraphTraversalContext* InTraversalContext, const UMovieGraphPipeline* InPipeline) override;
+};
+
+UCLASS()
+class UMovieGraphGlobalVariable_SequenceName final : public UMovieGraphGlobalVariable
+{
+	GENERATED_BODY()
+
+public:
+	UMovieGraphGlobalVariable_SequenceName();
+	virtual void UpdateValue(const FMovieGraphTraversalContext* InTraversalContext, const UMovieGraphPipeline* InPipeline) override;
+};
+
+UCLASS()
+class UMovieGraphGlobalVariable_FrameNumber final : public UMovieGraphGlobalVariable
+{
+	GENERATED_BODY()
+
+public:
+	UMovieGraphGlobalVariable_FrameNumber();
+	virtual void UpdateValue(const FMovieGraphTraversalContext* InTraversalContext, const UMovieGraphPipeline* InPipeline) override;
+};
+
+UCLASS()
+class UMovieGraphGlobalVariable_CameraName final : public UMovieGraphGlobalVariable
+{
+	GENERATED_BODY()
+
+public:
+	UMovieGraphGlobalVariable_CameraName();
+	virtual void UpdateValue(const FMovieGraphTraversalContext* InTraversalContext, const UMovieGraphPipeline* InPipeline) override;
 };
 
 /**
@@ -496,6 +552,10 @@ public:
 	UFUNCTION(BlueprintCallable, Category="Experimental")
 	TArray<UMovieGraphVariable*> GetVariables(const bool bIncludeGlobal = false) const;
 
+	/** Updates the values of all global variables. */
+	UFUNCTION(BlueprintCallable, Category="Experimental")
+	void UpdateGlobalVariableValues(const UMovieGraphPipeline* InPipeline);
+
 	/** Gets all inputs that have been defined on the graph. */
 	TArray<UMovieGraphInput*> GetInputs() const;
 
@@ -556,13 +616,6 @@ protected:
 	void VisitDownstreamNodes_Recursive(UMovieGraphNode* FromNode, const FVisitNodesCallback& VisitCallback, TSet<UMovieGraphNode*>& VisitedNodes) const;
 
 public:
-	// Names of global variables that are provided by the graph
-	static FName GlobalVariable_ShotName;
-	static FName GlobalVariable_SequenceName;
-	static FName GlobalVariable_FrameNumber;
-	static FName GlobalVariable_CameraName;
-	static FName GlobalVariable_RenderLayerName;
-	
 #if WITH_EDITOR
 	FOnMovieGraphChanged OnGraphChangedDelegate;
 	FOnMovieGraphVariablesChanged OnGraphVariablesChangedDelegate;
@@ -611,21 +664,29 @@ private:
 	/** Remove the specified output member from the graph. */
 	bool DeleteOutputMember(UMovieGraphOutput* OutputMemberToDelete);
 	
-	/** Add a new member of type T to MemberArray, with a unique name that includes BaseName in it. */
-	template<typename T>
-	T* AddMember(TArray<TObjectPtr<T>>& InMemberArray, const FName& InBaseName);
+	/**
+	 * Add a new member of type RetType to MemberArray (ArrType, which RetType must derive from), with a unique name
+	 * that includes BaseName in it.
+	 */
+	template<typename RetType, typename ArrType>
+	RetType* AddMember(TArray<TObjectPtr<ArrType>>& InMemberArray, const FName& InBaseName);
 
-	/** Adds a global variable to the graph with the provided name and value type. */
-	UMovieGraphVariable* AddGlobalVariable(const FName& InName, EMovieGraphValueType ValueType);
+	/** Adds a global variable of type T to the graph. */
+	template<typename T>
+	T* AddGlobalVariable();
 
 	/** Adds members to the graph that should always be available. */
 	void AddDefaultMembers();
 
 private:
-	/** All variables (user and global) which are available for use in the graph. */
+	/** All user (not global) variables which are available for use in the graph. */
 	UPROPERTY()
 	TArray<TObjectPtr<UMovieGraphVariable>> Variables;
 
+	/** All global variables which are available for use in the graph. */
+	UPROPERTY()
+	TArray<TObjectPtr<UMovieGraphGlobalVariable>> GlobalVariables;
+
 	/** All inputs which have been defined on the graph. */
 	UPROPERTY()
 	TArray<TObjectPtr<UMovieGraphInput>> Inputs;
