MoviePipeline: Added a pathtracer node and implementation to the graph. Also added some additional visibility-related settings to the Visibility modifier.

Renamed UMovieGraphDeferredRenderPassNode to UMovieGraphCoreRenderPassNode and exposed a few virtual methods so the bulk of the deferred rendering code can be shared with the path tracer.

#rnx

diff --git a/Engine/Plugins/MovieScene/MovieRenderPipeline/Source/MovieRenderPipelineCore/Private/Graph/MoviePipelineRenderLayerSubsystem.cpp b/Engine/Plugins/MovieScene/MovieRenderPipeline/Source/MovieRenderPipelineCore/Private/Graph/MoviePipelineRenderLayerSubsystem.cpp
--- a/Engine/Plugins/MovieScene/MovieRenderPipeline/Source/MovieRenderPipelineCore/Private/Graph/MoviePipelineRenderLayerSubsystem.cpp
+++ b/Engine/Plugins/MovieScene/MovieRenderPipeline/Source/MovieRenderPipelineCore/Private/Graph/MoviePipelineRenderLayerSubsystem.cpp
@@ -77,9 +77,26 @@ void UMoviePipelineMaterialModifier::UndoModifier()
 	ModifiedComponents.Empty();
 }
 
+UMoviePipelineVisibilityModifier::UMoviePipelineVisibilityModifier()
+	: bIsHidden(false)
+	, bCastShadowWhileHidden(false)
+	, bAffectIndirectLightingWhileHidden(false)
+	, bHoldout(false)
+{
+	
+}
+
 void UMoviePipelineVisibilityModifier::ApplyModifier(const UWorld* World)
 {
 	ModifiedActors.Empty();
+
+	// SetActorVisibilityState() takes a FActorVisibilityState, but all calls here will contain the same settings (with
+	// a differing actor), so create one copy to prevent constantly re-creating structs.
+	FActorVisibilityState NewVisibilityState;
+	NewVisibilityState.bIsHidden = bIsHidden;
+	NewVisibilityState.bCastShadowWhileHidden = bCastShadowWhileHidden;
+	NewVisibilityState.bAffectIndirectLightingWhileHidden = bAffectIndirectLightingWhileHidden;
+	NewVisibilityState.bHoldout = bHoldout;
 	
 	for (const UMoviePipelineCollection* Collection : Collections)
 	{
@@ -88,31 +105,63 @@ void UMoviePipelineVisibilityModifier::ApplyModifier(const UWorld* World)
 			continue;
 		}
 
-		for (AActor* Actor : Collection->GetMatchingActors(World, bUseInvertedActors))
+		const TArray<AActor*> MatchingActors = Collection->GetMatchingActors(World, bUseInvertedActors);
+		ModifiedActors.Reserve(MatchingActors.Num());
+
+		for (int32 Index = 0; Index < MatchingActors.Num(); ++Index)
 		{
-			ModifiedActors.Add(Actor, Actor->IsHidden());
-			SetActorHiddenState(Actor, bIsHidden);
+			const AActor* Actor = MatchingActors[Index];
+
+			// Save out visibility state before the modifier is applied
+			FActorVisibilityState OriginalVisibilityState;
+			OriginalVisibilityState.Actor = Actor;
+			OriginalVisibilityState.bIsHidden = Actor->IsHidden();
+			if (const UPrimitiveComponent* PrimitiveComponent = Actor->GetComponentByClass<UPrimitiveComponent>())
+			{
+				OriginalVisibilityState.bCastShadowWhileHidden = PrimitiveComponent->bCastHiddenShadow;
+				OriginalVisibilityState.bAffectIndirectLightingWhileHidden = PrimitiveComponent->bAffectIndirectLightingWhileHidden;
+				OriginalVisibilityState.bHoldout = PrimitiveComponent->bHoldout;
+			}
+			
+			ModifiedActors.Add(OriginalVisibilityState);
+
+			// Set new visibility state
+			NewVisibilityState.Actor = Actor;
+			SetActorVisibilityState(NewVisibilityState);
 		}
 	}
 }
 
 void UMoviePipelineVisibilityModifier::UndoModifier()
 {
-	for (const TTuple<TSoftObjectPtr<AActor>, bool>& Pair : ModifiedActors)
+	for (const FActorVisibilityState& PrevVisibilityState : ModifiedActors)
 	{
-		SetActorHiddenState(Pair.Key.LoadSynchronous(), Pair.Value);
+		SetActorVisibilityState(PrevVisibilityState);
 	}
 
 	ModifiedActors.Empty();
 }
 
-void UMoviePipelineVisibilityModifier::SetActorHiddenState(AActor* Actor, const bool bInIsHidden) const
+void UMoviePipelineVisibilityModifier::SetActorVisibilityState(const FActorVisibilityState& NewVisibilityState)
 {
-	Actor->SetActorHiddenInGame(bInIsHidden);
+	const TSoftObjectPtr<AActor> Actor = NewVisibilityState.Actor.LoadSynchronous();
+	if (!Actor)
+	{
+		return;
+	}
+	
+	Actor->SetActorHiddenInGame(NewVisibilityState.bIsHidden);
 
 #if WITH_EDITOR
-	Actor->SetIsTemporarilyHiddenInEditor(bInIsHidden);
+	Actor->SetIsTemporarilyHiddenInEditor(NewVisibilityState.bIsHidden);
 #endif
+
+	if (UPrimitiveComponent* PrimitiveComponent = Actor->GetComponentByClass<UPrimitiveComponent>())
+	{
+		PrimitiveComponent->bCastHiddenShadow = NewVisibilityState.bCastShadowWhileHidden;
+		PrimitiveComponent->bAffectIndirectLightingWhileHidden = NewVisibilityState.bAffectIndirectLightingWhileHidden;
+		PrimitiveComponent->bHoldout = NewVisibilityState.bHoldout;
+	}
 }
 
 // TODO: This really should be "DoesComponentMatchQuery()"
diff --git a/Engine/Plugins/MovieScene/MovieRenderPipeline/Source/MovieRenderPipelineCore/Private/Graph/Nodes/MovieGraphCoreRenderPassNode.cpp b/Engine/Plugins/MovieScene/MovieRenderPipeline/Source/MovieRenderPipelineCore/Private/Graph/Nodes/MovieGraphCoreRenderPassNode.cpp
new file mode 100644
--- /dev/null
+++ b/Engine/Plugins/MovieScene/MovieRenderPipeline/Source/MovieRenderPipelineCore/Private/Graph/Nodes/MovieGraphCoreRenderPassNode.cpp
@@ -0,0 +1,649 @@
+// Copyright Epic Games, Inc. All Rights Reserved.
+
+#include "Graph/Nodes/MovieGraphCoreRenderPassNode.h"
+
+#include "Graph/Nodes/MovieGraphOutputSettingNode.h"
+#include "Graph/MovieGraphDataTypes.h"
+#include "Graph/MovieGraphDefaultRenderer.h"
+#include "Graph/MovieGraphPipeline.h"
+#include "Graph/MoviePipelineRenderLayerSubsystem.h"
+#include "MovieRenderOverlappedImage.h"
+#include "MoviePipelineSurfaceReader.h"
+#include "EngineModule.h"
+#include "SceneManagement.h"
+#include "CanvasTypes.h"
+#include "Engine/EngineBaseTypes.h"
+#include "Engine/TextureRenderTarget2D.h"
+#include "LegacyScreenPercentageDriver.h"
+#include "OpenColorIODisplayExtension.h"
+#include "TextureResource.h"
+#include "Tasks/Task.h"
+
+// For the 1D Weight table for accumulation
+#include "MovieRenderPipelineDataTypes.h"
+
+void UMovieGraphCoreRenderPassNode::SetupImpl(const FMovieGraphRenderPassSetupData& InSetupData)
+{
+	// To make the implementation simpler, we make one instance of FMovieGraphDeferredRenderPas
+	// per camera, and per render layer. These objects can pull from common pools to share state,
+	// which gives us a better overview of how many resources are being used by MRQ.
+	for (const FMovieGraphRenderPassLayerData& LayerData : InSetupData.Layers)
+	{
+		TUniquePtr<FMovieGraphRenderPass> RendererInstance = MakeUnique<FMovieGraphRenderPass>();
+		RendererInstance->Setup(InSetupData.Renderer, this, LayerData);
+		CurrentInstances.Add(MoveTemp(RendererInstance));
+	}
+}
+
+void UMovieGraphCoreRenderPassNode::TeardownImpl()
+{
+	// We don't need to flush the rendering commands as we assume the MovieGraph
+	// Renderer has already done that once, so all data for all passes should
+	// have been submitted to the GPU (and subsequently read back) by now.
+	for (TUniquePtr<FMovieGraphRenderPass>& Instance : CurrentInstances)
+	{
+		Instance->Teardown();
+	}
+	CurrentInstances.Reset();
+}
+
+void UMovieGraphCoreRenderPassNode::RenderImpl(const FMovieGraphTraversalContext& InFrameTraversalContext, const FMovieGraphTimeStepData& InTimeData)
+{
+	for (const TUniquePtr<FMovieGraphRenderPass>& Instance : CurrentInstances)
+	{
+		UMoviePipelineRenderLayerSubsystem* LayerSubsystem =
+			Instance->GetRenderer()->GetWorld()->GetSubsystem<UMoviePipelineRenderLayerSubsystem>();
+		
+		// Apply all modifiers in the evaluated graph
+		if (LayerSubsystem)
+		{
+			TRACE_CPUPROFILER_EVENT_SCOPE(MRQ::DeferredRender::ActivateRenderLayer);
+			LayerSubsystem->SetActiveRenderLayerByName(Instance->GetBranchName());
+		}
+		
+		Instance->Render(InFrameTraversalContext, InTimeData);
+
+		// Revert all modifiers
+		if (LayerSubsystem)
+		{
+			TRACE_CPUPROFILER_EVENT_SCOPE(MRQ::DeferredRender::RevertRenderLayer);
+			LayerSubsystem->ClearActiveRenderLayer();
+		}
+	}
+}
+
+void UMovieGraphCoreRenderPassNode::GatherOutputPassesImpl(TArray<FMovieGraphRenderDataIdentifier>& OutExpectedPasses) const
+{
+	for (const TUniquePtr<FMovieGraphRenderPass>& Instance : CurrentInstances)
+	{
+		Instance->GatherOutputPassesImpl(OutExpectedPasses);
+	}
+}
+
+EViewModeIndex UMovieGraphCoreRenderPassNode::GetViewModeIndex() const
+{
+	return VMI_Lit;
+}
+
+FEngineShowFlags UMovieGraphCoreRenderPassNode::GetShowFlags() const
+{
+	return FEngineShowFlags(EShowFlagInitMode::ESFIM_Game);
+}
+
+void UMovieGraphCoreRenderPassNode::AddReferencedObjects(UObject* InThis, FReferenceCollector& Collector)
+{
+	Super::AddReferencedObjects(InThis, Collector);
+	
+	UMovieGraphCoreRenderPassNode* This = CastChecked<UMovieGraphCoreRenderPassNode>(InThis);
+	for (TUniquePtr<FMovieGraphRenderPass>& Instance : This->CurrentInstances)
+	{
+		Instance->AddReferencedObjects(Collector);
+	}
+}
+
+void UMovieGraphCoreRenderPassNode::FMovieGraphRenderPass::Setup(TWeakObjectPtr<UMovieGraphDefaultRenderer> InRenderer, TWeakObjectPtr<UMovieGraphCoreRenderPassNode> InRenderPassNode, const FMovieGraphRenderPassLayerData& InLayer)
+{
+	LayerData = InLayer;
+	Renderer = InRenderer;
+	RenderPassNode = InRenderPassNode;
+
+	RenderDataIdentifier.RootBranchName = LayerData.BranchName;
+	RenderDataIdentifier.RendererName = RenderPassNode->GetRendererName();
+	RenderDataIdentifier.SubResourceName = TEXT("beauty");
+	
+	UE::MovieGraph::DefaultRenderer::FCameraInfo CameraInfo = Renderer->GetCameraInfo(LayerData.CameraIdentifier);
+	RenderDataIdentifier.CameraName =  CameraInfo.CameraName;
+
+	// Figure out how big each sub-region (tile) is.
+	// const int32 TileSize = Graph->FindSetting(LayerData.LayerName, "highres.tileSize");
+	//int32 TileCount = 1; // ToDo: We should do tile sizes (ie: 1024x512) instead
+	//FIntPoint BackbufferResolution = FIntPoint(
+	//	FMath::CeilToInt((float)OutputCameraResolution.X / (float)TileCount),
+	//	FMath::CeilToInt((float)OutputCameraResolution.Y / (float)TileCount));
+
+	// BackbufferResolution = HighResSettings->CalculatePaddedBackbufferSize(BackbufferResolution);
+	
+	// Create a view render target for our given resolution. This is require for the FCanvas/FSceneViewAPI. We pool these.
+	// Renderer->GetOrCreateViewRenderTarget(BackbufferResolution);
+	// CreateSurfaceQueueImpl(BackbufferResolution);
+
+	SceneViewState.Allocate(Renderer->GetWorld()->GetFeatureLevel());
+
+}
+
+void UMovieGraphCoreRenderPassNode::FMovieGraphRenderPass::GatherOutputPassesImpl(TArray<FMovieGraphRenderDataIdentifier>& OutExpectedPasses) const
+{
+	OutExpectedPasses.Add(RenderDataIdentifier);
+}
+void UMovieGraphCoreRenderPassNode::FMovieGraphRenderPass::Teardown()
+{
+	FSceneViewStateInterface* Ref = SceneViewState.GetReference();
+	if (Ref)
+	{
+		Ref->ClearMIDPool();
+	}
+	SceneViewState.Destroy();
+}
+
+void UMovieGraphCoreRenderPassNode::FMovieGraphRenderPass::AddReferencedObjects(FReferenceCollector& Collector)
+{
+	FSceneViewStateInterface* Ref = SceneViewState.GetReference();
+	if (Ref)
+	{
+		Ref->AddReferencedObjects(Collector);
+	}
+}
+
+FName UMovieGraphCoreRenderPassNode::FMovieGraphRenderPass::GetBranchName() const
+{
+	return LayerData.BranchName;
+}
+
+TWeakObjectPtr<UMovieGraphDefaultRenderer> UMovieGraphCoreRenderPassNode::FMovieGraphRenderPass::GetRenderer() const
+{
+	return Renderer;
+}
+
+void UMovieGraphCoreRenderPassNode::FMovieGraphRenderPass::Render(const FMovieGraphTraversalContext& InFrameTraversalContext, const FMovieGraphTimeStepData& InTimeData)
+{
+	const bool bIncludeCDOs = true;
+	UMovieGraphOutputSettingNode* OutputSetting = InTimeData.EvaluatedConfig->GetSettingForBranch<UMovieGraphOutputSettingNode>(LayerData.BranchName, bIncludeCDOs);
+	if (!ensure(OutputSetting))
+	{
+		return;
+	}
+
+	// This is the size we actually render at.
+	UE::MovieGraph::DefaultRenderer::FRenderTargetInitParams RenderTargetInitParams;
+	RenderTargetInitParams.Size = OutputSetting->OutputResolution;
+	
+	// OCIO: Since this is a manually created Render target we don't need Gamma to be applied.
+	// We use this render target to render to via a display extension that utilizes Display Gamma
+	// which has a default value of 2.2 (DefaultDisplayGamma), therefore we need to set Gamma on this render target to 2.2 to cancel out any unwanted effects.
+	RenderTargetInitParams.TargetGamma = FOpenColorIORendering::DefaultDisplayGamma;
+	RenderTargetInitParams.PixelFormat = EPixelFormat::PF_FloatRGBA;
+
+	UTextureRenderTarget2D* RenderTarget = Renderer->GetOrCreateViewRenderTarget(RenderTargetInitParams);
+
+	FViewFamilyContextInitData InitData;
+	InitData.RenderTarget = RenderTarget->GameThread_GetRenderTargetResource();
+	InitData.World = Renderer->GetWorld();
+	InitData.CameraInfo = Renderer->GetCameraInfo(LayerData.CameraIdentifier);
+
+	// ToDo:
+	InitData.TimeData = InTimeData;
+	InitData.SceneCaptureSource = ESceneCaptureSource::SCS_FinalToneCurveHDR;
+	InitData.bWorldIsPaused = false;
+	InitData.GlobalScreenPercentageFraction = FLegacyScreenPercentageDriver::GetCVarResolutionFraction();
+	InitData.OverscanFraction = 0.f;
+	InitData.FrameIndex = 1;
+	InitData.bCameraCut = false;
+	InitData.AntiAliasingMethod = EAntiAliasingMethod::AAM_None;
+	InitData.SceneViewStateReference = SceneViewState.GetReference();
+
+
+	// Allocate the view we want to render from, then a Family for the view to live in.
+	// We apply a lot of MRQ-specific overrides to the ViewFamily and View, so do those next.
+	// Allocating the Scene View automatically 
+	TSharedRef<FSceneViewFamilyContext> ViewFamily = AllocateSceneViewFamilyContext(InitData);
+	AllocateSceneView(ViewFamily, InitData);
+	ApplyMoviePipelineOverridesToViewFamily(ViewFamily, InitData);
+
+	
+	FRenderTarget* RenderTargetResource = RenderTarget->GameThread_GetRenderTargetResource();
+	check(RenderTarget);
+
+	FHitProxyConsumer* HitProxyConsumer = nullptr;
+	const float DPIScale = 1.0f;
+	FCanvas Canvas = FCanvas(RenderTargetResource, HitProxyConsumer, Renderer->GetWorld(), Renderer->GetWorld()->GetFeatureLevel(), FCanvas::CDM_DeferDrawing, DPIScale);
+
+	// Submit the renderer to be rendered
+	GetRendererModule().BeginRenderingViewFamily(&Canvas, ViewFamily.ToSharedPtr().Get());
+
+	// If this was just to contribute to the history buffer, no need to go any further.
+	//if (InSampleState.bDiscardResult)
+	//{
+	//	return;
+	//}
+
+	// Take our per-frame Traversal Context and update it with context specific to this sample.
+	FMovieGraphTraversalContext UpdatedTraversalContext = InFrameTraversalContext;
+	UpdatedTraversalContext.Time = InTimeData;
+	UpdatedTraversalContext.RenderDataIdentifier = RenderDataIdentifier;
+
+
+	UE::MovieGraph::FMovieGraphSampleState SampleState;
+	SampleState.TraversalContext = MoveTemp(UpdatedTraversalContext);
+	SampleState.BackbufferResolution = RenderTargetInitParams.Size;
+	SampleState.AccumulatorResolution = RenderTargetInitParams.Size; // ToDo: Overscan
+	SampleState.bWriteSampleToDisk = false; // ToDo: From graph
+	SampleState.bRequiresAccumulator = InTimeData.bRequiresAccumulator;
+	SampleState.bFetchFromAccumulator = InTimeData.bIsLastTemporalSampleForFrame;
+
+	// Readback + Accumulate.
+	PostRendererSubmission(SampleState, RenderTargetInitParams, Canvas);
+}
+
+void UMovieGraphCoreRenderPassNode::FMovieGraphRenderPass::PostRendererSubmission(
+	const UE::MovieGraph::FMovieGraphSampleState& InSampleState,
+	const UE::MovieGraph::DefaultRenderer::FRenderTargetInitParams& InRenderTargetInitParams, FCanvas& InCanvas)
+{
+	// ToDo: Draw Letterboxing
+	
+
+	// We have a pool of accumulators - we multi-thread the accumulation on the task graph, and for each frame,
+	// the task has the previous samples as pre-reqs to keep the accumulation in order. However, each accumulator
+	// can only work on one frame at a time, so we create a pool of them to work concurrently. This needs a limit
+	// as large accumulations (16k) can take a lot of system RAM.
+	FMoviePipelineAccumulatorPoolPtr SampleAccumulatorPool = Renderer->GetOrCreateAccumulatorPool<FImageOverlappedAccumulator>();
+	UE::MovieGraph::DefaultRenderer::FSurfaceAccumulatorPool::FInstancePtr AccumulatorInstance = nullptr;
+	{
+		// SCOPE_CYCLE_COUNTER(STAT_MoviePipeline_WaitForAvailableAccumulator);
+		AccumulatorInstance = SampleAccumulatorPool->BlockAndGetAccumulator_GameThread(InSampleState.TraversalContext.Time.OutputFrameNumber, InSampleState.TraversalContext.RenderDataIdentifier);
+	}
+
+	FMoviePipelineSurfaceQueuePtr LocalSurfaceQueue = Renderer->GetOrCreateSurfaceQueue(InRenderTargetInitParams);
+	LocalSurfaceQueue->BlockUntilAnyAvailable();
+
+	UE::MovieGraph::FMovieGraphRenderDataAccumulationArgs AccumulationArgs;
+	{
+		AccumulationArgs.OutputMerger = Renderer->GetOwningGraph()->GetOutputMerger();
+		AccumulationArgs.ImageAccumulator = StaticCastSharedPtr<FImageOverlappedAccumulator>(AccumulatorInstance->Accumulator);
+		AccumulationArgs.bIsFirstSample = InSampleState.TraversalContext.Time.bIsFirstTemporalSampleForFrame; // FramePayload->IsFirstTile() && FramePayload->IsFirstTemporalSample()
+		AccumulationArgs.bIsLastSample = InSampleState.TraversalContext.Time.bIsLastTemporalSampleForFrame; // FramePayload->IsLastTile() && FramePayload->IsLastTemporalSample()
+	}
+
+	auto OnSurfaceReadbackFinished = [this, InSampleState, AccumulationArgs, AccumulatorInstance](TUniquePtr<FImagePixelData>&& InPixelData)
+	{
+		UE::Tasks::TTask<void> Task = UE::Tasks::Launch(UE_SOURCE_LOCATION, [PixelData = MoveTemp(InPixelData), InSampleState, AccumulationArgs, AccumulatorInstance]() mutable
+		{
+			// Enqueue a encode for this frame onto our worker thread.
+			UE::MovieGraph::AccumulateSample_TaskThread(MoveTemp(PixelData), InSampleState, AccumulationArgs);
+
+			// We have to defer clearing the accumulator until after sample accumulation has finished
+			if (AccumulationArgs.bIsLastSample)
+			{
+				// Final sample has now been executed, free the accumulator for reuse.
+				AccumulatorInstance->bIsActive = false;
+			}
+		}, AccumulatorInstance->TaskPrereq);
+
+		// Make the next accumulation task that uses this accumulator use the task we just created as a pre-req.
+		AccumulatorInstance->TaskPrereq = Task;
+
+		// Because we're run on a separate thread, we need to check validity differently. The standard
+		// TWeakObjectPtr will report non-valid during GC (even if the object it's pointing to isn't being
+		// GC'd).
+		const bool bEvenIfPendingKill = false;
+		const bool bThreadsafeTest = true;
+		const bool bValid = this->Renderer.IsValid(bEvenIfPendingKill, bThreadsafeTest);
+		if (ensureMsgf(bValid, TEXT("Renderer was garbage collected while outstanding tasks existed, outstanding tasks were not flushed properly during shutdown!")))
+		{
+			// The regular Get() will fail during GC so we use the above check to see if it's valid
+			// before ignoring it and directly getting the object.
+			this->Renderer.GetEvenIfUnreachable()->AddOutstandingRenderTask_AnyThread(Task);
+		}
+	};
+
+	FRenderTarget* RenderTarget = InCanvas.GetRenderTarget();
+
+	ENQUEUE_RENDER_COMMAND(CanvasRenderTargetResolveCommand)(
+		[LocalSurfaceQueue, OnSurfaceReadbackFinished, RenderTarget](FRHICommandListImmediate& RHICmdList) mutable
+		{
+			// The legacy surface reader takes the payload just so it can shuffle it into our callback, but we can just include the data
+			// directly in the callback, so this is just a dummy payload.
+			TSharedRef<FImagePixelDataPayload, ESPMode::ThreadSafe> FramePayload = MakeShared<FImagePixelDataPayload, ESPMode::ThreadSafe>();
+			LocalSurfaceQueue->OnRenderTargetReady_RenderThread(RenderTarget->GetRenderTargetTexture(), FramePayload, MoveTemp(OnSurfaceReadbackFinished));
+		});
+}
+
+TSharedRef<FSceneViewFamilyContext> UMovieGraphCoreRenderPassNode::FMovieGraphRenderPass::AllocateSceneViewFamilyContext(const FViewFamilyContextInitData& InInitData)
+{
+	FEngineShowFlags ShowFlags = RenderPassNode->GetShowFlags();
+	const EViewModeIndex ViewModeIndex = RenderPassNode->GetViewModeIndex();
+
+	const bool bIsPerspective = InInitData.CameraInfo.ViewInfo.ProjectionMode == ECameraProjectionMode::Type::Perspective;
+
+	// Allow the Engine Showflag system to override our engine showflags, based on our view mode index.
+	// This is required for certain debug view modes (to have matching show flags set for rendering).
+	ApplyViewMode(/*In*/ ViewModeIndex, bIsPerspective, /*InOut*/ShowFlags);
+
+	// And then we have to let another system override them again (based on cvars, etc.)
+	EngineShowFlagOverride(ESFIM_Game, ViewModeIndex, ShowFlags, false);
+
+	TSharedRef<FSceneViewFamilyContext> OutViewFamily = MakeShared<FSceneViewFamilyContext>(FSceneViewFamily::ConstructionValues(
+		InInitData.RenderTarget,
+		InInitData.World->Scene,
+		ShowFlags)
+		.SetTime(FGameTime::CreateUndilated(InInitData.TimeData.WorldSeconds, InInitData.TimeData.FrameDeltaTime))
+		.SetRealtimeUpdate(true));
+
+	// Used to specify if the Tone Curve is being applied or not to our Linear Output data
+	OutViewFamily->SceneCaptureSource = InInitData.SceneCaptureSource;
+	OutViewFamily->bWorldIsPaused = InInitData.bWorldIsPaused;
+	OutViewFamily->ViewMode = ViewModeIndex;
+	OutViewFamily->bOverrideVirtualTextureThrottle = true;
+
+	// ToDo: Let settings modify the ScreenPercentageInterface so third party screen percentages are supported.
+
+	// If UMoviePipelineViewFamilySetting never set a Screen percentage interface we fallback to default.
+	if (OutViewFamily->GetScreenPercentageInterface() == nullptr)
+	{
+		OutViewFamily->SetScreenPercentageInterface(new FLegacyScreenPercentageDriver(*OutViewFamily, InInitData.GlobalScreenPercentageFraction));
+	}
+
+	return OutViewFamily;
+}
+
+
+FSceneView* UMovieGraphCoreRenderPassNode::FMovieGraphRenderPass::AllocateSceneView(TSharedPtr<FSceneViewFamilyContext> InViewFamilyContext, FViewFamilyContextInitData& InInitData) const
+{
+	FSceneViewInitOptions ViewInitOptions;
+	ViewInitOptions.ViewFamily = InViewFamilyContext.Get();
+	ViewInitOptions.ViewOrigin = InInitData.CameraInfo.ViewInfo.Location;
+	ViewInitOptions.SetViewRectangle(FIntRect(FIntPoint(0, 0), FIntPoint(InInitData.RenderTarget->GetSizeXY().X, InInitData.RenderTarget->GetSizeXY().Y)));
+	ViewInitOptions.ViewRotationMatrix = FInverseRotationMatrix(InInitData.CameraInfo.ViewInfo.Rotation);
+	ViewInitOptions.ViewActor = InInitData.CameraInfo.ViewActor;
+
+	// Rotate the view 90 degrees to match the rest of the engine.
+	ViewInitOptions.ViewRotationMatrix = ViewInitOptions.ViewRotationMatrix * FMatrix(
+		FPlane(0, 0, 1, 0),
+		FPlane(1, 0, 0, 0),
+		FPlane(0, 1, 0, 0),
+		FPlane(0, 0, 0, 1));
+	float ViewFOV = InInitData.CameraInfo.ViewInfo.FOV;
+
+	// Inflate our FOV to support the overscan 
+	ViewFOV = 2.0f * FMath::RadiansToDegrees(FMath::Atan((1.0f + InInitData.OverscanFraction) * FMath::Tan(FMath::DegreesToRadians(ViewFOV * 0.5f))));
+	float DofSensorScale = 1.0f;
+
+	// ToDo: This isn't great but this appears to be how the system works... ideally we could fetch this from the camera itself
+	const EAspectRatioAxisConstraint AspectRatioAxisConstraint = GetDefault<ULocalPlayer>()->AspectRatioAxisConstraint;
+	FIntRect ViewportRect = FIntRect(FIntPoint(0, 0), FIntPoint(InInitData.RenderTarget->GetSizeXY().X, InInitData.RenderTarget->GetSizeXY().Y));
+	FMinimalViewInfo::CalculateProjectionMatrixGivenViewRectangle(InInitData.CameraInfo.ViewInfo, AspectRatioAxisConstraint, ViewportRect, /*InOut*/ ViewInitOptions);
+
+	// ToDo: High Res Tiling, Overscan support, letterboxing
+	ViewInitOptions.SceneViewStateInterface = InInitData.SceneViewStateReference;
+	ViewInitOptions.FOV = ViewFOV;
+	ViewInitOptions.DesiredFOV = ViewFOV;
+
+	FSceneView* View = new FSceneView(ViewInitOptions);
+	InViewFamilyContext->Views.Add(View);
+
+	View->StartFinalPostprocessSettings(View->ViewLocation);
+	// BlendPostProcessSettings(View, InOutSampleState, OptPayload);
+
+	// Scaling sensor size inversely with the the projection matrix [0][0] should physically
+	// cause the circle of confusion to be unchanged.
+	View->FinalPostProcessSettings.DepthOfFieldSensorWidth *= DofSensorScale;
+	// Modify the 'center' of the lens to be offset for high-res tiling, helps some effects (vignette) etc. still work.
+	// View->LensPrincipalPointOffsetScale = (FVector4f)CalculatePrinciplePointOffsetForTiling(InOutSampleState); // LWC_TODO: precision loss. CalculatePrinciplePointOffsetForTiling() could return float, it's normalized?
+	View->EndFinalPostprocessSettings(ViewInitOptions);
+	return View;
+}
+
+
+void UMovieGraphCoreRenderPassNode::FMovieGraphRenderPass::ApplyMoviePipelineOverridesToViewFamily(TSharedRef<FSceneViewFamilyContext> InOutFamily, const FViewFamilyContextInitData& InInitData)
+{
+	/*// A third set of overrides required to properly configure views to match the given showflags/etc.
+	// ToDo: There's now five(!) identical implementations of this, we should unify them.
+	// SetupViewForViewModeOverride(View);
+
+	// Override the view's FrameIndex to be based on our progress through the sequence. This greatly increases
+	// determinism with things like TAA.
+	InInitData.View->OverrideFrameIndexValue = InInitData.FrameIndex;
+	InInitData.View->bCameraCut = InInitData.bCameraCut;
+	InInitData.View->bIsOfflineRender = true;
+	InInitData.View->AntiAliasingMethod = InInitData.AntiAliasingMethod;
+
+	// Add any view extensions that were added to the scene to this View too
+	// OutViewFamily->ViewExtensions.Append(GEngine->ViewExtensions->GatherActiveExtensions(FSceneViewExtensionContext(GetWorld()->Scene)));
+	// ToDo: Allow render passes to add additional view extensions not added to the scene. (AddViewExtensions(*OutViewFamily, InOutSampleState))
+
+	// Set up the view extensions with this view family
+	//for (auto ViewExt : OutViewFamily->ViewExtensions)
+	//{
+	//	ViewExt->SetupViewFamily(*OutViewFamily.Get());
+	//}
+
+	// Set up the view
+	//for (int ViewExt = 0; ViewExt < OutViewFamily->ViewExtensions.Num(); ViewExt++)
+	//{
+	//	OutViewFamily->ViewExtensions[ViewExt]->SetupView(*OutViewFamily.Get(), *View);
+	//}
+
+	// Override the Motion Blur settings since these are controlled by the movie pipeline.
+	{
+		// FFrameRate OutputFrameRate = GetPipeline()->GetPipelinePrimaryConfig()->GetEffectiveFrameRate(GetPipeline()->GetTargetSequence());
+		FFrameRate OutputFrameRate = FFrameRate(24, 1); // ToDo, get this from config.
+
+		// We need to inversly scale the target FPS by time dilation to counteract slowmo. If scaling isn't applied then motion blur length
+		// stays the same length despite the smaller delta time and the blur ends up too long.
+		View->FinalPostProcessSettings.MotionBlurTargetFPS = FMath::RoundToInt(OutputFrameRate.AsDecimal() / FMath::Max(SMALL_NUMBER, InInitData.TimeData.TimeDilation));
+		View->FinalPostProcessSettings.MotionBlurAmount = InInitData.TimeData.MotionBlurFraction;
+		View->FinalPostProcessSettings.MotionBlurMax = 100.f;
+		View->FinalPostProcessSettings.bOverride_MotionBlurAmount = true;
+		View->FinalPostProcessSettings.bOverride_MotionBlurTargetFPS = true;
+		View->FinalPostProcessSettings.bOverride_MotionBlurMax = true;
+
+		// Skip the whole pass if they don't want motion blur.
+		if (FMath::IsNearlyZero(InInitData.TimeData.MotionBlurFraction))
+		{
+			OutViewFamily->EngineShowFlags.SetMotionBlur(false);
+		}
+	}
+
+	// Warn the user for invalid setting combinations / enforce hardware limitations
+	{
+
+		// Locked Exposure
+		const bool bAutoExposureAllowed = IsAutoExposureAllowed(InOutSampleState);
+		{
+			// If the rendering pass doesn't allow autoexposure and they dont' have manual exposure set up, warn.
+			if (!bAutoExposureAllowed && (View->FinalPostProcessSettings.AutoExposureMethod != EAutoExposureMethod::AEM_Manual))
+			{
+				// Skip warning if the project setting is disabled though, as exposure will be forced off in the renderer anyways.
+				const URendererSettings* RenderSettings = GetDefault<URendererSettings>();
+				if (RenderSettings->bDefaultFeatureAutoExposure != false)
+				{
+					UE_LOG(LogMovieRenderPipeline, Warning, TEXT("Camera Auto Exposure Method not supported by one or more render passes. Change the Auto Exposure Method to Manual!"));
+					View->FinalPostProcessSettings.AutoExposureMethod = EAutoExposureMethod::AEM_Manual;
+				}
+			}
+		}
+
+		// Orthographic cameras don't support anti-aliasing outside the path tracer (other than FXAA)
+		const bool bIsOrthographicCamera = !View->IsPerspectiveProjection();
+		if (bIsOrthographicCamera)
+		{
+			bool bIsSupportedAAMethod = View->AntiAliasingMethod == EAntiAliasingMethod::AAM_FXAA;
+			bool bIsPathTracer = OutViewFamily->EngineShowFlags.PathTracing;
+			bool bWarnJitters = InOutSampleState.ProjectionMatrixJitterAmount.SquaredLength() > SMALL_NUMBER;
+			if ((!bIsPathTracer && !bIsSupportedAAMethod) || bWarnJitters)
+			{
+				UE_LOG(LogMovieRenderPipeline, Warning, TEXT("Orthographic Cameras are only supported with PathTracer or Deferred with FXAA Anti-Aliasing"));
+			}
+		}
+
+		{
+			bool bMethodWasUnsupported = false;
+			if (View->AntiAliasingMethod == AAM_TemporalAA && !SupportsGen4TAA(View->GetShaderPlatform()))
+			{
+				UE_LOG(LogMovieRenderPipeline, Error, TEXT("TAA was requested but this hardware does not support it."));
+				bMethodWasUnsupported = true;
+			}
+			else if (View->AntiAliasingMethod == AAM_TSR && !SupportsTSR(View->GetShaderPlatform()))
+			{
+				UE_LOG(LogMovieRenderPipeline, Error, TEXT("TSR was requested but this hardware does not support it."));
+				bMethodWasUnsupported = true;
+			}
+
+			if (bMethodWasUnsupported)
+			{
+				View->AntiAliasingMethod = AAM_None;
+			}
+		}
+	}
+
+
+	// Anti Aliasing
+	{
+		// If we're not using TAA, TSR, or Path Tracing we will apply the View Matrix projection jitter. Normally TAA sets this
+		// inside FSceneRenderer::PreVisibilityFrameSetup. Path Tracing does its own anti-aliasing internally.
+		bool bApplyProjectionJitter = !bIsOrthographicCamera
+			&& !OutViewFamily->EngineShowFlags.PathTracing
+			&& !IsTemporalAccumulationBasedMethod(View->AntiAliasingMethod);
+		if (bApplyProjectionJitter)
+		{
+			View->ViewMatrices.HackAddTemporalAAProjectionJitter(InOutSampleState.ProjectionMatrixJitterAmount);
+		}
+	}
+
+	return OutViewFamily;*/
+}
+
+namespace UE::MovieGraph
+{
+	void AccumulateSample_TaskThread(TUniquePtr<FImagePixelData>&& InPixelData, const UE::MovieGraph::FMovieGraphSampleState InSampleState, const UE::MovieGraph::FMovieGraphRenderDataAccumulationArgs& InAccumulationParams)
+	{
+		TRACE_CPUPROFILER_EVENT_SCOPE(MoviePipeline_AccumulateSample);
+		
+		TUniquePtr<FImagePixelData> SamplePixelData = MoveTemp(InPixelData);
+
+		// Associate the sample state with the image as payload data, this allows downstream systems to fetch the values without us having to store the data
+		// separately and ensure they stay paired the whole way down.
+		TSharedPtr<UE::MovieGraph::FMovieGraphSampleState> SampleStatePayload = MakeShared<UE::MovieGraph::FMovieGraphSampleState>(InSampleState);
+		SamplePixelData->SetPayload(StaticCastSharedPtr<IImagePixelDataPayload>(SampleStatePayload));
+
+		TSharedPtr<IMovieGraphOutputMerger, ESPMode::ThreadSafe> OutputMergerPin = InAccumulationParams.OutputMerger.Pin();
+		if (!OutputMergerPin.IsValid())
+		{
+			return;
+		}
+
+		const bool bIsWellFormed = SamplePixelData->IsDataWellFormed();
+		check(bIsWellFormed);
+
+		if (SampleStatePayload->bWriteSampleToDisk)
+		{
+			// Debug Feature: Write the raw sample to disk for debugging purposes. We copy the data here,
+			// as we don't want to disturb the memory flow below.
+			TUniquePtr<FImagePixelData> SampleData = SamplePixelData->CopyImageData();
+			OutputMergerPin->OnSingleSampleDataAvailable_AnyThread(MoveTemp(SampleData));
+		}
+
+		// Optimization! If we don't need the accumulator (no tiling, no sub-sampling) then we'll skip it and just send it straight to the output stage.
+		// This reduces memory requirements and improves performance in the baseline case.
+		if (!SampleStatePayload->bRequiresAccumulator)
+		{
+			OutputMergerPin->OnCompleteRenderPassDataAvailable_AnyThread(MoveTemp(SamplePixelData));
+			return;
+		}
+
+		TSharedPtr<FImageOverlappedAccumulator> AccumulatorPin = InAccumulationParams.ImageAccumulator.Pin();
+		if (AccumulatorPin->NumChannels == 0)
+		{
+			LLM_SCOPE_BYNAME(TEXT("MoviePipeline/ImageAccumulatorInitMemory"));
+			const int32 ChannelCount = 4;
+			AccumulatorPin->InitMemory(SampleStatePayload->AccumulatorResolution, ChannelCount);
+			AccumulatorPin->ZeroPlanes();
+			AccumulatorPin->AccumulationGamma = 1.f;
+		}
+
+		// Accumulate the new sample to our target
+		{
+			// ToDo: Handle incorrectly sized samples (Post Process materials using r.screenpercentage)
+			
+			{
+				TRACE_CPUPROFILER_EVENT_SCOPE(MoviePipeline_AccumulatePixelData);
+				
+				// ToDo: High-res Tiling Support
+				const FIntPoint OverlappedPad = FIntPoint(0, 0);
+				const FIntPoint TileSize = SampleStatePayload->BackbufferResolution;
+				const FIntPoint OverlappedOffset = FIntPoint(0, 0);
+				const FVector2D OverlappedSubpixelShift = FVector2D(0.5f, 0.5f);
+				::MoviePipeline::FTileWeight1D WeightFunctionX;
+				::MoviePipeline::FTileWeight1D WeightFunctionY;
+				WeightFunctionX.InitHelper(OverlappedPad.X, TileSize.X, OverlappedPad.X);
+				WeightFunctionY.InitHelper(OverlappedPad.Y, TileSize.Y, OverlappedPad.Y);
+				
+				AccumulatorPin->AccumulatePixelData(*SamplePixelData, OverlappedOffset, OverlappedSubpixelShift, WeightFunctionX, WeightFunctionY);
+			}
+
+		}
+
+		// Finally on our last sample, we fetch the data out of the accumulator
+		// and move it to the Output Merger.
+		if (SampleStatePayload->bFetchFromAccumulator)
+		{
+			TRACE_CPUPROFILER_EVENT_SCOPE(MoviePipeline_FetchAccumulatedPixelData);
+			int32 FullSizeX = AccumulatorPin->PlaneSize.X;
+			int32 FullSizeY = AccumulatorPin->PlaneSize.Y;
+
+			// Now that a tile is fully built and accumulated we can notify the output builder that the
+			// data is ready so it can pass that onto the output containers (if needed).
+			if (SamplePixelData->GetType() == EImagePixelType::Float32)
+			{
+				// 32 bit FLinearColor
+				TUniquePtr<TImagePixelData<FLinearColor> > FinalPixelData = MakeUnique<TImagePixelData<FLinearColor>>(FIntPoint(FullSizeX, FullSizeY), SampleStatePayload);
+				AccumulatorPin->FetchFinalPixelDataLinearColor(FinalPixelData->Pixels);
+
+				// Send the data to the Output Builder
+				OutputMergerPin->OnCompleteRenderPassDataAvailable_AnyThread(MoveTemp(FinalPixelData));
+			}
+			else if (SamplePixelData->GetType() == EImagePixelType::Float16)
+			{
+				// 16 bit FLinearColor
+				TUniquePtr<TImagePixelData<FFloat16Color> > FinalPixelData = MakeUnique<TImagePixelData<FFloat16Color>>(FIntPoint(FullSizeX, FullSizeY), SampleStatePayload);
+				AccumulatorPin->FetchFinalPixelDataHalfFloat(FinalPixelData->Pixels);
+
+				// Send the data to the Output Builder
+				OutputMergerPin->OnCompleteRenderPassDataAvailable_AnyThread(MoveTemp(FinalPixelData));
+			}
+			else if (SamplePixelData->GetType() == EImagePixelType::Color)
+			{
+				// 8bit FColors
+				TUniquePtr<TImagePixelData<FColor>> FinalPixelData = MakeUnique<TImagePixelData<FColor>>(FIntPoint(FullSizeX, FullSizeY), SampleStatePayload);
+				AccumulatorPin->FetchFinalPixelDataByte(FinalPixelData->Pixels);
+
+				// Send the data to the Output Builder
+				OutputMergerPin->OnCompleteRenderPassDataAvailable_AnyThread(MoveTemp(FinalPixelData));
+			}
+			else
+			{
+				check(0);
+			}
+
+			{
+				TRACE_CPUPROFILER_EVENT_SCOPE(MoviePipeline_FreeAccumulatedPixelData);
+				// Free the memory in the accumulator.
+				AccumulatorPin->Reset();
+			}
+		}
+
+		{
+			TRACE_CPUPROFILER_EVENT_SCOPE(ReleaseSampleData);
+			SamplePixelData.Reset();
+		}
+	}
+}
diff --git a/Engine/Plugins/MovieScene/MovieRenderPipeline/Source/MovieRenderPipelineCore/Private/Graph/Nodes/MovieGraphDefRenderPassNode.cpp b/Engine/Plugins/MovieScene/MovieRenderPipeline/Source/MovieRenderPipelineCore/Private/Graph/Nodes/MovieGraphDefRenderPassNode.cpp
new file mode 100644
--- /dev/null
+++ b/Engine/Plugins/MovieScene/MovieRenderPipeline/Source/MovieRenderPipelineCore/Private/Graph/Nodes/MovieGraphDefRenderPassNode.cpp
@@ -0,0 +1,30 @@
+﻿// Copyright Epic Games, Inc. All Rights Reserved.
+
+#include "Graph/Nodes/MovieGraphDefRenderPassNode.h"
+
+#include "Engine/EngineBaseTypes.h"
+#include "ShowFlags.h"
+
+#if WITH_EDITOR
+FText UMovieGraphDeferredRenderPassNode::GetNodeTitle(const bool bGetDescriptive) const
+{
+	return NSLOCTEXT("MovieGraphNodes", "DeferredRenderPassGraphNode_Description", "Deferred Renderer");
+}
+#endif
+
+FString UMovieGraphDeferredRenderPassNode::GetRendererNameImpl() const
+{
+	static const FString RendererNameImpl(TEXT("Deferred"));
+	return RendererNameImpl;
+}
+
+EViewModeIndex UMovieGraphDeferredRenderPassNode::GetViewModeIndex() const
+{
+	return VMI_Lit;
+}
+
+FEngineShowFlags UMovieGraphDeferredRenderPassNode::GetShowFlags() const
+{
+	FEngineShowFlags ShowFlags = FEngineShowFlags(EShowFlagInitMode::ESFIM_Game);
+	return ShowFlags;
+}
\ No newline at end of file
diff --git a/Engine/Plugins/MovieScene/MovieRenderPipeline/Source/MovieRenderPipelineCore/Private/Graph/Nodes/MovieGraphDeferredRenderPassNode.cpp b/Engine/Plugins/MovieScene/MovieRenderPipeline/Source/MovieRenderPipelineCore/Private/Graph/Nodes/MovieGraphDeferredRenderPassNode.cpp
deleted file mode 100644
diff --git a/Engine/Plugins/MovieScene/MovieRenderPipeline/Source/MovieRenderPipelineCore/Private/Graph/Nodes/MovieGraphPathTracedRenderPassNode.cpp b/Engine/Plugins/MovieScene/MovieRenderPipeline/Source/MovieRenderPipelineCore/Private/Graph/Nodes/MovieGraphPathTracedRenderPassNode.cpp
new file mode 100644
--- /dev/null
+++ b/Engine/Plugins/MovieScene/MovieRenderPipeline/Source/MovieRenderPipelineCore/Private/Graph/Nodes/MovieGraphPathTracedRenderPassNode.cpp
@@ -0,0 +1,56 @@
+﻿// Copyright Epic Games, Inc. All Rights Reserved.
+
+#include "Graph/Nodes/MovieGraphPathTracedRenderPassNode.h"
+
+#include "Engine/EngineBaseTypes.h"
+#include "ShowFlags.h"
+
+#if WITH_EDITOR
+FText UMovieGraphPathTracedRenderPassNode::GetNodeTitle(const bool bGetDescriptive) const
+{
+	return NSLOCTEXT("MovieGraphNodes", "PathTracedRenderPassGraphNode_Description", "Path Traced Renderer");
+}
+#endif
+
+void UMovieGraphPathTracedRenderPassNode::SetupImpl(const FMovieGraphRenderPassSetupData& InSetupData)
+{
+	Super::SetupImpl(InSetupData);
+	
+	// Hide the progress display during the render
+	if (IConsoleVariable* ProgressDisplayCvar = IConsoleManager::Get().FindConsoleVariable(TEXT("r.PathTracing.ProgressDisplay")))
+	{
+		bOriginalProgressDisplayCvarValue = ProgressDisplayCvar->GetBool();
+		ProgressDisplayCvar->Set(false);
+	}
+}
+
+void UMovieGraphPathTracedRenderPassNode::TeardownImpl()
+{
+	Super::TeardownImpl();
+
+	// Restore the original setting for the progress display
+	if (IConsoleVariable* ProgressDisplayCvar = IConsoleManager::Get().FindConsoleVariable(TEXT("r.PathTracing.ProgressDisplay")))
+	{
+		ProgressDisplayCvar->Set(bOriginalProgressDisplayCvarValue);
+	}
+}
+
+FString UMovieGraphPathTracedRenderPassNode::GetRendererNameImpl() const
+{
+	static const FString RendererNameImpl(TEXT("PathTraced"));
+	return RendererNameImpl;
+}
+
+EViewModeIndex UMovieGraphPathTracedRenderPassNode::GetViewModeIndex() const
+{
+	return VMI_PathTracing;
+}
+
+FEngineShowFlags UMovieGraphPathTracedRenderPassNode::GetShowFlags() const
+{
+	FEngineShowFlags ShowFlags = FEngineShowFlags(EShowFlagInitMode::ESFIM_Game);
+	ShowFlags.SetPathTracing(true);
+	// TODO: SetMotionBlur()?
+	
+	return ShowFlags;
+}
\ No newline at end of file
diff --git a/Engine/Plugins/MovieScene/MovieRenderPipeline/Source/MovieRenderPipelineCore/Public/Graph/MovieGraphPlaceholderNodes.cpp b/Engine/Plugins/MovieScene/MovieRenderPipeline/Source/MovieRenderPipelineCore/Public/Graph/MovieGraphPlaceholderNodes.cpp
--- a/Engine/Plugins/MovieScene/MovieRenderPipeline/Source/MovieRenderPipelineCore/Public/Graph/MovieGraphPlaceholderNodes.cpp
+++ b/Engine/Plugins/MovieScene/MovieRenderPipeline/Source/MovieRenderPipelineCore/Public/Graph/MovieGraphPlaceholderNodes.cpp
@@ -7,36 +7,8 @@
 
 #define LOCTEXT_NAMESPACE "MovieGraphNode"
 
-static const FText NodeCategory_Rendering = LOCTEXT("NodeCategory_Rendering", "Rendering");
 static const FText NodeCategory_Settings = LOCTEXT("NodeCategory_Settings", "Settings");
 
-#if WITH_EDITOR
-FText UMovieGraphPathTracedRendererNode::GetNodeTitle(const bool bGetDescriptive) const
-{
-	static const FText PathTracedRendererNodeName = LOCTEXT("NodeName_PathTracedRenderer", "Path Traced Renderer");
-	return PathTracedRendererNodeName;
-}
-
-FText UMovieGraphPathTracedRendererNode::GetMenuCategory() const
-{
-	return NodeCategory_Rendering;
-}
-
-FLinearColor UMovieGraphPathTracedRendererNode::GetNodeTitleColor() const
-{
-	static const FLinearColor DeferredRendererNodeColor = FLinearColor(0.572f, 0.274f, 1.f);
-	return DeferredRendererNodeColor;
-}
-
-FSlateIcon UMovieGraphPathTracedRendererNode::GetIconAndTint(FLinearColor& OutColor) const
-{
-	static const FSlateIcon DeferredRendererIcon = FSlateIcon(FAppStyle::GetAppStyleSetName(), "SequenceRecorder.TabIcon");
-
-	OutColor = FLinearColor::White;
-	return DeferredRendererIcon;
-}
-#endif // WITH_EDITOR
-
 #if WITH_EDITOR
 FText UMovieGraphAntiAliasingNode::GetNodeTitle(const bool bGetDescriptive) const
 {
diff --git a/Engine/Plugins/MovieScene/MovieRenderPipeline/Source/MovieRenderPipelineCore/Public/Graph/MovieGraphPlaceholderNodes.h b/Engine/Plugins/MovieScene/MovieRenderPipeline/Source/MovieRenderPipelineCore/Public/Graph/MovieGraphPlaceholderNodes.h
--- a/Engine/Plugins/MovieScene/MovieRenderPipeline/Source/MovieRenderPipelineCore/Public/Graph/MovieGraphPlaceholderNodes.h
+++ b/Engine/Plugins/MovieScene/MovieRenderPipeline/Source/MovieRenderPipelineCore/Public/Graph/MovieGraphPlaceholderNodes.h
@@ -6,23 +6,6 @@
 
 #include "MovieGraphPlaceholderNodes.generated.h"
 
-/** A node which represents a path traced renderer. */
-UCLASS()
-class MOVIERENDERPIPELINECORE_API UMovieGraphPathTracedRendererNode : public UMovieGraphSettingNode
-{
-	GENERATED_BODY()
-
-public:
-	UMovieGraphPathTracedRendererNode() = default;
-
-#if WITH_EDITOR
-	virtual FText GetNodeTitle(const bool bGetDescriptive = false) const override;
-	virtual FText GetMenuCategory() const override;
-	virtual FLinearColor GetNodeTitleColor() const override;
-	virtual FSlateIcon GetIconAndTint(FLinearColor& OutColor) const override;
-#endif
-};
-
 /** A node which configures anti-aliasing settings. */
 UCLASS()
 class MOVIERENDERPIPELINECORE_API UMovieGraphAntiAliasingNode : public UMovieGraphSettingNode
diff --git a/Engine/Plugins/MovieScene/MovieRenderPipeline/Source/MovieRenderPipelineCore/Public/Graph/MoviePipelineRenderLayerSubsystem.h b/Engine/Plugins/MovieScene/MovieRenderPipeline/Source/MovieRenderPipelineCore/Public/Graph/MoviePipelineRenderLayerSubsystem.h
--- a/Engine/Plugins/MovieScene/MovieRenderPipeline/Source/MovieRenderPipelineCore/Public/Graph/MoviePipelineRenderLayerSubsystem.h
+++ b/Engine/Plugins/MovieScene/MovieRenderPipeline/Source/MovieRenderPipelineCore/Public/Graph/MoviePipelineRenderLayerSubsystem.h
@@ -113,6 +113,8 @@ class UMoviePipelineVisibilityModifier : public UMoviePipelineCollectionModifier
 	GENERATED_BODY()
 
 public:
+	UMoviePipelineVisibilityModifier();
+	
 	UFUNCTION(BlueprintCallable, Category = "Settings")
 	void SetHidden(const bool bInIsHidden) { bIsHidden = bInIsHidden; }
 
@@ -126,16 +128,45 @@ public:
 	virtual void UndoModifier() override;
 
 private:
-	void SetActorHiddenState(AActor* Actor, const bool bInIsHidden) const;
+	/** Various visibility properties for an actor. */
+	struct FActorVisibilityState
+	{
+		TSoftObjectPtr<AActor> Actor = nullptr;
+		uint8 bIsHidden : 1 = false;
+		uint8 bCastShadowWhileHidden : 1 = false;
+		uint8 bAffectIndirectLightingWhileHidden : 1 = false;
+		uint8 bHoldout : 1 = false;
+	};
+
+	/** Updates an actor's visibility state to the state contained in NewVisibilityState. */
+	void SetActorVisibilityState(const FActorVisibilityState& NewVisibilityState);
 
 private:
-	/** Maps an actor to its original hidden state. */
-	UPROPERTY(Transient)
-	TMap<TSoftObjectPtr<AActor>, bool> ModifiedActors;
+	/** Tracks actor visibility state prior to having the modifier applied. */
+	TArray<FActorVisibilityState> ModifiedActors;
 
 public:
+	/**
+	 * If true, the actor will not be visible and will not contribute to any secondary effects (shadows, indirect
+	 * lighting) unless their respective flags are set below.
+	 */
 	UPROPERTY(BlueprintReadWrite, EditAnywhere, Category = "Settings")
-	bool bIsHidden = false;
+	uint8 bIsHidden : 1;
+
+	/** If true, the primitive will cast shadows even if it is hidden. */
+	UPROPERTY(BlueprintReadWrite, EditAnywhere, Category = "Primitive")
+	uint8 bCastShadowWhileHidden : 1;
+
+	/** Controls whether the primitive should affect indirect lighting when hidden. */
+	UPROPERTY(BlueprintReadWrite, EditAnywhere, Category = "Primitive")
+	uint8 bAffectIndirectLightingWhileHidden : 1;
+
+	/**
+	 * If true, the primitive will render black with an alpha of 0, but all secondary effects (shadows, reflections,
+	 * indirect lighting) remain. This feature is currently only implemented in the Path Tracer.
+	 */
+	UPROPERTY(BlueprintReadWrite, EditAnywhere, Category = "Primitive|Path Tracing")
+	uint8 bHoldout : 1;
 };
 
 /**
diff --git a/Engine/Plugins/MovieScene/MovieRenderPipeline/Source/MovieRenderPipelineCore/Public/Graph/Nodes/MovieGraphCoreRenderPassNode.h b/Engine/Plugins/MovieScene/MovieRenderPipeline/Source/MovieRenderPipelineCore/Public/Graph/Nodes/MovieGraphCoreRenderPassNode.h
new file mode 100644
--- /dev/null
+++ b/Engine/Plugins/MovieScene/MovieRenderPipeline/Source/MovieRenderPipelineCore/Public/Graph/Nodes/MovieGraphCoreRenderPassNode.h
@@ -0,0 +1,159 @@
+// Copyright Epic Games, Inc. All Rights Reserved.
+
+#pragma once
+
+#include "Graph/Nodes/MovieGraphRenderPassNode.h"
+#include "Graph/MovieGraphDefaultRenderer.h" // For CameraInfo
+#include "Graph/MovieGraphDataTypes.h"
+#include "SceneTypes.h"
+#include "Camera/CameraTypes.h"
+#include "Engine/EngineTypes.h"
+#include "Async/TaskGraphFwd.h"
+#include "Styling/AppStyle.h"
+
+#include "MovieGraphCoreRenderPassNode.generated.h"
+
+// Forward Declares
+class UMovieGraphDefaultRenderer;
+namespace UE::MovieGraph::DefaultRenderer
+{
+	struct FRenderTargetInitParams;
+	struct IMovieGraphOutputMerger;
+}
+struct FMovieGraphRenderPassSetupData;
+struct FMovieGraphRenderPassLayerData;
+struct FImageOverlappedAccumulator;
+struct FEngineShowFlags;
+enum EViewModeIndex : int;
+
+// For FViewFamilyContextInitData
+class FRenderTarget;
+class UWorld;
+class FSceneView;
+class AActor;
+class FSceneViewFamilyContext;
+class FCanvas;
+
+namespace UE::MovieGraph
+{
+	struct FMovieGraphRenderDataAccumulationArgs
+	{
+	public:
+		TWeakPtr<FImageOverlappedAccumulator, ESPMode::ThreadSafe> ImageAccumulator;
+		TWeakPtr<IMovieGraphOutputMerger, ESPMode::ThreadSafe> OutputMerger;
+
+		// If it's the first sample then we will reset the accumulator to a clean slate before accumulating into it.
+		bool bIsFirstSample;
+		// If it's the last sample, then we will trigger moving the data to the output merger. It can be both the first and last sample at the same time.
+		bool bIsLastSample;
+	};
+
+	void AccumulateSample_TaskThread(TUniquePtr<FImagePixelData>&& InPixelData, const UE::MovieGraph::FMovieGraphSampleState InSampleState, const UE::MovieGraph::FMovieGraphRenderDataAccumulationArgs& InAccumulationParams);
+}
+
+/** Core functionality for a typical render pass node. */
+UCLASS(Abstract)
+class MOVIERENDERPIPELINECORE_API UMovieGraphCoreRenderPassNode : public UMovieGraphRenderPassNode
+{
+	GENERATED_BODY()
+
+public:
+#if WITH_EDITOR
+	virtual FLinearColor GetNodeTitleColor() const override
+	{
+		return FLinearColor(0.572f, 0.274f, 1.f);
+	}
+
+	virtual FSlateIcon GetIconAndTint(FLinearColor& OutColor) const override
+	{
+		static const FSlateIcon DeferredRendererIcon = FSlateIcon(FAppStyle::GetAppStyleSetName(), "SequenceRecorder.TabIcon");
+
+		OutColor = FLinearColor::White;
+		return DeferredRendererIcon;
+	}
+#endif
+
+protected:
+	// UMovieGraphRenderPassNode Interface
+	virtual void SetupImpl(const FMovieGraphRenderPassSetupData& InSetupData) override;
+	virtual void TeardownImpl() override;
+	virtual void RenderImpl(const FMovieGraphTraversalContext& InFrameTraversalContext, const FMovieGraphTimeStepData& InTimeData) override;
+	virtual void GatherOutputPassesImpl(TArray<FMovieGraphRenderDataIdentifier>& OutExpectedPasses) const override;
+	// ~UMovieGraphRenderPassNode Interface
+
+	/** Gets the view mode index that should be active for this renderer. */
+	virtual EViewModeIndex GetViewModeIndex() const;
+
+	/** Gets the show flags that should be active for this renderer. */
+	virtual FEngineShowFlags GetShowFlags() const;
+
+	static void AddReferencedObjects(UObject* InThis, FReferenceCollector& Collector);
+
+	struct FViewFamilyContextInitData
+	{
+		FViewFamilyContextInitData()
+			: RenderTarget(nullptr)
+			, World(nullptr)
+			, SceneCaptureSource(ESceneCaptureSource::SCS_MAX)
+			, bWorldIsPaused(false)
+			, GlobalScreenPercentageFraction(1.0f)
+			, OverscanFraction(0.f)
+			, FrameIndex(-1)
+			, bCameraCut(false)
+			, AntiAliasingMethod(EAntiAliasingMethod::AAM_None)
+			, View(nullptr)
+			, SceneViewStateReference(nullptr)
+		{
+		}
+
+		class FRenderTarget* RenderTarget;
+		class UWorld* World;
+		FMovieGraphTimeStepData TimeData;
+		ESceneCaptureSource SceneCaptureSource;
+		bool bWorldIsPaused;
+		float GlobalScreenPercentageFraction;
+		float OverscanFraction;
+		int32 FrameIndex;
+		bool bCameraCut;
+		EAntiAliasingMethod AntiAliasingMethod;
+
+		// Ownership of this will be passed to the FSceneViewFamilyContext
+		FSceneView* View;
+
+		FSceneViewStateInterface* SceneViewStateReference;
+
+		// Camera Setup
+		UE::MovieGraph::DefaultRenderer::FCameraInfo CameraInfo;
+	};
+
+	struct FMovieGraphRenderPass
+	{
+		void Setup(TWeakObjectPtr<UMovieGraphDefaultRenderer> InRenderer, TWeakObjectPtr<UMovieGraphCoreRenderPassNode> InRenderPassNode, const FMovieGraphRenderPassLayerData& InLayer);
+		void Teardown();
+		void Render(const FMovieGraphTraversalContext& InFrameTraversalContext, const FMovieGraphTimeStepData& InTimeData);
+		void GatherOutputPassesImpl(TArray<FMovieGraphRenderDataIdentifier>& OutExpectedPasses) const;
+		void AddReferencedObjects(FReferenceCollector& Collector);
+		FName GetBranchName() const;
+		TWeakObjectPtr<UMovieGraphDefaultRenderer> GetRenderer() const;
+
+	protected:
+		TSharedRef<FSceneViewFamilyContext> AllocateSceneViewFamilyContext(const FViewFamilyContextInitData& InInitData);
+		FSceneView* AllocateSceneView(TSharedPtr<FSceneViewFamilyContext> InViewFamilyContext, FViewFamilyContextInitData& InInitData) const;
+		void ApplyMoviePipelineOverridesToViewFamily(TSharedRef<FSceneViewFamilyContext> InOutFamily, const FViewFamilyContextInitData& InInitData);
+		void PostRendererSubmission(const UE::MovieGraph::FMovieGraphSampleState& InSampleState, const UE::MovieGraph::DefaultRenderer::FRenderTargetInitParams& InRenderTargetInitParams, FCanvas& InCanvas);
+
+	protected:
+		FMovieGraphRenderPassLayerData LayerData;
+
+		/** Unique identifier passed in GatherOutputPasses and with each render that identifies the data produced by this renderer. */
+		FMovieGraphRenderDataIdentifier RenderDataIdentifier;
+
+		// Scene View history used by the renderer 
+		FSceneViewStateReference SceneViewState;
+
+		TWeakObjectPtr<class UMovieGraphDefaultRenderer> Renderer;
+		TWeakObjectPtr<class UMovieGraphCoreRenderPassNode> RenderPassNode;
+	};
+
+	TArray<TUniquePtr<FMovieGraphRenderPass>> CurrentInstances;
+};
\ No newline at end of file
diff --git a/Engine/Plugins/MovieScene/MovieRenderPipeline/Source/MovieRenderPipelineCore/Public/Graph/Nodes/MovieGraphDefRenderPassNode.h b/Engine/Plugins/MovieScene/MovieRenderPipeline/Source/MovieRenderPipelineCore/Public/Graph/Nodes/MovieGraphDefRenderPassNode.h
new file mode 100644
--- /dev/null
+++ b/Engine/Plugins/MovieScene/MovieRenderPipeline/Source/MovieRenderPipelineCore/Public/Graph/Nodes/MovieGraphDefRenderPassNode.h
@@ -0,0 +1,29 @@
+﻿// Copyright Epic Games, Inc. All Rights Reserved.
+
+#pragma once
+
+#include "Graph/Nodes/MovieGraphCoreRenderPassNode.h"
+
+#include "MovieGraphDefRenderPassNode.generated.h"
+
+/** A render node which uses the deferred renderer. */
+UCLASS()
+class MOVIERENDERPIPELINECORE_API UMovieGraphDeferredRenderPassNode : public UMovieGraphCoreRenderPassNode
+{
+	GENERATED_BODY()
+
+public:
+#if WITH_EDITOR
+	virtual FText GetNodeTitle(const bool bGetDescriptive = false) const override;
+#endif
+
+protected:
+	// UMovieGraphRenderPassNode Interface
+	virtual FString GetRendererNameImpl() const override;
+	// ~UMovieGraphRenderPassNode Interface
+
+	// UMovieGraphCoreRenderPassNode Interface
+	virtual EViewModeIndex GetViewModeIndex() const override;
+	virtual FEngineShowFlags GetShowFlags() const override;
+	// ~UMovieGraphCoreRenderPassNode Interface
+};
\ No newline at end of file
diff --git a/Engine/Plugins/MovieScene/MovieRenderPipeline/Source/MovieRenderPipelineCore/Public/Graph/Nodes/MovieGraphDeferredRenderPassNode.h b/Engine/Plugins/MovieScene/MovieRenderPipeline/Source/MovieRenderPipelineCore/Public/Graph/Nodes/MovieGraphDeferredRenderPassNode.h
deleted file mode 100644
diff --git a/Engine/Plugins/MovieScene/MovieRenderPipeline/Source/MovieRenderPipelineCore/Public/Graph/Nodes/MovieGraphPathTracedRenderPassNode.h b/Engine/Plugins/MovieScene/MovieRenderPipeline/Source/MovieRenderPipelineCore/Public/Graph/Nodes/MovieGraphPathTracedRenderPassNode.h
new file mode 100644
--- /dev/null
+++ b/Engine/Plugins/MovieScene/MovieRenderPipeline/Source/MovieRenderPipelineCore/Public/Graph/Nodes/MovieGraphPathTracedRenderPassNode.h
@@ -0,0 +1,41 @@
+﻿// Copyright Epic Games, Inc. All Rights Reserved.
+
+#pragma once
+
+#include "Graph/Nodes/MovieGraphCoreRenderPassNode.h"
+
+#include "MovieGraphPathTracedRenderPassNode.generated.h"
+
+/** A render node which uses the path tracer. */
+UCLASS()
+class MOVIERENDERPIPELINECORE_API UMovieGraphPathTracedRenderPassNode : public UMovieGraphCoreRenderPassNode
+{
+	GENERATED_BODY()
+
+public:
+#if WITH_EDITOR
+	virtual FText GetNodeTitle(const bool bGetDescriptive = false) const override;
+#endif
+
+	// UMovieGraphRenderPassNode Interface
+	virtual void SetupImpl(const FMovieGraphRenderPassSetupData& InSetupData) override;
+	virtual void TeardownImpl() override;
+	// ~UMovieGraphRenderPassNode Interface
+
+protected:
+	// UMovieGraphRenderPassNode Interface
+	virtual FString GetRendererNameImpl() const override;
+	// ~UMovieGraphRenderPassNode Interface
+
+	// UMovieGraphCoreRenderPassNode Interface
+	virtual EViewModeIndex GetViewModeIndex() const override;
+	virtual FEngineShowFlags GetShowFlags() const override;
+	// ~UMovieGraphCoreRenderPassNode Interface
+
+private:
+	/**
+	 * The original value of the "r.PathTracing.ProgressDisplay" cvar before the render starts. The progress display
+	 * will be hidden during the render.
+	 */
+	bool bOriginalProgressDisplayCvarValue = false;
+};
\ No newline at end of file
