MovieGraph: Add support for soloing shots during rendering, and expanding them before rendering for the Handle Frames feature.

diff --git a/Engine/Plugins/MovieScene/MovieRenderPipeline/Source/MovieRenderPipelineCore/Private/Graph/MovieGraphLinearTimeStep.cpp b/Engine/Plugins/MovieScene/MovieRenderPipeline/Source/MovieRenderPipelineCore/Private/Graph/MovieGraphLinearTimeStep.cpp
--- a/Engine/Plugins/MovieScene/MovieRenderPipeline/Source/MovieRenderPipelineCore/Private/Graph/MovieGraphLinearTimeStep.cpp
+++ b/Engine/Plugins/MovieScene/MovieRenderPipeline/Source/MovieRenderPipelineCore/Private/Graph/MovieGraphLinearTimeStep.cpp
@@ -28,7 +28,7 @@ void UMovieGraphLinearTimeStep::TickProducingFrames()
 {
 	int32 CurrentShotIndex = GetOwningGraph()->GetCurrentShotIndex();
 	const TArray<TObjectPtr<UMoviePipelineExecutorShot>>& ActiveShotList = GetOwningGraph()->GetActiveShotList();
-	UMoviePipelineExecutorShot* CurrentCameraCut = ActiveShotList[CurrentShotIndex];
+	const TObjectPtr<UMoviePipelineExecutorShot>& CurrentCameraCut = ActiveShotList[CurrentShotIndex];
 
 	// When start up we want to override the engine's Custom Timestep with our own.
 	// This gives us the ability to completely control the engine tick/delta time before the frame
@@ -80,7 +80,6 @@ void UMovieGraphLinearTimeStep::TickProducingFrames()
 
 		// Ensure we've set it in the CurrentTimeStepData so things can fetch from it below.
 		CurrentTimeStepData.EvaluatedConfig = TObjectPtr<UMovieGraphEvaluatedConfig>(CurrentFrameData.EvaluatedConfig.Get());
-		CurrentTimeStepData.bExpandShotForTemporalSubSample = CurrentFrameData.TemporalSampleCount > 1;
 
 		// Sets up the render state, etc.
 		GetOwningGraph()->SetupShot(CurrentCameraCut);
@@ -340,6 +339,13 @@ void UMovieGraphLinearTimeStep::UpdateTemporalSampleCount()
 	CurrentFrameData.TemporalSampleCount = 1;
 }
 
+bool UMovieGraphLinearTimeStep::IsExpansionForTSRequired(const TObjectPtr<UMovieGraphEvaluatedConfig>& InConfig) const
+{
+	// ToDo: This needs to come from the config (once we have TemporalSampleCount there)
+	return false;
+}
+
+
 void UMovieGraphLinearTimeStep::Shutdown()
 {
 	// Shut down our custom timestep which reqstores some world settings we modified.
diff --git a/Engine/Plugins/MovieScene/MovieRenderPipeline/Source/MovieRenderPipelineCore/Private/Graph/MovieGraphPipeline.cpp b/Engine/Plugins/MovieScene/MovieRenderPipeline/Source/MovieRenderPipelineCore/Private/Graph/MovieGraphPipeline.cpp
--- a/Engine/Plugins/MovieScene/MovieRenderPipeline/Source/MovieRenderPipelineCore/Private/Graph/MovieGraphPipeline.cpp
+++ b/Engine/Plugins/MovieScene/MovieRenderPipeline/Source/MovieRenderPipelineCore/Private/Graph/MovieGraphPipeline.cpp
@@ -10,6 +10,7 @@
 #include "Graph/Nodes/MovieGraphFileOutputNode.h"
 #include "Graph/Nodes/MovieGraphModifierNode.h"
 #include "Graph/Nodes/MovieGraphOutputSettingNode.h"
+#include "Graph/Nodes/MovieGraphWarmUpSettingNode.h"
 #include "Graph/MovieGraphBlueprintLibrary.h"
 #include "MovieRenderPipelineCoreModule.h"
 #include "Misc/CoreDelegates.h"
@@ -245,9 +246,43 @@ void UMovieGraphPipeline::BuildShotListFromDataSource()
 	// Shots that are already in the list will be updated but their enable flag will be respected. 
 	GraphDataSourceInstance->UpdateShotList();
 
-	for (UMoviePipelineExecutorShot* Shot : GetCurrentJob()->ShotInfo)
+	for (int32 ShotIndex = 0; ShotIndex < GetCurrentJob()->ShotInfo.Num(); ShotIndex++)
 	{
-		Shot->ShotInfo.CurrentTimeInRoot = FFrameTime(Shot->ShotInfo.CurrentTickInRoot);
+		const TObjectPtr<UMoviePipelineExecutorShot>& Shot = GetCurrentJob()->ShotInfo[ShotIndex];
+
+		// Once the shot list is built, we now need to do some work to the calculated data. Handle Frames 
+		// should expand the range provided by the shot list, and we need to do this to all shots in advance.
+		// We do it in advance because to get the total number of frames we're going to render, we add up the
+		// range sizes for each shot, so handle frames need to be included in that range. Each shot can have
+		// different settings, so we build an evaluation context for each shot, flatten the graph and then
+		// read the handle frames config values.
+		// ToDo: This is being run on disabled shots as well, is that intended? The old system did this as well.
+
+		FMovieGraphTraversalContext CurrentContext;
+		CurrentContext.ShotIndex = ShotIndex;
+		CurrentContext.ShotCount = GetActiveShotList().Num();
+		CurrentContext.Job = GetCurrentJob();
+		CurrentContext.RootGraph = GetRootGraphForShot(Shot);
+
+		// Frame Rate, Handle Frames, Warm-up Frames are all global settings so we provide an empty time context.
+		FMovieGraphTimeStepData TimeContext;
+		CurrentContext.Time = TimeContext;
+
+		TObjectPtr<UMovieGraphEvaluatedConfig> EvaluatedConfig = CurrentContext.RootGraph->CreateFlattenedGraph(CurrentContext);
+		UMovieGraphOutputSettingNode* OutputNode = EvaluatedConfig->GetSettingForBranch<UMovieGraphOutputSettingNode>(UMovieGraphSettingNode::GlobalsPinName);
+
+		const FFrameRate SourceFrameRate = GetDataSourceInstance()->GetDisplayRate();
+		const FFrameRate FinalFrameRate = UMovieGraphBlueprintLibrary::GetEffectiveFrameRate(OutputNode, SourceFrameRate);
+		const FFrameRate TickResolution = GetDataSourceInstance()->GetTickResolution();
+
+
+		UMovieGraphWarmUpSettingNode* WarmUpNode = EvaluatedConfig->GetSettingForBranch<UMovieGraphWarmUpSettingNode>(UMovieGraphSettingNode::GlobalsPinName);
+		
+		const bool bPrePass = true;
+		const bool bExpandForTemporalSubSample = GetTimeStepInstance()->IsExpansionForTSRequired(EvaluatedConfig);
+		ExpandShot(Shot, OutputNode->HandleFrameCount, bExpandForTemporalSubSample, bPrePass, FinalFrameRate, TickResolution, WarmUpNode->NumWarmUpFrames);
+
+		Shot->ShotInfo.CurrentTimeInRoot = Shot->ShotInfo.TotalOutputRangeRoot.GetLowerBoundValue();
 		Shot->ShotInfo.CalculateWorkMetrics(GetDataSourceInstance());
 	}
 
@@ -436,10 +471,10 @@ void UMovieGraphPipeline::BeginFinalize()
 	}
 }
 
-void UMovieGraphPipeline::SetupShot(UMoviePipelineExecutorShot* InShot)
+void UMovieGraphPipeline::SetupShot(const TObjectPtr<UMoviePipelineExecutorShot>& InShot)
 {
 	// Set the new shot as the active shot. This enables the specified shot section and disables all other shot sections.
-	// SetSoloShot(InShot);
+	SetSoloShot(InShot);
 
 	// Loop through just our primary settings and let them know which shot we're about to start.
 	//TArray<UMoviePipelineSetting*> Settings = GetPipelinePrimaryConfig()->GetAllSettings();
@@ -471,11 +506,16 @@ void UMovieGraphPipeline::SetupShot(UMoviePipelineExecutorShot* InShot)
 	GraphRendererInstance->SetupRenderingPipelineForShot(InShot);
 }
 
-void UMovieGraphPipeline::TeardownShot(UMoviePipelineExecutorShot* InShot)
+void UMovieGraphPipeline::TeardownShot(const TObjectPtr<UMoviePipelineExecutorShot>& InShot)
 {
 	// Teardown any rendering architecture for this shot. This needs to happen first because it'll flush outstanding rendering commands
 	GraphRendererInstance->TeardownRenderingPipelineForShot(InShot);
 
+	for (const TObjectPtr<UMoviePipelineExecutorShot>& Shot : GetCurrentJob()->ShotInfo)
+	{
+		GetDataSourceInstance()->RestoreHierarchyForShot(Shot);
+	}
+
 	// some other stuff
 
 	CurrentShotIndex++;
@@ -492,6 +532,100 @@ void UMovieGraphPipeline::TeardownShot(UMoviePipelineExecutorShot* InShot)
 	}
 }
 
+void UMovieGraphPipeline::SetSoloShot(const TObjectPtr<UMoviePipelineExecutorShot>& InShot)
+{
+	// We need to 'solo' shots whichs means disabling any other sections that may overlap with the one currently being
+	// rendered. This is because temporal samples, handle frames, warmup frames, etc. all need to evaluate outside of
+	// their original bounds and we don't want to end up evaluating something that should have been clipped by the shot bounds.
+	for (const TObjectPtr<UMoviePipelineExecutorShot>& Shot : GetCurrentJob()->ShotInfo)
+	{
+		// Cache and mute all shots
+		GetDataSourceInstance()->CacheHierarchyForShot(Shot);
+		GetDataSourceInstance()->MuteShot(Shot);
+	}
+
+	// Historically shot expansion was done all at once up front, however this creates a lot of complications when a movie scene isn't filled with unique data
+	// such as re-using shots or using different parts of shots. To resolve this, we expand the entire tree needed for a given range, render it, and then restore the original
+	// values before moving onto the next shot so that each shot has no effect on the others.
+	{
+		UE_LOG(LogMovieRenderPipeline, Log, TEXT("Expanding Shot %d/%d (Shot: %s Camera: %s)"), CurrentShotIndex + 1, ActiveShotList.Num(), *InShot->OuterName, *InShot->InnerName);
+
+		// Enable the one hierarchy we do want for rendering. We will re-disable it later when we restore the current Sequence state.
+		GetDataSourceInstance()->UnmuteShot(InShot);
+
+		const FMovieGraphTimeStepData& TimeStepData = GetTimeStepInstance()->GetCalculatedTimeData();
+		TObjectPtr<UMovieGraphEvaluatedConfig> Config = TimeStepData.EvaluatedConfig;
+		UMovieGraphOutputSettingNode* OutputNode = Config->GetSettingForBranch<UMovieGraphOutputSettingNode>(UMovieGraphSettingNode::GlobalsPinName);
+
+		const FFrameRate SourceFrameRate = GetDataSourceInstance()->GetDisplayRate();
+		const FFrameRate FinalFrameRate = UMovieGraphBlueprintLibrary::GetEffectiveFrameRate(OutputNode, SourceFrameRate);
+		const FFrameRate TickResolution = GetDataSourceInstance()->GetTickResolution();
+
+		
+		UMovieGraphWarmUpSettingNode* WarmUpNode = Config->GetSettingForBranch<UMovieGraphWarmUpSettingNode>(UMovieGraphSettingNode::GlobalsPinName);
+
+		// Expand the shot to encompass handle frames (+warmup, etc.). This will modify the sections required for expansion, etc.
+		const bool bIsPrePass = false;
+		const bool bExpandForTemporalSubSample = GetTimeStepInstance()->IsExpansionForTSRequired(Config);
+		
+		ExpandShot(InShot, OutputNode->HandleFrameCount, bExpandForTemporalSubSample, bIsPrePass, FinalFrameRate, TickResolution, WarmUpNode->NumWarmUpFrames);
+	}
+}
+
+
+
+void UMovieGraphPipeline::ExpandShot(const TObjectPtr<UMoviePipelineExecutorShot>& InShot, const int32 InNumHandleFrames, const bool bInHasMultipleTemporalSamples, const bool bIsPrePass,
+	const FFrameRate& InDisplayRate, const FFrameRate& InTickResolution, const int32 InWarmUpFrames)
+{
+	int32 LeftDeltaFrames = 0;
+	int32 RightDeltaFrames = 0;
+
+	// Calculate the number of ticks added for warmup frames. These are added to both sides. The rendering
+	// code is unaware of handle frames, we just pretend the shot is bigger than it actually is.
+	LeftDeltaFrames += InNumHandleFrames;
+	RightDeltaFrames += InNumHandleFrames;
+
+	// We only expand the left side for temporal sub-sampling, as no camera timing allows you to beyond the end of frame.
+	if (bInHasMultipleTemporalSamples)
+	{
+		LeftDeltaFrames += 1;
+	}
+
+	// Check to see if the detected range was not aligned to a whole frame on the root. We produce a warning here because 
+	// if your shot starts on a sub-frame (say frame 3.5) the output frame will say "3", but when you go to look at frame 3
+	// in Sequencer, it will show different content than was actually evaluated. So we warn + round down to get them aligned
+	// on whole frames.
+	FFrameTime StartTimeInRoot = FFrameRate::TransformTime(InShot->ShotInfo.TotalOutputRangeRoot.GetLowerBoundValue(), InTickResolution, InDisplayRate);
+	if (bIsPrePass && StartTimeInRoot.GetSubFrame() != 0.f)
+	{
+		UE_LOG(LogMovieRenderPipeline, Warning, TEXT("Shot/Camera \"%s [%s]\" starts on a sub-frame. Rendered range has been rounded to the previous frame to make output numbers align with Sequencer."), 
+			*InShot->OuterName, *InShot->InnerName);
+		FFrameNumber NewStartFrame = FFrameRate::TransformTime(FFrameTime(StartTimeInRoot.GetFrame()), InDisplayRate, InTickResolution).FloorToFrame();
+		InShot->ShotInfo.TotalOutputRangeRoot.SetLowerBoundValue(NewStartFrame);
+	}
+
+	// We auto-expand into the warm-up ranges, but users are less concerned about 'early' data there. So we cache how many frames
+	// the user expects to check beforehand, so we can use this for a warning later.
+	const int32 LeftDeltaFramesUserPoV = LeftDeltaFrames;
+
+	// Warm Up frames are only on the left side. This comes after the above section so that we don't warn about partial data in the warm up section.
+	LeftDeltaFrames += InWarmUpFrames;
+
+	GetDataSourceInstance()->ExpandShot(InShot, LeftDeltaFrames, LeftDeltaFramesUserPoV, RightDeltaFrames, bIsPrePass);
+
+	// Expand the Total Output Range Root by Handle Frames. The expansion of TotalOutputRangeRoot has to come after we do partial evaluation checks,
+	// which is done by ExpandShot above, otherwise the expanded range makes it check the wrong area for partial evaluations.
+	if (bIsPrePass)
+	{
+		// We expand on the pre-pass so that we have the correct number of frames set up in our datastructures before we reach each shot so that metrics
+		// work as expected.
+		FFrameNumber LeftHandleTicks = FFrameRate::TransformTime(FFrameTime(InNumHandleFrames), InDisplayRate, InTickResolution).CeilToFrame().Value;
+		FFrameNumber RightHandleTicks = FFrameRate::TransformTime(FFrameTime(InNumHandleFrames), InDisplayRate, InTickResolution).CeilToFrame().Value;
+
+		InShot->ShotInfo.TotalOutputRangeRoot = UE::MovieScene::DilateRange(InShot->ShotInfo.TotalOutputRangeRoot, -LeftHandleTicks, RightHandleTicks);
+	}
+}
+
 void UMovieGraphPipeline::OnEngineTickEndFrame()
 {
 	LLM_SCOPE_BYNAME(TEXT("MovieGraphEndFrame"));
diff --git a/Engine/Plugins/MovieScene/MovieRenderPipeline/Source/MovieRenderPipelineCore/Private/Graph/MovieGraphSequenceDataSource.cpp b/Engine/Plugins/MovieScene/MovieRenderPipeline/Source/MovieRenderPipelineCore/Private/Graph/MovieGraphSequenceDataSource.cpp
--- a/Engine/Plugins/MovieScene/MovieRenderPipeline/Source/MovieRenderPipelineCore/Private/Graph/MovieGraphSequenceDataSource.cpp
+++ b/Engine/Plugins/MovieScene/MovieRenderPipeline/Source/MovieRenderPipelineCore/Private/Graph/MovieGraphSequenceDataSource.cpp
@@ -2,14 +2,17 @@
 
 #include "Graph/MovieGraphSequenceDataSource.h"
 #include "Graph/MovieGraphPipeline.h"
+#include "Sections/MovieSceneCameraCutSection.h"
+#include "Sections/MovieSceneCinematicShotSection.h"
+#include "MovieRenderPipelineCoreModule.h"
+#include "MoviePipelineBlueprintLibrary.h"
+#include "MoviePipelineUtils.h"
+#include "MoviePipelineQueue.h"
 #include "LevelSequence.h"
 #include "LevelSequenceActor.h"
 #include "LevelSequencePlayer.h"
-#include "MoviePipelineQueue.h"
 #include "CoreGlobals.h"
 #include "EngineUtils.h"
-#include "MovieRenderPipelineCoreModule.h"
-#include "MoviePipelineBlueprintLibrary.h"
 #include "Engine/GameViewportClient.h"
 
 UMovieGraphSequenceDataSource::UMovieGraphSequenceDataSource()
@@ -128,7 +131,7 @@ FFrameRate UMovieGraphSequenceDataSource::GetDisplayRate() const
 	return FFrameRate(24, 1);
 }
 
-void UMovieGraphSequenceDataSource::InitializeShot(UMoviePipelineExecutorShot* InShot)
+void UMovieGraphSequenceDataSource::InitializeShot(const TObjectPtr<UMoviePipelineExecutorShot>& InShot)
 {
 	//if (!InShot->ShotInfo.bEmulateFirstFrameMotionBlur)
 	{
@@ -141,6 +144,128 @@ void UMovieGraphSequenceDataSource::InitializeShot(UMoviePipelineExecutorShot* I
 	//	LevelSequenceActor->GetSequencePlayer()->Pause();
 	//}
 }
+void UMovieGraphSequenceDataSource::CacheHierarchyForShot(const TObjectPtr<UMoviePipelineExecutorShot>& InShot)
+{
+	// Save the whole hierarchy for this shot, then set it to be inactive.
+	TSharedPtr<MoviePipeline::FCameraCutSubSectionHierarchyNode> Node = InShot->ShotInfo.SubSectionHierarchy;
+	const bool bSaveSettings = true;
+	MoviePipeline::SaveOrRestoreSubSectionHierarchy(Node, bSaveSettings);
+}
+
+void UMovieGraphSequenceDataSource::RestoreHierarchyForShot(const TObjectPtr<UMoviePipelineExecutorShot>& InShot)
+{
+	TSharedPtr<MoviePipeline::FCameraCutSubSectionHierarchyNode> Node = InShot->ShotInfo.SubSectionHierarchy;
+
+	const bool bSaveSettings = false;
+	MoviePipeline::SaveOrRestoreSubSectionHierarchy(Node, bSaveSettings);
+}
+
+void UMovieGraphSequenceDataSource::MuteShot(const TObjectPtr<UMoviePipelineExecutorShot>& InShot)
+{
+	const bool bActive = false;
+	MoviePipeline::SetSubSectionHierarchyActive(InShot->ShotInfo.SubSectionHierarchy, bActive);
+}
+
+void UMovieGraphSequenceDataSource::UnmuteShot(const TObjectPtr<UMoviePipelineExecutorShot>& InShot)
+{
+	const bool bActive = true;
+	MoviePipeline::SetSubSectionHierarchyActive(InShot->ShotInfo.SubSectionHierarchy, bActive);
+}
+
+void UMovieGraphSequenceDataSource::ExpandShot(const TObjectPtr<UMoviePipelineExecutorShot>& InShot, const int32 InLeftDeltaFrames, const int32 InLeftDeltaFramesUserPoV, 
+	const int32 InRightDeltaFrames, const bool bInPrepass)
+{
+	const FFrameRate& DisplayRate = GetDisplayRate();
+	const FFrameRate& TickResolution = GetTickResolution();
+
+	TSharedPtr<MoviePipeline::FCameraCutSubSectionHierarchyNode> Node = InShot->ShotInfo.SubSectionHierarchy;
+	while (Node)
+	{
+		// We need to expand the inner playback bounds to cover three features:
+		// 1) Temporal Sampling (+1 frame left side)
+		// 2) Handle frames (+n frames both sides)
+		// 3) Non-emulated Warm Up (+n frames left side)
+		// To keep the inner movie scene and outer sequencer section in sync we can calculate the tick delta
+		// to each side and simply expand both sections like that - ignoring all start frame offsets, etc.
+		FFrameNumber LeftDeltaTicks = FFrameRate::TransformTime(FFrameTime(InLeftDeltaFrames), DisplayRate, TickResolution).CeilToFrame();
+		FFrameNumber RightDeltaTicks = FFrameRate::TransformTime(FFrameTime(InRightDeltaFrames), DisplayRate, TickResolution).CeilToFrame();
+		FFrameTime LeftDeltaTimeUserPoV = FFrameRate::TransformTime(FFrameTime(InLeftDeltaFramesUserPoV), DisplayRate, TickResolution);
+
+		// During pre-pass we cache which items we'd like to auto-expand later, and we print a warning
+		// for anything we can't automatically expand that is now getting partial evaluation.
+		if (bInPrepass)
+		{
+			if (Node->MovieScene.IsValid())
+			{
+				for (UMovieSceneSection* Section : Node->MovieScene->GetAllSections())
+				{
+					if (!Section)
+					{
+						continue;
+					}
+
+					// Their data is already cached for restore elsewhere.
+					if (Section == Node->Section || Section == Node->CameraCutSection)
+					{
+						continue;
+					}
+
+					if (Section->GetSupportsInfiniteRange())
+					{
+						Node->AdditionalSectionsToExpand.Add(MakeTuple(Section, Section->GetRange()));
+					}
+				}
+			}
+
+			// We only do our warnings during the pre-pass
+			// Check for sections that start in the expanded evaluation range and warn user. 
+			// Only check the frames user expects to (handle + temporal, no need for warm up frames to get checked as well)
+			MoviePipeline::CheckPartialSectionEvaluationAndWarn(LeftDeltaTimeUserPoV.CeilToFrame(), Node, InShot, DisplayRate);
+		}
+		else
+		{
+			if (Node->CameraCutSection.IsValid())
+			{
+				// Expand the camera cut section because there's no harm in doing it.
+				Node->CameraCutSection->SetRange(UE::MovieScene::DilateRange(Node->CameraCutSection->GetRange(), -LeftDeltaTicks, RightDeltaTicks));
+				Node->CameraCutSection->MarkAsChanged();
+			}
+
+			if (Node->Section.IsValid())
+			{
+				// Expand the MovieSceneSubSequenceSection
+				Node->Section->SetRange(UE::MovieScene::DilateRange(Node->Section->GetRange(), -LeftDeltaTicks, RightDeltaTicks));
+				Node->Section->MarkAsChanged();
+			}
+
+			if (Node->MovieScene.IsValid())
+			{
+				// Expand the Playback Range of the movie scene as well. Expanding this at the same time as expanding the 
+				// SubSequenceSection will result in no apparent change to the evaluated time. ToDo: This doesn't work if
+				// sub-sequences have different tick resolutions?
+				Node->MovieScene->SetPlaybackRange(UE::MovieScene::DilateRange(Node->MovieScene->GetPlaybackRange(), -LeftDeltaTicks, RightDeltaTicks));
+				Node->MovieScene->MarkAsChanged();
+			}
+
+			FFrameNumber LowerCheckBound = InShot->ShotInfo.TotalOutputRangeRoot.GetLowerBoundValue() - LeftDeltaTimeUserPoV.CeilToFrame();
+			FFrameNumber UpperCheckBound = InShot->ShotInfo.TotalOutputRangeRoot.GetLowerBoundValue();
+
+			TRange<FFrameNumber> CheckRange = TRange<FFrameNumber>(TRangeBound<FFrameNumber>::Exclusive(LowerCheckBound), TRangeBound<FFrameNumber>::Inclusive(UpperCheckBound));
+
+			for (const TTuple<UMovieSceneSection*, TRange<FFrameNumber>>& Pair : Node->AdditionalSectionsToExpand)
+			{
+				// Expand the section. Because it's an infinite range, we know the contents won't get shifted.
+				TRange<FFrameNumber> NewRange = TRange<FFrameNumber>::Hull(Pair.Key->GetRange(), CheckRange);
+				Pair.Key->SetRange(NewRange);
+				Pair.Key->MarkAsChanged();
+			}
+		}
+
+		// Travel up to the parent and expand it too.
+		Node = Node->GetParent();
+	}
+}
+
 
 void UMovieGraphSequenceDataSource::SyncDataSourceTime(const FFrameTime& InTime)
 {
@@ -157,4 +282,4 @@ namespace UE::MovieGraph
 
 		return RequestTime;
 	}
-}
\ No newline at end of file
+}
diff --git a/Engine/Plugins/MovieScene/MovieRenderPipeline/Source/MovieRenderPipelineCore/Private/Graph/Nodes/MovieGraphWarmUpSettingNode.cpp b/Engine/Plugins/MovieScene/MovieRenderPipeline/Source/MovieRenderPipelineCore/Private/Graph/Nodes/MovieGraphWarmUpSettingNode.cpp
new file mode 100644
--- /dev/null
+++ b/Engine/Plugins/MovieScene/MovieRenderPipeline/Source/MovieRenderPipelineCore/Private/Graph/Nodes/MovieGraphWarmUpSettingNode.cpp
@@ -0,0 +1,41 @@
+// Copyright Epic Games, Inc. All Rights Reserved.
+
+#include "Graph/Nodes/MovieGraphWarmUpSettingNode.h"
+#include "Styling/AppStyle.h"
+
+UMovieGraphWarmUpSettingNode::UMovieGraphWarmUpSettingNode()
+	: NumWarmUpFrames(64)
+	, bEmulateMotionBlur(true)
+{}
+
+EMovieGraphBranchRestriction UMovieGraphWarmUpSettingNode::GetBranchRestriction() const 
+{ 
+	return EMovieGraphBranchRestriction::Globals; 
+}
+
+#if WITH_EDITOR
+FText UMovieGraphWarmUpSettingNode::GetNodeTitle(const bool bGetDescriptive) const
+{
+	static const FText WarmUpSettingsNodeName = NSLOCTEXT("MoviePipelineGraph", "NodeName_WarmUpSettings", "Warm Up Settings");
+	return WarmUpSettingsNodeName;
+}
+
+FText UMovieGraphWarmUpSettingNode::GetMenuCategory() const 
+{
+	return NSLOCTEXT("MoviePipelineGraph", "Settings_Category", "Settings");
+}
+
+FLinearColor UMovieGraphWarmUpSettingNode::GetNodeTitleColor() const 
+{
+	static const FLinearColor WarmUpSettingsColor = FLinearColor(0.854f, 0.509f, 0.039f);
+	return WarmUpSettingsColor;
+}
+
+FSlateIcon UMovieGraphWarmUpSettingNode::GetIconAndTint(FLinearColor& OutColor) const 
+{
+	static const FSlateIcon SettingsIcon = FSlateIcon(FAppStyle::GetAppStyleSetName(), "Icons.Settings");
+
+	OutColor = FLinearColor::White;
+	return SettingsIcon;
+}
+#endif // WITH_EDITOR
\ No newline at end of file
diff --git a/Engine/Plugins/MovieScene/MovieRenderPipeline/Source/MovieRenderPipelineCore/Public/Graph/MovieGraphDataTypes.h b/Engine/Plugins/MovieScene/MovieRenderPipeline/Source/MovieRenderPipelineCore/Public/Graph/MovieGraphDataTypes.h
--- a/Engine/Plugins/MovieScene/MovieRenderPipeline/Source/MovieRenderPipelineCore/Public/Graph/MovieGraphDataTypes.h
+++ b/Engine/Plugins/MovieScene/MovieRenderPipeline/Source/MovieRenderPipelineCore/Public/Graph/MovieGraphDataTypes.h
@@ -110,6 +110,12 @@ public:
 	* needed to calculate the correct rendering timestep.
 	*/
 	virtual FMovieGraphTimeStepData GetCalculatedTimeData() const { return FMovieGraphTimeStepData(); }
+	
+	/**
+	* When expanding shots, do we need to expand the first frame by one extra to account for how
+	* temporal sub-sampling can sample outside the first frame (due to centered frame evals).
+	*/
+	virtual bool IsExpansionForTSRequired(const TObjectPtr<UMovieGraphEvaluatedConfig>& InConfig) const { return false; }
 
 	UMovieGraphPipeline* GetOwningGraph() const;
 };
@@ -161,7 +167,14 @@ public:
 	virtual void CacheDataPreJob(const FMovieGraphInitConfig& InInitConfig) {}
 	virtual void RestoreCachedDataPostJob() {}
 	virtual void UpdateShotList() {}
-	virtual void InitializeShot(UMoviePipelineExecutorShot* InShot) {}
+	virtual void InitializeShot(const TObjectPtr<UMoviePipelineExecutorShot>& InShot) {}
+	virtual void CacheHierarchyForShot(const TObjectPtr<UMoviePipelineExecutorShot>& InShot) {}
+	virtual void RestoreHierarchyForShot(const TObjectPtr<UMoviePipelineExecutorShot> &InShot) {}
+	virtual void MuteShot(const TObjectPtr<UMoviePipelineExecutorShot>& InShot) {}
+	virtual void UnmuteShot(const TObjectPtr<UMoviePipelineExecutorShot>& InShot) {}
+	virtual void ExpandShot(const TObjectPtr<UMoviePipelineExecutorShot>& InShot, const int32 InLeftDeltaFrames, const int32 InLeftDeltaFramesUserPoV,
+		const int32 InRightDeltaFrames, const bool bInPrepass) {}
+
 	UMovieGraphPipeline* GetOwningGraph() const;
 };
 
diff --git a/Engine/Plugins/MovieScene/MovieRenderPipeline/Source/MovieRenderPipelineCore/Public/Graph/MovieGraphLinearTimeStep.h b/Engine/Plugins/MovieScene/MovieRenderPipeline/Source/MovieRenderPipelineCore/Public/Graph/MovieGraphLinearTimeStep.h
--- a/Engine/Plugins/MovieScene/MovieRenderPipeline/Source/MovieRenderPipelineCore/Public/Graph/MovieGraphLinearTimeStep.h
+++ b/Engine/Plugins/MovieScene/MovieRenderPipeline/Source/MovieRenderPipelineCore/Public/Graph/MovieGraphLinearTimeStep.h
@@ -33,6 +33,7 @@ public:
 	virtual void TickProducingFrames() override;
 	virtual void Shutdown() override;
 	virtual FMovieGraphTimeStepData GetCalculatedTimeData() const override { return CurrentTimeStepData; }
+	virtual bool IsExpansionForTSRequired(const TObjectPtr<UMovieGraphEvaluatedConfig>& InConfig) const override;
 	// ~UMovieGraphTimeStepBase Interface
 
 protected:
diff --git a/Engine/Plugins/MovieScene/MovieRenderPipeline/Source/MovieRenderPipelineCore/Public/Graph/MovieGraphPipeline.h b/Engine/Plugins/MovieScene/MovieRenderPipeline/Source/MovieRenderPipelineCore/Public/Graph/MovieGraphPipeline.h
--- a/Engine/Plugins/MovieScene/MovieRenderPipeline/Source/MovieRenderPipelineCore/Public/Graph/MovieGraphPipeline.h
+++ b/Engine/Plugins/MovieScene/MovieRenderPipeline/Source/MovieRenderPipelineCore/Public/Graph/MovieGraphPipeline.h
@@ -66,9 +66,9 @@ public:
 	/** Which index of the Active Shot List are we currently on */
 	int32 GetCurrentShotIndex() const { return CurrentShotIndex; }
 	/** Called by the TimeStepInstance when it's time to set up for another shot. Don't call this unless you know what you're doing. */
-	void SetupShot(UMoviePipelineExecutorShot* InShot);
+	void SetupShot(const TObjectPtr<UMoviePipelineExecutorShot>& InShot);
 	/** Called by the TimeStepInstance when it's time to tear down the current shot. Don't call this unless you know what you're doing. */
-	void TeardownShot(UMoviePipelineExecutorShot* InShot);
+	void TeardownShot(const TObjectPtr<UMoviePipelineExecutorShot>& InShot);
 	/** Used occasionally to cross-reference other components. Don't call this unless you know what you're doing. */
 	UMovieGraphTimeStepBase* GetTimeStepInstance() const { return GraphTimeStepInstance; }
 	/** Used occasionally to cross-reference other components. Don't call this unless you know what you're doing. */
@@ -116,6 +116,15 @@ protected:
 	virtual void LoadPreviewWidget();
 	virtual void SetPreviewWidgetVisibleImpl(bool bInIsVisible);
 
+	// Update our data source to isolate the shot we're currently working on, so that expanded shots don't interfere with each other.
+	virtual void SetSoloShot(const TObjectPtr<UMoviePipelineExecutorShot>& InShot);
+	/*
+	* Expand our data source (for a specific shot) for various rendering features (such as handle frames). Gets called twice, once where it 
+	* expands the ranges, and again later to actually expand the tracks. This is required due to wanting "Handle Frames" to get counted in total frame counts.
+	*/
+	virtual void ExpandShot(const TObjectPtr<UMoviePipelineExecutorShot>& InShot, const int32 InNumHandleFrames, const bool bInHasMultipleTemporalSamples, const bool bIsPrePass,
+		const FFrameRate& InDisplayRate, const FFrameRate& InTickResolution, const int32 InWarmUpFrames);
+
 	// UMoviePipelineBase Interface
 	virtual void RequestShutdownImpl(bool bIsError) override;
 	virtual void ShutdownImpl(bool bIsError) override;
diff --git a/Engine/Plugins/MovieScene/MovieRenderPipeline/Source/MovieRenderPipelineCore/Public/Graph/MovieGraphSequenceDataSource.h b/Engine/Plugins/MovieScene/MovieRenderPipeline/Source/MovieRenderPipelineCore/Public/Graph/MovieGraphSequenceDataSource.h
--- a/Engine/Plugins/MovieScene/MovieRenderPipeline/Source/MovieRenderPipelineCore/Public/Graph/MovieGraphSequenceDataSource.h
+++ b/Engine/Plugins/MovieScene/MovieRenderPipeline/Source/MovieRenderPipelineCore/Public/Graph/MovieGraphSequenceDataSource.h
@@ -44,7 +44,13 @@ public:
 	virtual FFrameRate GetTickResolution() const override;
 	virtual FFrameRate GetDisplayRate() const override;
 	virtual void SyncDataSourceTime(const FFrameTime& InTime) override;
-	virtual void InitializeShot(UMoviePipelineExecutorShot* InShot) override;
+	virtual void InitializeShot(const TObjectPtr<UMoviePipelineExecutorShot>& InShot) override;
+	virtual void CacheHierarchyForShot(const TObjectPtr<UMoviePipelineExecutorShot>& InShot) override;
+	virtual void RestoreHierarchyForShot(const TObjectPtr<UMoviePipelineExecutorShot>& InShot) override;
+	virtual void MuteShot(const TObjectPtr<UMoviePipelineExecutorShot>& InShot) override;
+	virtual void UnmuteShot(const TObjectPtr<UMoviePipelineExecutorShot>& InShot) override;
+	virtual void ExpandShot(const TObjectPtr<UMoviePipelineExecutorShot>& InShot, const int32 InLeftDeltaFrames, const int32 InLeftDeltaFramesUserPoV,
+		const int32 InRightDeltaFrames, const bool bInPrepass) override;
 protected:
 	void CacheLevelSequenceData(ULevelSequence* InSequence);
 	void OnSequenceEvaluated(const UMovieSceneSequencePlayer& Player, FFrameTime CurrentTime, FFrameTime PreviousTime);
diff --git a/Engine/Plugins/MovieScene/MovieRenderPipeline/Source/MovieRenderPipelineCore/Public/Graph/MovieGraphTimeStepData.h b/Engine/Plugins/MovieScene/MovieRenderPipeline/Source/MovieRenderPipelineCore/Public/Graph/MovieGraphTimeStepData.h
--- a/Engine/Plugins/MovieScene/MovieRenderPipeline/Source/MovieRenderPipelineCore/Public/Graph/MovieGraphTimeStepData.h
+++ b/Engine/Plugins/MovieScene/MovieRenderPipeline/Source/MovieRenderPipelineCore/Public/Graph/MovieGraphTimeStepData.h
@@ -26,7 +26,6 @@ public:
 		, MotionBlurFraction(0.f)
 		, bIsFirstTemporalSampleForFrame(false)
 		, bIsLastTemporalSampleForFrame(false)
-		, bExpandShotForTemporalSubSample(false)
 		, bRequiresAccumulator(false)
 	{}
 
@@ -60,13 +59,6 @@ public:
 	UPROPERTY(BlueprintReadWrite, EditAnywhere, Category = "Movie Graph")
 	bool bIsLastTemporalSampleForFrame;
 
-	/**
-	* If this time struct is being used when setting up a shot, defines if the shot 
-	* should be expanded by 1 frame to account for temporal sub-sampling going 
-	* out of bounds on the first frame.
-	*/
-	UPROPERTY(BlueprintReadWrite, EditAnywhere, Category = "Movie Graph")
-	bool bExpandShotForTemporalSubSample;
 
 	/**
 	* Should be set to true for every sample if there is more than one temporal sample
diff --git a/Engine/Plugins/MovieScene/MovieRenderPipeline/Source/MovieRenderPipelineCore/Public/Graph/Nodes/MovieGraphOutputSettingNode.h b/Engine/Plugins/MovieScene/MovieRenderPipeline/Source/MovieRenderPipelineCore/Public/Graph/Nodes/MovieGraphOutputSettingNode.h
--- a/Engine/Plugins/MovieScene/MovieRenderPipeline/Source/MovieRenderPipelineCore/Public/Graph/Nodes/MovieGraphOutputSettingNode.h
+++ b/Engine/Plugins/MovieScene/MovieRenderPipeline/Source/MovieRenderPipelineCore/Public/Graph/Nodes/MovieGraphOutputSettingNode.h
@@ -46,6 +46,9 @@ public:
 	UPROPERTY(EditAnywhere, BlueprintReadWrite, Category = Overrides, meta = (InlineEditConditionToggle))
 	uint8 bOverride_FrameNumberOffset : 1;
 
+	UPROPERTY(EditAnywhere, BlueprintReadWrite, Category = Overrides, meta = (InlineEditConditionToggle))
+	uint8 bOverride_HandleFrameCount : 1;
+
 	/** What directory should all of our output files be relative to. */
 	UPROPERTY(EditAnywhere, BlueprintReadWrite, Category = "File Output", meta=(EditCondition="bOverride_OutputDirectory"))
 	FDirectoryPath OutputDirectory;
@@ -77,4 +80,12 @@ public:
 	UPROPERTY(EditAnywhere, BlueprintReadWrite,  Category = "File Output", meta = (EditCondition = "bOverride_FrameNumberOffset"))
 	int32 FrameNumberOffset;
 
+	/**
+	* Top level shot track sections will automatically be expanded by this many frames in both directions, and the resulting
+	* additional time will be rendered as part of that shot. The inner sequence should have sections long enough to cover
+	* this expanded range, otherwise these tracks will not evaluate during handle frames and may produce unexpected results.
+	* This can be used to generate handle frames for traditional non linear editing tools.
+	*/
+	UPROPERTY(EditAnywhere, BlueprintReadWrite, Category = "Frames", meta = (UIMin = 0, ClampMin = 0, EditCondition = "bOverride_HandleFrameCount"))
+	int32 HandleFrameCount;
 };
diff --git a/Engine/Plugins/MovieScene/MovieRenderPipeline/Source/MovieRenderPipelineCore/Public/Graph/Nodes/MovieGraphWarmUpSettingNode.h b/Engine/Plugins/MovieScene/MovieRenderPipeline/Source/MovieRenderPipelineCore/Public/Graph/Nodes/MovieGraphWarmUpSettingNode.h
new file mode 100644
--- /dev/null
+++ b/Engine/Plugins/MovieScene/MovieRenderPipeline/Source/MovieRenderPipelineCore/Public/Graph/Nodes/MovieGraphWarmUpSettingNode.h
@@ -0,0 +1,42 @@
+// Copyright Epic Games, Inc. All Rights Reserved.
+
+#pragma once
+#include "Graph/MovieGraphNode.h"
+#include "Misc/FrameRate.h"
+#include "MovieGraphWarmUpSettingNode.generated.h"
+
+
+UCLASS()
+class MOVIERENDERPIPELINECORE_API UMovieGraphWarmUpSettingNode : public UMovieGraphSettingNode
+{
+	GENERATED_BODY()
+public:
+	UMovieGraphWarmUpSettingNode();
+
+	// UMovieGraphSettingNode Interface
+	virtual EMovieGraphBranchRestriction GetBranchRestriction() const override;
+#if WITH_EDITOR
+	virtual FText GetNodeTitle(const bool bGetDescriptive) const override;
+	virtual FText GetMenuCategory() const override;
+	virtual FLinearColor GetNodeTitleColor() const override;
+	virtual FSlateIcon GetIconAndTint(FLinearColor& OutColor) const override;
+#endif
+	// ~UMovieGraphSettingNode Interface
+
+public:
+	UPROPERTY(EditAnywhere, BlueprintReadWrite, Category = Overrides, meta = (InlineEditConditionToggle))
+	uint8 bOverride_NumWarmUpFrames : 1;
+
+	UPROPERTY(EditAnywhere, BlueprintReadWrite, Category = Overrides, meta = (InlineEditConditionToggle))
+	uint8 bOverride_EmulateMotionBlur : 1;
+
+	/** At the start of each shot, how many frames should we run the engine (without writing renders to disk) to warm up various systems. */
+	UPROPERTY(EditAnywhere, BlueprintReadWrite, Category = "Settings", meta=(EditCondition="bOverride_NumWarmUpFrames"))
+	int32 NumWarmUpFrames;
+
+	/** 
+	If true, we will evaluate frame 0, and then wait NumWarmUpFrames frames, rendering each one as we go. Then we will evaluate frame 1, and then frame 0 again. This emulates motion blur on the first frame (which normally needs data before frame 0). If false, we will "walk" towards the first frame of the shot, starting NumWarmUpFrames before the shot normally starts. 
+	*/
+	UPROPERTY(EditAnywhere, BlueprintReadWrite, Category = "Settings", meta = (EditCondition = "bOverride_EmulateMotionBlur"))
+	bool bEmulateMotionBlur;
+};
