MoviePipeline: Created a base class for time step implementations, and spun off linear and stochastic versions from it. The stochastic version is mostly placeholder code for now, to be implemented later.

Also added a new Sampling Method node which can select the time step class to use, as well as set additional sampling properties (like temporal sample count).

#rnx

diff --git a/Engine/Plugins/MovieScene/MovieRenderPipeline/Source/MovieRenderPipelineCore/Private/Graph/MovieGraphCoreTimeStep.cpp b/Engine/Plugins/MovieScene/MovieRenderPipeline/Source/MovieRenderPipelineCore/Private/Graph/MovieGraphCoreTimeStep.cpp
new file mode 100644
--- /dev/null
+++ b/Engine/Plugins/MovieScene/MovieRenderPipeline/Source/MovieRenderPipelineCore/Private/Graph/MovieGraphCoreTimeStep.cpp
@@ -0,0 +1,547 @@
+// Copyright Epic Games, Inc. All Rights Reserved.
+
+#include "Graph/MovieGraphCoreTimeStep.h"
+
+#include "Graph/MovieGraphPipeline.h"
+#include "Graph/MovieGraphBlueprintLibrary.h"
+#include "Graph/Nodes/MovieGraphOutputSettingNode.h"
+#include "MoviePipelineQueue.h"
+#include "MovieRenderPipelineCoreModule.h"
+#include "LevelSequence.h"
+#include "MovieScene.h"
+#include "Camera/PlayerCameraManager.h"
+#include "Camera/CameraComponent.h"
+#include "Camera/CameraActor.h"
+#include "Interfaces/Interface_PostProcessVolume.h"
+#include "GameFramework/WorldSettings.h"
+#include "ProfilingDebugging/MiscTrace.h"
+
+UMovieGraphCoreTimeStep::UMovieGraphCoreTimeStep()
+{
+	CustomTimeStep = CreateDefaultSubobject<UMovieGraphEngineTimeStep>("MovieGraphEngineTimeStep");
+
+	// This set up our internal state to pick up on the first temporal sub-sample in the pattern
+	ResetForEndOfOutputFrame();
+	CurrentTimeStepData.OutputFrameNumber = 0;
+}
+
+void UMovieGraphCoreTimeStep::TickProducingFrames()
+{
+	int32 CurrentShotIndex = GetOwningGraph()->GetCurrentShotIndex();
+	const TArray<TObjectPtr<UMoviePipelineExecutorShot>>& ActiveShotList = GetOwningGraph()->GetActiveShotList();
+	const TObjectPtr<UMoviePipelineExecutorShot>& CurrentCameraCut = ActiveShotList[CurrentShotIndex];
+
+	// When start up we want to override the engine's Custom Timestep with our own.
+	// This gives us the ability to completely control the engine tick/delta time before the frame
+	// is started so that we don't have to always be thinking of delta times one frame ahead. We need
+	// to do this only once we're ready to set the timestep though, as Initialize can be called as
+	// a result of a OnBeginFrame, meaning that Initialize is called on the frame before TickProducingFrames
+	// so there would be one frame where it used the custom timestep (after initialize) before TPF was called.
+	if (GEngine->GetCustomTimeStep() != CustomTimeStep)
+	{
+		PrevCustomTimeStep = GEngine->GetCustomTimeStep();
+		GEngine->SetCustomTimeStep(CustomTimeStep);
+	}
+
+	// We cache the frame metrics for the duration of a single output frame, instead of recalculating them every tick.
+	// This is required for stochastic timesteps (which will jump around the actual evaluated time randomly within one
+	// output frame), but it helps resolve some of our issues related to time dilation. Historically when doing it every
+	// tick, we didn't know if we were going to overrun the end range of time represented until we had actually done it,
+	// partway through the frame, at which point we had to backtrack and abandon the work. So now we calculate it all
+	// in advance, check if our new time would overrun the end of the current camera cut, and if so, skip actually starting
+	// that frame. Unfortunately this does cause some complications with actor time dilations - the time dilation track
+	// will be updating each tick within the output sample and trying to change the world time dilation which we don't want.
+	
+	// ToDo: Still tbd how we'll fix that, currently the plan is to allow custom clocksources to disable world time dilation
+	// support, at which point we'll handle the time dilation by affecting the whole engine tick instead of just world tick.
+	//if (IsFirstTemporalSample())
+	{
+		UpdateFrameMetrics();
+
+		// Now that we've calculated the total range of time we're trying to represent, we can check to see
+		// if this would put us beyond our range of time this shot is supposed to represent.
+
+	}
+	// Handle shot initialization
+
+	if (CurrentCameraCut->ShotInfo.State == EMovieRenderShotState::Uninitialized)
+	{
+		UE_LOG(LogMovieRenderPipeline, Log, TEXT("MovieGraph Initializing Camera Cut [%d/%d] in [%s] %s."),
+			CurrentShotIndex + 1, ActiveShotList.Num(), *CurrentCameraCut->OuterName, *CurrentCameraCut->InnerName);
+
+		FMovieGraphTraversalContext Context = GetOwningGraph()->GetCurrentTraversalContext();
+		
+		// Update global variables before evaluating the graph
+		Context.RootGraph->UpdateGlobalVariableValues(GetOwningGraph());
+
+		// Evaluate the graph so we can fetch values for this shot.
+		UMovieGraphConfig* Config = GetOwningGraph()->GetRootGraphForShot(CurrentCameraCut);
+		CurrentFrameData.EvaluatedConfig = TStrongObjectPtr<UMovieGraphEvaluatedConfig>(Config->CreateFlattenedGraph(Context));
+		CurrentFrameData.TemporalSampleIndex = 0;
+
+		// Ensure we've set it in the CurrentTimeStepData so things can fetch from it below.
+		CurrentTimeStepData.EvaluatedConfig = TObjectPtr<UMovieGraphEvaluatedConfig>(CurrentFrameData.EvaluatedConfig.Get());
+
+		// Sets up the render state, etc.
+		GetOwningGraph()->SetupShot(CurrentCameraCut);
+
+		// Seeks the external data source to match
+		GetOwningGraph()->GetDataSourceInstance()->InitializeShot(CurrentCameraCut);
+
+		// Generate render layers from the evaluated graph
+		GetOwningGraph()->CreateLayersInRenderLayerSubsystem(CurrentTimeStepData.EvaluatedConfig);
+
+		// Update the contents of the render layers
+		// TODO: This should probably be done per-frame
+		GetOwningGraph()->UpdateLayerContentsInRenderLayerSubsystem(CurrentTimeStepData.EvaluatedConfig);
+
+		const bool bIncludeCDO = false;
+		UMovieGraphOutputSettingNode* OutputNode = CurrentFrameData.EvaluatedConfig->GetSettingForBranch<UMovieGraphOutputSettingNode>(UMovieGraphNode::GlobalsPinName, bIncludeCDO);
+
+		const FFrameRate TickResolution = GetOwningGraph()->GetDataSourceInstance()->GetTickResolution();
+		const FFrameRate SourceFrameRate = GetOwningGraph()->GetDataSourceInstance()->GetDisplayRate();
+		const FFrameRate FinalFrameRate = UMovieGraphBlueprintLibrary::GetEffectiveFrameRate(OutputNode, SourceFrameRate);
+
+		FFrameTime FrameTimePerOutputFrame = FFrameRate::TransformTime(FFrameTime(FFrameNumber(1)), FinalFrameRate, TickResolution);
+
+		// Dummy range so that the first rendering frame knows to start at the right time.
+		FFrameTime UpperBound = CurrentCameraCut->ShotInfo.CurrentTimeInRoot;
+		CurrentFrameData.LastOutputFrameRange = TRange<FFrameTime>(UpperBound - FrameTimePerOutputFrame, UpperBound);
+		CurrentFrameData.LastSampleRange = CurrentFrameData.LastOutputFrameRange;
+
+		// We can safely fall through to the below states as they're OK to process the same frame we set up.
+		UE_LOG(LogMovieRenderPipeline, Log, TEXT("MovieGraph Finished initializing Camera Cut [%d/%d] in [%s] %s."),
+			CurrentShotIndex + 1, ActiveShotList.Num(), *CurrentCameraCut->OuterName, *CurrentCameraCut->InnerName);
+
+		// Temp...
+		CurrentCameraCut->ShotInfo.State = EMovieRenderShotState::Rendering;
+	}
+
+	if (CurrentCameraCut->ShotInfo.State == EMovieRenderShotState::Rendering)
+	{
+		FMovieGraphTraversalContext Context = GetOwningGraph()->GetCurrentTraversalContext();
+
+		// Update global variables before evaluating the graph
+		Context.RootGraph->UpdateGlobalVariableValues(GetOwningGraph());
+		
+		if (IsFirstTemporalSample())
+		{
+			UMovieGraphConfig* Config = GetOwningGraph()->GetRootGraphForShot(CurrentCameraCut);
+			CurrentFrameData.EvaluatedConfig = TStrongObjectPtr<UMovieGraphEvaluatedConfig>(Config->CreateFlattenedGraph(Context));
+
+			// The temporal sample count can change every frame due to graph evaluations, so when we're on our first temporal
+			// sub-sample of the new frame, we need to re-fetch the value.
+			CurrentFrameData.TemporalSampleCount = GetTemporalSampleCount();
+
+			// Re-calculate some timing statistics about the given TS count
+			UpdateFrameMetrics();
+
+			CurrentCameraCut->ShotInfo.WorkMetrics.OutputFrameIndex++;
+
+			// We create a range of time that we want to represent with this frame in absolute root sequence
+			// space. This is because each frame can have a different temporal sample count, and when we jump
+			// over the time period the shutter is closed, it's no longer an easy fixed number. So instead,
+			// we build a time range we wish to represesnt with this frame, and then split that into sub-regions
+			// and simply move between pre-calculated regions. This lets us handle the fact that the delta time
+			// between output frames is going to be different when moving between different TS counts on frames.
+			
+			// To figure out what range of time we want to represent, we actually need to look at the last frame's
+			// calculated total output range (ignoring all sub-sampling). The new frame is then multiplied by time
+			// dilation, so our entire output range is shortened. We can't just jump between whole frames (due to
+			// time dilation) so we have to take the last range, figure out how long the new frame is, and then
+			// build ontop of that.
+			TRangeBound<FFrameTime> EndOfPreviousFrame = CurrentFrameData.LastOutputFrameRange.GetUpperBound();
+
+			// The upper bound is exclusive, so we initialize a new TRange with just the value to start inclusively
+			// ToDo: Apply time dilation by multiplying the calculated range by slow-mo duration.
+			CurrentFrameData.CurrentOutputFrameRange = TRange<FFrameTime>(EndOfPreviousFrame.GetValue(), EndOfPreviousFrame.GetValue() + CurrentFrameMetrics.FrameTimePerOutputFrame);
+
+			// Now that we've calculated the total range of time we're trying to represent, we can check to see
+			// if this would put us beyond our range of time this shot is supposed to represent.
+			if (CurrentFrameData.CurrentOutputFrameRange.GetUpperBoundValue() > CurrentCameraCut->ShotInfo.TotalOutputRangeRoot.GetUpperBoundValue())
+			{
+				// We're going to spend this frame tearing down the shot (not rendering anything), next frame
+				// we'll re-enter this loop and pick up the start of the next shot.
+				ResetForEndOfOutputFrame();
+
+				// The delta time isn't very relevant here but you must specify at _a_ delta time each frame to ensure
+				// we never have a frame that uses a stale delta time.
+				CustomTimeStep->SetCachedFrameTiming(UMovieGraphEngineTimeStep::FTimeStepCache(CurrentFrameMetrics.FrameRate.AsInterval()));
+
+				CurrentCameraCut->ShotInfo.State = EMovieRenderShotState::Finished;
+				GetOwningGraph()->TeardownShot(CurrentCameraCut);
+				return;
+			}
+
+			// Figure out ranges for how long the shutter is open and closed
+			{
+				TArray<TRange<FFrameTime>> SplitRange = CurrentFrameData.CurrentOutputFrameRange.Split(CurrentFrameData.CurrentOutputFrameRange.GetLowerBoundValue() + CurrentFrameMetrics.FrameTimeWhileShutterOpen);
+				if (ensure(SplitRange.Num() == 2))
+				{
+					CurrentFrameData.RangeShutterOpen = SplitRange[0];
+					CurrentFrameData.RangeShutterClosed = SplitRange[1];
+				}
+			}
+			
+			// Now split the range the shutter is open per temporal sample.
+			{
+				CurrentFrameData.TemporalRanges.Reset();
+				CurrentFrameData.TemporalRanges.Reserve(CurrentFrameData.TemporalSampleCount);
+				FFrameTime SplitStartTime = CurrentFrameData.RangeShutterOpen.GetLowerBoundValue();
+				TRange<FFrameTime> RemainingRange = CurrentFrameData.RangeShutterOpen;
+				for (int32 Index = 0; Index < CurrentFrameData.TemporalSampleCount - 1; Index++)
+				{
+					SplitStartTime += CurrentFrameMetrics.FrameTimePerTemporalSample;
+					TArray<TRange<FFrameTime>> NewRanges = RemainingRange.Split(SplitStartTime);
+					if (ensure(NewRanges.Num() == 2))
+					{
+						CurrentFrameData.TemporalRanges.Add(NewRanges[0]);
+						RemainingRange = NewRanges[1];
+					}
+				}
+
+				// Add the remaining range, it's the leftover from the last split.
+				CurrentFrameData.TemporalRanges.Add(RemainingRange);
+
+				//for (int32 Index = 0; Index < CurrentFrameData.TemporalRanges.Num(); Index++)
+				//{
+				//	UE_LOG(LogTemp, Warning, TEXT("Range: [%s, %s)"),
+				//		*LexToString(CurrentFrameData.TemporalRanges[Index].GetLowerBoundValue()),
+				//		*LexToString(CurrentFrameData.TemporalRanges[Index].GetUpperBoundValue()));
+				//}
+			}
+		}
+		
+		// The delta time for this frame is the difference between the current range index, and the last range index.
+		const TRange<FFrameTime>& PreviousRange = CurrentFrameData.LastSampleRange;
+		const TRange<FFrameTime>& NextRange = CurrentFrameData.TemporalRanges[GetNextTemporalRangeIndex()];
+
+		// Because some time-steps may not advance linearly through time, the abs value of the ranges needs to be taken
+		// to avoid negative delta times.
+		FFrameTime FrameDeltaTime = NextRange.GetLowerBoundValue() > PreviousRange.GetLowerBoundValue()
+			? NextRange.GetLowerBoundValue() - PreviousRange.GetLowerBoundValue()
+			: PreviousRange.GetLowerBoundValue() - NextRange.GetLowerBoundValue();
+
+		// ToDo: Propagate delta time multipliers to cloth
+
+		// Because we know what time range we're supposed to represent, we can just assign the CurrentTimeInRoot absolutely,
+		// instead of accumulating delta times. // ToDo: shutter timing, motion blur offset?
+		CurrentCameraCut->ShotInfo.CurrentTimeInRoot = NextRange.GetLowerBoundValue();
+
+		FFrameTime FinalEvalTime = CurrentCameraCut->ShotInfo.CurrentTimeInRoot + CurrentFrameMetrics.MotionBlurCenteringOffsetTime + CurrentFrameMetrics.ShutterOffsetFrameTime;
+		UE_LOG(LogMovieRenderPipeline, VeryVerbose, TEXT("FinalEvalTime: %s (Tick: %s SOffset: %s MOffset: %s)"),
+			*LexToString(FinalEvalTime),
+			*LexToString(CurrentCameraCut->ShotInfo.CurrentTimeInRoot), 
+			*LexToString(CurrentFrameMetrics.ShutterOffsetFrameTime),
+			*LexToString(CurrentFrameMetrics.MotionBlurCenteringOffsetTime));
+
+		// Now we need to fill out some of our current timestep data for the renderer portion to use.
+		double FrameDeltaTimeAsSeconds = CurrentFrameMetrics.TickResolution.AsSeconds(FrameDeltaTime);
+		CurrentTimeStepData.FrameDeltaTime = FrameDeltaTimeAsSeconds;
+		CurrentTimeStepData.WorldSeconds = CurrentTimeStepData.WorldSeconds + FrameDeltaTimeAsSeconds;
+		CurrentTimeStepData.WorldTimeDilation = 1.0; // Temp
+
+		// The combination of shutter angle percentage, non-uniform render frame delta times and dividing by sample
+		// count produce the correct length for motion blur in all cases.
+		CurrentTimeStepData.MotionBlurFraction = CurrentFrameMetrics.MotionBlurAmount / CurrentFrameData.TemporalSampleCount;
+		CurrentTimeStepData.bIsFirstTemporalSampleForFrame = IsFirstTemporalSample();
+		CurrentTimeStepData.bIsLastTemporalSampleForFrame = IsLastTemporalSample();
+		CurrentTimeStepData.bRequiresAccumulator = CurrentFrameData.TemporalSampleCount > 1;
+		CurrentTimeStepData.OutputFrameNumber = CurrentFrameData.OutputFrameNumber;
+		CurrentTimeStepData.EvaluatedConfig = TObjectPtr<UMovieGraphEvaluatedConfig>(CurrentFrameData.EvaluatedConfig.Get());
+
+		//UE_LOG(LogTemp, Warning, TEXT("F# %d bFirst: %d bLast: %d bReqAc: %d"),
+		//	CurrentTimeStepData.OutputFrameNumber, CurrentTimeStepData.bIsFirstTemporalSampleForFrame,
+		//	CurrentTimeStepData.bIsLastTemporalSampleForFrame, CurrentTimeStepData.bRequiresAccumulator);
+
+		// Calculate frame numbers and timecodes for the current sequence (root) and shot
+		{
+			constexpr bool bIncludeCDOs = true;
+			UMovieGraphOutputSettingNode* OutputSetting = CurrentFrameData.EvaluatedConfig->GetSettingForBranch<UMovieGraphOutputSettingNode>(UMovieGraphNode::GlobalsPinName, bIncludeCDOs);
+
+			// "Closest" isn't straightforward when using temporal sub-sampling, ie: A large enough shutter angle pushes a sample over the half way point and it rounds to
+			// the wrong one. Because temporal sub-sampling isn't centered around a frame (the centering is done via the final eval time) we can just subtract TSI*TPS to get our centered value.
+			const FFrameTime CenteringOffset = CurrentFrameData.TemporalSampleIndex * CurrentFrameMetrics.FrameTimePerTemporalSample;
+			FFrameTime CenteredFrameTime = CurrentCameraCut->ShotInfo.CurrentTimeInRoot - CenteringOffset;
+	
+			const FFrameRate SourceFrameRate = GetOwningGraph()->GetDataSourceInstance()->GetDisplayRate();
+			const FFrameRate EffectiveFrameRate = UMovieGraphBlueprintLibrary::GetEffectiveFrameRate(OutputSetting, SourceFrameRate);
+			const FFrameRate TickResolution = GetOwningGraph()->GetDataSourceInstance()->GetTickResolution();
+
+			constexpr bool bDropFrame = false;
+			CurrentTimeStepData.RootFrameNumber = FFrameRate::TransformTime(CenteredFrameTime, TickResolution, EffectiveFrameRate).RoundToFrame();
+			CurrentTimeStepData.RootTimeCode = FTimecode::FromFrameNumber(CurrentTimeStepData.RootFrameNumber, EffectiveFrameRate, bDropFrame);
+
+			// Calculate metrics for the shot as well
+			CenteredFrameTime = CenteredFrameTime * CurrentCameraCut->ShotInfo.OuterToInnerTransform;
+			CurrentTimeStepData.ShotFrameNumber = FFrameRate::TransformTime(CenteredFrameTime, TickResolution, EffectiveFrameRate).RoundToFrame();
+			CurrentTimeStepData.ShotTimeCode = FTimecode::FromFrameNumber(CurrentTimeStepData.ShotFrameNumber, EffectiveFrameRate, bDropFrame);
+		}
+
+		// Set our time step for the next frame. We use the undilated delta time for the Custom Timestep as the engine will
+		// apply the time dilation to the world tick for us, so we don't want to double up time dilation.
+		double UndilatedDeltaTime = CurrentFrameMetrics.TickResolution.AsSeconds(FrameDeltaTime);
+		CustomTimeStep->SetCachedFrameTiming(UMovieGraphEngineTimeStep::FTimeStepCache(UndilatedDeltaTime));
+		GetOwningGraph()->GetDataSourceInstance()->SyncDataSourceTime(FinalEvalTime);
+
+		// ToDo: This should be converted back to an 'effective' frame number (source frame in external data asset)
+		// so you can line up profiling with the acutal content on screen.
+		TRACE_BOOKMARK(TEXT("MRQ Frame %d [TS: %d]"), CurrentTimeStepData.OutputFrameNumber, CurrentFrameData.TemporalSampleIndex);
+
+		// Increment various post-frame counters to set them up for the next frame. This is okay
+		// because the rest of the Movie Graph Pipeline system is based on CurrentTimeStepData which accurately
+		// reflects which frame we're on.
+		{
+			// Update the last sample range to the one we just rendered.
+			CurrentFrameData.LastSampleRange = NextRange;
+			if (IsLastTemporalSample())
+			{
+				CurrentFrameData.LastOutputFrameRange = CurrentFrameData.CurrentOutputFrameRange;
+				
+				// Increment the output frame number only on the last temporal sample.
+				CurrentFrameData.OutputFrameNumber++;
+				
+				// If we've rendered the last temporal sub-sample, we've started a new output frame
+				// and we need to reset our temporal sample index.
+				CurrentFrameData.TemporalSampleIndex = 0;
+			}
+			else
+			{
+				// Each tick we increment the temporal sample we're on.
+				CurrentFrameData.TemporalSampleIndex++;
+			}
+		}
+
+		if (!ensure(CurrentCameraCut->ShotInfo.CurrentTickInRoot < CurrentCameraCut->ShotInfo.TotalOutputRangeRoot.GetUpperBoundValue()))
+		{
+			UE_LOG(LogMovieRenderPipeline, Error, TEXT("Shot ran past evaluation range, this shouldn't be possible."));
+		}
+	}
+
+}
+
+bool UMovieGraphCoreTimeStep::IsFirstTemporalSample() const
+{
+	return CurrentFrameData.TemporalSampleIndex == 0;
+}
+
+bool UMovieGraphCoreTimeStep::IsLastTemporalSample() const
+{
+	return CurrentFrameData.TemporalSampleIndex == CurrentFrameData.TemporalSampleCount - 1;
+}
+
+void UMovieGraphCoreTimeStep::ResetForEndOfOutputFrame()
+{
+	CurrentFrameData.TemporalSampleIndex = 0;
+}
+
+bool UMovieGraphCoreTimeStep::IsExpansionForTSRequired(const TObjectPtr<UMovieGraphEvaluatedConfig>& InConfig) const
+{
+	// ToDo: This needs to come from the config (once we have TemporalSampleCount there)
+	return false;
+}
+
+
+void UMovieGraphCoreTimeStep::Shutdown()
+{
+	// Shut down our custom timestep which reqstores some world settings we modified.
+
+	// ToDo: This takes 200ms because of an arbitrary sleep inside of SetCustomTimeStep when shutting down
+	// but that requires some bigger changes to custom timesteps to fix.
+	GEngine->SetCustomTimeStep(PrevCustomTimeStep);
+}
+
+void UMovieGraphCoreTimeStep::UpdateFrameMetrics()
+{
+	FOutputFrameMetrics FrameData;
+
+	// We inherit a tick resolution from the level sequence so that we can represent the same
+	// range of time that the level sequence does.
+	FrameData.TickResolution = GetOwningGraph()->GetDataSourceInstance()->GetTickResolution();
+	FrameData.FrameRate = GetOwningGraph()->GetDataSourceInstance()->GetDisplayRate(); // ToDo, needs to come from config (config can override)
+	FrameData.FrameTimePerOutputFrame = FFrameRate::TransformTime(FFrameTime(FFrameNumber(1)), FrameData.FrameRate, FrameData.TickResolution);
+
+	// Manually perform blending of the Post Process Volumes/Camera/Camera Modifiers to match what the renderer will do.
+	// This uses the primary camera specified by the PlayerCameraManager to get the motion blur amount so in the event of
+	// multi-camera rendering, all cameras will end up using the same motion blur amount defined by the primary camera).
+	FrameData.MotionBlurAmount = GetBlendedMotionBlurAmount();
+
+	// Calculate how long of a duration we want to represent where the camera shutter is open.
+	// If the shutter angle is effectively zero, lie about how long a frame is to prevent divide by zero
+	if (FrameData.MotionBlurAmount < (1.0 / 360.0))
+	{
+		FrameData.FrameTimeWhileShutterOpen = FrameData.FrameTimePerOutputFrame * (1.0 / 360.0);
+	}
+	else
+	{
+		FrameData.FrameTimeWhileShutterOpen = FrameData.FrameTimePerOutputFrame * FrameData.MotionBlurAmount;
+	}
+
+	// Now that we know how long the shutter is open, figure out how long each temporal sub-sample gets.
+	FrameData.FrameTimePerTemporalSample = FrameData.FrameTimeWhileShutterOpen / CurrentFrameData.TemporalSampleCount;
+
+	// The amount of time closed + time open should add up to exactly how long a frame is.
+	FrameData.FrameTimeWhileShutterClosed = FrameData.FrameTimePerOutputFrame - FrameData.FrameTimeWhileShutterOpen;
+
+	// Shutter timing is a bias applied to the final evaluation time to let us change what we consider a frame
+	// ie: Do we consider a frame the start of the timespan we captured? Or is the frame the end of the timespan?
+	// We default to Centered so that the center of your evaluated time is what you see in Level Sequences.
+	// ToDo: Shutter Timing Offset
+	//switch (CameraSettings->ShutterTiming)
+	//{
+	//	// Subtract the entire time the shutter is open.
+	//case EMoviePipelineShutterTiming::FrameClose:
+	//	Output.ShutterOffsetTicks = -Output.TicksWhileShutterOpen;
+	//	break;
+	//	// Only subtract half the time the shutter is open.
+	//case EMoviePipelineShutterTiming::FrameCenter:
+	FrameData.ShutterOffsetFrameTime = -FrameData.FrameTimeWhileShutterOpen / 2.0;
+	//	break;
+	//	// No offset needed
+	//case EMoviePipelineShutterTiming::FrameOpen:
+	//	break;
+	//}
+
+	// Then, calculate our motion blur offset. Motion Blur in the engine is always
+	// centered around the object so we offset our time sampling by half of the
+	// motion blur distance so that the distance blurred represents that time.
+	FrameData.MotionBlurCenteringOffsetTime = FrameData.FrameTimePerTemporalSample / 2.0;
+
+	CurrentFrameMetrics = FrameData;
+}
+
+float UMovieGraphCoreTimeStep::GetBlendedMotionBlurAmount()
+{
+	// 0.5f is the default engine motion blur in the event no Post Process/Camera overrides it.
+	float FinalMotionBlurAmount = 0.5f;
+
+	APlayerCameraManager* PlayerCameraManager = GetWorld()->GetFirstPlayerController()->PlayerCameraManager;
+	if (PlayerCameraManager)
+	{
+		// Apply any motion blur settings from post process volumes in the world
+		FVector ViewLocation = PlayerCameraManager->GetCameraLocation();
+		for (IInterface_PostProcessVolume* PPVolume : GetWorld()->PostProcessVolumes)
+		{
+			const FPostProcessVolumeProperties VolumeProperties = PPVolume->GetProperties();
+
+			// Skip any volumes which are either disabled or don't modify blur amount
+			if (!VolumeProperties.bIsEnabled || !VolumeProperties.Settings->bOverride_MotionBlurAmount)
+			{
+				continue;
+			}
+
+			float LocalWeight = FMath::Clamp(VolumeProperties.BlendWeight, 0.0f, 1.0f);
+
+			if (!VolumeProperties.bIsUnbound)
+			{
+				float DistanceToPoint = 0.0f;
+				PPVolume->EncompassesPoint(ViewLocation, 0.0f, &DistanceToPoint);
+
+				if (DistanceToPoint >= 0 && DistanceToPoint < VolumeProperties.BlendRadius)
+				{
+					LocalWeight *= FMath::Clamp(1.0f - DistanceToPoint / VolumeProperties.BlendRadius, 0.0f, 1.0f);
+				}
+				else
+				{
+					LocalWeight = 0.0f;
+				}
+			}
+
+			if (LocalWeight > 0.0f)
+			{
+				FinalMotionBlurAmount = FMath::Lerp(FinalMotionBlurAmount, VolumeProperties.Settings->MotionBlurAmount, LocalWeight);
+			}
+		}
+
+		// Now try from the camera, which takes priority over post processing volumes.
+		ACameraActor* CameraActor = Cast<ACameraActor>(PlayerCameraManager->GetViewTarget());
+		if (CameraActor)
+		{
+			UCameraComponent* CameraComponent = CameraActor->GetCameraComponent();
+			if (CameraComponent && CameraComponent->PostProcessSettings.bOverride_MotionBlurAmount)
+			{
+				FinalMotionBlurAmount = FMath::Lerp(FinalMotionBlurAmount, CameraComponent->PostProcessSettings.MotionBlurAmount, CameraComponent->PostProcessBlendWeight);
+			}
+		}
+
+		// Apply any motion blur settings from post processing blends attached to the camera manager
+		TArray<FPostProcessSettings> const* CameraAnimPPSettings;
+		TArray<float> const* CameraAnimPPBlendWeights;
+		PlayerCameraManager->GetCachedPostProcessBlends(CameraAnimPPSettings, CameraAnimPPBlendWeights);
+		for (int32 PPIdx = 0; PPIdx < CameraAnimPPBlendWeights->Num(); ++PPIdx)
+		{
+			if ((*CameraAnimPPSettings)[PPIdx].bOverride_MotionBlurAmount)
+			{
+				FinalMotionBlurAmount = FMath::Lerp(FinalMotionBlurAmount, (*CameraAnimPPSettings)[PPIdx].MotionBlurAmount, (*CameraAnimPPBlendWeights)[PPIdx]);
+			}
+		}
+	}
+
+	return FinalMotionBlurAmount;
+}
+
+UMovieGraphEngineTimeStep::UMovieGraphEngineTimeStep()
+	: PrevMinUndilatedFrameTime(0.f)
+	, PrevMaxUndilatedFrameTime(0.f)
+{
+}
+
+bool UMovieGraphEngineTimeStep::Initialize(UEngine* InEngine)
+{
+	// We cache the Min/Max times the world has because we modify these during
+	// renders to match our sample length, otherwise stuff gets out of sync.
+	if (AWorldSettings* WorldSettings = GetWorld()->GetWorldSettings())
+	{
+		PrevMinUndilatedFrameTime = WorldSettings->MinUndilatedFrameTime;
+		PrevMaxUndilatedFrameTime = WorldSettings->MaxUndilatedFrameTime;
+	}
+
+	return true;
+}
+
+void UMovieGraphEngineTimeStep::Shutdown(UEngine* InEngine)
+{
+	if (AWorldSettings* WorldSettings = GetWorld()->GetWorldSettings())
+	{
+		WorldSettings->MinUndilatedFrameTime = PrevMinUndilatedFrameTime;
+		WorldSettings->MaxUndilatedFrameTime = PrevMaxUndilatedFrameTime;
+	}
+}
+
+void UMovieGraphEngineTimeStep::SetCachedFrameTiming(const FTimeStepCache& InTimeCache)
+{
+	if (ensureMsgf(!FMath::IsNearlyZero(InTimeCache.UndilatedDeltaTime), TEXT("An incorrect or uninitialized time step was used! Delta Time of 0 isn't allowed.")))
+	{
+		TimeCache = InTimeCache;
+	}
+	else
+	{
+		UE_LOG(LogMovieRenderPipeline, Error, TEXT("SetCachedFrameTiming called with zero delta time, falling back to 1/24"));
+		TimeCache = FTimeStepCache(1 / 24.0);
+	}
+}
+
+bool UMovieGraphEngineTimeStep::UpdateTimeStep(UEngine* /*InEngine*/)
+{
+	if (ensureMsgf(!FMath::IsNearlyZero(TimeCache.UndilatedDeltaTime), TEXT("An incorrect or uninitialized time step was used! Delta Time of 0 isn't allowed.")))
+	{
+		FApp::UpdateLastTime();
+		FApp::SetDeltaTime(TimeCache.UndilatedDeltaTime);
+		FApp::SetCurrentTime(FApp::GetCurrentTime() + FApp::GetDeltaTime());
+
+		// The UWorldSettings can clamp the delta time inside the Level Tick function, this creates a mismatch between component
+		// velocity and render thread velocity and becomes an issue at high temporal sample counts. 
+		if (AWorldSettings* WorldSettings = GetWorld()->GetWorldSettings())
+		{
+			WorldSettings->MinUndilatedFrameTime = TimeCache.UndilatedDeltaTime;
+			WorldSettings->MaxUndilatedFrameTime = TimeCache.UndilatedDeltaTime;
+		}
+	}
+
+	// Clear our cached time to ensure we're always explicitly telling this what to do and never relying on the last set value.
+	// (This will cause the ensure above to check on the next tick if someone didn't reset our value.)
+	TimeCache = FTimeStepCache();
+
+	// Return false so the engine doesn't run its own logic to overwrite FApp timings.
+	return false;
+}
diff --git a/Engine/Plugins/MovieScene/MovieRenderPipeline/Source/MovieRenderPipelineCore/Private/Graph/MovieGraphDataTypes.cpp b/Engine/Plugins/MovieScene/MovieRenderPipeline/Source/MovieRenderPipelineCore/Private/Graph/MovieGraphDataTypes.cpp
--- a/Engine/Plugins/MovieScene/MovieRenderPipeline/Source/MovieRenderPipelineCore/Private/Graph/MovieGraphDataTypes.cpp
+++ b/Engine/Plugins/MovieScene/MovieRenderPipeline/Source/MovieRenderPipelineCore/Private/Graph/MovieGraphDataTypes.cpp
@@ -1,14 +1,12 @@
 // Copyright Epic Games, Inc. All Rights Reserved.
 
 #include "Graph/MovieGraphDataTypes.h"
-#include "Graph/MovieGraphLinearTimeStep.h"
 #include "Graph/MovieGraphDefaultRenderer.h"
 #include "Graph/MovieGraphSequenceDataSource.h"
 #include "Graph/MovieGraphPipeline.h"
 
 FMovieGraphInitConfig::FMovieGraphInitConfig()
 {
-	TimeStepClass = UMovieGraphLinearTimeStep::StaticClass();
 	RendererClass = UMovieGraphDefaultRenderer::StaticClass();
 	DataSourceClass = UMovieGraphSequenceDataSource::StaticClass();
 	bRenderViewport = false;
diff --git a/Engine/Plugins/MovieScene/MovieRenderPipeline/Source/MovieRenderPipelineCore/Private/Graph/MovieGraphLinTimeStep.cpp b/Engine/Plugins/MovieScene/MovieRenderPipeline/Source/MovieRenderPipelineCore/Private/Graph/MovieGraphLinTimeStep.cpp
new file mode 100644
--- /dev/null
+++ b/Engine/Plugins/MovieScene/MovieRenderPipeline/Source/MovieRenderPipelineCore/Private/Graph/MovieGraphLinTimeStep.cpp
@@ -0,0 +1,20 @@
+﻿// Copyright Epic Games, Inc. All Rights Reserved.
+
+#include "Graph/MovieGraphLinTimeStep.h"
+
+#include "Graph/Nodes/MovieGraphSamplingMethodNode.h"
+
+int32 UMovieGraphLinearTimeStep::GetNextTemporalRangeIndex() const
+{
+	// Linear time step just steps through the temporal ranges in order
+	return CurrentFrameData.TemporalSampleIndex;
+}
+
+int32 UMovieGraphLinearTimeStep::GetTemporalSampleCount() const
+{
+	constexpr bool bIncludeCDOs = true;
+	const UMovieGraphSamplingMethodNode* SamplingMethod =
+		CurrentFrameData.EvaluatedConfig->GetSettingForBranch<UMovieGraphSamplingMethodNode>(UMovieGraphNode::GlobalsPinName, bIncludeCDOs);
+
+	return SamplingMethod->TemporalSampleCount;
+}
diff --git a/Engine/Plugins/MovieScene/MovieRenderPipeline/Source/MovieRenderPipelineCore/Private/Graph/MovieGraphLinearTimeStep.cpp b/Engine/Plugins/MovieScene/MovieRenderPipeline/Source/MovieRenderPipelineCore/Private/Graph/MovieGraphLinearTimeStep.cpp
deleted file mode 100644
diff --git a/Engine/Plugins/MovieScene/MovieRenderPipeline/Source/MovieRenderPipelineCore/Private/Graph/MovieGraphPipeline.cpp b/Engine/Plugins/MovieScene/MovieRenderPipeline/Source/MovieRenderPipelineCore/Private/Graph/MovieGraphPipeline.cpp
--- a/Engine/Plugins/MovieScene/MovieRenderPipeline/Source/MovieRenderPipelineCore/Private/Graph/MovieGraphPipeline.cpp
+++ b/Engine/Plugins/MovieScene/MovieRenderPipeline/Source/MovieRenderPipelineCore/Private/Graph/MovieGraphPipeline.cpp
@@ -4,11 +4,13 @@
 
 #include "Graph/MovieGraphCVarManager.h"
 #include "Graph/MovieGraphDataTypes.h"
+#include "Graph/MovieGraphLinTimeStep.h"
 #include "Graph/MovieGraphOutputMerger.h"
 #include "Graph/MoviePipelineRenderLayerSubsystem.h"
 #include "Graph/Nodes/MovieGraphCollectionNode.h"
 #include "Graph/Nodes/MovieGraphFileOutputNode.h"
 #include "Graph/Nodes/MovieGraphModifierNode.h"
+#include "Graph/Nodes/MovieGraphSamplingMethodNode.h"
 #include "Graph/Nodes/MovieGraphOutputSettingNode.h"
 #include "Graph/Nodes/MovieGraphWarmUpSettingNode.h"
 #include "Graph/MovieGraphBlueprintLibrary.h"
@@ -79,7 +81,6 @@ void UMovieGraphPipeline::Initialize(UMoviePipelineExecutorJob* InJob, const FMo
 	}
 
 	// Create instances of our different classes from the InitConfig
-	GraphTimeStepInstance = NewObject<UMovieGraphTimeStepBase>(this, InitConfig.TimeStepClass);
 	GraphRendererInstance = NewObject<UMovieGraphRendererBase>(this, InitConfig.RendererClass);
 	GraphDataSourceInstance = NewObject<UMovieGraphDataSourceBase>(this, InitConfig.DataSourceClass);
 	
@@ -246,6 +247,9 @@ void UMovieGraphPipeline::BuildShotListFromDataSource()
 	// Shots that are already in the list will be updated but their enable flag will be respected. 
 	GraphDataSourceInstance->UpdateShotList();
 
+	GraphTimeStepInstances.Reset();
+	GraphTimeStepInstances.Reserve(GetCurrentJob()->ShotInfo.Num());
+
 	for (int32 ShotIndex = 0; ShotIndex < GetCurrentJob()->ShotInfo.Num(); ShotIndex++)
 	{
 		const TObjectPtr<UMoviePipelineExecutorShot>& Shot = GetCurrentJob()->ShotInfo[ShotIndex];
@@ -269,6 +273,25 @@ void UMovieGraphPipeline::BuildShotListFromDataSource()
 		CurrentContext.Time = TimeContext;
 
 		TObjectPtr<UMovieGraphEvaluatedConfig> EvaluatedConfig = CurrentContext.RootGraph->CreateFlattenedGraph(CurrentContext);
+		
+		// Create the time step instance for this shot. The time step method cannot vary per branch or per frame, so it
+		// is fetched from the Globals branch. This is the earliest point in the pipeline where the evaluated graph is
+		// available, hence why the instances are generated all at once. Doing it at this point also prevents a circular
+		// dependency between the graph and the time step node. Generally the time step class generates the evaluated
+		// graph, but the evaluated graph also specifies the time step class to use. Therefore, we need to determine the
+		// time step class to use by evaluating the graph outside of the time step class.
+		UMovieGraphSamplingMethodNode* SamplingMethodNode = EvaluatedConfig->GetSettingForBranch<UMovieGraphSamplingMethodNode>(UMovieGraphSettingNode::GlobalsPinName);
+		if (UClass* SamplingMethodClass = SamplingMethodNode->SamplingMethodClass.TryLoadClass<UMovieGraphTimeStepBase>())
+		{
+			GraphTimeStepInstances.Add(NewObject<UMovieGraphTimeStepBase>(this, SamplingMethodClass));
+		}
+		else
+		{
+			GraphTimeStepInstances.Add(NewObject<UMovieGraphLinearTimeStep>(this));
+			UE_LOG(LogMovieRenderPipeline, Warning, TEXT("The shot '%s' specified a time step of type '%s', but it could not be loaded. Defaulting to linear."),
+				*Shot->OuterName, *SamplingMethodNode->SamplingMethodClass.GetAssetName());
+		}
+		
 		UMovieGraphOutputSettingNode* OutputNode = EvaluatedConfig->GetSettingForBranch<UMovieGraphOutputSettingNode>(UMovieGraphSettingNode::GlobalsPinName);
 
 		const FFrameRate SourceFrameRate = GetDataSourceInstance()->GetDisplayRate();
@@ -279,13 +302,17 @@ void UMovieGraphPipeline::BuildShotListFromDataSource()
 		UMovieGraphWarmUpSettingNode* WarmUpNode = EvaluatedConfig->GetSettingForBranch<UMovieGraphWarmUpSettingNode>(UMovieGraphSettingNode::GlobalsPinName);
 		
 		const bool bPrePass = true;
-		const bool bExpandForTemporalSubSample = GetTimeStepInstance()->IsExpansionForTSRequired(EvaluatedConfig);
+		const bool bExpandForTemporalSubSample = GraphTimeStepInstances.Last()->IsExpansionForTSRequired(EvaluatedConfig);
 		ExpandShot(Shot, OutputNode->HandleFrameCount, bExpandForTemporalSubSample, bPrePass, FinalFrameRate, TickResolution, WarmUpNode->NumWarmUpFrames);
 
 		Shot->ShotInfo.CurrentTimeInRoot = Shot->ShotInfo.TotalOutputRangeRoot.GetLowerBoundValue();
 		Shot->ShotInfo.CalculateWorkMetrics(GetDataSourceInstance());
 	}
 
+	// Initialize the time step instance to be the first one. This is usually set up by the shot, but some logic may attempt
+	// to access the time step instance before the shot has a chance to initialize.
+	GraphTimeStepInstance = GraphTimeStepInstances[0];
+
 	// The active shot-list is a subset of the whole shot-list; The ShotInfo contains information about every range it detected to render
 	// but if the user has turned the shot off in the UI then we don't want to render it.
 	ActiveShotList.Empty();
@@ -326,8 +353,6 @@ void UMovieGraphPipeline::OnEngineTickBeginFrame()
 
 void UMovieGraphPipeline::TickProducingFrames()
 {
-	check(GraphTimeStepInstance);
-
 	// Move any output frames that have been finished from the Output Merger
 	// into the actual outputs. This will generate new futures (for actual 
 	// disk writes) which we keep track of below.
@@ -366,7 +391,7 @@ void UMovieGraphPipeline::TickProducingFrames()
 	//}
 
 
-	GraphTimeStepInstance->TickProducingFrames();
+	GetTimeStepInstance()->TickProducingFrames();
 }
 
 void UMovieGraphPipeline::TickFinalizeOutputContainers(const bool bInForceFinish)
@@ -495,6 +520,8 @@ void UMovieGraphPipeline::SetupShot(const TObjectPtr<UMoviePipelineExecutorShot>
 	//	}
 	//}
 
+	GraphTimeStepInstance = GraphTimeStepInstances[CurrentShotIndex];
+
 	const FMovieGraphTimeStepData& TimeStepData = GetTimeStepInstance()->GetCalculatedTimeData();
 	const UMovieGraphEvaluatedConfig* EvaluatedConfig = TimeStepData.EvaluatedConfig;
 
@@ -626,6 +653,13 @@ void UMovieGraphPipeline::ExpandShot(const TObjectPtr<UMoviePipelineExecutorShot
 	}
 }
 
+UMovieGraphTimeStepBase* UMovieGraphPipeline::GetTimeStepInstance() const
+{
+	check(GraphTimeStepInstance);
+
+	return GraphTimeStepInstance;
+}
+
 void UMovieGraphPipeline::OnEngineTickEndFrame()
 {
 	LLM_SCOPE_BYNAME(TEXT("MovieGraphEndFrame"));
@@ -650,7 +684,8 @@ void UMovieGraphPipeline::OnEngineTickEndFrame()
 void UMovieGraphPipeline::RenderFrame()
 {
 	check(GraphRendererInstance);
-	const FMovieGraphTimeStepData& TimeStepData = GraphTimeStepInstance->GetCalculatedTimeData();
+
+	const FMovieGraphTimeStepData& TimeStepData = GetTimeStepInstance()->GetCalculatedTimeData();
 
 	GraphRendererInstance->Render(TimeStepData);
 }
diff --git a/Engine/Plugins/MovieScene/MovieRenderPipeline/Source/MovieRenderPipelineCore/Private/Graph/MovieGraphStochasticTimeStep.cpp b/Engine/Plugins/MovieScene/MovieRenderPipeline/Source/MovieRenderPipelineCore/Private/Graph/MovieGraphStochasticTimeStep.cpp
new file mode 100644
--- /dev/null
+++ b/Engine/Plugins/MovieScene/MovieRenderPipeline/Source/MovieRenderPipelineCore/Private/Graph/MovieGraphStochasticTimeStep.cpp
@@ -0,0 +1,27 @@
+﻿// Copyright Epic Games, Inc. All Rights Reserved.
+
+#include "Graph/MovieGraphStochasticTimeStep.h"
+
+#include "Graph/Nodes/MovieGraphSamplingMethodNode.h"
+
+int32 UMovieGraphStochasticTimeStep::GetNextTemporalRangeIndex() const
+{
+	// Instead of going forwards through the temporal ranges, iterate them in reverse (as an example)
+	// TODO: Actual implementation
+	return (CurrentFrameData.TemporalRanges.Num() - 1) - CurrentFrameData.TemporalSampleIndex;
+}
+
+int32 UMovieGraphStochasticTimeStep::GetTemporalSampleCount() const
+{
+	constexpr bool bIncludeCDOs = true;
+	const UMovieGraphSamplingMethodNode* SamplingMethod =
+		CurrentFrameData.EvaluatedConfig->GetSettingForBranch<UMovieGraphSamplingMethodNode>(UMovieGraphNode::GlobalsPinName, bIncludeCDOs);
+
+	return SamplingMethod->SampleCount;
+}
+
+bool UMovieGraphStochasticTimeStep::IsLastTemporalSample() const
+{
+	// TODO: Return early based on wall time or noise threshold?
+	return Super::IsLastTemporalSample();
+}
diff --git a/Engine/Plugins/MovieScene/MovieRenderPipeline/Source/MovieRenderPipelineCore/Private/Graph/Nodes/MovieGraphSamplingMethodNode.cpp b/Engine/Plugins/MovieScene/MovieRenderPipeline/Source/MovieRenderPipelineCore/Private/Graph/Nodes/MovieGraphSamplingMethodNode.cpp
new file mode 100644
--- /dev/null
+++ b/Engine/Plugins/MovieScene/MovieRenderPipeline/Source/MovieRenderPipelineCore/Private/Graph/Nodes/MovieGraphSamplingMethodNode.cpp
@@ -0,0 +1,46 @@
+﻿// Copyright Epic Games, Inc. All Rights Reserved.
+
+#include "Graph/Nodes/MovieGraphSamplingMethodNode.h"
+
+#include "Graph/MovieGraphConfig.h"
+#include "Graph/MovieGraphLinTimeStep.h"
+#include "Styling/AppStyle.h"
+
+UMovieGraphSamplingMethodNode::UMovieGraphSamplingMethodNode()
+	: SamplingMethodClass(UMovieGraphLinearTimeStep::StaticClass())
+	, TemporalSampleCount(1)
+	, SampleCount(1)
+{
+}
+
+EMovieGraphBranchRestriction UMovieGraphSamplingMethodNode::GetBranchRestriction() const
+{
+	return EMovieGraphBranchRestriction::Globals;
+}
+
+#if WITH_EDITOR
+FText UMovieGraphSamplingMethodNode::GetNodeTitle(const bool bGetDescriptive) const
+{
+	static const FText SamplingMethodNodeName = NSLOCTEXT("MovieGraphNodes", "NodeName_SamplingMethod", "Sampling Method");
+	return SamplingMethodNodeName;
+}
+
+FText UMovieGraphSamplingMethodNode::GetMenuCategory() const
+{
+	return NSLOCTEXT("MovieGraphNodes", "SamplingMethodGraphNode_Category", "Rendering");
+}
+
+FLinearColor UMovieGraphSamplingMethodNode::GetNodeTitleColor() const
+{
+	static const FLinearColor SamplingMethodNodeColor = FLinearColor(0.572f, 0.274f, 1.f);
+	return SamplingMethodNodeColor;
+}
+
+FSlateIcon UMovieGraphSamplingMethodNode::GetIconAndTint(FLinearColor& OutColor) const
+{
+	static const FSlateIcon SamplingMethodPresetIcon = FSlateIcon(FAppStyle::GetAppStyleSetName(), "DerivedData.Cache.Statistics");
+
+	OutColor = FLinearColor::White;
+	return SamplingMethodPresetIcon;
+}
+#endif // WITH_EDITOR
\ No newline at end of file
diff --git a/Engine/Plugins/MovieScene/MovieRenderPipeline/Source/MovieRenderPipelineCore/Public/Graph/MovieGraphCoreTimeStep.h b/Engine/Plugins/MovieScene/MovieRenderPipeline/Source/MovieRenderPipelineCore/Public/Graph/MovieGraphCoreTimeStep.h
new file mode 100644
--- /dev/null
+++ b/Engine/Plugins/MovieScene/MovieRenderPipeline/Source/MovieRenderPipelineCore/Public/Graph/MovieGraphCoreTimeStep.h
@@ -0,0 +1,201 @@
+// Copyright Epic Games, Inc. All Rights Reserved.
+#pragma once
+
+#include "Graph/MovieGraphDataTypes.h"
+#include "Graph/MovieGraphConfig.h"
+#include "Misc/FrameRate.h"
+#include "Misc/FrameTime.h"
+#include "Engine/EngineCustomTimeStep.h"
+#include "UObject/StrongObjectPtr.h"
+
+#include "MovieGraphCoreTimeStep.generated.h"
+
+// Forward Declares
+class UMovieGraphEngineTimeStep;
+
+/**
+ * Provides common logic for typical time-step functionality.
+ * 
+ * The number of temporal sub-samples is read from the graph for each output frame, and we then take the time the
+ * shutter is open and break it into that many sub-samples. Subclasses must implement GetNextTemporalRangeIndex() to
+ * indicate the index of the next temporal sub-sample. 
+ */
+UCLASS(Abstract)
+class MOVIERENDERPIPELINECORE_API UMovieGraphCoreTimeStep : public UMovieGraphTimeStepBase
+{
+	GENERATED_BODY()
+public:
+	UMovieGraphCoreTimeStep();
+
+	// UMovieGraphTimeStepBase Interface
+	virtual void TickProducingFrames() override;
+	virtual void Shutdown() override;
+	virtual FMovieGraphTimeStepData GetCalculatedTimeData() const override { return CurrentTimeStepData; }
+	virtual bool IsExpansionForTSRequired(const TObjectPtr<UMovieGraphEvaluatedConfig>& InConfig) const override;
+	// ~UMovieGraphTimeStepBase Interface
+
+protected:
+	/**
+	 * Gets the index of the next temporal range. The index returned should not be the index returned by the prior call
+	 * (this would result in a frame delta time of zero, causing the engine to crash).
+	 */
+	virtual int32 GetNextTemporalRangeIndex() const PURE_VIRTUAL(UMovieGraphCoreTimeStep::GetNextTemporalRangeIndex, return 0; );
+
+	/** Determines if the current sample being rendered is the last. */
+	virtual bool IsLastTemporalSample() const;
+
+	/** Gets the number of samples that should be used in each frame. */
+	virtual int32 GetTemporalSampleCount() const PURE_VIRTUAL(UMovieGraphCoreTimeStep::UpdateTemporalSampleCount, return 0; ); 
+	
+	virtual void UpdateFrameMetrics();
+	virtual bool IsFirstTemporalSample() const;
+	virtual void ResetForEndOfOutputFrame();
+	virtual float GetBlendedMotionBlurAmount();
+
+protected:
+	/** This is the output data needed by the rest of MRQ to produce a frame. */
+	UPROPERTY(Transient)
+	FMovieGraphTimeStepData CurrentTimeStepData;
+
+	struct FOutputFrameMetrics
+	{
+		FOutputFrameMetrics()
+			: MotionBlurAmount(0.f)
+		{}
+
+		/** The human readable frame rate, adjusted by the config file (24fps, 30fps, etc.) */
+		FFrameRate FrameRate;
+
+		/** The internal resolution data is actually stored at. (24,000, 120,000 etc.) */
+		FFrameRate TickResolution;
+
+		/** The amount of time (in Tick Resolution) that represents one output frame (ie: 1/24). */
+		FFrameTime FrameTimePerOutputFrame;
+
+		/** The 0-1 Motion Blur amount the camera has. */
+		float MotionBlurAmount;
+
+		/** The amount of time (in Tick Resolution) that the shutter is open for. Could be <= FrameTimePerOutputFrame. */
+		FFrameTime FrameTimeWhileShutterOpen;
+
+		/** 
+		* The amount of time (in Tick Resolution) per temporal sample. The actual time slice for a temporal sub-sample
+		* may be different than this (as we have to jump over the time period that the shutter is closed too for one 
+		* temporal sub-sample.
+		*/
+		FFrameTime FrameTimePerTemporalSample;
+		
+		/** The inverse of FrameTimeWhileShutterOpen. ShutterClosed+ShutterOpen should add up to FrameTimePerOutputFrame. */
+		FFrameTime FrameTimeWhileShutterClosed;
+
+		/** A constant offset applied to the final evaluated time to allow users to influence what counts as a frame. (the time before, during, or after a frame) */
+		FFrameTime ShutterOffsetFrameTime;
+
+		/** An offset to the evaluated time to get us into the center of the time period (since motion blur is bi-directional and centered) */
+		FFrameTime MotionBlurCenteringOffsetTime;
+	};
+
+	struct FCurrentFrameData
+	{
+		FCurrentFrameData()
+			: TemporalSampleIndex(0)
+			, TemporalSampleCount(0)
+			, OutputFrameNumber(0)
+		{
+		}
+
+		/** Which temporal sub-sample are we working on? [0,TemporalSampleCount) */
+		int32 TemporalSampleIndex;
+		
+		/** How many temporal sub-samples are there for the current output frame? */
+		int32 TemporalSampleCount;
+		
+		/** Which output frame are we working on, relative to zero.*/
+		int32 OutputFrameNumber;
+
+		/** A range of time (in Tick Resolution) for the last output frame being worked on. Updated before first TS of next frame. */
+		TRange<FFrameTime> LastOutputFrameRange;
+
+		/** A range of time (in Tick Resolution) for the current frame being worked on. Represents the whole output frame (shutter closed + open). */
+		TRange<FFrameTime> CurrentOutputFrameRange;
+		
+		/** A range of time (in Tick Resolution) that the shutter is opened. */
+		TRange<FFrameTime> RangeShutterOpen;
+		
+		/** A range of time (in Tick Resolution) that the shutter is closed. Can be an empty range (when using MotionBlurAmount=1.0) */
+		TRange<FFrameTime> RangeShutterClosed;
+
+		/**
+		* A range of time (in Tick Resolution) for the last temporal sub-sample (ie: what the last tick rendered with). We don't track
+		* CurrentSampleRange as it's just TemporalRanges[TemporalSampleIndex].
+		*/
+		TRange<FFrameTime> LastSampleRange;
+
+		/** A pre-calculated (when the output frame starts) range of time, breaking RangeShutterOpen into smaller chunks, one per temporal sub-sample. */
+		TArray<TRange<FFrameTime>> TemporalRanges;
+
+		/** Updated at the start of each Temporal Sample. */
+		TStrongObjectPtr<UMovieGraphEvaluatedConfig> EvaluatedConfig;
+	};
+
+	/**
+	* A set of cached values that are true for the current output frame. They will be recalculated
+	* at the start of the next output frame.
+	*/
+	FOutputFrameMetrics CurrentFrameMetrics;
+
+	/**
+	* Frame Metrics are constant data updated once per output frame, while CurrentFrameData is keeping track
+	* across engine ticks to work towards producing a single frame.
+	*/
+	FCurrentFrameData CurrentFrameData;
+
+	/**
+	* A custom timestep owned by this object that is used to inform the engine what the delta time for each
+	* frame should be. 
+	*/
+	UPROPERTY(Transient)
+	TObjectPtr<UMovieGraphEngineTimeStep> CustomTimeStep;
+
+	/** The previous custom timestep the engine was using, if any. */
+	UPROPERTY(Transient)
+	TObjectPtr<UEngineCustomTimeStep> PrevCustomTimeStep;
+};
+
+UCLASS()
+class MOVIERENDERPIPELINECORE_API UMovieGraphEngineTimeStep : public UEngineCustomTimeStep
+{
+	GENERATED_BODY()
+public:
+	UMovieGraphEngineTimeStep();
+
+	struct FTimeStepCache
+	{
+		FTimeStepCache() 
+			: UndilatedDeltaTime(0.0)
+		{}
+
+		FTimeStepCache(double InUndilatedDeltaTime)
+			: UndilatedDeltaTime(InUndilatedDeltaTime)
+		{}
+
+		double UndilatedDeltaTime;
+	};
+
+public:
+	void SetCachedFrameTiming(const FTimeStepCache& InTimeCache);
+	
+	// UEngineCustomTimeStep Interface
+	virtual bool Initialize(UEngine* InEngine) override;
+	virtual void Shutdown(UEngine* InEngine) override;
+	virtual bool UpdateTimeStep(UEngine* InEngine) override;
+	virtual ECustomTimeStepSynchronizationState GetSynchronizationState() const override { return ECustomTimeStepSynchronizationState::Synchronized; }
+	// ~UEngineCustomTimeStep Interface
+
+	/** We don't do any thinking on our own, instead we just spit out the numbers stored in our time cache. */
+	FTimeStepCache TimeCache;
+
+	// Not cached in TimeCache as TimeCache is reset every frame.
+	float PrevMinUndilatedFrameTime;
+	float PrevMaxUndilatedFrameTime;
+};
diff --git a/Engine/Plugins/MovieScene/MovieRenderPipeline/Source/MovieRenderPipelineCore/Public/Graph/MovieGraphDataTypes.h b/Engine/Plugins/MovieScene/MovieRenderPipeline/Source/MovieRenderPipelineCore/Public/Graph/MovieGraphDataTypes.h
--- a/Engine/Plugins/MovieScene/MovieRenderPipeline/Source/MovieRenderPipelineCore/Public/Graph/MovieGraphDataTypes.h
+++ b/Engine/Plugins/MovieScene/MovieRenderPipeline/Source/MovieRenderPipelineCore/Public/Graph/MovieGraphDataTypes.h
@@ -34,13 +34,6 @@ struct MOVIERENDERPIPELINECORE_API FMovieGraphInitConfig
 	
 	FMovieGraphInitConfig();
 	
-	/** 
-	* Which class should the UMovieGraphPipeline use to handle calculating per frame 
-	* timesteps? Defaults to UMovieGraphLinearTimeStep.
-	*/
-	UPROPERTY(BlueprintReadWrite, EditAnywhere, Category = "Movie Graph")
-	TSubclassOf<UMovieGraphTimeStepBase> TimeStepClass;
-	
 	/**
 	* Which class should the UMovieGraphPipeline use to look for render layers and
 	* request renders from. Defaults to UMovieGraphDefaultRenderer.
diff --git a/Engine/Plugins/MovieScene/MovieRenderPipeline/Source/MovieRenderPipelineCore/Public/Graph/MovieGraphLinTimeStep.h b/Engine/Plugins/MovieScene/MovieRenderPipeline/Source/MovieRenderPipelineCore/Public/Graph/MovieGraphLinTimeStep.h
new file mode 100644
--- /dev/null
+++ b/Engine/Plugins/MovieScene/MovieRenderPipeline/Source/MovieRenderPipelineCore/Public/Graph/MovieGraphLinTimeStep.h
@@ -0,0 +1,25 @@
+﻿// Copyright Epic Games, Inc. All Rights Reserved.
+
+#pragma once
+
+#include "MovieGraphCoreTimeStep.h"
+
+#include "MovieGraphLinTimeStep.generated.h"
+
+/**
+ * Advances time forward linearly until the end of the range of time that is being rendered is reached. This is useful
+ * for deferred rendering (where there's a small number of temporal sub-samples and no feedback mechanism for measuring
+ * noise in the final image).
+ */
+UCLASS(BlueprintType, meta = (DisplayName = "Linear Time Step"))
+class MOVIERENDERPIPELINECORE_API UMovieGraphLinearTimeStep : public UMovieGraphCoreTimeStep
+{
+	GENERATED_BODY()
+
+public:
+	UMovieGraphLinearTimeStep() = default;
+
+protected:
+	virtual int32 GetNextTemporalRangeIndex() const override;
+	virtual int32 GetTemporalSampleCount() const override;
+};
diff --git a/Engine/Plugins/MovieScene/MovieRenderPipeline/Source/MovieRenderPipelineCore/Public/Graph/MovieGraphLinearTimeStep.h b/Engine/Plugins/MovieScene/MovieRenderPipeline/Source/MovieRenderPipelineCore/Public/Graph/MovieGraphLinearTimeStep.h
deleted file mode 100644
diff --git a/Engine/Plugins/MovieScene/MovieRenderPipeline/Source/MovieRenderPipelineCore/Public/Graph/MovieGraphPipeline.h b/Engine/Plugins/MovieScene/MovieRenderPipeline/Source/MovieRenderPipelineCore/Public/Graph/MovieGraphPipeline.h
--- a/Engine/Plugins/MovieScene/MovieRenderPipeline/Source/MovieRenderPipelineCore/Public/Graph/MovieGraphPipeline.h
+++ b/Engine/Plugins/MovieScene/MovieRenderPipeline/Source/MovieRenderPipelineCore/Public/Graph/MovieGraphPipeline.h
@@ -70,7 +70,7 @@ public:
 	/** Called by the TimeStepInstance when it's time to tear down the current shot. Don't call this unless you know what you're doing. */
 	void TeardownShot(const TObjectPtr<UMoviePipelineExecutorShot>& InShot);
 	/** Used occasionally to cross-reference other components. Don't call this unless you know what you're doing. */
-	UMovieGraphTimeStepBase* GetTimeStepInstance() const { return GraphTimeStepInstance; }
+	UMovieGraphTimeStepBase* GetTimeStepInstance() const;
 	/** Used occasionally to cross-reference other components. Don't call this unless you know what you're doing. */
 	UFUNCTION(BlueprintCallable, Category = "Movie Graph")
 	UMovieGraphRendererBase* GetRendererInstance() const { return GraphRendererInstance; }
@@ -134,6 +134,15 @@ protected:
 
 
 protected:
+	/** Time step instances for each shot, where the index into the array corresponds to the shot index. */
+	UPROPERTY(Transient)
+	TArray<TObjectPtr<UMovieGraphTimeStepBase>> GraphTimeStepInstances;
+
+	/**
+	 * Sometimes the shot index can be incremented before the time step instance should be changed, so the current time
+	 * step instance is tracked. This should generally be used to access the current time step instance, rather than
+	 * indexing into GraphTimeStepInstances with the shot index.
+	 */
 	UPROPERTY(Transient)
 	TObjectPtr<UMovieGraphTimeStepBase> GraphTimeStepInstance;
 
diff --git a/Engine/Plugins/MovieScene/MovieRenderPipeline/Source/MovieRenderPipelineCore/Public/Graph/MovieGraphStochasticTimeStep.h b/Engine/Plugins/MovieScene/MovieRenderPipeline/Source/MovieRenderPipelineCore/Public/Graph/MovieGraphStochasticTimeStep.h
new file mode 100644
--- /dev/null
+++ b/Engine/Plugins/MovieScene/MovieRenderPipeline/Source/MovieRenderPipelineCore/Public/Graph/MovieGraphStochasticTimeStep.h
@@ -0,0 +1,26 @@
+﻿// Copyright Epic Games, Inc. All Rights Reserved.
+
+#pragma once
+
+#include "MovieGraphCoreTimeStep.h"
+
+#include "MovieGraphStochasticTimeStep.generated.h"
+
+/**
+ * Advances time in a non-linear (stochastic) manner within the range of time that represents the current frame. Useful for
+ * path traced images which have a varying amount of noise (and thus need a varying amount of samples) based on their
+ * content.
+ */
+UCLASS(BlueprintType, meta = (DisplayName = "Stochastic Time Step"))
+class MOVIERENDERPIPELINECORE_API UMovieGraphStochasticTimeStep : public UMovieGraphCoreTimeStep
+{
+	GENERATED_BODY()
+
+public:
+	UMovieGraphStochasticTimeStep() = default;
+
+protected:
+	virtual int32 GetNextTemporalRangeIndex() const override;
+	virtual int32 GetTemporalSampleCount() const override;
+	virtual bool IsLastTemporalSample() const override;
+};
diff --git a/Engine/Plugins/MovieScene/MovieRenderPipeline/Source/MovieRenderPipelineCore/Public/Graph/Nodes/MovieGraphSamplingMethodNode.h b/Engine/Plugins/MovieScene/MovieRenderPipeline/Source/MovieRenderPipelineCore/Public/Graph/Nodes/MovieGraphSamplingMethodNode.h
new file mode 100644
--- /dev/null
+++ b/Engine/Plugins/MovieScene/MovieRenderPipeline/Source/MovieRenderPipelineCore/Public/Graph/Nodes/MovieGraphSamplingMethodNode.h
@@ -0,0 +1,72 @@
+﻿// Copyright Epic Games, Inc. All Rights Reserved.
+
+#pragma once
+
+#include "Graph/MovieGraphNode.h"
+
+#include "MovieGraphSamplingMethodNode.generated.h"
+
+/** A node which configures sampling properties for renderers. */
+UCLASS()
+class MOVIERENDERPIPELINECORE_API UMovieGraphSamplingMethodNode : public UMovieGraphSettingNode
+{
+	GENERATED_BODY()
+
+public:
+	UMovieGraphSamplingMethodNode();
+
+	virtual EMovieGraphBranchRestriction GetBranchRestriction() const override;
+
+#if WITH_EDITOR
+	virtual FText GetNodeTitle(const bool bGetDescriptive = false) const override;
+	virtual FText GetMenuCategory() const override;
+	virtual FLinearColor GetNodeTitleColor() const override;
+	virtual FSlateIcon GetIconAndTint(FLinearColor& OutColor) const override;
+#endif
+
+public:
+	UPROPERTY(EditAnywhere, BlueprintReadWrite, Category = "Settings", meta = (InlineEditConditionToggle))
+	uint8 bOverride_SamplingMethodClass : 1;
+
+	UPROPERTY(EditAnywhere, BlueprintReadWrite, Category = "Settings", meta = (InlineEditConditionToggle))
+	uint8 bOverride_TemporalSampleCount : 1;
+
+	UPROPERTY(EditAnywhere, BlueprintReadWrite, Category = "Settings", meta = (InlineEditConditionToggle))
+	uint8 bOverride_SampleCount : 1;
+	
+	UPROPERTY(EditAnywhere, BlueprintReadWrite, Category = "Settings", meta = (InlineEditConditionToggle))
+	uint8 bOverride_NoiseThreshold : 1;
+
+	UPROPERTY(EditAnywhere, BlueprintReadWrite, Category = "Settings", meta = (InlineEditConditionToggle))
+	uint8 bOverride_MaxWallTime : 1;
+
+	/** The type of sampling the render should use. */
+	UPROPERTY(BlueprintReadWrite, EditAnywhere, Category = "Settings", NoClear, meta = (EditCondition = "bOverride_SamplingMethodClass", DisplayName = "Sampling Method", MetaClass = "/Script/MovieRenderPipelineCore.MovieGraphTimeStepBase", ShowDisplayNames))
+	FSoftClassPath SamplingMethodClass;
+
+	/** The number of temporal samples which should be taken on one frame. Applies only to linear sampling. */
+	UPROPERTY(BlueprintReadWrite, EditAnywhere, Category = "Settings|Linear Sampling", meta = (UIMin = 1, ClampMin = 1, EditCondition = "bOverride_TemporalSampleCount"))
+	int32 TemporalSampleCount;
+
+	/**
+	 * The number of samples which should be taken on one frame. The number of samples taken may be less if the max
+	 * wall time or noise threshold is reached. Applies only to stochastic sampling.
+	 */
+	UPROPERTY(BlueprintReadWrite, EditAnywhere, Category = "Settings|Stochastic Sampling", meta = (UIMin = 1, ClampMin = 1, EditCondition = "bOverride_SampleCount"))
+	int32 SampleCount;
+
+	/**
+	 * [UNIMPLEMENTED] The noise threshold that must be reached on the current frame before samples stop being collected.
+	 * Samples may stop being collected before the noise threshold is reached if the max wall time is met, or the sample
+	 * count is set too low. Applies only to stochastic sampling.
+	 */
+	UPROPERTY(BlueprintReadWrite, EditAnywhere, Category = "Settings|Stochastic Sampling", meta = (UIMin = 0, ClampMin = 0, EditCondition = "bOverride_NoiseThreshold"))
+	float NoiseThreshold;
+
+	/**
+	 * [UNIMPLEMENTED] The maximum amount of aggregate time (seconds) that can be spent rendering all samples for one frame.
+	 * Applies only to stochastic sampling.
+	 */
+	UPROPERTY(BlueprintReadWrite, EditAnywhere, Category = "Settings|Stochastic Sampling", meta = (UIMin = 0, ClampMin = 0, Unit = "s", EditCondition = "bOverride_MaxWallTime"))
+	float MaxWallTime;
+};
\ No newline at end of file
