MoviePipeline: Fixed dllexport omissions on UMoviePipelineDeferredPass_* classes.

#rnx

diff --git a/Engine/Plugins/MovieScene/MovieRenderPipeline/Source/MovieRenderPipelineRenderPasses/Public/MoviePipelineDeferredPasses.h b/Engine/Plugins/MovieScene/MovieRenderPipeline/Source/MovieRenderPipelineRenderPasses/Public/MoviePipelineDeferredPasses.h
--- a/Engine/Plugins/MovieScene/MovieRenderPipeline/Source/MovieRenderPipelineRenderPasses/Public/MoviePipelineDeferredPasses.h
+++ b/Engine/Plugins/MovieScene/MovieRenderPipeline/Source/MovieRenderPipelineRenderPasses/Public/MoviePipelineDeferredPasses.h
@@ -191,7 +191,7 @@ private:
 
 
 UCLASS(BlueprintType)
-class UMoviePipelineDeferredPass_Unlit : public UMoviePipelineDeferredPassBase
+class MOVIERENDERPIPELINERENDERPASSES_API UMoviePipelineDeferredPass_Unlit : public UMoviePipelineDeferredPassBase
 {
 	GENERATED_BODY()
 
@@ -213,7 +213,7 @@ public:
 };
 
 UCLASS(BlueprintType)
-class UMoviePipelineDeferredPass_DetailLighting : public UMoviePipelineDeferredPassBase
+class MOVIERENDERPIPELINERENDERPASSES_API UMoviePipelineDeferredPass_DetailLighting : public UMoviePipelineDeferredPassBase
 {
 	GENERATED_BODY()
 
@@ -236,7 +236,7 @@ public:
 };
 
 UCLASS(BlueprintType)
-class UMoviePipelineDeferredPass_LightingOnly : public UMoviePipelineDeferredPassBase
+class MOVIERENDERPIPELINERENDERPASSES_API UMoviePipelineDeferredPass_LightingOnly : public UMoviePipelineDeferredPassBase
 {
 	GENERATED_BODY()
 
@@ -259,7 +259,7 @@ public:
 };
 
 UCLASS(BlueprintType)
-class UMoviePipelineDeferredPass_ReflectionsOnly : public UMoviePipelineDeferredPassBase
+class MOVIERENDERPIPELINERENDERPASSES_API UMoviePipelineDeferredPass_ReflectionsOnly : public UMoviePipelineDeferredPassBase
 {
 	GENERATED_BODY()
 
@@ -282,7 +282,7 @@ public:
 
 
 UCLASS(BlueprintType)
-class UMoviePipelineDeferredPass_PathTracer : public UMoviePipelineDeferredPassBase
+class MOVIERENDERPIPELINERENDERPASSES_API UMoviePipelineDeferredPass_PathTracer : public UMoviePipelineDeferredPassBase
 {
 	GENERATED_BODY()
 
