Fix bad memory access when closing the editor.

diff --git a/Engine/Plugins/MovieScene/MovieRenderPipeline/Source/MovieRenderPipelineEditor/Private/MovieRenderPipelineEditorModule.cpp b/Engine/Plugins/MovieScene/MovieRenderPipeline/Source/MovieRenderPipelineEditor/Private/MovieRenderPipelineEditorModule.cpp
--- a/Engine/Plugins/MovieScene/MovieRenderPipeline/Source/MovieRenderPipelineEditor/Private/MovieRenderPipelineEditorModule.cpp
+++ b/Engine/Plugins/MovieScene/MovieRenderPipeline/Source/MovieRenderPipelineEditor/Private/MovieRenderPipelineEditorModule.cpp
@@ -202,10 +202,13 @@ void FMovieRenderPipelineEditorModule::RegisterTypeCustomizations()
 
 void FMovieRenderPipelineEditorModule::UnregisterTypeCustomizations()
 {
-	FPropertyEditorModule* PropertyModule = FModuleManager::GetModulePtr<FPropertyEditorModule>("PropertyEditor");
-	if (PropertyModule)
+	if (UObjectInitialized())
 	{
-		PropertyModule->UnregisterCustomPropertyTypeLayout(FMoviePipelineConsoleVariableEntry::StaticStruct()->GetFName());
+		FPropertyEditorModule* PropertyModule = FModuleManager::GetModulePtr<FPropertyEditorModule>("PropertyEditor");
+		if (PropertyModule)
+		{
+			PropertyModule->UnregisterCustomPropertyTypeLayout(FMoviePipelineConsoleVariableEntry::StaticStruct()->GetFName());
+		}
 	}
 }
 
