MoviePipeline: Workaround for Motion Vectors not being generated for skeletal meshes during MRQ renders. Ideally we'd fix this by moving when GFrameCounter is incremented but that is too big of a change for right now, so instead we just assume that GFrameCounter has already been incremented, and then decrement it when submitting MRQ renders, using some new, hopefully temporary, fields on the FSceneViewFamily to override the value the FSceneViewFamily thinks it is rendering with.

diff --git a/Engine/Plugins/MovieScene/MovieRenderPipeline/Source/MovieRenderPipelineCore/Private/MoviePipeline.cpp b/Engine/Plugins/MovieScene/MovieRenderPipeline/Source/MovieRenderPipelineCore/Private/MoviePipeline.cpp
--- a/Engine/Plugins/MovieScene/MovieRenderPipeline/Source/MovieRenderPipelineCore/Private/MoviePipeline.cpp
+++ b/Engine/Plugins/MovieScene/MovieRenderPipeline/Source/MovieRenderPipelineCore/Private/MoviePipeline.cpp
@@ -89,6 +89,7 @@ UMoviePipeline::UMoviePipeline()
 	, bShutdownRequested(false)
 	, bIsTransitioningState(false)
 	, AccumulatedTickSubFrameDeltas(0.f)
+	, bHasRenderedFirstViewThisFrame(false)
 	, CurrentJob(nullptr)
 {
 	CustomTimeStep = CreateDefaultSubobject<UMoviePipelineCustomTimeStep>("MoviePipelineCustomTimeStep");
diff --git a/Engine/Plugins/MovieScene/MovieRenderPipeline/Source/MovieRenderPipelineCore/Private/MoviePipelineRendering.cpp b/Engine/Plugins/MovieScene/MovieRenderPipeline/Source/MovieRenderPipelineCore/Private/MoviePipelineRendering.cpp
--- a/Engine/Plugins/MovieScene/MovieRenderPipeline/Source/MovieRenderPipelineCore/Private/MoviePipelineRendering.cpp
+++ b/Engine/Plugins/MovieScene/MovieRenderPipeline/Source/MovieRenderPipelineCore/Private/MoviePipelineRendering.cpp
@@ -256,6 +256,9 @@ void UMoviePipeline::RenderFrame()
 
 	TArray<UMoviePipelineRenderPass*> InputBuffers = FindSettingsForShot<UMoviePipelineRenderPass>(ActiveShotList[CurrentShotIndex]);
 
+	// Reset our flag for this frame.
+	bHasRenderedFirstViewThisFrame = false;
+
 	// If this is the first sample for a new frame, we want to notify the output builder that it should expect data to accumulate for this frame.
 	if (CachedOutputState.IsFirstTemporalSample())
 	{
diff --git a/Engine/Plugins/MovieScene/MovieRenderPipeline/Source/MovieRenderPipelineCore/Private/MoviePipelineUtils.cpp b/Engine/Plugins/MovieScene/MovieRenderPipeline/Source/MovieRenderPipelineCore/Private/MoviePipelineUtils.cpp
--- a/Engine/Plugins/MovieScene/MovieRenderPipeline/Source/MovieRenderPipelineCore/Private/MoviePipelineUtils.cpp
+++ b/Engine/Plugins/MovieScene/MovieRenderPipeline/Source/MovieRenderPipelineCore/Private/MoviePipelineUtils.cpp
@@ -128,6 +128,23 @@ namespace UE
 
 			return AntiAliasingMethod;
 		}
+
+		uint64 GetRendererFrameCount()
+		{
+			// The rendering module relies on GFrameCounter during submission to line several things up that need to happen once per engine-tick.
+			// Unfortunately by the time the Movie Render Queue code gets called (in FCoreDelegates::OnEndFrame) GFrameCounter has already been
+			// incremented compared to the Tick() that frame that Skeletal Meshes, etc. may have used. To solve this, we manually override which
+			// tick the renderer thinks this frame is for.
+
+#if WITH_EDITORONLY_DATA
+			if (GIsEditor)
+			{
+				return GFrameCounter - 1;
+			}
+#endif
+			return GFrameCounter;
+		}
+
 	}
 }
 
diff --git a/Engine/Plugins/MovieScene/MovieRenderPipeline/Source/MovieRenderPipelineCore/Public/MoviePipeline.h b/Engine/Plugins/MovieScene/MovieRenderPipeline/Source/MovieRenderPipelineCore/Public/MoviePipeline.h
--- a/Engine/Plugins/MovieScene/MovieRenderPipeline/Source/MovieRenderPipelineCore/Public/MoviePipeline.h
+++ b/Engine/Plugins/MovieScene/MovieRenderPipeline/Source/MovieRenderPipelineCore/Public/MoviePipeline.h
@@ -427,6 +427,9 @@ public:
 	UE_DEPRECATED(5.1, "Use SetViewportInitArgs instead.")
 	UPROPERTY(Transient)
 	TSubclassOf<UMovieRenderDebugWidget> DebugWidgetClass;
+	
+	/** Used to track first-render submissions (for 3d renders) to set the correct flags on the renderer module. */
+	bool bHasRenderedFirstViewThisFrame;
 
 private:
 	/** Keep track of which job we're working on. This holds our Configuration + which shots we're supposed to render from it. */
diff --git a/Engine/Plugins/MovieScene/MovieRenderPipeline/Source/MovieRenderPipelineCore/Public/MoviePipelineUtils.h b/Engine/Plugins/MovieScene/MovieRenderPipeline/Source/MovieRenderPipelineCore/Public/MoviePipelineUtils.h
--- a/Engine/Plugins/MovieScene/MovieRenderPipeline/Source/MovieRenderPipelineCore/Public/MoviePipelineUtils.h
+++ b/Engine/Plugins/MovieScene/MovieRenderPipeline/Source/MovieRenderPipelineCore/Public/MoviePipelineUtils.h
@@ -124,6 +124,9 @@ namespace UE
 	{
 		MOVIERENDERPIPELINECORE_API TArray<UClass*> FindMoviePipelineSettingClasses(UClass* InBaseClass);
 		MOVIERENDERPIPELINECORE_API EAntiAliasingMethod GetEffectiveAntiAliasingMethod(const UMoviePipelineAntiAliasingSetting* InSetting);
+		
+		UE_DEPRECATED(5.3, "Do not use, this is here as a temporary workaround for another issue.")
+		MOVIERENDERPIPELINECORE_API uint64 GetRendererFrameCount();
 	}
 
 	namespace MoviePipeline
diff --git a/Engine/Plugins/MovieScene/MovieRenderPipeline/Source/MovieRenderPipelineRenderPasses/Private/MoviePipelineImagePassBase.cpp b/Engine/Plugins/MovieScene/MovieRenderPipeline/Source/MovieRenderPipelineRenderPasses/Private/MoviePipelineImagePassBase.cpp
--- a/Engine/Plugins/MovieScene/MovieRenderPipeline/Source/MovieRenderPipelineRenderPasses/Private/MoviePipelineImagePassBase.cpp
+++ b/Engine/Plugins/MovieScene/MovieRenderPipeline/Source/MovieRenderPipelineRenderPasses/Private/MoviePipelineImagePassBase.cpp
@@ -194,6 +194,26 @@ TSharedPtr<FSceneViewFamilyContext> UMoviePipelineImagePassBase::CalculateViewFa
 	OutViewFamily->bWorldIsPaused = InOutSampleState.bWorldIsPaused;
 	OutViewFamily->ViewMode = ViewModeIndex;
 	OutViewFamily->bOverrideVirtualTextureThrottle = true;
+	
+	PRAGMA_DISABLE_DEPRECATION_WARNINGS
+	OutViewFamily->OverrideFrameCounter = UE::MovieRenderPipeline::GetRendererFrameCount();
+	PRAGMA_ENABLE_DEPRECATION_WARNINGS
+	
+	// Kept as an if/else statement to avoid the confusion with setting all of these values to some permutation of !/!!bHasRenderedFirstViewThisFrame.
+	if (!GetPipeline()->bHasRenderedFirstViewThisFrame)
+	{
+		GetPipeline()->bHasRenderedFirstViewThisFrame = true;
+		
+		OutViewFamily->bIsFirstViewInMultipleViewFamily = true;
+		OutViewFamily->bIsMultipleViewFamily = true;
+		OutViewFamily->bAdditionalViewFamily = false;
+	}
+	else
+	{
+		OutViewFamily->bIsFirstViewInMultipleViewFamily = false;
+		OutViewFamily->bAdditionalViewFamily = true;
+		OutViewFamily->bIsMultipleViewFamily = true;
+	}
 
 	const bool bIsPerspective = true;
 	ApplyViewMode(OutViewFamily->ViewMode, bIsPerspective, OutViewFamily->EngineShowFlags);
diff --git a/Engine/Source/Runtime/Engine/Private/SceneView.cpp b/Engine/Source/Runtime/Engine/Private/SceneView.cpp
--- a/Engine/Source/Runtime/Engine/Private/SceneView.cpp
+++ b/Engine/Source/Runtime/Engine/Private/SceneView.cpp
@@ -2836,6 +2836,7 @@ FSceneViewFamily::FSceneViewFamily(const ConstructionValues& CVS)
 	}
 }
 
+PRAGMA_DISABLE_DEPRECATION_WARNINGS // TOptional can't be deprecated without emitting warnings in destructor
 FSceneViewFamily::~FSceneViewFamily()
 {
 	// If a screen percentage was given for the view family, delete it since any new copy of a view family will Fork it.
@@ -2862,6 +2863,7 @@ FSceneViewFamily::~FSceneViewFamily()
 		delete reinterpret_cast<ISceneViewFamilyExtention*>(SecondarySpatialUpscalerInterface);
 	}
 }
+PRAGMA_ENABLE_DEPRECATION_WARNINGS
 
 ERHIFeatureLevel::Type FSceneViewFamily::GetFeatureLevel() const
 {
diff --git a/Engine/Source/Runtime/Engine/Public/SceneView.h b/Engine/Source/Runtime/Engine/Public/SceneView.h
--- a/Engine/Source/Runtime/Engine/Public/SceneView.h
+++ b/Engine/Source/Runtime/Engine/Public/SceneView.h
@@ -1984,6 +1984,11 @@ public:
 	/** Copy from main thread GFrameCounter to be accessible on render thread side. GFrameCounter is incremented once per engine tick, so multi views of the same frame have the same value. */
 	uint64 FrameCounter = 0;
 
+	// When deleted, remove PRAGMA_DISABLE_DEPRECATION_WARNINGS / PRAGMA_ENABLE_DEPRECATION_WARNINGS from ~FSceneViewFamily()
+	/** If set, overrides FrameCounter value, useful for determistic Movie Render Queue behaviour. **/
+	UE_DEPRECATED(5.3, "Do not use, this is here as a temporary workaround for another issue.")
+	TOptional<uint64> OverrideFrameCounter;
+
 	/** Indicates this view family is an additional one. */
 	bool bAdditionalViewFamily;
 
diff --git a/Engine/Source/Runtime/Renderer/Private/SceneRendering.cpp b/Engine/Source/Runtime/Renderer/Private/SceneRendering.cpp
--- a/Engine/Source/Runtime/Renderer/Private/SceneRendering.cpp
+++ b/Engine/Source/Runtime/Renderer/Private/SceneRendering.cpp
@@ -4755,7 +4755,10 @@ void FRendererModule::BeginRenderingViewFamilies(FCanvas* Canvas, TArrayView<FSc
 
 	for (FSceneViewFamily* ViewFamily : ViewFamilies)
 	{
-		ViewFamily->FrameCounter = GFrameCounter;
+		PRAGMA_DISABLE_DEPRECATION_WARNINGS
+		// If optional parameter OverrideFrameCounter is set, use its value instead of GFrameCounter.
+		ViewFamily->FrameCounter = ViewFamily->OverrideFrameCounter.IsSet() ? ViewFamily->OverrideFrameCounter.GetValue() : GFrameCounter;
+		PRAGMA_ENABLE_DEPRECATION_WARNINGS
 
 #if !(UE_BUILD_SHIPPING || UE_BUILD_TEST)
 		{
