MoviePipeline: Exposed view mode index (on the deferred renderer) and show flags on both the path tracer and deferred renderer.

Getting the show flags properly integrated into the graph is tricky because show flags are inside a struct, and not exposed as individual overrideable properties by default. To make them integrate well into the graph:
1) A details customization was added, which makes each show flag overrideable (via the usual checkbox).
2) Show flags are part of a new class, and that class implements the new IMovieGraphTraversableObject interface. IMovieGraphTraversableObject allows classes to define how they merge during traversal (so each show flag can be merged individually).
3) The Active Render Settings editor was updated to ask IMovieGraphTraversableObject which properties have been merged.

#rnx

diff --git a/Engine/Plugins/MovieScene/MovieRenderPipeline/Source/MovieRenderPipelineCore/Private/Graph/MovieGraphConfig.cpp b/Engine/Plugins/MovieScene/MovieRenderPipeline/Source/MovieRenderPipelineCore/Private/Graph/MovieGraphConfig.cpp
--- a/Engine/Plugins/MovieScene/MovieRenderPipeline/Source/MovieRenderPipelineCore/Private/Graph/MovieGraphConfig.cpp
+++ b/Engine/Plugins/MovieScene/MovieRenderPipeline/Source/MovieRenderPipelineCore/Private/Graph/MovieGraphConfig.cpp
@@ -1000,10 +1000,13 @@ void UMovieGraphConfig::CopyOverriddenProperties(UMovieGraphNode* FromNode, UMov
 		}
 
 		// If our destination node already has this property marked as overridden, then some other node in the graph has
-		// taken priority and set the value to something, so we don't want to override it.
+		// taken priority and set the value to something, so we don't want to override it. The exception to this is
+		// an object implementing IMovieGraphTraversableObject -- they determine when/how property values are updated.
+		const FObjectProperty* ObjectProperty = CastField<FObjectProperty>(DestNodeProperty);
+		const bool bIsMergeableObject = ObjectProperty && ObjectProperty->PropertyClass->ImplementsInterface(UMovieGraphTraversableObject::StaticClass());
 		const bool bAlreadyOverriddenOnDestNode = bIsDynamic
 			? ToNode->IsDynamicPropertyOverridden(PropertyName)
-			: EditConditionProperty->GetPropertyValue_InContainer(ToNode);
+			: EditConditionProperty->GetPropertyValue_InContainer(ToNode) && !bIsMergeableObject;
 		if (bAlreadyOverriddenOnDestNode)
 		{
 			continue;
@@ -1058,6 +1061,24 @@ void UMovieGraphConfig::CopyOverriddenProperties(UMovieGraphNode* FromNode, UMov
 			EditConditionProperty->SetPropertyValue_InContainer(ToNode, true);
 		}
 
+		// Before using the normal property copying procedure, check to see if this property is an IMovieGraphTraversableObject.
+		// These objects define a particular way they should have their properties merged.
+		if (bIsMergeableObject)
+		{
+			const IMovieGraphTraversableObject* SourceTraversableObject = Cast<IMovieGraphTraversableObject>(
+				ObjectProperty->GetObjectPropertyValue(ObjectProperty->ContainerPtrToValuePtr<void>(FromNode)));
+			IMovieGraphTraversableObject* DestTraversableObject = Cast<IMovieGraphTraversableObject>(
+				ObjectProperty->GetObjectPropertyValue(ObjectProperty->ContainerPtrToValuePtr<void>(ToNode)));
+			
+			if (DestTraversableObject && SourceTraversableObject)
+			{
+				DestTraversableObject->Merge(SourceTraversableObject);
+
+				// Property has been copied via Merge(), don't run the normal copy procedure
+				continue;
+			}
+		}
+
 		// Now we need to copy the value from the source to the destination
 		if (bIsDynamic)
 		{
diff --git a/Engine/Plugins/MovieScene/MovieRenderPipeline/Source/MovieRenderPipelineCore/Private/Graph/Nodes/MovieGraphCoreRenderPassNode.cpp b/Engine/Plugins/MovieScene/MovieRenderPipeline/Source/MovieRenderPipelineCore/Private/Graph/Nodes/MovieGraphCoreRenderPassNode.cpp
--- a/Engine/Plugins/MovieScene/MovieRenderPipeline/Source/MovieRenderPipelineCore/Private/Graph/Nodes/MovieGraphCoreRenderPassNode.cpp
+++ b/Engine/Plugins/MovieScene/MovieRenderPipeline/Source/MovieRenderPipelineCore/Private/Graph/Nodes/MovieGraphCoreRenderPassNode.cpp
@@ -22,6 +22,143 @@
 // For the 1D Weight table for accumulation
 #include "MovieRenderPipelineDataTypes.h"
 
+UMovieGraphShowFlags::UMovieGraphShowFlags()
+	: ShowFlags(ESFIM_Game)
+{
+	// Using ESFIM_Game as a base, update the show flags based on the serialized state
+	CopyOverridesToShowFlags(ShowFlags);
+}
+
+FEngineShowFlags UMovieGraphShowFlags::GetShowFlags() const
+{
+	// A copy is necessary because we keep ShowFlags around as the default (including values set by the renderer as
+	// default), and sparse overrides are stored on top of it. The sparse overrides need to be applied to a copy of the
+	// default.
+	
+	FEngineShowFlags Copy(ShowFlags);
+	CopyOverridesToShowFlags(Copy);
+
+	return Copy;
+}
+
+bool UMovieGraphShowFlags::IsShowFlagEnabled(const uint32 ShowFlagIndex) const
+{
+	if (IsShowFlagOverridden(ShowFlagIndex))
+	{
+		if (const bool* EnableState = ShowFlagEnableState.Find(ShowFlagIndex))
+		{
+			return *EnableState;
+		}
+	}
+	
+	return ShowFlags.GetSingleFlag(ShowFlagIndex); 
+}
+
+void UMovieGraphShowFlags::SetShowFlagEnabled(const uint32 ShowFlagIndex, const bool bIsShowFlagEnabled)
+{
+	if (!IsShowFlagOverridden(ShowFlagIndex))
+	{
+		UE_LOG(LogMovieRenderPipeline, Warning, TEXT("SetShowFlagEnabled() call for '%s' ignored because the flag has not been marked as overridden."), *ShowFlags.FindNameByIndex(ShowFlagIndex));
+		return;
+	}
+	
+	bool& EnableState = ShowFlagEnableState.FindOrAdd(ShowFlagIndex);
+	EnableState = bIsShowFlagEnabled;
+}
+
+bool UMovieGraphShowFlags::IsShowFlagOverridden(const uint32 ShowFlagIndex) const
+{
+	return OverriddenShowFlags.Contains(ShowFlagIndex);
+}
+	
+void UMovieGraphShowFlags::SetShowFlagOverridden(const uint32 ShowFlagIndex, const bool bIsOverridden)
+{
+	if (bIsOverridden)
+	{
+		OverriddenShowFlags.Add(ShowFlagIndex);
+
+		// If the flag has not been marked as overridden in the past, then set the enable state to the default from
+		// ShowFlags. Otherwise, marking the flag as overridden will pick up the previously-set enable state.
+		if (!ShowFlagEnableState.Contains(ShowFlagIndex))
+		{
+			ShowFlagEnableState.Add(ShowFlagIndex, ShowFlags.GetSingleFlag(ShowFlagIndex));
+		}
+	}
+	else
+	{
+		OverriddenShowFlags.Remove(ShowFlagIndex);
+
+		// State from ShowFlagEnableState isn't removed; this is retained to restore the previously-set state if the
+		// user re-enables the override
+	}
+}
+
+const TSet<uint32>& UMovieGraphShowFlags::GetOverriddenShowFlags() const
+{
+	return OverriddenShowFlags;
+}
+
+void UMovieGraphShowFlags::ApplyDefaultShowFlagValue(const uint32 ShowFlagIndex, const bool bShowFlagState)
+{
+	// If there's no current override on this flag, then the default can be set
+	if (!OverriddenShowFlags.Contains(ShowFlagIndex))
+	{
+		ShowFlags.SetSingleFlag(ShowFlagIndex, bShowFlagState);
+	}
+}
+
+void UMovieGraphShowFlags::Merge(const IMovieGraphTraversableObject* InSourceObject)
+{
+	const UMovieGraphShowFlags* InSourceFlags = Cast<UMovieGraphShowFlags>(InSourceObject);
+	checkf(InSourceFlags, TEXT("UMovieGraphShowFlags cannot merge with null or an object of another type."));
+	
+	for (const uint32& ShowFlagIndex : InSourceFlags->GetOverriddenShowFlags())
+	{
+		// The graph is iterated in reverse, from Outputs to Inputs. If this property is already overridden (on a
+		// object upstream, aka this object), then don't merge in the downstream (InSourceObject) override.
+		if (!OverriddenShowFlags.Contains(ShowFlagIndex))
+		{
+			SetShowFlagOverridden(ShowFlagIndex, true);
+			SetShowFlagEnabled(ShowFlagIndex, InSourceFlags->IsShowFlagEnabled(ShowFlagIndex));
+		}
+	}
+}
+
+TArray<TPair<FString, FString>> UMovieGraphShowFlags::GetMergedProperties() const
+{
+	TArray<TPair<FString, FString>> MergedProperties;
+	
+	for (const uint32& ShowFlagIndex : OverriddenShowFlags)
+	{
+		MergedProperties.Add({ FEngineShowFlags::FindNameByIndex(ShowFlagIndex), LexToString(IsShowFlagEnabled(ShowFlagIndex)) });
+	}
+
+	// Sort by show flag name (helps w/ keeping order consistent in the UI)
+	MergedProperties.Sort([](const TPair<FString, FString>& Pair1, const TPair<FString, FString>& Pair2)
+	{
+		return Pair1.Key < Pair2.Key;
+	});
+
+	return MergedProperties;
+}
+
+void UMovieGraphShowFlags::CopyOverridesToShowFlags(FEngineShowFlags& InEngineShowFlags) const
+{
+	for (const uint32 ShowFlagIndex : OverriddenShowFlags)
+	{
+		// ShowFlagEnableState should always have an entry if there's an override, but check just in case
+		if (const bool* EnableState = ShowFlagEnableState.Find(ShowFlagIndex))
+		{
+			InEngineShowFlags.SetSingleFlag(ShowFlagIndex, *EnableState);
+		}
+	}
+}
+
+UMovieGraphCoreRenderPassNode::UMovieGraphCoreRenderPassNode()
+{
+	ShowFlags = CreateDefaultSubobject<UMovieGraphShowFlags>(TEXT("ShowFlags"));
+}
+
 void UMovieGraphCoreRenderPassNode::SetupImpl(const FMovieGraphRenderPassSetupData& InSetupData)
 {
 	// To make the implementation simpler, we make one instance of FMovieGraphDeferredRenderPas
@@ -87,7 +224,12 @@ EViewModeIndex UMovieGraphCoreRenderPassNode::GetViewModeIndex() const
 
 FEngineShowFlags UMovieGraphCoreRenderPassNode::GetShowFlags() const
 {
-	return FEngineShowFlags(EShowFlagInitMode::ESFIM_Game);
+	return ShowFlags->GetShowFlags();
+}
+
+TArray<uint32> UMovieGraphCoreRenderPassNode::GetDefaultShowFlags() const
+{
+	return TArray<uint32>();
 }
 
 void UMovieGraphCoreRenderPassNode::AddReferencedObjects(UObject* InThis, FReferenceCollector& Collector)
@@ -319,22 +461,22 @@ void UMovieGraphCoreRenderPassNode::FMovieGraphRenderPass::PostRendererSubmissio
 
 TSharedRef<FSceneViewFamilyContext> UMovieGraphCoreRenderPassNode::FMovieGraphRenderPass::AllocateSceneViewFamilyContext(const FViewFamilyContextInitData& InInitData)
 {
-	FEngineShowFlags ShowFlags = RenderPassNode->GetShowFlags();
-	const EViewModeIndex ViewModeIndex = RenderPassNode->GetViewModeIndex();
+	FEngineShowFlags NewShowFlags = Cast<UMovieGraphCoreRenderPassNode>(LayerData.RenderPassNode)->GetShowFlags();
+	const EViewModeIndex ViewModeIndex = Cast<UMovieGraphCoreRenderPassNode>(LayerData.RenderPassNode)->GetViewModeIndex();
 
 	const bool bIsPerspective = InInitData.CameraInfo.ViewInfo.ProjectionMode == ECameraProjectionMode::Type::Perspective;
 
 	// Allow the Engine Showflag system to override our engine showflags, based on our view mode index.
 	// This is required for certain debug view modes (to have matching show flags set for rendering).
-	ApplyViewMode(/*In*/ ViewModeIndex, bIsPerspective, /*InOut*/ShowFlags);
+	ApplyViewMode(/*In*/ ViewModeIndex, bIsPerspective, /*InOut*/NewShowFlags);
 
 	// And then we have to let another system override them again (based on cvars, etc.)
-	EngineShowFlagOverride(ESFIM_Game, ViewModeIndex, ShowFlags, false);
+	EngineShowFlagOverride(ESFIM_Game, ViewModeIndex, NewShowFlags, false);
 
 	TSharedRef<FSceneViewFamilyContext> OutViewFamily = MakeShared<FSceneViewFamilyContext>(FSceneViewFamily::ConstructionValues(
 		InInitData.RenderTarget,
 		InInitData.World->Scene,
-		ShowFlags)
+		NewShowFlags)
 		.SetTime(FGameTime::CreateUndilated(InInitData.TimeData.WorldSeconds, InInitData.TimeData.FrameDeltaTime))
 		.SetRealtimeUpdate(true));
 
diff --git a/Engine/Plugins/MovieScene/MovieRenderPipeline/Source/MovieRenderPipelineCore/Private/Graph/Nodes/MovieGraphDefRenderPassNode.cpp b/Engine/Plugins/MovieScene/MovieRenderPipeline/Source/MovieRenderPipelineCore/Private/Graph/Nodes/MovieGraphDefRenderPassNode.cpp
--- a/Engine/Plugins/MovieScene/MovieRenderPipeline/Source/MovieRenderPipelineCore/Private/Graph/Nodes/MovieGraphDefRenderPassNode.cpp
+++ b/Engine/Plugins/MovieScene/MovieRenderPipeline/Source/MovieRenderPipelineCore/Private/Graph/Nodes/MovieGraphDefRenderPassNode.cpp
@@ -5,6 +5,11 @@
 #include "Engine/EngineBaseTypes.h"
 #include "ShowFlags.h"
 
+UMovieGraphDeferredRenderPassNode::UMovieGraphDeferredRenderPassNode()
+	: ViewModeIndex(VMI_Lit)
+{
+}
+
 #if WITH_EDITOR
 FText UMovieGraphDeferredRenderPassNode::GetNodeTitle(const bool bGetDescriptive) const
 {
@@ -20,11 +25,5 @@ FString UMovieGraphDeferredRenderPassNode::GetRendererNameImpl() const
 
 EViewModeIndex UMovieGraphDeferredRenderPassNode::GetViewModeIndex() const
 {
-	return VMI_Lit;
-}
-
-FEngineShowFlags UMovieGraphDeferredRenderPassNode::GetShowFlags() const
-{
-	FEngineShowFlags ShowFlags = FEngineShowFlags(EShowFlagInitMode::ESFIM_Game);
-	return ShowFlags;
+	return ViewModeIndex;
 }
\ No newline at end of file
diff --git a/Engine/Plugins/MovieScene/MovieRenderPipeline/Source/MovieRenderPipelineCore/Private/Graph/Nodes/MovieGraphPathTracedRenderPassNode.cpp b/Engine/Plugins/MovieScene/MovieRenderPipeline/Source/MovieRenderPipelineCore/Private/Graph/Nodes/MovieGraphPathTracedRenderPassNode.cpp
--- a/Engine/Plugins/MovieScene/MovieRenderPipeline/Source/MovieRenderPipelineCore/Private/Graph/Nodes/MovieGraphPathTracedRenderPassNode.cpp
+++ b/Engine/Plugins/MovieScene/MovieRenderPipeline/Source/MovieRenderPipelineCore/Private/Graph/Nodes/MovieGraphPathTracedRenderPassNode.cpp
@@ -5,6 +5,12 @@
 #include "Engine/EngineBaseTypes.h"
 #include "ShowFlags.h"
 
+UMovieGraphPathTracedRenderPassNode::UMovieGraphPathTracedRenderPassNode()
+{
+	ShowFlags->ApplyDefaultShowFlagValue(VMI_PathTracing, true);
+	// TODO: Showflag for SetMotionBlur()?
+}
+
 #if WITH_EDITOR
 FText UMovieGraphPathTracedRenderPassNode::GetNodeTitle(const bool bGetDescriptive) const
 {
@@ -45,12 +51,3 @@ EViewModeIndex UMovieGraphPathTracedRenderPassNode::GetViewModeIndex() const
 {
 	return VMI_PathTracing;
 }
-
-FEngineShowFlags UMovieGraphPathTracedRenderPassNode::GetShowFlags() const
-{
-	FEngineShowFlags ShowFlags = FEngineShowFlags(EShowFlagInitMode::ESFIM_Game);
-	ShowFlags.SetPathTracing(true);
-	// TODO: SetMotionBlur()?
-	
-	return ShowFlags;
-}
\ No newline at end of file
diff --git a/Engine/Plugins/MovieScene/MovieRenderPipeline/Source/MovieRenderPipelineCore/Public/Graph/MovieGraphConfig.h b/Engine/Plugins/MovieScene/MovieRenderPipeline/Source/MovieRenderPipelineCore/Public/Graph/MovieGraphConfig.h
--- a/Engine/Plugins/MovieScene/MovieRenderPipeline/Source/MovieRenderPipelineCore/Public/Graph/MovieGraphConfig.h
+++ b/Engine/Plugins/MovieScene/MovieRenderPipeline/Source/MovieRenderPipelineCore/Public/Graph/MovieGraphConfig.h
@@ -1,11 +1,13 @@
 // Copyright Epic Games, Inc. All Rights Reserved.
 
 #pragma once
+
 #include "CoreMinimal.h"
 #include "EdGraph/EdGraph.h"
 #include "Graph/MovieGraphNode.h"
 #include "Graph/MovieGraphTraversalContext.h"
 #include "MovieGraphValueContainer.h"
+#include "UObject/Interface.h"
 
 #include "MovieGraphConfig.generated.h"
 
@@ -482,6 +484,33 @@ public:
 	TMap<FName, FMovieGraphEvaluatedBranchConfig> BranchConfigMapping;
 };
 
+UINTERFACE(MinimalAPI)
+class UMovieGraphTraversableObject : public UInterface
+{
+	GENERATED_BODY()
+};
+
+/**
+ * Provides a way for objects, which would otherwise not be mergeable during a traversal, to merge in a well-defined way.
+ * Also allows objects to expose which properties have been affected by the merge.
+ */
+class IMovieGraphTraversableObject
+{
+	GENERATED_BODY()
+
+public:
+	/** Merges the contents of InSourceClass into this object. */
+	virtual void Merge(const IMovieGraphTraversableObject* InSourceObject) { }
+
+	/**
+	 * Gets properties, and their associated values, which have been modified by a merge.
+	 * Key = property name, value = stringified value
+	 * The stringified value is a representation of the value which will usually be displayed in the UI. It does not need
+	 * to be a serialized representation.
+	 */
+	virtual TArray<TPair<FString, FString>> GetMergedProperties() const { return {}; }
+};
+
 /**
 * This is the runtime representation of the UMoviePipelineEdGraph which contains the actual strongly
 * typed graph network that is read by the MoviePipeline. There is an editor-only representation of
diff --git a/Engine/Plugins/MovieScene/MovieRenderPipeline/Source/MovieRenderPipelineCore/Public/Graph/Nodes/MovieGraphCoreRenderPassNode.h b/Engine/Plugins/MovieScene/MovieRenderPipeline/Source/MovieRenderPipelineCore/Public/Graph/Nodes/MovieGraphCoreRenderPassNode.h
--- a/Engine/Plugins/MovieScene/MovieRenderPipeline/Source/MovieRenderPipelineCore/Public/Graph/Nodes/MovieGraphCoreRenderPassNode.h
+++ b/Engine/Plugins/MovieScene/MovieRenderPipeline/Source/MovieRenderPipelineCore/Public/Graph/Nodes/MovieGraphCoreRenderPassNode.h
@@ -3,9 +3,12 @@
 #pragma once
 
 #include "Graph/Nodes/MovieGraphRenderPassNode.h"
+
+#include "Graph/MovieGraphConfig.h"
 #include "Graph/MovieGraphDefaultRenderer.h" // For CameraInfo
 #include "Graph/MovieGraphDataTypes.h"
 #include "SceneTypes.h"
+#include "ShowFlags.h"
 #include "Camera/CameraTypes.h"
 #include "Engine/EngineTypes.h"
 #include "Async/TaskGraphFwd.h"
@@ -23,7 +26,6 @@ namespace UE::MovieGraph::DefaultRenderer
 struct FMovieGraphRenderPassSetupData;
 struct FMovieGraphRenderPassLayerData;
 struct FImageOverlappedAccumulator;
-struct FEngineShowFlags;
 enum EViewModeIndex : int;
 
 // For FViewFamilyContextInitData
@@ -51,6 +53,79 @@ namespace UE::MovieGraph
 	void AccumulateSample_TaskThread(TUniquePtr<FImagePixelData>&& InPixelData, const UE::MovieGraph::FMovieGraphSampleState InSampleState, const UE::MovieGraph::FMovieGraphRenderDataAccumulationArgs& InAccumulationParams);
 }
 
+/**
+ * Stores show flag enable/disable state, as well as per-flag override state.
+ *
+ * Here, enable/disable state refers to the actual state of the show flag itself (ie, whether it is turned on or off).
+ * This is the value that will be applied to renders. "Override" state is whether the flag has been marked as overridden
+ * in the UI (ie, whether it is a value that graph traversal should respect). Flags must be marked as overridden in order
+ * for their values to deviate from the default.
+ */
+UCLASS(BlueprintType)
+class MOVIERENDERPIPELINECORE_API UMovieGraphShowFlags : public UObject, public IMovieGraphTraversableObject
+{
+	GENERATED_BODY()
+
+public:
+	UMovieGraphShowFlags();
+
+	/** Gets a copy of the show flags. If they need to be modified, use one of the Set*() methods. */
+	FEngineShowFlags GetShowFlags() const;
+
+	/** Gets the enable/disable state for a specific show flag. */
+	bool IsShowFlagEnabled(const uint32 ShowFlagIndex) const;
+
+	/** Sets the enable/disable state for a specific show flag. The flag must be marked as overridden for this to have an effect. */
+	void SetShowFlagEnabled(const uint32 ShowFlagIndex, const bool bIsShowFlagEnabled);
+
+	/** Gets the override state for a specific show flag. */
+	bool IsShowFlagOverridden(const uint32 ShowFlagIndex) const;
+
+	/** Sets the override state for a specific show flag. */
+	void SetShowFlagOverridden(const uint32 ShowFlagIndex, const bool bIsOverridden);
+
+	/** Gets the indices of all show flags that have been overridden. */
+	const TSet<uint32>& GetOverriddenShowFlags() const;
+
+	/**
+	 * Applies a default value to the show flags. However, the default will not be applied if a user has already
+	 * specified an override for the flag. This should generally only be called from renderers in order for them to
+	 * specify defaults that are relevant.
+	 */
+	void ApplyDefaultShowFlagValue(const uint32 ShowFlagIndex, const bool bShowFlagState);
+
+	// IMovieGraphTraversableObject interface
+	virtual void Merge(const IMovieGraphTraversableObject* InSourceObject) override;
+	virtual TArray<TPair<FString, FString>> GetMergedProperties() const override;
+	// ~IMovieGraphTraversableObject interface
+	
+private:
+	/** Copies the overrides from ShowFlagEnableState to InEngineShowFlags. */
+	void CopyOverridesToShowFlags(FEngineShowFlags& InEngineShowFlags) const;
+
+private:
+	/**
+	 * The show flags which have been marked as overridden in the UI (note this is NOT the enabled/disabled state of the
+	 * flag itself).
+	 */
+	UPROPERTY()
+	TSet<uint32> OverriddenShowFlags;
+
+	/**
+	 * If the flag has been marked as overridden, this stores the enable/disable state of the flag. Key = show flag
+	 * index, value = enable/disable state of the flag.
+	 */
+	UPROPERTY()
+	TMap<uint32, bool> ShowFlagEnableState;
+
+	/**
+	 * The default show flag state. Since FEngineShowFlags itself cannot be serialized (it is not a USTRUCT), we need to
+	 * keep track of sparse overrides (ShowFlagEnableState) on top of the base state that these flags are initialized
+	 * with (ESFIM_Game, plus any additional flags that a specialized renderer specifies).
+	 */
+	FEngineShowFlags ShowFlags;
+};
+
 /** Core functionality for a typical render pass node. */
 UCLASS(Abstract)
 class MOVIERENDERPIPELINECORE_API UMovieGraphCoreRenderPassNode : public UMovieGraphRenderPassNode
@@ -58,6 +133,8 @@ class MOVIERENDERPIPELINECORE_API UMovieGraphCoreRenderPassNode : public UMovieG
 	GENERATED_BODY()
 
 public:
+	UMovieGraphCoreRenderPassNode();
+
 #if WITH_EDITOR
 	virtual FLinearColor GetNodeTitleColor() const override
 	{
@@ -73,6 +150,16 @@ public:
 	}
 #endif
 
+public:
+	// Note: Since *individual* show flags are overridden instead of the entire ShowFlags property, manually set to
+	// overridden so the traversal picks the changes up (otherwise they will be ignored).
+	UPROPERTY(EditAnywhere, BlueprintReadWrite, Category = Overrides, meta = (InlineEditConditionToggle))
+	uint8 bOverride_ShowFlags : 1 = 1;
+
+	/** The show flags that should be active during a render for this node. */
+	UPROPERTY(EditAnywhere, BlueprintReadWrite, Category = "Show Flags")
+	TObjectPtr<UMovieGraphShowFlags> ShowFlags;
+
 protected:
 	// UMovieGraphRenderPassNode Interface
 	virtual void SetupImpl(const FMovieGraphRenderPassSetupData& InSetupData) override;
@@ -87,6 +174,9 @@ protected:
 	/** Gets the show flags that should be active for this renderer. */
 	virtual FEngineShowFlags GetShowFlags() const;
 
+	/** Gets any show flags that should be applied as defaults, before user changes are applied. */
+	virtual TArray<uint32> GetDefaultShowFlags() const;
+
 	static void AddReferencedObjects(UObject* InThis, FReferenceCollector& Collector);
 
 	struct FViewFamilyContextInitData
diff --git a/Engine/Plugins/MovieScene/MovieRenderPipeline/Source/MovieRenderPipelineCore/Public/Graph/Nodes/MovieGraphDefRenderPassNode.h b/Engine/Plugins/MovieScene/MovieRenderPipeline/Source/MovieRenderPipelineCore/Public/Graph/Nodes/MovieGraphDefRenderPassNode.h
--- a/Engine/Plugins/MovieScene/MovieRenderPipeline/Source/MovieRenderPipelineCore/Public/Graph/Nodes/MovieGraphDefRenderPassNode.h
+++ b/Engine/Plugins/MovieScene/MovieRenderPipeline/Source/MovieRenderPipelineCore/Public/Graph/Nodes/MovieGraphDefRenderPassNode.h
@@ -13,6 +13,8 @@ class MOVIERENDERPIPELINECORE_API UMovieGraphDeferredRenderPassNode : public UMo
 	GENERATED_BODY()
 
 public:
+	UMovieGraphDeferredRenderPassNode();
+
 #if WITH_EDITOR
 	virtual FText GetNodeTitle(const bool bGetDescriptive = false) const override;
 #endif
@@ -24,6 +26,13 @@ protected:
 
 	// UMovieGraphCoreRenderPassNode Interface
 	virtual EViewModeIndex GetViewModeIndex() const override;
-	virtual FEngineShowFlags GetShowFlags() const override;
 	// ~UMovieGraphCoreRenderPassNode Interface
+
+public:
+	UPROPERTY(EditAnywhere, BlueprintReadWrite, Category = Overrides, meta = (InlineEditConditionToggle))
+	uint8 bOverride_ViewModeIndex : 1;
+
+	/** The view mode index that will be applied to renders. */
+	UPROPERTY(EditAnywhere, BlueprintReadWrite, Category = "Settings", meta=(EditCondition="bOverride_ViewModeIndex", InvalidEnumValues = "VMI_PathTracing"))
+	TEnumAsByte<EViewModeIndex> ViewModeIndex;
 };
\ No newline at end of file
diff --git a/Engine/Plugins/MovieScene/MovieRenderPipeline/Source/MovieRenderPipelineCore/Public/Graph/Nodes/MovieGraphPathTracedRenderPassNode.h b/Engine/Plugins/MovieScene/MovieRenderPipeline/Source/MovieRenderPipelineCore/Public/Graph/Nodes/MovieGraphPathTracedRenderPassNode.h
--- a/Engine/Plugins/MovieScene/MovieRenderPipeline/Source/MovieRenderPipelineCore/Public/Graph/Nodes/MovieGraphPathTracedRenderPassNode.h
+++ b/Engine/Plugins/MovieScene/MovieRenderPipeline/Source/MovieRenderPipelineCore/Public/Graph/Nodes/MovieGraphPathTracedRenderPassNode.h
@@ -13,6 +13,8 @@ class MOVIERENDERPIPELINECORE_API UMovieGraphPathTracedRenderPassNode : public U
 	GENERATED_BODY()
 
 public:
+	UMovieGraphPathTracedRenderPassNode();
+
 #if WITH_EDITOR
 	virtual FText GetNodeTitle(const bool bGetDescriptive = false) const override;
 #endif
@@ -29,7 +31,6 @@ protected:
 
 	// UMovieGraphCoreRenderPassNode Interface
 	virtual EViewModeIndex GetViewModeIndex() const override;
-	virtual FEngineShowFlags GetShowFlags() const override;
 	// ~UMovieGraphCoreRenderPassNode Interface
 
 private:
diff --git a/Engine/Plugins/MovieScene/MovieRenderPipeline/Source/MovieRenderPipelineEditor/Private/Customizations/Graph/MovieGraphShowFlagsCustomization.h b/Engine/Plugins/MovieScene/MovieRenderPipeline/Source/MovieRenderPipelineEditor/Private/Customizations/Graph/MovieGraphShowFlagsCustomization.h
new file mode 100644
--- /dev/null
+++ b/Engine/Plugins/MovieScene/MovieRenderPipeline/Source/MovieRenderPipelineEditor/Private/Customizations/Graph/MovieGraphShowFlagsCustomization.h
@@ -0,0 +1,131 @@
+﻿// Copyright Epic Games, Inc. All Rights Reserved.
+
+#pragma once
+
+#include "DetailWidgetRow.h"
+#include "EditorShowFlags.h"
+#include "Graph/Nodes/MovieGraphCoreRenderPassNode.h"
+#include "IDetailChildrenBuilder.h"
+#include "IDetailGroup.h"
+#include "PropertyHandle.h"
+#include "Widgets/Input/SCheckBox.h"
+
+#define LOCTEXT_NAMESPACE "MoviePipelineEditor"
+
+/** Customize how show flags appear in the details panel. */
+class FMovieGraphShowFlagsCustomization : public IPropertyTypeCustomization
+{
+public:
+	static TSharedRef<IPropertyTypeCustomization> MakeInstance()
+	{
+		return MakeShared<FMovieGraphShowFlagsCustomization>();
+	}
+
+protected:
+	//~ Begin IPropertyTypeCustomization interface
+	virtual void CustomizeHeader(TSharedRef<IPropertyHandle> PropertyHandle, FDetailWidgetRow& HeaderRow, IPropertyTypeCustomizationUtils& CustomizationUtils) override {}
+	
+	virtual void CustomizeChildren(TSharedRef<IPropertyHandle> PropertyHandle, IDetailChildrenBuilder& ChildBuilder, IPropertyTypeCustomizationUtils& CustomizationUtils) override
+	{
+		static const TMap<EShowFlagGroup, FText> GroupNames =
+		{
+			{SFG_Normal, LOCTEXT("GraphNormalSF", "Normal")},
+			{SFG_PostProcess, LOCTEXT("GraphPostProcessSF", "Post Processing")},
+			{SFG_LightTypes, LOCTEXT("GraphLightTypesSF", "Light Types")},
+			{SFG_LightingComponents, LOCTEXT("GraphLightingComponentsSF", "Lighting Components")},
+			{SFG_LightingFeatures, LOCTEXT("GraphLightingFeaturesSF", "Lighting Features")},
+			{SFG_Lumen, LOCTEXT("GraphLumenSF", "Lumen")},
+			{SFG_Nanite, LOCTEXT("GraphNaniteSF", "Nanite")},
+			{SFG_Developer, LOCTEXT("GraphDeveloperSF", "Developer")},
+			{SFG_Visualize, LOCTEXT("GraphVisualizeSF", "Visualize")},
+			{SFG_Advanced, LOCTEXT("GraphAdvancedSF", "Advanced")}
+		};
+
+		TArray<UObject*> OuterObjects;
+		PropertyHandle->GetOuterObjects(OuterObjects);
+		
+		check(PropertyHandle->IsValidHandle());
+
+		const FObjectProperty* ObjectProperty = CastField<FObjectProperty>(PropertyHandle->GetProperty());
+		UMovieGraphShowFlags* ShowFlagObject = Cast<UMovieGraphShowFlags>(
+			ObjectProperty->GetObjectPropertyValue(ObjectProperty->ContainerPtrToValuePtr<void>(OuterObjects[0])));
+
+		// Group together the show flags with the groups they show up under in the UI
+		TMap<EShowFlagGroup, TArray<FShowFlagData>> GroupedShowFlags;
+		for (FShowFlagData& ShowFlag : GetShowFlagMenuItems())
+		{
+			TArray<FShowFlagData>& GroupShowFlags = GroupedShowFlags.FindOrAdd(ShowFlag.Group);
+			GroupShowFlags.Add(MoveTemp(ShowFlag));
+		}
+
+		// Make each show flag group a group in the UI, and each show flag a row under its respective group
+		for (const TTuple<EShowFlagGroup, FText>& ShowFlagGroup : GroupNames)
+		{
+			const EShowFlagGroup& Group = ShowFlagGroup.Key;
+			const FText& GroupName = ShowFlagGroup.Value;
+			
+			const TArray<FShowFlagData>* ShowFlags = GroupedShowFlags.Find(Group);
+			if (!ShowFlags)
+			{
+				continue;
+			}
+			
+			IDetailGroup& FlagGroup = ChildBuilder.AddGroup(FName(GroupName.ToString()), GroupName);
+
+			for (const FShowFlagData& ShowFlag : *ShowFlags)
+			{
+				const uint32 ShowFlagIndex = ShowFlag.EngineShowFlagIndex;
+				
+				FlagGroup.AddWidgetRow()
+				.NameContent()
+				[
+					SNew(SHorizontalBox)
+					+ SHorizontalBox::Slot()
+					.VAlign(VAlign_Center)
+					.AutoWidth()
+					.Padding(0, 0, 1, 0)
+					[
+						SNew(SCheckBox)
+						.IsChecked_Lambda([ShowFlagObject, ShowFlagIndex]()
+						{
+							return ShowFlagObject->IsShowFlagOverridden(ShowFlagIndex)
+								? ECheckBoxState::Checked : ECheckBoxState::Unchecked;
+						})
+						.OnCheckStateChanged_Lambda([ShowFlagObject, ShowFlagIndex](const ECheckBoxState NewState)
+						{
+							const bool bIsOverridden = (NewState == ECheckBoxState::Checked);
+							ShowFlagObject->SetShowFlagOverridden(ShowFlagIndex, bIsOverridden);
+						})
+					]
+
+					+ SHorizontalBox::Slot()
+					.VAlign(VAlign_Center)
+					[
+						SNew(STextBlock)
+						.IsEnabled_Lambda([ShowFlagObject, ShowFlagIndex]() { return ShowFlagObject->IsShowFlagOverridden(ShowFlagIndex); })
+						.Text(ShowFlag.DisplayName)
+						.Font(CustomizationUtils.GetRegularFont())
+					]
+				]
+				.ValueContent()
+				[
+					SNew(SCheckBox)
+					.IsEnabled_Lambda([ShowFlagObject, ShowFlagIndex]() { return ShowFlagObject->IsShowFlagOverridden(ShowFlagIndex); })
+					.IsChecked_Lambda([ShowFlagObject, ShowFlagIndex]()
+					{
+						return ShowFlagObject->IsShowFlagEnabled(ShowFlagIndex)
+							? ECheckBoxState::Checked : ECheckBoxState::Unchecked;
+					})
+					.OnCheckStateChanged_Lambda([ShowFlagObject, ShowFlagIndex](const ECheckBoxState NewState)
+					{
+						const bool bIsUsed = (NewState == ECheckBoxState::Checked);
+						ShowFlagObject->SetShowFlagEnabled(ShowFlagIndex, bIsUsed);
+					})
+				];
+			}
+		}
+	}
+	//~ End IPropertyTypeCustomization interface
+};
+
+#undef LOCTEXT_NAMESPACE
diff --git a/Engine/Plugins/MovieScene/MovieRenderPipeline/Source/MovieRenderPipelineEditor/Private/Graph/MovieGraphAssetToolkit.cpp b/Engine/Plugins/MovieScene/MovieRenderPipeline/Source/MovieRenderPipelineEditor/Private/Graph/MovieGraphAssetToolkit.cpp
--- a/Engine/Plugins/MovieScene/MovieRenderPipeline/Source/MovieRenderPipelineEditor/Private/Graph/MovieGraphAssetToolkit.cpp
+++ b/Engine/Plugins/MovieScene/MovieRenderPipeline/Source/MovieRenderPipelineEditor/Private/Graph/MovieGraphAssetToolkit.cpp
@@ -4,8 +4,10 @@
 
 #include "Customizations/Graph/GraphMemberCustomization.h"
 #include "Customizations/Graph/GraphNodeCustomization.h"
+#include "Customizations/Graph/MovieGraphShowFlagsCustomization.h"
 #include "Framework/Commands/GenericCommands.h"
 #include "Graph/MovieGraphConfig.h"
+#include "Graph/Nodes/MovieGraphCoreRenderPassNode.h"
 #include "MovieEdGraphNode.h"
 #include "MovieGraphSchema.h"
 #include "PropertyEditorModule.h"
@@ -230,6 +232,10 @@ TSharedRef<SDockTab> FMovieGraphAssetToolkit::SpawnTab_RenderGraphDetails(const
 	SelectedGraphObjectsDetailsWidget->RegisterInstancedCustomPropertyLayout(
 		UMovieGraphNode::StaticClass(),
 		FOnGetDetailCustomizationInstance::CreateStatic(&FGraphNodeCustomization::MakeInstance));
+
+	SelectedGraphObjectsDetailsWidget->RegisterInstancedCustomPropertyTypeLayout(
+		UMovieGraphShowFlags::StaticClass()->GetFName(),
+		FOnGetPropertyTypeCustomizationInstance::CreateStatic(&FMovieGraphShowFlagsCustomization::MakeInstance));
 	
 	return SNew(SDockTab)
 		.TabColorScale(GetTabColorScale())
diff --git a/Engine/Plugins/MovieScene/MovieRenderPipeline/Source/MovieRenderPipelineEditor/Private/Graph/SMovieGraphActiveRenderSettingsTabContent.cpp b/Engine/Plugins/MovieScene/MovieRenderPipeline/Source/MovieRenderPipelineEditor/Private/Graph/SMovieGraphActiveRenderSettingsTabContent.cpp
--- a/Engine/Plugins/MovieScene/MovieRenderPipeline/Source/MovieRenderPipelineEditor/Private/Graph/SMovieGraphActiveRenderSettingsTabContent.cpp
+++ b/Engine/Plugins/MovieScene/MovieRenderPipeline/Source/MovieRenderPipelineEditor/Private/Graph/SMovieGraphActiveRenderSettingsTabContent.cpp
@@ -12,6 +12,7 @@
 #include "MoviePipelineQueue.h"
 #include "MoviePipelineQueueSubsystem.h"
 #include "SPositiveActionButton.h"
+#include "Graph/Nodes/MovieGraphCoreRenderPassNode.h"
 #include "Widgets/Input/SComboBox.h"
 
 #define LOCTEXT_NAMESPACE "SMoviePipelineActiveRenderSettingsTabContent"
@@ -47,6 +48,23 @@ FString FActiveRenderSettingsTreeElement::GetValue() const
 		}
 	}
 
+	// If the property implements IMovieGraphTraversableObject, get the value via GetMergedProperties().
+	// Value will be formatted as a newline-delimited list of "PropertyName = PropertyValue" strings.
+	const FObjectProperty* ObjectProperty = CastField<FObjectProperty>(SettingsProperty);
+	if (ObjectProperty && ObjectProperty->PropertyClass->ImplementsInterface(UMovieGraphTraversableObject::StaticClass()))
+	{
+		const IMovieGraphTraversableObject* MergeableProperty = Cast<IMovieGraphTraversableObject>(
+			ObjectProperty->GetObjectPropertyValue(ObjectProperty->ContainerPtrToValuePtr<void>(SettingsNode)));
+		
+		TArray<FString> MergedValues;
+		Algo::Transform(MergeableProperty->GetMergedProperties(), MergedValues, [](const TPair<FString, FString>& Pair)
+		{
+			return FString::Format(TEXT("{0} = {1}"), {Pair.Key, Pair.Value});
+		});
+		
+		return FString::Join(MergedValues, TEXT("\n"));
+	}
+
 	// Otherwise, ask the property for its value directly
 	SettingsProperty->ExportTextItem_InContainer(ValueString, SettingsNode, nullptr, nullptr, PPF_None);
 	return ValueString;
@@ -242,7 +260,7 @@ TSharedRef<SWidget> SMovieGraphActiveRenderSettingsTreeItem::GenerateWidgetForCo
 
 			+ SHorizontalBox::Slot()
 			.AutoWidth()
-			.VAlign(VAlign_Center)
+			.VAlign(VAlign_Top)
 			.HAlign(HAlign_Right)
 			[
 				SNew(STextBlock)
