MoviePipeline: Various file/class renames that need to be made. No logic changes.

#rnx

diff --git a/Engine/Plugins/MovieScene/MovieRenderPipeline/Source/MovieRenderPipelineCore/Private/Graph/MovieGraphLinTimeStep.cpp b/Engine/Plugins/MovieScene/MovieRenderPipeline/Source/MovieRenderPipelineCore/Private/Graph/MovieGraphLinTimeStep.cpp
deleted file mode 100644
diff --git a/Engine/Plugins/MovieScene/MovieRenderPipeline/Source/MovieRenderPipelineCore/Private/Graph/MovieGraphLinearTimeStep.cpp b/Engine/Plugins/MovieScene/MovieRenderPipeline/Source/MovieRenderPipelineCore/Private/Graph/MovieGraphLinearTimeStep.cpp
new file mode 100644
--- /dev/null
+++ b/Engine/Plugins/MovieScene/MovieRenderPipeline/Source/MovieRenderPipelineCore/Private/Graph/MovieGraphLinearTimeStep.cpp
@@ -0,0 +1,20 @@
+﻿// Copyright Epic Games, Inc. All Rights Reserved.
+
+#include "Graph/MovieGraphLinearTimeStep.h"
+
+#include "Graph/Nodes/MovieGraphSamplingMethodNode.h"
+
+int32 UMovieGraphLinearTimeStep::GetNextTemporalRangeIndex() const
+{
+	// Linear time step just steps through the temporal ranges in order
+	return CurrentFrameData.TemporalSampleIndex;
+}
+
+int32 UMovieGraphLinearTimeStep::GetTemporalSampleCount() const
+{
+	constexpr bool bIncludeCDOs = true;
+	const UMovieGraphSamplingMethodNode* SamplingMethod =
+		CurrentFrameData.EvaluatedConfig->GetSettingForBranch<UMovieGraphSamplingMethodNode>(UMovieGraphNode::GlobalsPinName, bIncludeCDOs);
+
+	return SamplingMethod->TemporalSampleCount;
+}
diff --git a/Engine/Plugins/MovieScene/MovieRenderPipeline/Source/MovieRenderPipelineCore/Private/Graph/MovieGraphPipeline.cpp b/Engine/Plugins/MovieScene/MovieRenderPipeline/Source/MovieRenderPipelineCore/Private/Graph/MovieGraphPipeline.cpp
--- a/Engine/Plugins/MovieScene/MovieRenderPipeline/Source/MovieRenderPipelineCore/Private/Graph/MovieGraphPipeline.cpp
+++ b/Engine/Plugins/MovieScene/MovieRenderPipeline/Source/MovieRenderPipelineCore/Private/Graph/MovieGraphPipeline.cpp
@@ -4,7 +4,7 @@
 
 #include "Graph/MovieGraphCVarManager.h"
 #include "Graph/MovieGraphDataTypes.h"
-#include "Graph/MovieGraphLinTimeStep.h"
+#include "Graph/MovieGraphLinearTimeStep.h"
 #include "Graph/MovieGraphOutputMerger.h"
 #include "Graph/MoviePipelineRenderLayerSubsystem.h"
 #include "Graph/Nodes/MovieGraphCollectionNode.h"
diff --git a/Engine/Plugins/MovieScene/MovieRenderPipeline/Source/MovieRenderPipelineCore/Private/Graph/Nodes/MovieGraphDefRenderPassNode.cpp b/Engine/Plugins/MovieScene/MovieRenderPipeline/Source/MovieRenderPipelineCore/Private/Graph/Nodes/MovieGraphDefRenderPassNode.cpp
deleted file mode 100644
diff --git a/Engine/Plugins/MovieScene/MovieRenderPipeline/Source/MovieRenderPipelineCore/Private/Graph/Nodes/MovieGraphDeferredRenderPassNode.cpp b/Engine/Plugins/MovieScene/MovieRenderPipeline/Source/MovieRenderPipelineCore/Private/Graph/Nodes/MovieGraphDeferredRenderPassNode.cpp
new file mode 100644
--- /dev/null
+++ b/Engine/Plugins/MovieScene/MovieRenderPipeline/Source/MovieRenderPipelineCore/Private/Graph/Nodes/MovieGraphDeferredRenderPassNode.cpp
@@ -0,0 +1,29 @@
+﻿// Copyright Epic Games, Inc. All Rights Reserved.
+
+#include "Graph/Nodes/MovieGraphDeferredRenderPassNode.h"
+
+#include "Engine/EngineBaseTypes.h"
+#include "ShowFlags.h"
+
+UMovieGraphDeferredRenderPassNode::UMovieGraphDeferredRenderPassNode()
+	: ViewModeIndex(VMI_Lit)
+{
+}
+
+#if WITH_EDITOR
+FText UMovieGraphDeferredRenderPassNode::GetNodeTitle(const bool bGetDescriptive) const
+{
+	return NSLOCTEXT("MovieGraphNodes", "DeferredRenderPassGraphNode_Description", "Deferred Renderer");
+}
+#endif
+
+FString UMovieGraphDeferredRenderPassNode::GetRendererNameImpl() const
+{
+	static const FString RendererNameImpl(TEXT("Deferred"));
+	return RendererNameImpl;
+}
+
+EViewModeIndex UMovieGraphDeferredRenderPassNode::GetViewModeIndex() const
+{
+	return ViewModeIndex;
+}
\ No newline at end of file
diff --git a/Engine/Plugins/MovieScene/MovieRenderPipeline/Source/MovieRenderPipelineCore/Private/Graph/Nodes/MovieGraphSamplingMethodNode.cpp b/Engine/Plugins/MovieScene/MovieRenderPipeline/Source/MovieRenderPipelineCore/Private/Graph/Nodes/MovieGraphSamplingMethodNode.cpp
--- a/Engine/Plugins/MovieScene/MovieRenderPipeline/Source/MovieRenderPipelineCore/Private/Graph/Nodes/MovieGraphSamplingMethodNode.cpp
+++ b/Engine/Plugins/MovieScene/MovieRenderPipeline/Source/MovieRenderPipelineCore/Private/Graph/Nodes/MovieGraphSamplingMethodNode.cpp
@@ -3,7 +3,7 @@
 #include "Graph/Nodes/MovieGraphSamplingMethodNode.h"
 
 #include "Graph/MovieGraphConfig.h"
-#include "Graph/MovieGraphLinTimeStep.h"
+#include "Graph/MovieGraphLinearTimeStep.h"
 #include "Styling/AppStyle.h"
 
 UMovieGraphSamplingMethodNode::UMovieGraphSamplingMethodNode()
diff --git a/Engine/Plugins/MovieScene/MovieRenderPipeline/Source/MovieRenderPipelineCore/Public/Graph/MovieGraphLinTimeStep.h b/Engine/Plugins/MovieScene/MovieRenderPipeline/Source/MovieRenderPipelineCore/Public/Graph/MovieGraphLinTimeStep.h
deleted file mode 100644
diff --git a/Engine/Plugins/MovieScene/MovieRenderPipeline/Source/MovieRenderPipelineCore/Public/Graph/MovieGraphLinearTimeStep.h b/Engine/Plugins/MovieScene/MovieRenderPipeline/Source/MovieRenderPipelineCore/Public/Graph/MovieGraphLinearTimeStep.h
new file mode 100644
--- /dev/null
+++ b/Engine/Plugins/MovieScene/MovieRenderPipeline/Source/MovieRenderPipelineCore/Public/Graph/MovieGraphLinearTimeStep.h
@@ -0,0 +1,25 @@
+﻿// Copyright Epic Games, Inc. All Rights Reserved.
+
+#pragma once
+
+#include "MovieGraphCoreTimeStep.h"
+
+#include "MovieGraphLinearTimeStep.generated.h"
+
+/**
+ * Advances time forward linearly until the end of the range of time that is being rendered is reached. This is useful
+ * for deferred rendering (where there's a small number of temporal sub-samples and no feedback mechanism for measuring
+ * noise in the final image).
+ */
+UCLASS(BlueprintType, meta = (DisplayName = "Linear Time Step"))
+class MOVIERENDERPIPELINECORE_API UMovieGraphLinearTimeStep : public UMovieGraphCoreTimeStep
+{
+	GENERATED_BODY()
+
+public:
+	UMovieGraphLinearTimeStep() = default;
+
+protected:
+	virtual int32 GetNextTemporalRangeIndex() const override;
+	virtual int32 GetTemporalSampleCount() const override;
+};
diff --git a/Engine/Plugins/MovieScene/MovieRenderPipeline/Source/MovieRenderPipelineCore/Public/Graph/Nodes/MovieGraphDefRenderPassNode.h b/Engine/Plugins/MovieScene/MovieRenderPipeline/Source/MovieRenderPipelineCore/Public/Graph/Nodes/MovieGraphDefRenderPassNode.h
deleted file mode 100644
diff --git a/Engine/Plugins/MovieScene/MovieRenderPipeline/Source/MovieRenderPipelineCore/Public/Graph/Nodes/MovieGraphDeferredRenderPassNode.h b/Engine/Plugins/MovieScene/MovieRenderPipeline/Source/MovieRenderPipelineCore/Public/Graph/Nodes/MovieGraphDeferredRenderPassNode.h
new file mode 100644
--- /dev/null
+++ b/Engine/Plugins/MovieScene/MovieRenderPipeline/Source/MovieRenderPipelineCore/Public/Graph/Nodes/MovieGraphDeferredRenderPassNode.h
@@ -0,0 +1,38 @@
+﻿// Copyright Epic Games, Inc. All Rights Reserved.
+
+#pragma once
+
+#include "Graph/Nodes/MovieGraphCoreRenderPassNode.h"
+
+#include "MovieGraphDeferredRenderPassNode.generated.h"
+
+/** A render node which uses the deferred renderer. */
+UCLASS()
+class MOVIERENDERPIPELINECORE_API UMovieGraphDeferredRenderPassNode : public UMovieGraphCoreRenderPassNode
+{
+	GENERATED_BODY()
+
+public:
+	UMovieGraphDeferredRenderPassNode();
+
+#if WITH_EDITOR
+	virtual FText GetNodeTitle(const bool bGetDescriptive = false) const override;
+#endif
+
+protected:
+	// UMovieGraphRenderPassNode Interface
+	virtual FString GetRendererNameImpl() const override;
+	// ~UMovieGraphRenderPassNode Interface
+
+	// UMovieGraphCoreRenderPassNode Interface
+	virtual EViewModeIndex GetViewModeIndex() const override;
+	// ~UMovieGraphCoreRenderPassNode Interface
+
+public:
+	UPROPERTY(EditAnywhere, BlueprintReadWrite, Category = Overrides, meta = (InlineEditConditionToggle))
+	uint8 bOverride_ViewModeIndex : 1;
+
+	/** The view mode index that will be applied to renders. */
+	UPROPERTY(EditAnywhere, BlueprintReadWrite, Category = "Settings", meta=(EditCondition="bOverride_ViewModeIndex", InvalidEnumValues = "VMI_PathTracing"))
+	TEnumAsByte<EViewModeIndex> ViewModeIndex;
+};
\ No newline at end of file
diff --git a/Engine/Plugins/MovieScene/MovieRenderPipeline/Source/MovieRenderPipelineEditor/Private/Customizations/Graph/GraphMemberCustomization.h b/Engine/Plugins/MovieScene/MovieRenderPipeline/Source/MovieRenderPipelineEditor/Private/Customizations/Graph/GraphMemberCustomization.h
deleted file mode 100644
diff --git a/Engine/Plugins/MovieScene/MovieRenderPipeline/Source/MovieRenderPipelineEditor/Private/Customizations/Graph/GraphNodeCustomization.h b/Engine/Plugins/MovieScene/MovieRenderPipeline/Source/MovieRenderPipelineEditor/Private/Customizations/Graph/GraphNodeCustomization.h
deleted file mode 100644
diff --git a/Engine/Plugins/MovieScene/MovieRenderPipeline/Source/MovieRenderPipelineEditor/Private/Customizations/Graph/MovieGraphMemberCustomization.h b/Engine/Plugins/MovieScene/MovieRenderPipeline/Source/MovieRenderPipelineEditor/Private/Customizations/Graph/MovieGraphMemberCustomization.h
new file mode 100644
--- /dev/null
+++ b/Engine/Plugins/MovieScene/MovieRenderPipeline/Source/MovieRenderPipelineEditor/Private/Customizations/Graph/MovieGraphMemberCustomization.h
@@ -0,0 +1,115 @@
+// Copyright Epic Games, Inc. All Rights Reserved.
+
+#pragma once
+
+#include "DetailCategoryBuilder.h"
+#include "DetailLayoutBuilder.h"
+#include "DetailWidgetRow.h"
+#include "IDetailCustomization.h"
+#include "Graph/MovieGraphConfig.h"
+#include "PropertyHandle.h"
+#include "Widgets/Input/SEditableTextBox.h"
+#include "Widgets/Text/STextBlock.h"
+
+#define LOCTEXT_NAMESPACE "MoviePipelineEditor"
+
+/** Customize how members for a graph appear in the details panel. */
+class FMovieGraphMemberCustomization : public IDetailCustomization
+{
+public:
+	static TSharedRef<IDetailCustomization> MakeInstance()
+	{
+		return MakeShared<FMovieGraphMemberCustomization>();
+	}
+
+protected:
+	//~ Begin IDetailCustomization interface
+	virtual void CustomizeDetails(const TSharedPtr<IDetailLayoutBuilder>& DetailBuilder) override
+	{
+		CustomizeDetails(*DetailBuilder);
+	}
+
+	void OnNameChanged(const FText& InText, UMovieGraphMember* MovieGraphMember) const
+	{
+		NameEditableTextBox->SetError(FText::GetEmpty());
+
+		FText Error;
+		if (!MovieGraphMember->CanRename(InText, Error))
+		{
+			NameEditableTextBox->SetError(Error);
+		}
+	}
+	
+	void OnNameCommitted(const FText& InText, ETextCommit::Type Arg, UMovieGraphMember* MovieGraphMember) const
+	{
+		if (MovieGraphMember->SetMemberName(InText.ToString()))
+		{
+			NameEditableTextBox->SetError(FText::GetEmpty());
+		}
+	}
+
+	virtual void CustomizeDetails(IDetailLayoutBuilder& DetailBuilder) override
+	{
+		TArray<TWeakObjectPtr<UObject>> ObjectsBeingCustomized;
+		DetailBuilder.GetObjectsBeingCustomized(ObjectsBeingCustomized);
+
+		for (const TWeakObjectPtr<UObject>& CustomizedObject : ObjectsBeingCustomized)
+		{
+			// Note: The graph members inherit the Value property from a base class, so the enable/disable state cannot
+			// be driven by UPROPERTY metadata. Hence why this needs to be done w/ a details customization.
+			
+			// Enable/disable the value property for inputs/outputs based on whether it is specified as a branch or not
+			if (const UMovieGraphInterfaceBase* InterfaceBase = Cast<UMovieGraphInterfaceBase>(CustomizedObject))
+			{
+				const TSharedRef<IPropertyHandle> ValueProperty = DetailBuilder.GetProperty("Value", UMovieGraphValueContainer::StaticClass());
+				if (ValueProperty->IsValidHandle())
+				{
+					DetailBuilder.EditDefaultProperty(ValueProperty)->IsEnabled(!InterfaceBase->bIsBranch);
+				}
+			}
+
+			// Enable/disable the value property for variables based on the editable state
+			if (const UMovieGraphVariable* Variable = Cast<UMovieGraphVariable>(CustomizedObject))
+			{
+				const TSharedRef<IPropertyHandle> ValueProperty = DetailBuilder.GetProperty("Value", UMovieGraphValueContainer::StaticClass());
+				if (ValueProperty->IsValidHandle())
+				{
+					DetailBuilder.EditDefaultProperty(ValueProperty)->IsEnabled(Variable->IsEditable());
+				}
+			}
+
+			UMovieGraphMember* MemberObject = Cast<UMovieGraphMember>(CustomizedObject);
+			if (!MemberObject)
+			{
+				continue;
+			}
+
+			// Add a custom row for the Name property (to allow for proper validation)
+			IDetailCategoryBuilder& GeneralCategory = DetailBuilder.EditCategory("General");
+			GeneralCategory.AddCustomRow(FText::GetEmpty())
+			.NameContent()
+			[
+				SNew(STextBlock)
+				.Text(LOCTEXT("MemberPropertyLabel_Name", "Name"))
+				.Font(DetailBuilder.GetDetailFont())
+			]
+			.ValueContent()
+			[
+				SAssignNew(NameEditableTextBox, SEditableTextBox)
+				.Text_Lambda([MemberObject]() { return FText::FromString(MemberObject->GetMemberName()); })
+				.OnTextChanged(this, &FMovieGraphMemberCustomization::OnNameChanged, MemberObject)
+				.OnTextCommitted(this, &FMovieGraphMemberCustomization::OnNameCommitted, MemberObject)
+				.IsReadOnly_Lambda([MemberObject]() { return !MemberObject->IsEditable(); })
+				.SelectAllTextWhenFocused(true)
+				.Font(IDetailLayoutBuilder::GetDetailFont())
+			];
+		}
+	}
+	//~ End IDetailCustomization interface
+
+private:
+	/** Text box for the "Name" property. */
+	TSharedPtr<SEditableTextBox> NameEditableTextBox;
+};
+
+#undef LOCTEXT_NAMESPACE
\ No newline at end of file
diff --git a/Engine/Plugins/MovieScene/MovieRenderPipeline/Source/MovieRenderPipelineEditor/Private/Customizations/Graph/MovieGraphNodeCustomization.h b/Engine/Plugins/MovieScene/MovieRenderPipeline/Source/MovieRenderPipelineEditor/Private/Customizations/Graph/MovieGraphNodeCustomization.h
new file mode 100644
--- /dev/null
+++ b/Engine/Plugins/MovieScene/MovieRenderPipeline/Source/MovieRenderPipelineEditor/Private/Customizations/Graph/MovieGraphNodeCustomization.h
@@ -0,0 +1,57 @@
+// Copyright Epic Games, Inc. All Rights Reserved.
+
+#pragma once
+
+#include "DetailLayoutBuilder.h"
+#include "IDetailCustomization.h"
+#include "Graph/MovieGraphConfig.h"
+#include "PropertyHandle.h"
+
+#define LOCTEXT_NAMESPACE "MoviePipelineEditor"
+
+/** Customize how graph nodes appear in the details panel. */
+class FMovieGraphNodeCustomization : public IDetailCustomization
+{
+public:
+	static TSharedRef<IDetailCustomization> MakeInstance()
+	{
+		return MakeShared<FMovieGraphNodeCustomization>();
+	}
+
+protected:
+	//~ Begin IDetailCustomization interface
+	virtual void CustomizeDetails(const TSharedPtr<IDetailLayoutBuilder>& DetailBuilder) override
+	{
+		CustomizeDetails(*DetailBuilder);
+	}
+
+	virtual void CustomizeDetails(IDetailLayoutBuilder& DetailBuilder) override
+	{
+		TArray<TWeakObjectPtr<UObject>> ObjectsBeingCustomized;
+		DetailBuilder.GetObjectsBeingCustomized(ObjectsBeingCustomized);
+
+		// Hide the "Properties" category, which houses the dynamic properties, if there are no dynamic properties
+		// on the node. An empty category in the details panel is messy/confusing.
+		for (const TWeakObjectPtr<UObject>& CustomizedObject : ObjectsBeingCustomized)
+		{
+			if (const UMovieGraphNode* Node = Cast<UMovieGraphNode>(CustomizedObject))
+			{
+				if (Node->GetDynamicPropertyDescriptions().IsEmpty())
+				{
+					DetailBuilder.HideCategory("Properties");
+					break;
+				}
+			}
+		}
+
+		// Also hide the "Properties" category if more than one object is selected. There needs to be some more work done
+		// in property bags to be able to handle this condition without crashing in many cases.
+		if (ObjectsBeingCustomized.Num() > 1)
+		{
+			DetailBuilder.HideCategory("Properties");
+		}
+	}
+	//~ End IDetailCustomization interface
+};
+
+#undef LOCTEXT_NAMESPACE
\ No newline at end of file
diff --git a/Engine/Plugins/MovieScene/MovieRenderPipeline/Source/MovieRenderPipelineEditor/Private/Graph/MovieGraphAssetToolkit.cpp b/Engine/Plugins/MovieScene/MovieRenderPipeline/Source/MovieRenderPipelineEditor/Private/Graph/MovieGraphAssetToolkit.cpp
--- a/Engine/Plugins/MovieScene/MovieRenderPipeline/Source/MovieRenderPipelineEditor/Private/Graph/MovieGraphAssetToolkit.cpp
+++ b/Engine/Plugins/MovieScene/MovieRenderPipeline/Source/MovieRenderPipelineEditor/Private/Graph/MovieGraphAssetToolkit.cpp
@@ -2,8 +2,8 @@
 
 #include "MovieGraphAssetToolkit.h"
 
-#include "Customizations/Graph/GraphMemberCustomization.h"
-#include "Customizations/Graph/GraphNodeCustomization.h"
+#include "Customizations/Graph/MovieGraphMemberCustomization.h"
+#include "Customizations/Graph/MovieGraphNodeCustomization.h"
 #include "Customizations/Graph/MovieGraphShowFlagsCustomization.h"
 #include "Framework/Commands/GenericCommands.h"
 #include "Graph/MovieGraphConfig.h"
@@ -227,11 +227,11 @@ TSharedRef<SDockTab> FMovieGraphAssetToolkit::SpawnTab_RenderGraphDetails(const
 
 	SelectedGraphObjectsDetailsWidget->RegisterInstancedCustomPropertyLayout(
 		UMovieGraphMember::StaticClass(),
-		FOnGetDetailCustomizationInstance::CreateStatic(&FGraphMemberCustomization::MakeInstance));
+		FOnGetDetailCustomizationInstance::CreateStatic(&FMovieGraphMemberCustomization::MakeInstance));
 
 	SelectedGraphObjectsDetailsWidget->RegisterInstancedCustomPropertyLayout(
 		UMovieGraphNode::StaticClass(),
-		FOnGetDetailCustomizationInstance::CreateStatic(&FGraphNodeCustomization::MakeInstance));
+		FOnGetDetailCustomizationInstance::CreateStatic(&FMovieGraphNodeCustomization::MakeInstance));
 
 	SelectedGraphObjectsDetailsWidget->RegisterInstancedCustomPropertyTypeLayout(
 		UMovieGraphShowFlags::StaticClass()->GetFName(),
