Movie Render Pipeline: Find whether the shot info's outer name contains the requested shot sequence names to render as opposed to the other way around.

diff --git a/Engine/Plugins/MovieScene/MovieRenderPipeline/Source/MovieRenderPipelineEditor/Private/MovieRenderPipelineEditorModule.cpp b/Engine/Plugins/MovieScene/MovieRenderPipeline/Source/MovieRenderPipelineEditor/Private/MovieRenderPipelineEditorModule.cpp
--- a/Engine/Plugins/MovieScene/MovieRenderPipeline/Source/MovieRenderPipelineEditor/Private/MovieRenderPipelineEditorModule.cpp
+++ b/Engine/Plugins/MovieScene/MovieRenderPipeline/Source/MovieRenderPipelineEditor/Private/MovieRenderPipelineEditorModule.cpp
@@ -144,10 +144,13 @@ class FMovieRenderPipelineRenderer : public IMovieRendererInterface
 			UMoviePipelineEditorBlueprintLibrary::EnsureJobHasDefaultSettings(ActiveJob);
 		}
 
-		TArray<FString> ShotNames;
+		TArray<FString> SequenceNames;
 		for (UMovieSceneCinematicShotSection* ShotSection : InSections)
 		{
-			ShotNames.Add(ShotSection->GetShotDisplayName());
+			if (ShotSection && ShotSection->GetSequence())
+			{
+				SequenceNames.Add(ShotSection->GetSequence()->GetName());
+			}
 		}
 
 		if (!GEditor->IsTransactionActive())
@@ -162,7 +165,22 @@ class FMovieRenderPipelineRenderer : public IMovieRendererInterface
 			// If no specified shots, enabled them all
 			if (Shot)
 			{
-				Shot->bEnabled = ShotNames.Num() == 0 ? true : ShotNames.Contains(Shot->OuterName);
+				if (SequenceNames.Num() == 0)
+				{
+					Shot->bEnabled = true;
+				}
+				else
+				{
+					Shot->bEnabled = false;
+					for (const FString& SequenceName : SequenceNames)
+					{
+						if (Shot->OuterName.Contains(SequenceName))
+						{
+							Shot->bEnabled = true;
+							break;
+						}
+					}
+				}
 			}
 		}
 
@@ -243,4 +261,4 @@ void FMovieRenderPipelineEditorModule::ShutdownModule()
 
 #undef LOCTEXT_NAMESPACE
 
-IMPLEMENT_MODULE(FMovieRenderPipelineEditorModule, MovieRenderPipelineEditor)
\ No newline at end of file
+IMPLEMENT_MODULE(FMovieRenderPipelineEditorModule, MovieRenderPipelineEditor)
