MoviePipeline: Added a new "Command Line Encoder" node to the graph.

Currently doesn't work with EXRs. Support for that will be coming in a later change.

#rnx

diff --git a/Engine/Plugins/MovieScene/MovieRenderPipeline/Source/MovieRenderPipelineCore/Private/Graph/MovieGraphBlueprintLibrary.cpp b/Engine/Plugins/MovieScene/MovieRenderPipeline/Source/MovieRenderPipelineCore/Private/Graph/MovieGraphBlueprintLibrary.cpp
--- a/Engine/Plugins/MovieScene/MovieRenderPipeline/Source/MovieRenderPipelineCore/Private/Graph/MovieGraphBlueprintLibrary.cpp
+++ b/Engine/Plugins/MovieScene/MovieRenderPipeline/Source/MovieRenderPipelineCore/Private/Graph/MovieGraphBlueprintLibrary.cpp
@@ -166,22 +166,21 @@ FString UMovieGraphBlueprintLibrary::ResolveFilenameFormatArguments(const FStrin
 
 	// Fix-up slashes
 	FPaths::NormalizeFilename(BaseFilename);
+	FPaths::RemoveDuplicateSlashes(BaseFilename);
 
+	// If we end with a "." character, remove it. The extension will put it back on. We can end up with this sometimes
+	// after resolving file format strings, ie: {sequence_name}.{frame_number} becomes {sequence_name}. for
+	// videos (which can't use frame_numbers).
+	BaseFilename.RemoveFromEnd(TEXT("."));
+
+	FString Extension = FString::Format(TEXT(".{ext}"), NamedArgs);
+	FString ThisTry = BaseFilename + Extension;
 	
 	// Check that it's a valid path that we can write to.
-	if (UE::MoviePipeline::CanWriteToFile(*BaseFilename, bOverwriteExisting))
-	{
-		return BaseFilename;
-	}
-
-	// The base filename must contain {file_dup} at this point, otherwise we're going to
-	// be stuck in an infinite loop of never resolving names as we'd just keep checking the
-	// base name!
-	if (!BaseFilename.Contains(TEXT("{file_dup}")))
+	if (UE::MoviePipeline::CanWriteToFile(*ThisTry, bOverwriteExisting))
 	{
-		BaseFilename.Append(TEXT("{file_dup}"));
+		return ThisTry;
 	}
-	FString ThisTry = BaseFilename;
 
 	// If we got here it means that there was already a file there with that name and bOverwriteExisting was false.
 	// So we start by swapping in _(2) where the file_dup token is, and try again. If that fails, _(3), etc. We start
diff --git a/Engine/Plugins/MovieScene/MovieRenderPipeline/Source/MovieRenderPipelineCore/Private/Graph/MovieGraphPipeline.cpp b/Engine/Plugins/MovieScene/MovieRenderPipeline/Source/MovieRenderPipelineCore/Private/Graph/MovieGraphPipeline.cpp
--- a/Engine/Plugins/MovieScene/MovieRenderPipeline/Source/MovieRenderPipelineCore/Private/Graph/MovieGraphPipeline.cpp
+++ b/Engine/Plugins/MovieScene/MovieRenderPipeline/Source/MovieRenderPipelineCore/Private/Graph/MovieGraphPipeline.cpp
@@ -10,6 +10,7 @@
 #include "Graph/Nodes/MovieGraphCollectionNode.h"
 #include "Graph/Nodes/MovieGraphFileOutputNode.h"
 #include "Graph/Nodes/MovieGraphModifierNode.h"
+#include "Graph/Nodes/MovieGraphRenderLayerNode.h"
 #include "Graph/Nodes/MovieGraphSamplingMethodNode.h"
 #include "Graph/Nodes/MovieGraphOutputSettingNode.h"
 #include "Graph/Nodes/MovieGraphWarmUpSettingNode.h"
@@ -169,6 +170,11 @@ void UMovieGraphPipeline::SetPreviewWidgetVisibleImpl(bool bInIsVisible)
 	}
 }
 
+TArray<FMovieGraphRenderOutputData>& UMovieGraphPipeline::GetGeneratedOutputData()
+{
+	return GeneratedOutputData;
+}
+
 void UMovieGraphPipeline::CreateLayersInRenderLayerSubsystem(const UMovieGraphEvaluatedConfig* EvaluatedConfig) const
 {
 	UMoviePipelineRenderLayerSubsystem* LayerSubsystem = GetWorld()->GetSubsystem<UMoviePipelineRenderLayerSubsystem>();
@@ -182,9 +188,18 @@ void UMovieGraphPipeline::CreateLayersInRenderLayerSubsystem(const UMovieGraphEv
 	// One render layer is generated per branch
 	for (const FName& BranchName : EvaluatedConfig->GetBranchNames())
 	{
-		UMoviePipelineRenderLayer* RenderLayer = NewObject<UMoviePipelineRenderLayer>();
-		RenderLayer->SetRenderLayerName(BranchName);
-		LayerSubsystem->AddRenderLayer(RenderLayer);
+		// Don't add the branch unless it has an active Render Layer node
+		constexpr bool bIncludeCDOs = false;
+		constexpr bool bExactMatch = true;
+		const UMovieGraphRenderLayerNode* RenderLayerNode =
+			EvaluatedConfig->GetSettingForBranch<UMovieGraphRenderLayerNode>(BranchName, bIncludeCDOs, bExactMatch);
+		
+		if (RenderLayerNode && !RenderLayerNode->IsDisabled())
+		{
+			UMoviePipelineRenderLayer* RenderLayer = NewObject<UMoviePipelineRenderLayer>();
+			RenderLayer->SetRenderLayerName(BranchName);
+			LayerSubsystem->AddRenderLayer(RenderLayer);
+		}
 	}
 }
 
@@ -406,7 +421,6 @@ void UMovieGraphPipeline::TickFinalizeOutputContainers(const bool bInForceFinish
 		for (const TObjectPtr<UMovieGraphFileOutputNode>& Node : GetOutputNodesUsed())
 		{
 			bAllContainsFinishedProcessing &= Node->IsFinishedWritingToDisk();
-
 		}
 	
 		// If we aren't forcing a finish, early out after one loop to keep
@@ -419,7 +433,6 @@ void UMovieGraphPipeline::TickFinalizeOutputContainers(const bool bInForceFinish
 		// If they've reached here, they're forcing them to finish so we'll sleep for a touch to give
 		// everyone a chance to actually do work before asking them if they're done.
 		FPlatformProcess::Sleep(0.1f);
-	
 	}
 
 	// If an output container is still working, we'll early out to keep the UI responsive.
@@ -446,33 +459,33 @@ void UMovieGraphPipeline::TickPostFinalizeExport(const bool bInForceFinish)
 	check(PipelineState == EMovieRenderPipelineState::Export);
 	UE_LOG(LogMovieRenderPipeline, Verbose, TEXT("[%d] PostFinalize Export (Start)."), GFrameCounter);
 
-	// ToDo: Loop through any extensions (such as XML export) and let them export using all of the
+	// Loop through any extensions (such as XML export) and let them export using all of the
 	// data that was generated during this run such as containers, output names and lengths.
 	// Tick all containers until they all report that they have finalized.
 	bool bAllContainsFinishedProcessing = true;
 
-	//do
-	//{
-	//	bAllContainsFinishedProcessing = true;
-	//
-	//	// Ask the containers if they're all done processing.
-	//	for (UMoviePipelineSetting* Setting : GetPipelinePrimaryConfig()->GetAllSettings())
-	//	{
-	//		bAllContainsFinishedProcessing &= Setting->HasFinishedExporting();
-	//	}
-	//
-	//	// If we aren't forcing a finish, early out after one loop to keep
-	//	// the editor/ui responsive.
-	//	if (!bInForceFinish || bAllContainsFinishedProcessing)
-	//	{
-	//		break;
-	//	}
-	//
-	//	// If they've reached here, they're forcing them to finish so we'll sleep for a touch to give
-	//	// everyone a chance to actually do work before asking them if they're done.
-	//	FPlatformProcess::Sleep(1.f);
-	//
-	//} while (true);
+	do
+	{
+		bAllContainsFinishedProcessing = true;
+
+		// Ask the nodes if they're all done processing.
+		constexpr bool bIncludeCDOs = false;
+		constexpr bool bExactMatch = false;
+		for (const TPair<FName, UMovieGraphPostRenderNode*>& Pair : GetSettingForActiveRenderLayers<UMovieGraphPostRenderNode>(bIncludeCDOs, bExactMatch))
+		{
+			bAllContainsFinishedProcessing &= Pair.Value->HasFinishedExporting();
+		}
+	
+		// If we aren't forcing a finish, early out after one loop to keep the editor/ui responsive.
+		if (!bInForceFinish || bAllContainsFinishedProcessing)
+		{
+			break;
+		}
+	
+		// If they've reached here, they're forcing them to finish so we'll sleep for a touch to give
+		// everyone a chance to actually do work before asking them if they're done.
+		FPlatformProcess::Sleep(1.f);
+	} while (true);
 
 	UE_LOG(LogMovieRenderPipeline, Verbose, TEXT("[%d] PostFinalize Export (End)."), GFrameCounter);
 
@@ -496,6 +509,20 @@ void UMovieGraphPipeline::BeginFinalize()
 	}
 }
 
+void UMovieGraphPipeline::BeginExport()
+{
+	constexpr bool bIncludeCDOs = false;
+	constexpr bool bExactMatch = false;
+	TArray<TPair<FName, UMovieGraphPostRenderNode*>> PostRenderNodes = GetSettingForActiveRenderLayers<UMovieGraphPostRenderNode>(bIncludeCDOs, bExactMatch);
+	for (TPair<FName, UMovieGraphPostRenderNode*>& Pair : PostRenderNodes)
+	{
+		UMovieGraphPostRenderNode* PostRenderNode = Pair.Value;
+		FName& BranchName = Pair.Key;
+		
+		PostRenderNode->BeginExport(this, BranchName);
+	}
+}
+
 void UMovieGraphPipeline::SetupShot(const TObjectPtr<UMoviePipelineExecutorShot>& InShot)
 {
 	// Set the new shot as the active shot. This enables the specified shot section and disables all other shot sections.
@@ -599,8 +626,6 @@ void UMovieGraphPipeline::SetSoloShot(const TObjectPtr<UMoviePipelineExecutorSho
 	}
 }
 
-
-
 void UMovieGraphPipeline::ExpandShot(const TObjectPtr<UMoviePipelineExecutorShot>& InShot, const int32 InNumHandleFrames, const bool bInHasMultipleTemporalSamples, const bool bIsPrePass,
 	const FFrameRate& InDisplayRate, const FFrameRate& InTickResolution, const int32 InWarmUpFrames)
 {
@@ -653,6 +678,27 @@ void UMovieGraphPipeline::ExpandShot(const TObjectPtr<UMoviePipelineExecutorShot
 	}
 }
 
+template<typename T>
+TArray<TPair<FName, T*>> UMovieGraphPipeline::GetSettingForActiveRenderLayers(const bool bIncludeCDOs, const bool bExactMatch)
+{
+	TArray<TPair<FName, T*>> FoundSettings;
+	
+	UMoviePipelineRenderLayerSubsystem* LayerSubsystem = GetWorld()->GetSubsystem<UMoviePipelineRenderLayerSubsystem>();
+	const FMovieGraphTimeStepData& TimeStepData = GetTimeStepInstance()->GetCalculatedTimeData();
+	const TObjectPtr<UMovieGraphEvaluatedConfig> EvaluatedConfig = TimeStepData.EvaluatedConfig;
+
+	for (const UMoviePipelineRenderLayer* RenderLayer : LayerSubsystem->GetRenderLayers())
+	{
+		FName LayerName = RenderLayer->GetRenderLayerName();
+		for (T* SettingNode : EvaluatedConfig->GetSettingsForBranch<T>(LayerName, bIncludeCDOs, bExactMatch))
+		{
+			FoundSettings.Add({LayerName, SettingNode});
+		}
+	}
+
+	return FoundSettings;
+}
+
 UMovieGraphTimeStepBase* UMovieGraphPipeline::GetTimeStepInstance() const
 {
 	check(GraphTimeStepInstance);
@@ -770,7 +816,7 @@ void UMovieGraphPipeline::ShutdownImpl(bool bIsError)
 	{
 		// All frames have been written to disk but we're doing a post-export step (such as encoding). Flush this operation as well.
 		// Export automatically switches our state to Finished so no need to manually transition afterwards.
-		//TickPostFinalizeExport(true);
+		TickPostFinalizeExport(true);
 	}
 }
 
@@ -847,7 +893,7 @@ void UMovieGraphPipeline::TransitionToState(const EMovieRenderPipelineState InNe
 			// after finalize finishes, because the futures won't be available until actually written to disk.
 			ProcessOutstandingFutures();
 
-			//BeginExport();
+			BeginExport();
 		}
 		break;
 	case EMovieRenderPipelineState::Export:
@@ -1081,8 +1127,8 @@ void UMovieGraphPipeline::ProcessOutstandingFutures()
 			ShotOutputData->Shot = FutureData.Shot;
 		}
 
-		// Add the filepath to the renderpass data.
-		ShotOutputData->RenderPassData.FindOrAdd(FutureData.DataIdentifier).FilePaths.Add(FutureData.FilePath);
+		// Add the filepath to the render layer data.
+		ShotOutputData->RenderLayerData.FindOrAdd(FutureData.DataIdentifier).FilePaths.Add(FutureData.FilePath);
 
 		// Sometime futures can be completed, but will be set to an error state (such as we couldn't write to the specified disk path.)
 		if (!OutputFuture.Get<0>().Get())
diff --git a/Engine/Plugins/MovieScene/MovieRenderPipeline/Source/MovieRenderPipelineCore/Private/Graph/Nodes/MovieGraphCommandLineEncoderNode.cpp b/Engine/Plugins/MovieScene/MovieRenderPipeline/Source/MovieRenderPipelineCore/Private/Graph/Nodes/MovieGraphCommandLineEncoderNode.cpp
new file mode 100644
--- /dev/null
+++ b/Engine/Plugins/MovieScene/MovieRenderPipeline/Source/MovieRenderPipelineCore/Private/Graph/Nodes/MovieGraphCommandLineEncoderNode.cpp
@@ -0,0 +1,541 @@
+// Copyright Epic Games, Inc. All Rights Reserved.
+
+#include "Graph/Nodes/MovieGraphCommandLineEncoderNode.h"
+
+#include "Graph/MovieGraphBlueprintLibrary.h"
+#include "Graph/MovieGraphPipeline.h"
+#include "Graph/Nodes/MovieGraphOutputSettingNode.h"
+#include "HAL/FileManager.h"
+#include "HAL/PlatformFileManager.h"
+#include "Misc/CoreDelegates.h"
+#include "Misc/FileHelper.h"
+#include "MoviePipelineCommandLineEncoderSettings.h"
+#include "MoviePipelineUtils.h"
+#include "MovieRenderPipelineCoreModule.h"
+#include "Styling/AppStyle.h"
+
+#define LOCTEXT_NAMESPACE "MovieGraph"
+
+UMovieGraphCommandLineEncoderNode::UMovieGraphCommandLineEncoderNode()
+	: bDeleteSourceFiles(false)
+	, bSkipEncodeOnRenderCanceled(true)
+	, bRetainInputTextFiles(false)
+{
+	// Sensible defaults for cmdline arguments
+	CommandLineFormat = TEXT("-hide_banner -y -loglevel error {VideoInputs} {AudioInputs} -acodec {AudioCodec} -vcodec {VideoCodec} {Quality} \"{OutputPath}\"");
+	VideoInputStringFormat = TEXT("-f concat -safe 0 -i \"{InputFile}\" -r {FrameRate}");
+	AudioInputStringFormat = TEXT("-f concat -safe 0 -i \"{InputFile}\"");
+	EncodeSettings = TEXT("-crf 20");
+}
+
+#if WITH_EDITOR
+FText UMovieGraphCommandLineEncoderNode::GetNodeTitle(const bool bGetDescriptive) const
+{
+	static const FText CommandLineEncoderNodeName = LOCTEXT("NodeName_CommandLineEncoder", "Command Line Encoder");
+	return CommandLineEncoderNodeName;
+}
+
+FText UMovieGraphCommandLineEncoderNode::GetMenuCategory() const
+{
+	return LOCTEXT("CommandLineEncoder_Category", "Utility");
+}
+
+FLinearColor UMovieGraphCommandLineEncoderNode::GetNodeTitleColor() const
+{
+	static const FLinearColor ModifierNodeColor = FLinearColor(0.92f, 0.42f, 0.2f);
+	return ModifierNodeColor;
+}
+
+FSlateIcon UMovieGraphCommandLineEncoderNode::GetIconAndTint(FLinearColor& OutColor) const
+{
+	static const FSlateIcon ModifierIcon = FSlateIcon(FAppStyle::GetAppStyleSetName(), "ObjectBrowser.TabIcon");
+
+	OutColor = FLinearColor::White;
+	return ModifierIcon;
+}
+#endif // WITH_EDITOR
+
+EMovieGraphBranchRestriction UMovieGraphCommandLineEncoderNode::GetBranchRestriction() const
+{
+	return EMovieGraphBranchRestriction::RenderLayer;
+}
+
+void UMovieGraphCommandLineEncoderNode::StartEncodingProcess(TArray<FMovieGraphRenderOutputData>& InGeneratedData, const bool bInIsShotEncode)
+{
+	// If the format string isn't split per shot, we can't start encoding now.
+	if (bInIsShotEncode && !NeedsPerShotFlushing())
+	{
+		return;
+	}
+
+	// Early out if there are any issues with settings. This will log and shutdown the pipeline if there are errors.
+	TArray<FText> ValidationErrors;
+	if (!AreSettingsValid(ValidationErrors))
+	{
+		for (const FText& Error : ValidationErrors)
+		{
+			UE_LOG(LogMovieRenderPipelineIO, Error, TEXT("%s"), *Error.ToString());
+		}
+		
+		CachedPipeline->RequestShutdown(true);
+		
+		return;
+	}
+
+	// Get the base output file path; it still needs to be validated and resolved at this point
+	const UMovieGraphOutputSettingNode* OutputSettingNode = GetSettingOnBranch<UMovieGraphOutputSettingNode>();
+	const FString OutputFilename = FileNameFormatOverride.Len() > 0 ? FileNameFormatOverride : OutputSettingNode->FileNameFormat;
+	FString FilePathFormatString = OutputSettingNode->OutputDirectory.Path / OutputFilename;
+	
+	constexpr bool bTestRenderPass = false;
+	constexpr bool bTestFrameNumber = false;
+	UE::MoviePipeline::ValidateOutputFormatString(FilePathFormatString, bTestRenderPass, bTestFrameNumber);
+	UE::MoviePipeline::RemoveFrameNumberFormatStrings(FilePathFormatString, true);
+
+	// Launch the encoder, once per render layer
+	for (TTuple<FMovieGraphRenderDataIdentifier, FEncoderParams>& RenderLayerEncoderParams : GenerateRenderLayerEncoderParams(InGeneratedData))
+	{
+		FMovieGraphRenderDataIdentifier& RenderIdentifier = RenderLayerEncoderParams.Key;
+		FEncoderParams& EncoderParams = RenderLayerEncoderParams.Value;
+		
+		FString FinalFilePath = GetResolvedOutputFilename(RenderIdentifier, EncoderParams.Shot, FilePathFormatString);
+
+		// Manipulate the in/out data in case scripting tries to get access to the files. It's not a perfect solution
+		// because the encoding won't be finished by the time the scripting layer is called. Need to manipulate the original
+		// and not the copy that we're currently iterating through.
+		FMovieGraphRenderDataIdentifier CommandLineEncoderIdentifier;
+		CommandLineEncoderIdentifier.RootBranchName = CachedBranchName;
+		CommandLineEncoderIdentifier.RendererName = FString(TEXT("CommandLineEncoder"));
+		for (FMovieGraphRenderOutputData& OutputData : InGeneratedData)
+		{
+			if (OutputData.Shot != EncoderParams.Shot)
+			{
+				continue;
+			}
+
+			OutputData.RenderLayerData.FindOrAdd(CommandLineEncoderIdentifier).FilePaths.Add(FinalFilePath);
+		}
+		
+		EncoderParams.NamedArguments.Add(TEXT("OutputPath"), FinalFilePath);
+		
+		LaunchEncoder(EncoderParams);
+	}
+}
+
+void UMovieGraphCommandLineEncoderNode::BeginExport(UMovieGraphPipeline* InMoviePipeline, const FName& InBranchName)
+{
+	CachedPipeline = InMoviePipeline;
+	CachedBranchName = InBranchName;
+	
+	// When we start exporting, we remove the OnEndFrame delegate because if the user hit escape to cancel a movie render
+	// no frames will be ticked anymore. Instead we'll call OnTick from HasFinishedExporting() by hand as that will
+	// get called in a loop until it is finished.
+	FCoreDelegates::OnEndFrame.RemoveAll(this);
+
+	// This is called at the end of the movie render. If we were starting encode jobs per shot, then there's already
+	// an encode job going for all shots, so we early out.
+	if (NeedsPerShotFlushing())
+	{
+		return;
+	}
+
+	// However, if they didn't want a per-shot flush (ie: rendering all shots into one file) then we start now.
+	TArray<FMovieGraphRenderOutputData>& OutputData = CachedPipeline->GetGeneratedOutputData();
+	constexpr bool bIsShotEncode = false;
+	StartEncodingProcess(OutputData, bIsShotEncode);
+}
+
+bool UMovieGraphCommandLineEncoderNode::HasFinishedExporting()
+{
+	// Manually tick the output (which cleans up ActiveEncodeJobs). This is needed because manually canceling a job
+	// stops ticking the engine and repeatedly calls HasFinishedExporting().
+	OnTick();
+	
+	return ActiveEncodeJobs.IsEmpty();
+}
+
+bool UMovieGraphCommandLineEncoderNode::AreSettingsValid(TArray<FText>& OutErrors) const
+{
+	OutErrors.Empty();
+	
+	const UMoviePipelineCommandLineEncoderSettings* EncoderSettings = GetDefault<UMoviePipelineCommandLineEncoderSettings>();
+
+	// Validate project settings
+	{
+		if (EncoderSettings->ExecutablePath.IsEmpty())
+		{
+			OutErrors.Add(LOCTEXT("CommandLineEncode_MissingExecutable", "No encoder executable has been specified in the Project Settings. Please set an encoder executable in Project Settings > Movie Pipeline CLI Encoder."));
+		}
+	}
+
+	// Validate node settings
+	{
+		if (VideoCodec.IsEmpty())
+		{
+			OutErrors.Add(LOCTEXT("CommandLineEncode_MissingVideoCodec", "No video codec has been specified on the Command Line Encoder node."));
+		}
+
+		if (AudioCodec.IsEmpty() && CommandLineFormat.Contains("{AudioCodec}"))
+		{
+			OutErrors.Add(LOCTEXT("CommandLineEncode_MissingAudioCodec", "The {AudioCodec} token is in the Command Line Format property on the Command Line Encoder node, but no audio codec has been specified."));
+		}
+
+		if (OutputFileExtension.IsEmpty())
+		{
+			OutErrors.Add(LOCTEXT("CommandLineEncode_MissingFileExtension", "No file extension has been specified on the Command Line Encoder node."));
+		}
+	}
+
+	return OutErrors.IsEmpty();
+}
+
+void UMovieGraphCommandLineEncoderNode::OnTick()
+{
+	for (int32 Index = ActiveEncodeJobs.Num() - 1; Index >= 0; Index--)
+	{
+		FActiveJob& Job = ActiveEncodeJobs[Index];
+		TArray<uint8> ProcessOutput;
+		FString Results = FPlatformProcess::ReadPipe(Job.ReadPipe);
+		if (Results.Len() > 0)
+		{
+			// The default global arguments suppress everything but errors. We unfortunately can't tell stdout from stderr
+			// using a non-blocking platform process launch, so we'll just promote all the output as an error log.
+			UE_LOG(LogMovieRenderPipeline, Error, TEXT("Command Line Encoder: %s"), *Results);
+		}
+
+		// If they hit escape during a render, (potentially) cancel the encode job
+		bool bCancelEncode = false;
+		if (bSkipEncodeOnRenderCanceled && CachedPipeline->IsShutdownRequested())
+		{
+			bCancelEncode = true;
+
+			// We have to specifically terminate the process instead of just closing it
+			constexpr bool bKillTree = true;
+			FPlatformProcess::TerminateProc(Job.ProcessHandle, bKillTree);
+			FPlatformProcess::WaitForProc(Job.ProcessHandle);
+		}
+
+		// If the process has finished we'll clean up
+		if (!FPlatformProcess::IsProcRunning(Job.ProcessHandle) || bCancelEncode)
+		{
+			FPlatformProcess::ClosePipe(Job.ReadPipe, Job.WritePipe);
+			FPlatformProcess::CloseProc(Job.ProcessHandle);
+
+			IFileManager& FileManager = IFileManager::Get();
+			for (const FString& FilePath : Job.FilesToDelete)
+			{
+				constexpr bool bRequireExist = false;
+				constexpr bool bEvenReadOnly = false;
+				constexpr bool bQuiet = false;
+				FileManager.Delete(*FilePath, bRequireExist, bEvenReadOnly, bQuiet);
+			}
+
+			ActiveEncodeJobs.RemoveAt(Index);
+		}
+	}
+}
+
+bool UMovieGraphCommandLineEncoderNode::NeedsPerShotFlushing() const
+{
+	const UMovieGraphOutputSettingNode* OutputSettingNode = GetSettingOnBranch<UMovieGraphOutputSettingNode>();
+	
+	const FString OutputFilename = FileNameFormatOverride.Len() > 0 ? FileNameFormatOverride : OutputSettingNode->FileNameFormat;
+	const FString FullPath = OutputSettingNode->OutputDirectory.Path / OutputFilename;
+	
+	if (FullPath.Contains(TEXT("{shot_name}")) || FullPath.Contains(TEXT("{camera_name}")))
+	{
+		return true;
+	}
+
+	return false;
+}
+
+TMap<FMovieGraphRenderDataIdentifier, UMovieGraphCommandLineEncoderNode::FEncoderParams> UMovieGraphCommandLineEncoderNode::GenerateRenderLayerEncoderParams(TArray<FMovieGraphRenderOutputData>& InGeneratedData) const
+{
+	// We produce one set of encoder parameters (ie, one file) per render layer
+	TMap<FMovieGraphRenderDataIdentifier, FEncoderParams> RenderLayerEncoderParams;
+
+	// The path shouldn't have quotes on it as it's already kept as a separate argument right up until creating the process, at which point
+	// the platform puts quotes around the FString if needed.
+	const UMoviePipelineCommandLineEncoderSettings* EncoderSettings = GetDefault<UMoviePipelineCommandLineEncoderSettings>();
+	FString ExecutablePathNoQuotes = EncoderSettings->ExecutablePath.Replace(TEXT("\""), TEXT(""));
+	FPaths::NormalizeFilename(ExecutablePathNoQuotes);
+
+	UMovieGraphOutputSettingNode* OutputSettingNode = GetSettingOnBranch<UMovieGraphOutputSettingNode>();
+	const FFrameRate SourceFrameRate = CachedPipeline->GetDataSourceInstance()->GetDisplayRate();
+	const FFrameRate EffectiveFrameRate = UMovieGraphBlueprintLibrary::GetEffectiveFrameRate(OutputSettingNode, SourceFrameRate);
+
+	const FString SanitizedEncodeSettings = EncodeSettings.Replace(LINE_TERMINATOR, TEXT(" "));
+	FStringFormatNamedArguments SharedArguments;
+	SharedArguments.Add(TEXT("Executable"), ExecutablePathNoQuotes);
+	SharedArguments.Add(TEXT("AudioCodec"), AudioCodec);
+	SharedArguments.Add(TEXT("VideoCodec"), VideoCodec);
+	SharedArguments.Add(TEXT("FrameRate"), EffectiveFrameRate.AsDecimal());
+	SharedArguments.Add(TEXT("Quality"), SanitizedEncodeSettings);
+	
+	for (FMovieGraphRenderOutputData& GeneratedRenderData : InGeneratedData)
+	{
+		// Search for audio to attach to every render layer
+		TMap<FString, TArray<FString>> AudioPathsByExtension;
+		for (const TPair<FMovieGraphRenderDataIdentifier, FMovieGraphRenderLayerOutputData>& InnerRenderLayer : GeneratedRenderData.RenderLayerData)
+		{
+			if (InnerRenderLayer.Key.RendererName == FString(TEXT("Audio")))
+			{
+				for (const FString& FilePath : InnerRenderLayer.Value.FilePaths)
+				{
+					FString Extension = FPaths::GetExtension(FilePath);
+					AudioPathsByExtension.FindOrAdd(Extension).Add(FilePath);
+				}
+			}
+		}
+
+		TArray<FMovieGraphRenderDataIdentifier> RenderDataToRemove;
+		
+		for (const TPair<FMovieGraphRenderDataIdentifier, FMovieGraphRenderLayerOutputData>& RenderLayer : GeneratedRenderData.RenderLayerData)
+		{
+			const FMovieGraphRenderDataIdentifier& RenderIdentifier = RenderLayer.Key;
+			const FMovieGraphRenderLayerOutputData& RenderOutputData = RenderLayer.Value;
+			
+			// TODO: Is "Audio" set on the renderer name or something else? We don't have audio yet...
+			if (RenderIdentifier.RendererName == FString(TEXT("Audio")))
+			{
+				continue;
+			}
+
+			// This node should only be responsible for exporting movies for images generated on the node's branch
+			if (RenderIdentifier.RootBranchName != CachedBranchName)
+			{
+				continue;
+			}
+
+			FEncoderParams& EncoderParams = RenderLayerEncoderParams.FindOrAdd(RenderIdentifier);
+			EncoderParams.Shot = GeneratedRenderData.Shot;
+			EncoderParams.RenderDataIdentifier = RenderIdentifier;
+			EncoderParams.NamedArguments = SharedArguments;
+			
+			for (const FString& FilePath : RenderOutputData.FilePaths)
+			{
+				FString Extension = FPaths::GetExtension(FilePath);
+				EncoderParams.FilesByExtensionType.FindOrAdd(Extension).Add(FilePath);
+			}
+
+			// Attach audio
+			EncoderParams.FilesByExtensionType.Append(AudioPathsByExtension);
+
+			RenderDataToRemove.Add(RenderIdentifier);
+		}
+		
+		// If we're going to delete the source files, don't make them available to the scripting layer callback
+		// because the files will be deleted out from underneath at a random point, so don't want the scripting
+		// layer to think they can rely on them actually existing. If scripting layer really needs source files,
+		// they need to not use the Command Line Encoder setting and instead roll their own.
+		if (bDeleteSourceFiles)
+		{
+			for (const FMovieGraphRenderDataIdentifier& RenderIdentifier : RenderDataToRemove)
+			{
+				GeneratedRenderData.RenderLayerData.Remove(RenderIdentifier);
+			}
+		}
+	}
+
+	return RenderLayerEncoderParams;
+}
+
+FString UMovieGraphCommandLineEncoderNode::GetResolvedOutputFilename(const FMovieGraphRenderDataIdentifier& RenderIdentifier, const TWeakObjectPtr<UMoviePipelineExecutorShot>& Shot, const FString& FileNameFormatString) const
+{
+	TMap<FString, FString> FormatOverrides;
+	FormatOverrides.Add(TEXT("render_pass"), RenderIdentifier.RendererName);
+	FormatOverrides.Add(TEXT("ext"), OutputFileExtension);
+	if (Shot.IsValid())
+	{
+		FormatOverrides.Add(TEXT("shot_name"), Shot->OuterName);
+		FormatOverrides.Add(TEXT("camera_name"), Shot->InnerName);
+	}
+
+	// We must manually resolve {version} tokens because they're intended to be per-shot/global (ie: individual file names
+	// don't have versions when using image sequences).
+	FMovieGraphFilenameResolveParams ResolveParams;
+	ResolveParams.InitializationTime = CachedPipeline->GetInitializationTime();
+	ResolveParams.Job = CachedPipeline->GetCurrentJob();
+	ResolveParams.Shot = Shot.IsValid() ? Shot.Get() : nullptr;
+	ResolveParams.FileNameFormatOverrides = FormatOverrides;
+	ResolveParams.EvaluatedConfig = CachedPipeline->GetTimeStepInstance()->GetCalculatedTimeData().EvaluatedConfig;
+	ResolveParams.RenderDataIdentifier = RenderIdentifier;
+	// int32 VersionNumber = 0; // TODO: Track versions
+	// FileNameFormatString.ReplaceInline(TEXT("{version}"), *FString::Printf(TEXT("v%0*d"), 3, VersionNumber));
+
+	FMovieGraphResolveArgs FinalFormatArgs;
+	FString FinalFilePath = UMovieGraphBlueprintLibrary::ResolveFilenameFormatArguments(FileNameFormatString, ResolveParams, FinalFormatArgs);
+
+	if (FPaths::IsRelative(FinalFilePath))
+	{
+		FinalFilePath = FPaths::ConvertRelativePathToFull(FinalFilePath);
+	}
+
+	FPaths::NormalizeFilename(FinalFilePath);
+	FPaths::CollapseRelativeDirectories(FinalFilePath);
+
+	const FString FinalFileDirectory = FPaths::GetPath(FinalFilePath);
+
+	// Ensure the output directory is created
+	IPlatformFile& FileManager = FPlatformFileManager::Get().GetPlatformFile();
+	if (!FileManager.CreateDirectoryTree(*FinalFileDirectory))
+	{
+		UE_LOG(LogMovieRenderPipelineIO, Error, TEXT("Failed to create directory for output path '%s'"), *FinalFileDirectory);
+	}
+
+	return FinalFilePath;
+}
+
+void UMovieGraphCommandLineEncoderNode::GenerateTemporaryEncoderInputFiles(const FEncoderParams& InParams, TArray<FString>& OutVideoInputFilePaths, TArray<FString>& OutAudioInputFilePaths) const
+{
+	UMovieGraphOutputSettingNode* OutputSettingNode = GetSettingOnBranch<UMovieGraphOutputSettingNode>();
+
+	// Generate a text file for each input type which lists the files for that input type. We generate a FGuid in case there are
+	// multiple encode jobs going at once.
+	TStringBuilder<64> StringBuilder;
+
+	double InFrameRate = InParams.NamedArguments[TEXT("FrameRate")].DoubleValue;
+	double FrameRateAsDuration = 1.0 / InFrameRate;
+
+	for (const TTuple<FString, TArray<FString>>& Pair : InParams.FilesByExtensionType)
+	{
+		FGuid FileGuid = FGuid::NewGuid();
+		FString FilePath = OutputSettingNode->OutputDirectory.Path / FileGuid.ToString() + TEXT("_input");
+
+		TMap<FString, FString> FormatOverrides;
+		FormatOverrides.Add(TEXT("ext"), TEXT("txt"));
+
+		FMovieGraphFilenameResolveParams ResolveParams;
+		ResolveParams.InitializationTime = CachedPipeline->GetInitializationTime();
+		ResolveParams.Job = CachedPipeline->GetCurrentJob();
+		ResolveParams.Shot = InParams.Shot.Get();
+		ResolveParams.FileNameFormatOverrides = FormatOverrides;
+		ResolveParams.EvaluatedConfig = CachedPipeline->GetTimeStepInstance()->GetCalculatedTimeData().EvaluatedConfig;
+		ResolveParams.RenderDataIdentifier = InParams.RenderDataIdentifier;
+
+		FMovieGraphResolveArgs FinalFormatArgs;
+		FString FinalFilePath = UMovieGraphBlueprintLibrary::ResolveFilenameFormatArguments(FilePath, ResolveParams, FinalFormatArgs);
+
+		UE_LOG(LogMovieRenderPipelineIO, Log, TEXT("Generated Path '%s' for input data."), *FinalFilePath);
+		StringBuilder.Reset();
+		for (const FString& Path : Pair.Value)
+		{
+			StringBuilder.Appendf(TEXT("file 'file:%s'%s"), *Path, LINE_TERMINATOR);
+
+			// Some encoders require the duration of each file to be listed after the file.
+			if (Pair.Key != TEXT("wav"))
+			{
+				StringBuilder.Appendf(TEXT("duration %f%s"), FrameRateAsDuration, LINE_TERMINATOR);
+			}
+		}
+
+		// Save this to disk.
+		FFileHelper::SaveStringToFile(StringBuilder.ToString(), *FinalFilePath);
+
+		// Not a great solution but best we've got right now
+		if (Pair.Key == TEXT("wav"))
+		{
+			OutAudioInputFilePaths.Add(FinalFilePath);
+		}
+		else
+		{
+			OutVideoInputFilePaths.Add(FinalFilePath);
+		}
+	}
+}
+
+FString UMovieGraphCommandLineEncoderNode::BuildEncoderCommand(const FEncoderParams& InParams, TArray<FString>& InVideoInputFilePaths, TArray<FString>& InAudioInputFilePaths) const
+{
+	// Build our final command line arguments
+	FStringFormatNamedArguments FinalNamedArgs = InParams.NamedArguments;
+	
+	const FString SanitizedVideoInputStringFormat = VideoInputStringFormat.Replace(LINE_TERMINATOR, TEXT(" "));
+	const FString SanitizedAudioInputStringFormat = AudioInputStringFormat.Replace(LINE_TERMINATOR, TEXT(" "));
+	const FString SanitizedCommandLineFormat = CommandLineFormat.Replace(LINE_TERMINATOR, TEXT(" "));
+
+	FString VideoInputArg;
+	FString AudioInputArg;
+
+	for (const FString& FilePath : InVideoInputFilePaths)
+	{
+		FStringFormatNamedArguments NamedArgs;
+		NamedArgs.Add(TEXT("InputFile"), *FilePath);
+		NamedArgs.Add(TEXT("FrameRate"), InParams.NamedArguments[TEXT("FrameRate")]);
+
+		VideoInputArg += TEXT(" ") + FString::Format(*SanitizedVideoInputStringFormat, NamedArgs);
+	}
+
+	for (const FString& FilePath : InAudioInputFilePaths)
+	{
+		FStringFormatNamedArguments NamedArgs;
+		NamedArgs.Add(TEXT("InputFile"), *FilePath);
+
+		AudioInputArg += TEXT(" ") + FString::Format(*SanitizedAudioInputStringFormat, NamedArgs);
+	}
+
+	FinalNamedArgs.Add(TEXT("VideoInputs"), VideoInputArg);
+	FinalNamedArgs.Add(TEXT("AudioInputs"), AudioInputArg);
+
+	return FString::Format(*SanitizedCommandLineFormat, FinalNamedArgs);;
+}
+
+void UMovieGraphCommandLineEncoderNode::LaunchEncoder(const FEncoderParams& InParams)
+{
+	TArray<FString> VideoInputs;
+	TArray<FString> AudioInputs;
+	GenerateTemporaryEncoderInputFiles(InParams, VideoInputs, AudioInputs);
+	
+	const FString CommandToRun = BuildEncoderCommand(InParams, VideoInputs, AudioInputs);
+	
+	UE_LOG(LogMovieRenderPipelineIO, Log, TEXT("Final Command Line Arguments: %s"), *CommandToRun);
+
+	void* PipeRead = nullptr;
+	void* PipeWrite = nullptr;
+	verify(FPlatformProcess::CreatePipe(PipeRead, PipeWrite));
+
+	constexpr bool bLaunchDetached = false;
+	constexpr bool bLaunchHidden = true;
+	constexpr bool bLaunchReallyHidden = bLaunchHidden;
+	const FString ExecutableArg = FString::Format(TEXT("{Executable}"), InParams.NamedArguments);
+	const FProcHandle ProcessHandle = FPlatformProcess::CreateProc(*ExecutableArg, *CommandToRun, bLaunchDetached, bLaunchHidden, bLaunchReallyHidden, nullptr, 0, nullptr, PipeWrite, PipeRead);
+	if (ProcessHandle.IsValid())
+	{
+		FActiveJob& NewJob = ActiveEncodeJobs.AddDefaulted_GetRef();
+		NewJob.ProcessHandle = ProcessHandle;
+		NewJob.ReadPipe = PipeRead;
+		NewJob.WritePipe = PipeWrite;
+
+		// Delete the auto-generated input files (unless they have been flagged to be retained)
+		if (!bRetainInputTextFiles)
+		{
+			NewJob.FilesToDelete.Append(VideoInputs);
+			NewJob.FilesToDelete.Append(AudioInputs);
+		}
+
+		// Delete source files if requested
+		if (bDeleteSourceFiles)
+		{
+			for (const TTuple<FString, TArray<FString>>& Pair : InParams.FilesByExtensionType)
+			{
+				NewJob.FilesToDelete.Append(Pair.Value);
+			}
+		}
+	}
+	else
+	{
+		UE_LOG(LogMovieRenderPipeline, Error, TEXT("Failed to launch encoder process, see output log for more details."));
+		CachedPipeline->Shutdown(true);
+	}
+}
+
+template<typename T>
+T* UMovieGraphCommandLineEncoderNode::GetSettingOnBranch(const bool bIncludeCDOs, const bool bExactMatch) const
+{
+	const TObjectPtr<UMovieGraphEvaluatedConfig> EvaluatedConfig =
+		CachedPipeline->GetTimeStepInstance()->GetCalculatedTimeData().EvaluatedConfig;
+	
+	return EvaluatedConfig->GetSettingForBranch<T>(CachedBranchName, bIncludeCDOs, bExactMatch);
+}
+
+#undef LOCTEXT_NAMESPACE // "MovieGraph"
diff --git a/Engine/Plugins/MovieScene/MovieRenderPipeline/Source/MovieRenderPipelineCore/Public/Graph/MovieGraphDataTypes.h b/Engine/Plugins/MovieScene/MovieRenderPipeline/Source/MovieRenderPipelineCore/Public/Graph/MovieGraphDataTypes.h
--- a/Engine/Plugins/MovieScene/MovieRenderPipeline/Source/MovieRenderPipelineCore/Public/Graph/MovieGraphDataTypes.h
+++ b/Engine/Plugins/MovieScene/MovieRenderPipeline/Source/MovieRenderPipelineCore/Public/Graph/MovieGraphDataTypes.h
@@ -59,7 +59,7 @@ struct MOVIERENDERPIPELINECORE_API FMovieGraphInitConfig
 };
 
 USTRUCT(BlueprintType)
-struct FMovieGraphRenderPassOutputData
+struct FMovieGraphRenderLayerOutputData
 {
 	GENERATED_BODY()
 public:
@@ -78,11 +78,11 @@ public:
 	TWeakObjectPtr<UMoviePipelineExecutorShot> Shot;
 
 	/**
-	* A mapping between render passes (such as "beauty") and an array containing the files written for that shot.
-	* Will be multiple files if using image sequences
+	* A mapping between render layers (such as "beauty") and an array containing the files written for that shot.
+	* Will be multiple files if using image sequences.
 	*/
 	UPROPERTY(VisibleAnywhere, BlueprintReadOnly, Category = "Movie Graph")
-	TMap<FMovieGraphRenderDataIdentifier, FMovieGraphRenderPassOutputData> RenderPassData;
+	TMap<FMovieGraphRenderDataIdentifier, FMovieGraphRenderLayerOutputData> RenderLayerData;
 };
 
 
diff --git a/Engine/Plugins/MovieScene/MovieRenderPipeline/Source/MovieRenderPipelineCore/Public/Graph/MovieGraphNode.h b/Engine/Plugins/MovieScene/MovieRenderPipeline/Source/MovieRenderPipelineCore/Public/Graph/MovieGraphNode.h
--- a/Engine/Plugins/MovieScene/MovieRenderPipeline/Source/MovieRenderPipelineCore/Public/Graph/MovieGraphNode.h
+++ b/Engine/Plugins/MovieScene/MovieRenderPipeline/Source/MovieRenderPipelineCore/Public/Graph/MovieGraphNode.h
@@ -21,6 +21,7 @@ class UMovieGraphInput;
 class UMovieGraphMember;
 class UMovieGraphOutput;
 class UMovieGraphPin;
+class UMovieGraphPipeline;
 class UMovieGraphVariable;
 struct FMovieGraphEvaluationContext;
 struct FMovieGraphTraversalContext;
@@ -350,4 +351,20 @@ public:
 	* resolving the settings of the graph, the node only needs to read its own values.
 	*/
 	virtual void GetFormatResolveArgs(FMovieGraphResolveArgs& OutMergedFormatArgs) const {}
-};
\ No newline at end of file
+};
+
+/**
+ * A node which runs after all renders have completed.
+ */
+UCLASS(Abstract)
+class MOVIERENDERPIPELINECORE_API UMovieGraphPostRenderNode : public UMovieGraphSettingNode
+{
+	GENERATED_BODY()
+
+public:
+	/** Begins the export process for this node. */
+	virtual void BeginExport(UMovieGraphPipeline* InMoviePipeline, const FName& InBranchName) PURE_VIRTUAL(UMovieGraphPostRenderNode::BeginExport, );
+
+	/** Returns true if this node has finished its export process, else false. */
+	virtual bool HasFinishedExporting() PURE_VIRTUAL(UMovieGraphPostRenderNode::HasFinishedExporting, return true; );
+};
diff --git a/Engine/Plugins/MovieScene/MovieRenderPipeline/Source/MovieRenderPipelineCore/Public/Graph/MovieGraphPipeline.h b/Engine/Plugins/MovieScene/MovieRenderPipeline/Source/MovieRenderPipelineCore/Public/Graph/MovieGraphPipeline.h
--- a/Engine/Plugins/MovieScene/MovieRenderPipeline/Source/MovieRenderPipelineCore/Public/Graph/MovieGraphPipeline.h
+++ b/Engine/Plugins/MovieScene/MovieRenderPipeline/Source/MovieRenderPipelineCore/Public/Graph/MovieGraphPipeline.h
@@ -14,6 +14,7 @@ class UMoviePipelineExecutorJob;
 class UMoviePipelineExecutorShot;
 class UMovieGraphFileOutputNode;
 class IImageWriteQueue;
+struct FMovieGraphRenderOutputData;
 
 namespace UE::MovieGraph
 {
@@ -63,26 +64,41 @@ public:
 
 	/** Get the Active Shot list, which is the full shot list generated from the external data source, with disabled shots removed. */
 	const TArray<TObjectPtr<UMoviePipelineExecutorShot>>& GetActiveShotList() const { return ActiveShotList; }
+	
 	/** Which index of the Active Shot List are we currently on */
 	int32 GetCurrentShotIndex() const { return CurrentShotIndex; }
+	
 	/** Called by the TimeStepInstance when it's time to set up for another shot. Don't call this unless you know what you're doing. */
 	void SetupShot(const TObjectPtr<UMoviePipelineExecutorShot>& InShot);
+	
 	/** Called by the TimeStepInstance when it's time to tear down the current shot. Don't call this unless you know what you're doing. */
 	void TeardownShot(const TObjectPtr<UMoviePipelineExecutorShot>& InShot);
+	
 	/** Used occasionally to cross-reference other components. Don't call this unless you know what you're doing. */
 	UMovieGraphTimeStepBase* GetTimeStepInstance() const;
+	
 	/** Used occasionally to cross-reference other components. Don't call this unless you know what you're doing. */
 	UFUNCTION(BlueprintCallable, Category = "Movie Graph")
 	UMovieGraphRendererBase* GetRendererInstance() const { return GraphRendererInstance; }
+	
 	/** Used occasionally to cross-reference other components. Don't call this unless you know what you're doing. */
 	UMovieGraphDataSourceBase* GetDataSourceInstance() const { return GraphDataSourceInstance; }
+	
 	/** Gets the Output Merger for this Movie Pipeline which is responsible for gathering all of the data coming in for a given output frame, before making it available to the MovieGraphPipeline. */
 	TSharedPtr<UE::MovieGraph::IMovieGraphOutputMerger> GetOutputMerger() const { return OutputMerger; }
+	
 	/** Writing images to disk is an async process. When you start writing, declare a future with the filename you will eventually write to, and complete the future once it is on disk. */
 	void AddOutputFuture(TFuture<bool>&& InOutputFuture, const UE::MovieGraph::FMovieGraphOutputFutureData& InData);
+	
 	/** Used by the Renderer Instance to disable the preview widget before rendering so it isn't baked into the UI Renderer. */
 	void SetPreviewWidgetVisible(bool bInIsVisible) { SetPreviewWidgetVisibleImpl(bInIsVisible); }
 
+	/**
+	 * Gets render data that was generated for each shot. This data is mutable, so processes that run after renders
+	 * complete can add their additional outputs to the generated data. Care should be taken when removing output data.
+	 */
+	TArray<FMovieGraphRenderOutputData>& GetGeneratedOutputData();
+
 	/**
 	 * Resets the render layer subsystem and updates it to reflect the render layers that are present in the evaluated
 	 * graph. Does not update the contents of the layers, however. See UpdateLayerContentsInRenderLayerSubsystem().
@@ -113,11 +129,13 @@ protected:
 	virtual void TransitionToState(const EMovieRenderPipelineState InNewState);
 	virtual const TSet<TObjectPtr<UMovieGraphFileOutputNode>> GetOutputNodesUsed() const;
 	virtual void BeginFinalize();
+	virtual void BeginExport();
 	virtual void LoadPreviewWidget();
 	virtual void SetPreviewWidgetVisibleImpl(bool bInIsVisible);
 
 	// Update our data source to isolate the shot we're currently working on, so that expanded shots don't interfere with each other.
 	virtual void SetSoloShot(const TObjectPtr<UMoviePipelineExecutorShot>& InShot);
+	
 	/*
 	* Expand our data source (for a specific shot) for various rendering features (such as handle frames). Gets called twice, once where it 
 	* expands the ranges, and again later to actually expand the tracks. This is required due to wanting "Handle Frames" to get counted in total frame counts.
@@ -125,6 +143,10 @@ protected:
 	virtual void ExpandShot(const TObjectPtr<UMoviePipelineExecutorShot>& InShot, const int32 InNumHandleFrames, const bool bInHasMultipleTemporalSamples, const bool bIsPrePass,
 		const FFrameRate& InDisplayRate, const FFrameRate& InTickResolution, const int32 InWarmUpFrames);
 
+	/** Gets a setting across all active render layers. Key = branch name, value = node object. */
+	template<typename T>
+	TArray<TPair<FName, T*>> GetSettingForActiveRenderLayers(const bool bIncludeCDOs, const bool bExactMatch);
+
 	// UMoviePipelineBase Interface
 	virtual void RequestShutdownImpl(bool bIsError) override;
 	virtual void ShutdownImpl(bool bIsError) override;
diff --git a/Engine/Plugins/MovieScene/MovieRenderPipeline/Source/MovieRenderPipelineCore/Public/Graph/Nodes/MovieGraphCommandLineEncoderNode.h b/Engine/Plugins/MovieScene/MovieRenderPipeline/Source/MovieRenderPipelineCore/Public/Graph/Nodes/MovieGraphCommandLineEncoderNode.h
new file mode 100644
--- /dev/null
+++ b/Engine/Plugins/MovieScene/MovieRenderPipeline/Source/MovieRenderPipelineCore/Public/Graph/Nodes/MovieGraphCommandLineEncoderNode.h
@@ -0,0 +1,195 @@
+﻿// Copyright Epic Games, Inc. All Rights Reserved.
+
+#pragma once
+
+#include "Graph/MovieGraphNode.h"
+
+#include "MovieGraphCommandLineEncoderNode.generated.h"
+
+class UMovieGraphPipeline;
+struct FMovieGraphRenderOutputData;
+
+/**
+ * A node which kicks off an encode process after all renders have completed.
+ */
+UCLASS()
+class MOVIERENDERPIPELINECORE_API UMovieGraphCommandLineEncoderNode : public UMovieGraphPostRenderNode
+{
+	GENERATED_BODY()
+
+public:
+	UMovieGraphCommandLineEncoderNode();
+
+#if WITH_EDITOR
+	virtual FText GetNodeTitle(const bool bGetDescriptive = false) const override;
+	virtual FText GetMenuCategory() const override;
+	virtual FLinearColor GetNodeTitleColor() const override;
+	virtual FSlateIcon GetIconAndTint(FLinearColor& OutColor) const override;
+#endif
+
+	// UMovieGraphNode interface
+	virtual EMovieGraphBranchRestriction GetBranchRestriction() const override;
+	// ~UMovieGraphNode interface
+
+	/** Begins the encode process (as long as there are no validation errors). */
+	void StartEncodingProcess(TArray<FMovieGraphRenderOutputData>& InGeneratedData, const bool bInIsShotEncode);
+	
+	// UMovieGraphPostRenderNode interface
+	virtual void BeginExport(UMovieGraphPipeline* InMoviePipeline, const FName& InBranchName) override;
+	virtual bool HasFinishedExporting() override;
+	// ~UMovieGraphPostRenderNode interface
+
+public:
+	UPROPERTY(EditAnywhere, BlueprintReadWrite, Category = Overrides, meta = (InlineEditConditionToggle))
+	uint8 bOverride_FileNameFormatOverride : 1;
+	
+	UPROPERTY(EditAnywhere, BlueprintReadWrite, Category = Overrides, meta = (InlineEditConditionToggle))
+	uint8 bOverride_Quality : 1;
+
+	UPROPERTY(EditAnywhere, BlueprintReadWrite, Category = Overrides, meta = (InlineEditConditionToggle))
+	uint8 bOverride_bDeleteSourceFiles : 1;
+	
+	UPROPERTY(EditAnywhere, BlueprintReadWrite, Category = Overrides, meta = (InlineEditConditionToggle))
+	uint8 bOverride_bSkipEncodeOnRenderCanceled : 1;
+
+	UPROPERTY(EditAnywhere, BlueprintReadWrite, Category = Overrides, meta = (InlineEditConditionToggle))
+	uint8 bOverride_VideoCodec : 1;
+
+	UPROPERTY(EditAnywhere, BlueprintReadWrite, Category = Overrides, meta = (InlineEditConditionToggle))
+	uint8 bOverride_AudioCodec : 1;
+
+	UPROPERTY(EditAnywhere, BlueprintReadWrite, Category = Overrides, meta = (InlineEditConditionToggle))
+	uint8 bOverride_OutputFileExtension : 1;
+
+	UPROPERTY(EditAnywhere, BlueprintReadWrite, Category = Overrides, meta = (InlineEditConditionToggle))
+	uint8 bOverride_CommandLineFormat : 1;
+
+	UPROPERTY(EditAnywhere, BlueprintReadWrite, Category = Overrides, meta = (InlineEditConditionToggle))
+	uint8 bOverride_VideoInputStringFormat : 1;
+
+	UPROPERTY(EditAnywhere, BlueprintReadWrite, Category = Overrides, meta = (InlineEditConditionToggle))
+	uint8 bOverride_AudioInputStringFormat : 1;
+
+	UPROPERTY(EditAnywhere, BlueprintReadWrite, Category = Overrides, meta = (InlineEditConditionToggle))
+	uint8 bOverride_EncodeSettings : 1;
+
+	UPROPERTY(EditAnywhere, BlueprintReadWrite, Category = Overrides, meta = (InlineEditConditionToggle))
+	uint8 bOverride_bRetainInputTextFiles : 1;
+
+	/** 
+	* File name format string override. If specified it will override the FileNameFormat from the Output setting.
+	* If {shot_name} or {camera_name} is used, encoding will begin after each shot finishes rendering.
+	* Can be different from the main one in the Output setting so you can render out frames to individual
+	* shot folders but encode to one file.
+	*/
+	UPROPERTY(EditAnywhere, BlueprintReadWrite, Category = "General", meta=(EditCondition="bOverride_FileNameFormatOverride"))
+	FString FileNameFormatOverride;
+
+	/** Whether the source files should be deleted on disk after encoding. */
+	UPROPERTY(EditAnywhere, BlueprintReadWrite, Category = "General", meta=(EditCondition="bOverride_bDeleteSourceFiles"))
+	bool bDeleteSourceFiles;
+
+	/** Whether encoding should be skipped on frames produced if rendering was canceled before finishing. */
+	UPROPERTY(EditAnywhere, BlueprintReadWrite, Category = "General", meta=(EditCondition="bOverride_bSkipEncodeOnRenderCanceled"))
+	bool bSkipEncodeOnRenderCanceled;
+
+	/** Which video codec should we use? Run 'MovieRenderPipeline.DumpCLIEncoderCodecs' for options. */ 
+	UPROPERTY(EditAnywhere, BlueprintReadWrite, Category = "General", meta=(EditCondition="bOverride_VideoCodec"))
+	FString VideoCodec;
+
+	/** Which audio codec should we use? Run 'MovieRenderPipeline.DumpCLIEncoderCodecs' for options. */
+	UPROPERTY(EditAnywhere, BlueprintReadWrite, Category = "General", meta=(EditCondition="bOverride_AudioCodec"))
+	FString AudioCodec;
+
+	/** Extension for the output files. Many encoders use this to determine the container type they are placed in. Should be without dot, ie: "webm". */
+	UPROPERTY(EditAnywhere, BlueprintReadWrite, Category = "General", meta=(EditCondition="bOverride_OutputFileExtension"))
+	FString OutputFileExtension;
+
+	/** The format string used when building the final command line argument to launch. */
+	UPROPERTY(EditAnywhere, BlueprintReadWrite, Category = "Command Line Arguments", meta=(EditCondition="bOverride_CommandLineFormat", MultiLine=true))
+	FString CommandLineFormat;
+	
+	/** Format string used for each video input. */
+	UPROPERTY(EditAnywhere, BlueprintReadWrite, Category = "Command Line Arguments", meta=(EditCondition="bOverride_VideoInputStringFormat", MultiLine=true))
+	FString VideoInputStringFormat;
+	
+	/** Format string used for each audio input. */
+	UPROPERTY(EditAnywhere, BlueprintReadWrite, Category = "Command Line Arguments", meta=(EditCondition="bOverride_AudioInputStringFormat", MultiLine=true))
+	FString AudioInputStringFormat;
+	
+	/** Additional flags used for specifying encode quality. */
+	UPROPERTY(EditAnywhere, BlueprintReadWrite, Category = "Command Line Arguments", meta=(EditCondition="bOverride_EncodeSettings", MultiLine=true))
+	FString EncodeSettings;
+
+	/** Retain the intermediate audio and video input text files that are passed to the encoder. */
+	UPROPERTY(EditAnywhere, BlueprintReadWrite, Category = "Debug", meta=(EditCondition="bOverride_bRetainInputTextFiles"))
+	bool bRetainInputTextFiles;
+
+private:
+	struct FActiveJob
+	{
+		FActiveJob()
+			: ReadPipe(nullptr)
+			, WritePipe(nullptr)
+		{}
+
+		FProcHandle ProcessHandle;
+		void* ReadPipe;
+		void* WritePipe;
+
+		/** Files that should be deleted after the process has finished. */
+		TArray<FString> FilesToDelete;
+	};
+
+	struct FEncoderParams
+	{
+		FStringFormatNamedArguments NamedArguments;
+		TMap<FString, TArray<FString>> FilesByExtensionType;
+		TWeakObjectPtr<UMoviePipelineExecutorShot> Shot;
+		FMovieGraphRenderDataIdentifier RenderDataIdentifier;
+	};
+
+	/**
+	 * Returns true if all project and node settings are valid. Returns false if there are validation errors (and
+	 * provides the errors in OutErrors).
+	 */
+	bool AreSettingsValid(TArray<FText>& OutErrors) const;
+
+	/** Manages the actively running encode jobs, potentially aborting them if pipeline needs to shut down. */
+	void OnTick();
+
+	/** Returns true if there should be one export per shot, else false. */
+	bool NeedsPerShotFlushing() const;
+
+	/** Generates encoder parameters for each render layer. */
+	TMap<FMovieGraphRenderDataIdentifier, FEncoderParams> GenerateRenderLayerEncoderParams(TArray<FMovieGraphRenderOutputData>& InGeneratedData) const;
+
+	/** Gets the path to the file that the encode process should write to. */
+	FString GetResolvedOutputFilename(const FMovieGraphRenderDataIdentifier& RenderIdentifier, const TWeakObjectPtr<UMoviePipelineExecutorShot>& Shot, const FString& FileNameFormatString) const;
+
+	/**
+	 * Generates the temporary video and audio input files that will be passed to the encoder process. These inputs
+	 * typically contain a list of the files (eg, image files) that will be used in the encode process.
+	 */
+	void GenerateTemporaryEncoderInputFiles(const FEncoderParams& InParams, TArray<FString>& OutVideoInputFilePaths, TArray<FString>& OutAudioInputFilePaths) const;
+
+	/** Generates the encoder command that should be executed. */
+	FString BuildEncoderCommand(const FEncoderParams& InParams, TArray<FString>& InVideoInputFilePaths, TArray<FString>& InAudioInputFilePaths) const;
+
+	/** Starts the actual encoder process. After the process is started, it will be monitored in OnTick(). */
+	void LaunchEncoder(const FEncoderParams& InParams);
+
+	/** Gets a setting on the branch that this node exists on. */
+	template<typename T>
+	T* GetSettingOnBranch(const bool bIncludeCDOs = true, const bool bExactMatch = true) const;
+
+private:
+	/** Encode jobs which are currently running. */
+	TArray<FActiveJob> ActiveEncodeJobs;
+
+	/** The pipeline that started the export. */
+	TWeakObjectPtr<UMovieGraphPipeline> CachedPipeline;
+
+	/** The branch that the export is occurring on. */
+	FName CachedBranchName;
+};
\ No newline at end of file
diff --git a/Engine/Plugins/MovieScene/MovieRenderPipeline/Source/MovieRenderPipelineRenderPasses/Private/MovieGraphImageSequenceOutputNode.cpp b/Engine/Plugins/MovieScene/MovieRenderPipeline/Source/MovieRenderPipelineRenderPasses/Private/MovieGraphImageSequenceOutputNode.cpp
--- a/Engine/Plugins/MovieScene/MovieRenderPipeline/Source/MovieRenderPipelineRenderPasses/Private/MovieGraphImageSequenceOutputNode.cpp
+++ b/Engine/Plugins/MovieScene/MovieRenderPipeline/Source/MovieRenderPipelineRenderPasses/Private/MovieGraphImageSequenceOutputNode.cpp
@@ -134,10 +134,9 @@ void UMovieGraphImageSequenceOutputNode::OnReceiveImageDataImpl(UMovieGraphPipel
 		// Generate one string that puts the directory combined with the filename format.
 		FString FileNameFormatString = OutputSettingNode->OutputDirectory.Path / OutputSettingNode->FileNameFormat;
 
-		// ToDo: Validate the string, ie: ensure it has {render_pass} in there somewhere there are multiple render passes
-		// in the output data, include {camera_name} if there are multiple cameras for that render pass, etc. Validation
-		// should insert {file_dup} tokens so it can put them at a logical place (ie: before frame numbers?)
-		FileNameFormatString += TEXT(".{ext}");
+		constexpr bool bIncludeRenderPass = false;
+		constexpr bool bTestFrameNumber = true;
+		UE::MoviePipeline::ValidateOutputFormatString(FileNameFormatString, bIncludeRenderPass, bTestFrameNumber);
 
 		// Map the .ext to be specific to our output data.
 		TMap<FString, FString> AdditionalFormatArgs;
@@ -425,7 +424,6 @@ FString UMovieGraphImageSequenceOutputNode_EXR::ResolveOutputFilename(
 	
 	// If we have more than one resolution we'll store it as "_Add" / "_Add(1)" etc via {ExtraTag}.
 	FString FileNameFormatString = OutputSettings->FileNameFormat + "{ExtraTag}";
-	FileNameFormatString += TEXT(".{ext}");
 
 	const FString FilePathFormatString = OutputSettings->OutputDirectory.Path / FileNameFormatString;
 	
