MoviePipeline: Added a "Casts Shadows" property to the Visibility modifier.

#rnx

diff --git a/Engine/Plugins/MovieScene/MovieRenderPipeline/Source/MovieRenderPipelineCore/Private/Graph/MoviePipelineRenderLayerSubsystem.cpp b/Engine/Plugins/MovieScene/MovieRenderPipeline/Source/MovieRenderPipelineCore/Private/Graph/MoviePipelineRenderLayerSubsystem.cpp
--- a/Engine/Plugins/MovieScene/MovieRenderPipeline/Source/MovieRenderPipelineCore/Private/Graph/MoviePipelineRenderLayerSubsystem.cpp
+++ b/Engine/Plugins/MovieScene/MovieRenderPipeline/Source/MovieRenderPipelineCore/Private/Graph/MoviePipelineRenderLayerSubsystem.cpp
@@ -79,6 +79,7 @@ void UMoviePipelineMaterialModifier::UndoModifier()
 
 UMoviePipelineVisibilityModifier::UMoviePipelineVisibilityModifier()
 	: bIsHidden(false)
+	, bCastsShadows(true)
 	, bCastShadowWhileHidden(false)
 	, bAffectIndirectLightingWhileHidden(false)
 	, bHoldout(false)
@@ -94,6 +95,7 @@ void UMoviePipelineVisibilityModifier::ApplyModifier(const UWorld* World)
 	// a differing actor), so create one copy to prevent constantly re-creating structs.
 	FActorVisibilityState NewVisibilityState;
 	NewVisibilityState.bIsHidden = bIsHidden;
+	NewVisibilityState.bCastsShadows = bCastsShadows;
 	NewVisibilityState.bCastShadowWhileHidden = bCastShadowWhileHidden;
 	NewVisibilityState.bAffectIndirectLightingWhileHidden = bAffectIndirectLightingWhileHidden;
 	NewVisibilityState.bHoldout = bHoldout;
@@ -118,6 +120,7 @@ void UMoviePipelineVisibilityModifier::ApplyModifier(const UWorld* World)
 			OriginalVisibilityState.bIsHidden = Actor->IsHidden();
 			if (const UPrimitiveComponent* PrimitiveComponent = Actor->GetComponentByClass<UPrimitiveComponent>())
 			{
+				OriginalVisibilityState.bCastsShadows = PrimitiveComponent->CastShadow;
 				OriginalVisibilityState.bCastShadowWhileHidden = PrimitiveComponent->bCastHiddenShadow;
 				OriginalVisibilityState.bAffectIndirectLightingWhileHidden = PrimitiveComponent->bAffectIndirectLightingWhileHidden;
 				OriginalVisibilityState.bHoldout = PrimitiveComponent->bHoldout;
@@ -158,9 +161,12 @@ void UMoviePipelineVisibilityModifier::SetActorVisibilityState(const FActorVisib
 
 	if (UPrimitiveComponent* PrimitiveComponent = Actor->GetComponentByClass<UPrimitiveComponent>())
 	{
-		PrimitiveComponent->bCastHiddenShadow = NewVisibilityState.bCastShadowWhileHidden;
-		PrimitiveComponent->bAffectIndirectLightingWhileHidden = NewVisibilityState.bAffectIndirectLightingWhileHidden;
-		PrimitiveComponent->bHoldout = NewVisibilityState.bHoldout;
+		// TODO: These could potentially cause a large rendering penalty due to dirtying the render state; investigate potential
+		// ways to optimize this
+		PrimitiveComponent->SetCastShadow(NewVisibilityState.bCastsShadows);
+		PrimitiveComponent->SetCastHiddenShadow(NewVisibilityState.bCastShadowWhileHidden);
+		PrimitiveComponent->SetAffectIndirectLightingWhileHidden(NewVisibilityState.bAffectIndirectLightingWhileHidden);
+		PrimitiveComponent->SetHoldout(NewVisibilityState.bHoldout);
 	}
 }
 
diff --git a/Engine/Plugins/MovieScene/MovieRenderPipeline/Source/MovieRenderPipelineCore/Public/Graph/MoviePipelineRenderLayerSubsystem.h b/Engine/Plugins/MovieScene/MovieRenderPipeline/Source/MovieRenderPipelineCore/Public/Graph/MoviePipelineRenderLayerSubsystem.h
--- a/Engine/Plugins/MovieScene/MovieRenderPipeline/Source/MovieRenderPipelineCore/Public/Graph/MoviePipelineRenderLayerSubsystem.h
+++ b/Engine/Plugins/MovieScene/MovieRenderPipeline/Source/MovieRenderPipelineCore/Public/Graph/MoviePipelineRenderLayerSubsystem.h
@@ -133,6 +133,7 @@ private:
 	{
 		TSoftObjectPtr<AActor> Actor = nullptr;
 		uint8 bIsHidden : 1 = false;
+		uint8 bCastsShadows : 1 = true;
 		uint8 bCastShadowWhileHidden : 1 = false;
 		uint8 bAffectIndirectLightingWhileHidden : 1 = false;
 		uint8 bHoldout : 1 = false;
@@ -152,6 +153,10 @@ public:
 	 */
 	UPROPERTY(BlueprintReadWrite, EditAnywhere, Category = "Settings")
 	uint8 bIsHidden : 1;
+	
+	/** If true, the primitive will cast shadows. */
+	UPROPERTY(BlueprintReadWrite, EditAnywhere, Category = "Primitive")
+	uint8 bCastsShadows : 1;
 
 	/** If true, the primitive will cast shadows even if it is hidden. */
 	UPROPERTY(BlueprintReadWrite, EditAnywhere, Category = "Primitive")
