MRQ: Initial OpenColorIO CPU processor integration in the graph output nodes. Editor only.

diff --git a/Engine/Plugins/MovieScene/MovieRenderPipeline/Source/MovieRenderPipelineCore/Public/Graph/Nodes/MovieGraphOutputSettingNode.h b/Engine/Plugins/MovieScene/MovieRenderPipeline/Source/MovieRenderPipelineCore/Public/Graph/Nodes/MovieGraphOutputSettingNode.h
--- a/Engine/Plugins/MovieScene/MovieRenderPipeline/Source/MovieRenderPipelineCore/Public/Graph/Nodes/MovieGraphOutputSettingNode.h
+++ b/Engine/Plugins/MovieScene/MovieRenderPipeline/Source/MovieRenderPipelineCore/Public/Graph/Nodes/MovieGraphOutputSettingNode.h
@@ -3,6 +3,7 @@
 #pragma once
 #include "Graph/MovieGraphNode.h"
 #include "Misc/FrameRate.h"
+#include "OpenColorIOColorSpace.h"
 #include "MovieGraphOutputSettingNode.generated.h"
 
 
@@ -49,6 +50,11 @@ public:
 	UPROPERTY(EditAnywhere, BlueprintReadWrite, Category = Overrides, meta = (InlineEditConditionToggle))
 	uint8 bOverride_HandleFrameCount : 1;
 
+#if WITH_EDITORONLY_DATA
+	UPROPERTY(EditAnywhere, BlueprintReadWrite, Category = Overrides, meta = (InlineEditConditionToggle))
+	uint8 bOverride_OCIOConfiguration : 1;
+#endif
+
 	/** What directory should all of our output files be relative to. */
 	UPROPERTY(EditAnywhere, BlueprintReadWrite, Category = "File Output", meta=(EditCondition="bOverride_OutputDirectory"))
 	FDirectoryPath OutputDirectory;
@@ -88,4 +94,16 @@ public:
 	*/
 	UPROPERTY(EditAnywhere, BlueprintReadWrite, Category = "Frames", meta = (UIMin = 0, ClampMin = 0, EditCondition = "bOverride_HandleFrameCount"))
 	int32 HandleFrameCount;
+
+#if WITH_EDITORONLY_DATA
+	/**
+	* OCIO configuration/transform settings.
+	* 
+	* Note: There are differences from the previous implementation in MRQ given that we are now doing CPU-side processing.
+	* 1) This feature only works in editor-mode, since the OpenColorIO library is currently unavailable in game builds.
+	* 2) Users are now responsible for setting the renderer output space to Final Color (HDR) in Linear Working Color Space (SCS_FinalColorHDR).
+	*/
+	UPROPERTY(EditAnywhere, BlueprintReadWrite, Category = "File Output", meta = (EditCondition = "bOverride_OCIOConfiguration"))
+	FOpenColorIODisplayConfiguration OCIOConfiguration;
+#endif
 };
diff --git a/Engine/Plugins/MovieScene/MovieRenderPipeline/Source/MovieRenderPipelineRenderPasses/MovieRenderPipelineRenderPasses.Build.cs b/Engine/Plugins/MovieScene/MovieRenderPipeline/Source/MovieRenderPipelineRenderPasses/MovieRenderPipelineRenderPasses.Build.cs
--- a/Engine/Plugins/MovieScene/MovieRenderPipeline/Source/MovieRenderPipelineRenderPasses/MovieRenderPipelineRenderPasses.Build.cs
+++ b/Engine/Plugins/MovieScene/MovieRenderPipeline/Source/MovieRenderPipelineRenderPasses/MovieRenderPipelineRenderPasses.Build.cs
@@ -40,5 +40,15 @@ public class MovieRenderPipelineRenderPasses : ModuleRules
 
 		// Required for UEOpenExr
 		AddEngineThirdPartyPrivateStaticDependencies(Target, "zlib");
+
+		if (Target.bBuildEditor == true)
+		{
+			PrivateDependencyModuleNames.AddRange(
+				new string[] {
+					"ImageCore",
+					"OpenColorIOWrapper"
+				}
+			);
+		}
 	}
 }
diff --git a/Engine/Plugins/MovieScene/MovieRenderPipeline/Source/MovieRenderPipelineRenderPasses/Private/MovieGraphImageSequenceOutputNode.cpp b/Engine/Plugins/MovieScene/MovieRenderPipeline/Source/MovieRenderPipelineRenderPasses/Private/MovieGraphImageSequenceOutputNode.cpp
--- a/Engine/Plugins/MovieScene/MovieRenderPipeline/Source/MovieRenderPipelineRenderPasses/Private/MovieGraphImageSequenceOutputNode.cpp
+++ b/Engine/Plugins/MovieScene/MovieRenderPipeline/Source/MovieRenderPipelineRenderPasses/Private/MovieGraphImageSequenceOutputNode.cpp
@@ -11,10 +11,69 @@
 #include "Graph/MovieGraphBlueprintLibrary.h"
 #include "Modules/ModuleManager.h"
 #include "MoviePipelineUtils.h"
+#include "MovieRenderPipelineCoreModule.h"
 #include "ImageWriteQueue.h"
 #include "Misc/Paths.h"
 #include "Async/TaskGraphInterfaces.h"
 
+#if WITH_EDITOR
+#include "ImageCore.h" // For GetImageView()
+#include "OpenColorIOConfiguration.h"
+#include "OpenColorIOColorTransform.h"
+#include "OpenColorIOWrapper.h"
+
+namespace UE::MovieGraph::Private
+{
+	struct FOpenColorIOPixelPreProcessor
+	{
+		FOpenColorIOPixelPreProcessor(FOpenColorIOWrapperProcessor&& InProcessor)
+			: Processor(InProcessor)
+		{ }
+
+		void operator()(FImagePixelData* PixelData)
+		{
+			check(PixelData);
+			Processor.TransformImage(PixelData->GetImageView());
+		}
+
+		FOpenColorIOWrapperProcessor Processor;
+	};
+
+	/**
+	 * Convenience function to create an OpenColorIO CPU processor based on the specified conversion settings. Editor-only.
+	 *
+	 * @return The pixel preprocessor if successful, nullptr otherwise.
+	*/
+	static FPixelPreProcessor CreateOpenColorIOPixelPreProcessor(const FOpenColorIOColorConversionSettings& InConversionSettings)
+	{
+		const TObjectPtr<UOpenColorIOConfiguration>& ConfigurationSource = InConversionSettings.ConfigurationSource;
+		if (IsValid(ConfigurationSource))
+		{
+			TObjectPtr<const UOpenColorIOColorTransform> ColorTransform = ConfigurationSource->FindTransform(InConversionSettings);
+			if (IsValid(ColorTransform))
+			{
+				FOpenColorIOWrapperProcessor Processor;
+				if (ColorTransform->GetTransformProcessor(Processor))
+				{
+					return FOpenColorIOPixelPreProcessor(MoveTemp(Processor));
+				}
+			}
+			else
+			{
+				UE_LOG(LogMovieRenderPipeline, Warning, TEXT("Invalid conversion settings, bypassing OpenColorIO transform."));
+			}
+		}
+		else
+		{
+			UE_LOG(LogMovieRenderPipeline, Warning, TEXT("Invalid configuration source, bypassing OpenColorIO transform."));
+		}
+
+		return {};
+	}
+} //end namespace UE::MovieGraph::Private
+
+#endif
+
 UMovieGraphImageSequenceOutputNode::UMovieGraphImageSequenceOutputNode()
 {
 	ImageWriteQueue = &FModuleManager::Get().LoadModuleChecked<IImageWriteQueueModule>("ImageWriteQueue").GetWriteQueue();
@@ -178,6 +237,17 @@ void UMovieGraphImageSequenceOutputNode::OnReceiveImageDataImpl(UMovieGraphPipel
 		TileImageTask->Filename = FileName;
 		TileImageTask->PixelData = RenderData.Value->CopyImageData();
 
+#if WITH_EDITOR
+		if (OutputSettingNode->OCIOConfiguration.bIsEnabled)
+		{
+			FPixelPreProcessor OCIOPixelPreProcessor = UE::MovieGraph::Private::CreateOpenColorIOPixelPreProcessor(OutputSettingNode->OCIOConfiguration.ColorConfiguration);
+			if (OCIOPixelPreProcessor)
+			{
+				TileImageTask->PixelPreProcessors.Emplace(MoveTemp(OCIOPixelPreProcessor));
+			}
+		}
+#endif
+
 		EImagePixelType PixelType = TileImageTask->PixelData->GetType();
 
 		// Perform compositing if any composited passes were found earlier
