MoviePipeline: Updated the audio renderer to ensure the audio is correct when more than one temporal sample is used.

diff --git a/Engine/Plugins/MovieScene/MovieRenderPipeline/Source/MovieRenderPipelineCore/Private/MoviePipelineAudioRendering.cpp b/Engine/Plugins/MovieScene/MovieRenderPipeline/Source/MovieRenderPipelineCore/Private/MoviePipelineAudioRendering.cpp
--- a/Engine/Plugins/MovieScene/MovieRenderPipeline/Source/MovieRenderPipelineCore/Private/MoviePipelineAudioRendering.cpp
+++ b/Engine/Plugins/MovieScene/MovieRenderPipeline/Source/MovieRenderPipelineCore/Private/MoviePipelineAudioRendering.cpp
@@ -170,14 +170,12 @@ void UMoviePipeline::ProcessAudioTick()
 		AudioDeltaTime = GetPipelinePrimaryConfig()->GetEffectiveFrameRate(TargetSequence).AsInterval();
 	}
 
-	{
-		// Handle any game logic that changed Audio State.
-		MixerDevice->Update(true);
-	}
+	// Handle any game logic that changed Audio State.
+	MixerDevice->Update(true);
 
+	if (bCanRenderAudio)
 	{
 		// Process work that has been submitted from the game thread to the audio thread over the temporal samples of this frame.
-		FFrameRate OutputFrameRate = GetPipelinePrimaryConfig()->GetEffectiveFrameRate(TargetSequence);
-		NRTPlatform->RenderAudio(OutputFrameRate.AsInterval());
+		NRTPlatform->RenderAudio(AudioDeltaTime);
 	}
 }
