MoviePipeline: Added a new "default" graph asset to the plugin, which is copied when a new graph asset is created. Also allowed the Project Settings to configure what the default graph is.

Updated the asset drop-down in the queue UI to allow a user to create a new graph asset and assign it to the job in one action.

#rnx

diff --git a/Engine/Plugins/MovieScene/MovieRenderPipeline/Source/MovieRenderPipelineEditor/Private/Graph/MovieGraphConfigFactory.cpp b/Engine/Plugins/MovieScene/MovieRenderPipeline/Source/MovieRenderPipelineEditor/Private/Graph/MovieGraphConfigFactory.cpp
--- a/Engine/Plugins/MovieScene/MovieRenderPipeline/Source/MovieRenderPipelineEditor/Private/Graph/MovieGraphConfigFactory.cpp
+++ b/Engine/Plugins/MovieScene/MovieRenderPipeline/Source/MovieRenderPipelineEditor/Private/Graph/MovieGraphConfigFactory.cpp
@@ -1,9 +1,11 @@
 // Copyright Epic Games, Inc. All Rights Reserved.
 
 #include "Graph/MovieGraphConfigFactory.h"
+
 #include "Graph/MovieGraphConfig.h"
-#include "AssetTypeCategories.h"
 #include "HAL/IConsoleManager.h"
+#include "MovieRenderPipelineCoreModule.h"
+#include "MovieRenderPipelineSettings.h"
 
 #include UE_INLINE_GENERATED_CPP_BY_NAME(MovieGraphConfigFactory)
 
@@ -16,6 +18,28 @@ UMovieGraphConfigFactory::UMovieGraphConfigFactory()
 
 UObject* UMovieGraphConfigFactory::FactoryCreateNew(UClass* Class, UObject* InParent, FName Name, EObjectFlags Flags, UObject* Context, FFeedbackContext* Warn)
 {
+	const UMovieRenderPipelineProjectSettings* ProjectSettings = GetDefault<UMovieRenderPipelineProjectSettings>();
+	const TSoftObjectPtr<UMovieGraphConfig> ProjectDefaultGraph = ProjectSettings->DefaultGraph;
+
+	// Duplicate the default graph asset provided in Project Settings
+	if (const UMovieGraphConfig* DefaultGraph = ProjectDefaultGraph.LoadSynchronous())
+	{
+		return DuplicateObject<UMovieGraphConfig>(DefaultGraph, InParent, Name);
+	}
+
+	// If the default couldn't be loaded, try loading the default supplied by MRQ. Note that this could be the same as
+	// the default above if it wasn't changed by the user; failure to load this indicates a larger issue.
+	const FSoftObjectPath DefaultGraphPath(UMovieRenderPipelineProjectSettings::GetDefaultGraphPath());
+	if (const UMovieGraphConfig* DefaultGraph = Cast<UMovieGraphConfig>(DefaultGraphPath.TryLoad()))
+	{
+		UE_LOG(LogMovieRenderPipeline, Warning, TEXT("Could not load the default graph [%s] specified in Project Settings. Falling back to MRQ-supplied default."), *DefaultGraphPath.GetAssetPathString());
+		
+		return DuplicateObject<UMovieGraphConfig>(DefaultGraph, InParent, Name);
+	}
+
+	// This should never happen, but create an empty graph as a last resort (which is most likely just an Input node and an Output node).
+	UE_LOG(LogMovieRenderPipeline, Warning, TEXT("Could not load the default graph [%s] supplied by MRQ. Falling back to an empty graph."), *DefaultGraphPath.GetAssetPathString());
+	
 	return NewObject<UMovieGraphConfig>(InParent, Class, Name, Flags);
 }
 
diff --git a/Engine/Plugins/MovieScene/MovieRenderPipeline/Source/MovieRenderPipelineEditor/Private/MovieRenderPipelineSettings.cpp b/Engine/Plugins/MovieScene/MovieRenderPipeline/Source/MovieRenderPipelineEditor/Private/MovieRenderPipelineSettings.cpp
--- a/Engine/Plugins/MovieScene/MovieRenderPipeline/Source/MovieRenderPipelineEditor/Private/MovieRenderPipelineSettings.cpp
+++ b/Engine/Plugins/MovieScene/MovieRenderPipeline/Source/MovieRenderPipelineEditor/Private/MovieRenderPipelineSettings.cpp
@@ -1,6 +1,7 @@
 // Copyright Epic Games, Inc. All Rights Reserved.
 
 #include "MovieRenderPipelineSettings.h"
+
 #include "MoviePipelinePIEExecutor.h"
 #include "MoviePipelineNewProcessExecutor.h"
 #include "MoviePipeline.h"
@@ -17,7 +18,15 @@ UMovieRenderPipelineProjectSettings::UMovieRenderPipelineProjectSettings()
 	DefaultRemoteExecutor = UMoviePipelineNewProcessExecutor::StaticClass();
 	DefaultExecutorJob = UMoviePipelineExecutorJob::StaticClass();
 	DefaultPipeline = UMoviePipeline::StaticClass();
+	DefaultGraph = GetDefaultGraphPath();
 
 	DefaultClasses.Add(UMoviePipelineImageSequenceOutput_JPG::StaticClass());
 	DefaultClasses.Add(UMoviePipelineDeferredPassBase::StaticClass());
 }
+
+FSoftObjectPath UMovieRenderPipelineProjectSettings::GetDefaultGraphPath()
+{
+	static const FSoftObjectPath DefaultGraphPath(TEXT("/MovieRenderPipeline/DefaultRenderGraph.DefaultRenderGraph"));
+	
+	return DefaultGraphPath;
+}
diff --git a/Engine/Plugins/MovieScene/MovieRenderPipeline/Source/MovieRenderPipelineEditor/Private/Widgets/SMoviePipelineQueueEditor.cpp b/Engine/Plugins/MovieScene/MovieRenderPipeline/Source/MovieRenderPipelineEditor/Private/Widgets/SMoviePipelineQueueEditor.cpp
--- a/Engine/Plugins/MovieScene/MovieRenderPipeline/Source/MovieRenderPipelineEditor/Private/Widgets/SMoviePipelineQueueEditor.cpp
+++ b/Engine/Plugins/MovieScene/MovieRenderPipeline/Source/MovieRenderPipelineEditor/Private/Widgets/SMoviePipelineQueueEditor.cpp
@@ -14,6 +14,7 @@
 #include "MoviePipelineCommands.h"
 #include "MoviePipelineEditorBlueprintLibrary.h"
 #include "Graph/MovieGraphConfig.h"
+#include "Graph/MovieGraphConfigFactory.h"
 
 // Slate Includes
 #include "Widgets/Input/SComboButton.h"
@@ -42,6 +43,7 @@
 #include "AssetRegistry/AssetData.h"
 
 // Misc
+#include "AssetToolsModule.h"
 #include "LevelSequence.h"
 #include "Engine/EngineTypes.h"
 #include "Framework/Application/SlateApplication.h"
@@ -58,6 +60,28 @@ struct FMoviePipelineMapTreeItem;
 struct FMoviePipelineShotItem;
 class SQueueJobListRow;
 
+namespace UE::MovieGraph::Private
+{
+	/** Returns a new saved UMovieGraphConfig if one could be created, else nullptr. */
+	UMovieGraphConfig* CreateNewSavedGraphAsset()
+	{
+		if (UMovieGraphConfigFactory* GraphFactory = NewObject<UMovieGraphConfigFactory>())
+		{
+			// Make the new graph via save dialog
+			const FAssetToolsModule& AssetToolsModule = FAssetToolsModule::GetModule();
+			UObject* NewAsset = AssetToolsModule.Get().CreateAssetWithDialog(GraphFactory->GetSupportedClass(), GraphFactory);
+
+			// Don't ensure here because a "cancel" in the dialog can cause the returned asset to be null
+			if (UMovieGraphConfig* NewGraph = Cast<UMovieGraphConfig>(NewAsset))
+			{
+				return NewGraph;
+			}
+		}
+
+		return nullptr;
+	}
+}
+
 struct IMoviePipelineQueueTreeItem : TSharedFromThis<IMoviePipelineQueueTreeItem>
 {
 	virtual ~IMoviePipelineQueueTreeItem() {}
@@ -289,6 +313,18 @@ public:
 		Job->SetGraphConfig(NewGraph);
 	}
 
+	void OnCreateNewGraphAndAssign() const
+	{
+		if (const UMovieGraphConfig* NewGraph = UE::MovieGraph::Private::CreateNewSavedGraphAsset())
+		{
+			UMoviePipelineExecutorJob* Job = WeakJob.Get();
+			if (ensureMsgf(Job, TEXT("Could not assign new graph to job: Job is invalid.")))
+			{
+				Job->SetGraphPreset(NewGraph);
+			}
+		}
+	}
+
 	EVisibility GetPrimaryConfigModifiedVisibility() const
 	{
 		UMoviePipelineExecutorJob* Job = WeakJob.Get();
@@ -389,12 +425,13 @@ public:
 			WeakJob,
 			FOnAssetSelected::CreateRaw(this, &FMoviePipelineQueueJobTreeItem::OnPickPresetFromAsset),
 			FExecuteAction::CreateRaw(this, &FMoviePipelineQueueJobTreeItem::OnPickNewPreset),
-			FExecuteAction::CreateRaw(this, &FMoviePipelineQueueJobTreeItem::OnReplaceWithRenderGraph)
+			FExecuteAction::CreateRaw(this, &FMoviePipelineQueueJobTreeItem::OnReplaceWithRenderGraph),
+			FExecuteAction::CreateRaw(this, &FMoviePipelineQueueJobTreeItem::OnCreateNewGraphAndAssign)
 			);
 	}
 
 	static TSharedRef<SWidget> OnGenerateConfigPresetPickerMenuFromClass(UClass* InClass,
-		TWeakObjectPtr<UMoviePipelineExecutorJob> TargetJob, FOnAssetSelected InOnAssetSelected, FExecuteAction InNewConfig, FExecuteAction InNewRenderGraph)
+		TWeakObjectPtr<UMoviePipelineExecutorJob> TargetJob, FOnAssetSelected InOnAssetSelected, FExecuteAction InNewConfig, FExecuteAction InNewRenderGraph, FExecuteAction InCreateNewGraphAndAssign)
 	{
 		FMenuBuilder MenuBuilder(true, nullptr);
 		IContentBrowserSingleton& ContentBrowser = FModuleManager::LoadModuleChecked<FContentBrowserModule>("ContentBrowser").Get();
@@ -424,18 +461,34 @@ public:
 			AssetPickerConfig.OnAssetSelected = InOnAssetSelected;
 		}
 
-		MenuBuilder.BeginSection(NAME_None, LOCTEXT("NewConfig_MenuSection", "New Configuration"));
+		const bool bIsGraphFeatureEnabled = IConsoleManager::Get().FindConsoleVariable(TEXT("MoviePipeline.EnableRenderGraph"))->GetBool();
+
+		if (bIsGraphFeatureEnabled && TargetJob->IsUsingGraphConfiguration())
+		{
+			MenuBuilder.BeginSection(NAME_None, LOCTEXT("NewConfig_MenuSection", "New Configuration"));
+			{
+				MenuBuilder.AddMenuEntry(
+					LOCTEXT("NewGraph_Label", "New Graph"),
+					LOCTEXT("NewGraph_Tooltip", "Creates a new graph asset and assigns it to this job."),
+					FSlateIcon(),
+					FUIAction(InCreateNewGraphAndAssign),
+					NAME_None,
+					EUserInterfaceActionType::Button
+				);
+			}
+		}
+
+		MenuBuilder.BeginSection(NAME_None, LOCTEXT("CurrentConfig_MenuSection", "Current Configuration"));
 		{
 			MenuBuilder.AddMenuEntry(
-				LOCTEXT("NewConfig_Label", "Clear Config"),
-				LOCTEXT("NewConfig_Tooltip", "Resets the changes to the config and goes back to the defaults."),
+				LOCTEXT("ClearConfig_Label", "Clear Config"),
+				LOCTEXT("ClearConfig_Tooltip", "Resets the changes to the config and goes back to the defaults."),
 				FSlateIcon(),
 				FUIAction(InNewConfig),
 				NAME_None,
 				EUserInterfaceActionType::Button
 			);
 
-			const bool bIsGraphFeatureEnabled = IConsoleManager::Get().FindConsoleVariable(TEXT("MoviePipeline.EnableRenderGraph"))->GetBool();
 			if (bIsGraphFeatureEnabled && !TargetJob->IsUsingGraphConfiguration())
 			{
 				MenuBuilder.AddMenuEntry(
@@ -929,6 +982,18 @@ struct FMoviePipelineShotItem : IMoviePipelineQueueTreeItem
 		Shot->SetGraphConfig(NewGraph);
 	}
 
+	void OnCreateNewGraphAndAssign() const
+	{
+		if (const UMovieGraphConfig* NewGraph = UE::MovieGraph::Private::CreateNewSavedGraphAsset())
+		{
+			UMoviePipelineExecutorShot* Shot = WeakShot.Get();
+			if (ensureMsgf(Shot, TEXT("Could not assign new graph to shot: Shot is invalid.")))
+			{
+				Shot->SetGraphPreset(NewGraph);
+			}
+		}
+	}
+
 	EVisibility GetShotConfigModifiedVisibility() const
 	{
 		UMoviePipelineExecutorShot* Shot = WeakShot.Get();
@@ -952,7 +1017,8 @@ struct FMoviePipelineShotItem : IMoviePipelineQueueTreeItem
 			WeakJob,
 			FOnAssetSelected::CreateRaw(this, &FMoviePipelineShotItem::OnPickShotPresetFromAsset),
 			FExecuteAction::CreateRaw(this, &FMoviePipelineShotItem::OnPickNewShotPreset),
-			FExecuteAction::CreateRaw(this, &FMoviePipelineShotItem::OnReplaceWithRenderGraph)
+			FExecuteAction::CreateRaw(this, &FMoviePipelineShotItem::OnReplaceWithRenderGraph),
+			FExecuteAction::CreateRaw(this, &FMoviePipelineShotItem::OnCreateNewGraphAndAssign)
 			);
 	}
 };
diff --git a/Engine/Plugins/MovieScene/MovieRenderPipeline/Source/MovieRenderPipelineEditor/Public/MovieRenderPipelineSettings.h b/Engine/Plugins/MovieScene/MovieRenderPipeline/Source/MovieRenderPipelineEditor/Public/MovieRenderPipelineSettings.h
--- a/Engine/Plugins/MovieScene/MovieRenderPipeline/Source/MovieRenderPipelineEditor/Public/MovieRenderPipelineSettings.h
+++ b/Engine/Plugins/MovieScene/MovieRenderPipeline/Source/MovieRenderPipelineEditor/Public/MovieRenderPipelineSettings.h
@@ -8,6 +8,7 @@
 #include "Engine/EngineTypes.h"
 #include "MovieRenderPipelineSettings.generated.h"
 
+class UMovieGraphConfig;
 class UMoviePipelineExecutorBase;
 class UMoviePipeline;
 class UMoviePipelinePrimaryConfig;
@@ -70,6 +71,10 @@ public:
 	UPROPERTY(config, EditAnywhere, BlueprintReadWrite, meta = (MetaClass = "/Script/MovieRenderPipelineCore.MoviePipelineBase"), Category="Movie Render Pipeline")
 	FSoftClassPath DefaultPipeline;
 
+	/** The graph that newly-created graph assets will be based off of. */
+	UPROPERTY(config, EditAnywhere, BlueprintReadWrite, Category="Movie Render Pipeline")
+	TSoftObjectPtr<UMovieGraphConfig> DefaultGraph;
+
 	/**
 	* The settings specified here will automatically be added to a Movie Pipeline Primary Configuration when using the UI. 
 	* This does not apply to scripting and does not apply to runtime. It is only a convenience function so that when a job is
@@ -80,4 +85,8 @@ public:
 	* add settings by hand for each job you create. */
 	UPROPERTY(config, EditAnywhere, BlueprintReadWrite, Category = "Movie Render Pipeline", meta = (MetaClass = "/Script/MovieRenderPipelineCore.MoviePipelineSetting"), DisplayName = "Default Job Settings Classes")
 	TArray<FSoftClassPath> DefaultClasses;
+
+public:
+	/** Gets the path to the default graph asset supplied by MRQ. */
+	static FSoftObjectPath GetDefaultGraphPath();
 };
\ No newline at end of file
