MovieGraph: Warm Up frame support. This uses a simplified system for warm-up frames compared to the Movie Pipeline implementation. We now always render warm-up frame (due to the number of systems in Unreal which now require rendering to properly warm up (Lumen, Virtual Textures, Nanite, TSR/TAA, etc.), and the user gets to choose how many frames should be rendered.

The only other choice the user now has is whether or not we emulate Motion Blur on the first frame. If true (default), we evaluate Frame 0 and pass the user specified amount of frames rendering the engine, then we evaluate Frame 1, and then finally Frame 0 again, this time actually writing it to disk. This allows users to not have any additional data before the first frame but still have motion blur. If there is real data, then setting Emulate Motion Blur to false will instead "walk" the sequence up to the first frame, ie: if the number of warmup frames is set to 10, then we will evaluate (and render) starting on Frame -10, -9, etc. and then at Frame 0 we start saving the renders to disk, and there is no 0->1->0 back and forth trick.

diff --git a/Engine/Plugins/MovieScene/MovieRenderPipeline/Source/MovieRenderPipelineCore/Private/Graph/MovieGraphCoreTimeStep.cpp b/Engine/Plugins/MovieScene/MovieRenderPipeline/Source/MovieRenderPipelineCore/Private/Graph/MovieGraphCoreTimeStep.cpp
--- a/Engine/Plugins/MovieScene/MovieRenderPipeline/Source/MovieRenderPipelineCore/Private/Graph/MovieGraphCoreTimeStep.cpp
+++ b/Engine/Plugins/MovieScene/MovieRenderPipeline/Source/MovieRenderPipelineCore/Private/Graph/MovieGraphCoreTimeStep.cpp
@@ -23,6 +23,7 @@ UMovieGraphCoreTimeStep::UMovieGraphCoreTimeStep()
 	// This set up our internal state to pick up on the first temporal sub-sample in the pattern
 	ResetForEndOfOutputFrame();
 	CurrentTimeStepData.OutputFrameNumber = 0;
+	CurrentTimeStepData.RenderedFrameNumber = 0;
 }
 
 void UMovieGraphCoreTimeStep::TickProducingFrames()
@@ -39,6 +40,7 @@ void UMovieGraphCoreTimeStep::TickProducingFrames()
 	// so there would be one frame where it used the custom timestep (after initialize) before TPF was called.
 	if (GEngine->GetCustomTimeStep() != CustomTimeStep)
 	{
+		// ToDo: This will restore the wrong timestep at the end of a render if we have different TimeStep instances.
 		PrevCustomTimeStep = GEngine->GetCustomTimeStep();
 		GEngine->SetCustomTimeStep(CustomTimeStep);
 	}
@@ -56,28 +58,33 @@ void UMovieGraphCoreTimeStep::TickProducingFrames()
 	// support, at which point we'll handle the time dilation by affecting the whole engine tick instead of just world tick.
 	//if (IsFirstTemporalSample())
 	{
-		UpdateFrameMetrics();
 
 		// Now that we've calculated the total range of time we're trying to represent, we can check to see
 		// if this would put us beyond our range of time this shot is supposed to represent.
 
 	}
-	// Handle shot initialization
+	
+	// Some state transitions (such as moving to a previous point in time when paused) require a Jump operation
+	// to correctly evaluate the new time, so setting this flag will issue an extra Jump instruction once the
+	// final time has been calculated.
+	bool bShouldJump = false;
 
+	// Handle shot initialization
 	if (CurrentCameraCut->ShotInfo.State == EMovieRenderShotState::Uninitialized)
 	{
-		UE_LOG(LogMovieRenderPipeline, Log, TEXT("MovieGraph Initializing Camera Cut [%d/%d] in [%s] %s."),
+		UE_LOG(LogMovieRenderPipeline, Log, TEXT("MovieGraph Initializing Camera Cut [%d/%d] OuterName: [%s] InnerName: %s."),
 			CurrentShotIndex + 1, ActiveShotList.Num(), *CurrentCameraCut->OuterName, *CurrentCameraCut->InnerName);
-
+		
+		// We need some evaluation context to be able to properly set up this shot, but we
+		// re-evaluate before actually doing anything.
 		FMovieGraphTraversalContext Context = GetOwningGraph()->GetCurrentTraversalContext();
 		
 		// Update global variables before evaluating the graph
 		Context.RootGraph->UpdateGlobalVariableValues(GetOwningGraph());
 
-		// Evaluate the graph so we can fetch values for this shot.
-		UMovieGraphConfig* Config = GetOwningGraph()->GetRootGraphForShot(CurrentCameraCut);
-		CurrentFrameData.EvaluatedConfig = TStrongObjectPtr<UMovieGraphEvaluatedConfig>(Config->CreateFlattenedGraph(Context));
-		CurrentFrameData.TemporalSampleIndex = 0;
+		CurrentFrameData.EvaluatedConfig = TStrongObjectPtr<UMovieGraphEvaluatedConfig>(Context.RootGraph->CreateFlattenedGraph(Context));
+		CurrentFrameData.TemporalSampleCount = GetTemporalSampleCount();
+		UpdateFrameMetrics();
 
 		// Ensure we've set it in the CurrentTimeStepData so things can fetch from it below.
 		CurrentTimeStepData.EvaluatedConfig = TObjectPtr<UMovieGraphEvaluatedConfig>(CurrentFrameData.EvaluatedConfig.Get());
@@ -95,66 +102,222 @@ void UMovieGraphCoreTimeStep::TickProducingFrames()
 		// TODO: This should probably be done per-frame
 		GetOwningGraph()->UpdateLayerContentsInRenderLayerSubsystem(CurrentTimeStepData.EvaluatedConfig);
 
-		const bool bIncludeCDO = false;
-		UMovieGraphOutputSettingNode* OutputNode = CurrentFrameData.EvaluatedConfig->GetSettingForBranch<UMovieGraphOutputSettingNode>(UMovieGraphNode::GlobalsPinName, bIncludeCDO);
+		// We have three possible options we can go to after the initialization state. If they have no warm up frames,
+		// and don't want to emulate motion blur, we go directly into the Rendering State (which starts on Frame 0, then Frame 1, etc.)
+		// If they do have warm up frames and are not emulating motion blur, we "walk" the level sequence up to the first frame.
+		// If they have warm up frames and are emulating motion blur, then we evaluate Frame 0, wait the specified number of frames,
+		// and then do motion blur emulation.
+		CurrentCameraCut->ShotInfo.WorkMetrics.TotalEngineWarmUpFrameCount = CurrentCameraCut->ShotInfo.NumEngineWarmUpFramesRemaining;
 
-		const FFrameRate TickResolution = GetOwningGraph()->GetDataSourceInstance()->GetTickResolution();
-		const FFrameRate SourceFrameRate = GetOwningGraph()->GetDataSourceInstance()->GetDisplayRate();
-		const FFrameRate FinalFrameRate = UMovieGraphBlueprintLibrary::GetEffectiveFrameRate(OutputNode, SourceFrameRate);
+		if (CurrentCameraCut->ShotInfo.NumEngineWarmUpFramesRemaining == 0 && !CurrentCameraCut->ShotInfo.bEmulateFirstFrameMotionBlur)
+		{
+			// If there's no warm-up frames and they don't want to emulate motion blur, skip to rendering state next frame.
+			CurrentCameraCut->ShotInfo.State = EMovieRenderShotState::Rendering;
+			
+			// Set up a fake "Previous" range that is just 1 frame further back than the start.
+			FFrameTime UpperBound = CurrentCameraCut->ShotInfo.CurrentTimeInRoot;
+			CurrentFrameData.LastOutputFrameRange = TRange<FFrameTime>(UpperBound - CurrentFrameMetrics.FrameTimePerOutputFrame, UpperBound);
+			CurrentFrameData.LastSampleRange = CurrentFrameData.LastOutputFrameRange;
 
-		FFrameTime FrameTimePerOutputFrame = FFrameRate::TransformTime(FFrameTime(FFrameNumber(1)), FinalFrameRate, TickResolution);
+			GetOwningGraph()->GetDataSourceInstance()->PlayDataSource();
+		}
+		else
+		{
+			const bool bHasWarmUpFrames = CurrentCameraCut->ShotInfo.NumEngineWarmUpFramesRemaining > 0;
+			const bool bEmulateMotionBlur = CurrentCameraCut->ShotInfo.bEmulateFirstFrameMotionBlur;
 
-		// Dummy range so that the first rendering frame knows to start at the right time.
-		FFrameTime UpperBound = CurrentCameraCut->ShotInfo.CurrentTimeInRoot;
-		CurrentFrameData.LastOutputFrameRange = TRange<FFrameTime>(UpperBound - FrameTimePerOutputFrame, UpperBound);
-		CurrentFrameData.LastSampleRange = CurrentFrameData.LastOutputFrameRange;
+			if(bHasWarmUpFrames)
+			{
+				CurrentCameraCut->ShotInfo.State = EMovieRenderShotState::WarmingUp;
 
-		// We can safely fall through to the below states as they're OK to process the same frame we set up.
-		UE_LOG(LogMovieRenderPipeline, Log, TEXT("MovieGraph Finished initializing Camera Cut [%d/%d] in [%s] %s."),
-			CurrentShotIndex + 1, ActiveShotList.Num(), *CurrentCameraCut->OuterName, *CurrentCameraCut->InnerName);
+				// If they have warm up frames and are not emulating motion blur, we need to actually move their start range back in the sequence.
+				int32 NumFramesToGoBack = bEmulateMotionBlur ? 0 : CurrentCameraCut->ShotInfo.NumEngineWarmUpFramesRemaining;
+				
+				FFrameTime UpperBound = CurrentCameraCut->ShotInfo.CurrentTimeInRoot;
+				FFrameTime NewStartTime = UpperBound - (CurrentFrameMetrics.FrameTimePerOutputFrame * NumFramesToGoBack);
+				CurrentFrameData.LastOutputFrameRange = TRange<FFrameTime>(NewStartTime - CurrentFrameMetrics.FrameTimePerOutputFrame, NewStartTime);
+				CurrentFrameData.LastSampleRange = CurrentFrameData.LastOutputFrameRange;
 
-		// Temp...
-		CurrentCameraCut->ShotInfo.State = EMovieRenderShotState::Rendering;
+				if(bEmulateMotionBlur)
+				{
+					GetOwningGraph()->GetDataSourceInstance()->PauseDataSource();
+					// We have to assign this in this scenario, because we don't automatically advance forward
+					// when we're warming up, but emulating motion blur (meaning we sit on frame zero)
+					TRangeBound<FFrameTime> EndOfPreviousFrame = CurrentFrameData.LastOutputFrameRange.GetUpperBound();
+					CurrentFrameData.CurrentOutputFrameRange = TRange<FFrameTime>(EndOfPreviousFrame.GetValue(), EndOfPreviousFrame.GetValue() + CurrentFrameMetrics.FrameTimePerOutputFrame);
+
+					// Because we're going backwards in time (while paused) we need to issue a jump command.
+					bShouldJump = true;
+				}
+				else
+				{
+					GetOwningGraph()->GetDataSourceInstance()->PlayDataSource();
+				}
+			}
+			else
+			{
+				// If there are no warm up frames, then they're already in the motion blur state.
+				CurrentCameraCut->ShotInfo.State = EMovieRenderShotState::MotionBlur;
+
+				// Set up our frame time to just be the first frame, with a fake past.
+				FFrameTime UpperBound = CurrentCameraCut->ShotInfo.CurrentTimeInRoot;
+				FFrameTime NewStartTime = UpperBound;
+				CurrentFrameData.LastOutputFrameRange = TRange<FFrameTime>(NewStartTime - CurrentFrameMetrics.FrameTimePerOutputFrame, NewStartTime);
+				CurrentFrameData.LastSampleRange = CurrentFrameData.LastOutputFrameRange;
+				
+				// This doens't get automatically incremented below, so we need to specify the output range they're working on.
+				CurrentFrameData.CurrentOutputFrameRange = TRange<FFrameTime>(NewStartTime, NewStartTime + CurrentFrameMetrics.FrameTimePerOutputFrame);
+
+
+				// We pause the sequence so it accurately reflects what we're actually doing with the time.
+				GetOwningGraph()->GetDataSourceInstance()->PauseDataSource();
+			}
+		}
+
+		// We intentionally fall through to the below states so that we don't have extra frames between state changes.
+		// Extra frames are problematic for things like motion blur which depend on the data from the previous engine tick.
+		UE_LOG(LogMovieRenderPipeline, Log, TEXT("MovieGraph Finished initializing Camera Cut [%d/%d] OuterName: [%s]  InnerName: %s."),
+			CurrentShotIndex + 1, ActiveShotList.Num(), *CurrentCameraCut->OuterName, *CurrentCameraCut->InnerName);
 	}
 
-	if (CurrentCameraCut->ShotInfo.State == EMovieRenderShotState::Rendering)
+	// We only automatically advance the CurrentOutputFrameRange during Rendering, or when warm-up frames
+	// are counting down (if we're not going to emulate motion blur).
+	const bool bIncrementBecauseRenderingState = CurrentCameraCut->ShotInfo.State == EMovieRenderShotState::Rendering;
+	const bool bIncrementBecauseWarmUpState = CurrentCameraCut->ShotInfo.State == EMovieRenderShotState::WarmingUp && !CurrentCameraCut->ShotInfo.bEmulateFirstFrameMotionBlur;
+	bool bIncrementInternalCounters = bIncrementBecauseRenderingState || bIncrementBecauseWarmUpState;
+	
+	// Due to the motion blur emulation frame only using one TS, we need to potentially override the flags we send to the renderer,
+	// otherwise the renderer gets confused because it's asked to render a frame that wasn't scheduled.
+	TOptional<bool> bIsFirstTemporalSampleOverride;
+	TOptional<bool> bIsLastTemporalSampleOverride;
+
+	if (IsFirstTemporalSample())
 	{
+		// Update our context and stats
 		FMovieGraphTraversalContext Context = GetOwningGraph()->GetCurrentTraversalContext();
 
 		// Update global variables before evaluating the graph
 		Context.RootGraph->UpdateGlobalVariableValues(GetOwningGraph());
+		CurrentFrameData.EvaluatedConfig = TStrongObjectPtr<UMovieGraphEvaluatedConfig>(Context.RootGraph->CreateFlattenedGraph(Context));
+
+		// The temporal sample count can change every frame due to graph evaluations, so when we're on our first temporal
+		// sub-sample of the new frame, we need to re-fetch the value. 
 		
-		if (IsFirstTemporalSample())
+		// ToDo: The time-step implementations can use the active
+		// shot state to decide not to do a full TS count (ie: during warm-up) if applicable.
+		CurrentFrameData.TemporalSampleCount = GetTemporalSampleCount();
+
+		// Re-calculate some timing statistics about the given TS count
+		UpdateFrameMetrics();
+
+		// We create a range of time that we want to represent with this frame in absolute root sequence
+		// space. This is because each frame can have a different temporal sample count, and when we jump
+		// over the time period the shutter is closed, it's no longer an easy fixed number. So instead,
+		// we build a time range we wish to represesnt with this frame, and then split that into sub-regions
+		// and simply move between pre-calculated regions. This lets us handle the fact that the delta time
+		// between output frames is going to be different when moving between different TS counts on frames.
+
+		// To figure out what range of time we want to represent, we actually need to look at the last frame's
+		// calculated total output range (ignoring all sub-sampling). The new frame is then multiplied by time
+		// dilation, so our entire output range is shortened. We can't just jump between whole frames (due to
+		// time dilation) so we have to take the last range, figure out how long the new frame is, and then
+		// build ontop of that.
+		TRangeBound<FFrameTime> EndOfPreviousFrame = CurrentFrameData.LastOutputFrameRange.GetUpperBound();
+
+		if (bIncrementInternalCounters)
 		{
-			UMovieGraphConfig* Config = GetOwningGraph()->GetRootGraphForShot(CurrentCameraCut);
-			CurrentFrameData.EvaluatedConfig = TStrongObjectPtr<UMovieGraphEvaluatedConfig>(Config->CreateFlattenedGraph(Context));
+			// ToDo: Apply time dilation by multiplying the calculated range by slow-mo duration.
+			// The upper bound is exclusive, so we initialize a new TRange with just the value to start inclusively.
+			CurrentFrameData.CurrentOutputFrameRange = TRange<FFrameTime>(EndOfPreviousFrame.GetValue(), EndOfPreviousFrame.GetValue() + CurrentFrameMetrics.FrameTimePerOutputFrame);
+		}
 
-			// The temporal sample count can change every frame due to graph evaluations, so when we're on our first temporal
-			// sub-sample of the new frame, we need to re-fetch the value.
-			CurrentFrameData.TemporalSampleCount = GetTemporalSampleCount();
+		// Update CurrentFrameData.RangeShutterOpen and CurrentFrameData.RangeShutterClosed
+		UpdateShutterRanges();
+		// Update CurrentFrameData.TemporalRanges
+		UpdateTemporalRanges();
 
-			// Re-calculate some timing statistics about the given TS count
-			UpdateFrameMetrics();
+		// Warm-up Frames follow similar logic to the main rendering. They need to support temporal sub-sampling because
+		// the Path Tracer requires many samples to produce a high quality image, and the high quality image is important
+		// to seed a denoiser (which requires looking at previous frames for temporal stability). So we now render temporal
+		// sub-samples during warm-up.
+		if (CurrentCameraCut->ShotInfo.State == EMovieRenderShotState::WarmingUp)
+		{
+			CurrentTimeStepData.bDiscardOutput = true;
+
+			// If we're going to emulate motion blur, we don't actually want to move the 
+			// evaluated point in the sequence, otherwise we repeatedly jump back and forth
+			// between temporal sub-sample locations when we're supposed to be warming up
+			// only the first frame. We could possibly do this by overriding the temporal
+			// sample count to 1, but we want to avoid having different behavior between
+			// emulating motion blur and not.
+			if (CurrentCameraCut->ShotInfo.bEmulateFirstFrameMotionBlur)
+			{
+				for (int32 Index = 1; Index < CurrentFrameData.TemporalRanges.Num(); Index++)
+				{
+					CurrentFrameData.TemporalRanges[Index] = CurrentFrameData.TemporalRanges[0];
+				}
+			}
 
-			CurrentCameraCut->ShotInfo.WorkMetrics.OutputFrameIndex++;
+			if (CurrentCameraCut->ShotInfo.NumEngineWarmUpFramesRemaining == 0)
+			{
+				GetOwningGraph()->GetDataSourceInstance()->PlayDataSource();
+				CurrentCameraCut->ShotInfo.State = CurrentCameraCut->ShotInfo.bEmulateFirstFrameMotionBlur ? EMovieRenderShotState::MotionBlur : EMovieRenderShotState::Rendering;
+			}
 
-			// We create a range of time that we want to represent with this frame in absolute root sequence
-			// space. This is because each frame can have a different temporal sample count, and when we jump
-			// over the time period the shutter is closed, it's no longer an easy fixed number. So instead,
-			// we build a time range we wish to represesnt with this frame, and then split that into sub-regions
-			// and simply move between pre-calculated regions. This lets us handle the fact that the delta time
-			// between output frames is going to be different when moving between different TS counts on frames.
-			
-			// To figure out what range of time we want to represent, we actually need to look at the last frame's
-			// calculated total output range (ignoring all sub-sampling). The new frame is then multiplied by time
-			// dilation, so our entire output range is shortened. We can't just jump between whole frames (due to
-			// time dilation) so we have to take the last range, figure out how long the new frame is, and then
-			// build ontop of that.
-			TRangeBound<FFrameTime> EndOfPreviousFrame = CurrentFrameData.LastOutputFrameRange.GetUpperBound();
-
-			// The upper bound is exclusive, so we initialize a new TRange with just the value to start inclusively
-			// ToDo: Apply time dilation by multiplying the calculated range by slow-mo duration.
-			CurrentFrameData.CurrentOutputFrameRange = TRange<FFrameTime>(EndOfPreviousFrame.GetValue(), EndOfPreviousFrame.GetValue() + CurrentFrameMetrics.FrameTimePerOutputFrame);
+			// We decrement at the end of this as we check before to see if we should move to the next state, but want to ensure we do at least one warm up frame.
+			CurrentCameraCut->ShotInfo.WorkMetrics.EngineWarmUpFrameIndex++;
+			CurrentCameraCut->ShotInfo.NumEngineWarmUpFramesRemaining--;
+		}
+
+		if (CurrentCameraCut->ShotInfo.State == EMovieRenderShotState::MotionBlur)
+		{
+			CurrentTimeStepData.bDiscardOutput = true;
+
+			// This fires the first frame of motion blur evaluation, where we want to set their position to the end of the first temporal sub-sample
+			if (!CurrentCameraCut->ShotInfo.bHasEvaluatedMotionBlurFrame)
+			{
+				// Instead of jumping all the way to the end of the first temporal sub-sample, we jump to one tick less. 
+				// This is because if you have a sequence that is 1 frame long, and 1 TS, the end of the Temporal Sub-Sample
+				// will be out of bounds for evaluation, and you won't properly emulate motion blur.
+				FFrameTime TemporalUpperBound = CurrentFrameData.TemporalRanges[0].GetUpperBoundValue() - FFrameTime(FFrameNumber(1));
+
+				// The code below is going to pick the lower bound of the next temporal range (which will be the first range)
+				// We're going to rewrite the range to just be our nearly empty range, and then the rest of the timing should
+				// just fall out.
+				TRange<FFrameTime> NewRange = TRange<FFrameTime>(TemporalUpperBound, CurrentFrameData.TemporalRanges[0].GetUpperBoundValue());
+				CurrentFrameData.TemporalRanges.Empty();
+				CurrentFrameData.TemporalRanges.Add(NewRange);
+				bIsFirstTemporalSampleOverride = true;
+				bIsLastTemporalSampleOverride = true;
+
+				GetOwningGraph()->GetDataSourceInstance()->PlayDataSource();
+				CurrentCameraCut->ShotInfo.bHasEvaluatedMotionBlurFrame = true;
+			}
+			else
+			{
+				// Update the temporal ranges for this output frame to match their actual TS settings again
+				// (we modified them above to make the motion blur emulation frame jump to a non-standard time).
+				UpdateTemporalRanges();
+
+				// We have to override this back to true because we're falling through into the Rendering State,
+				// and the rendering state should increment.
+				bIncrementInternalCounters = true;
+
+				// We've done the motion blur emulation frame (above), but we didn't increment internal counters, so
+				// we can safely fall through to rendering here and it'll start on the rendering frame like normal. 
+				// We _must_ fall through to rendering this tick, as we need the motion vectors generated by the difference 
+				// between last frame and this one to be used for the renders this frame.
+				CurrentCameraCut->ShotInfo.State = EMovieRenderShotState::Rendering;
+				GetOwningGraph()->GetDataSourceInstance()->PlayDataSource();
+			}
+		}
+
+		if (CurrentCameraCut->ShotInfo.State == EMovieRenderShotState::Rendering)
+		{
+			CurrentTimeStepData.bDiscardOutput = false;
+
+			// We only track work metrics during Rendering
+			CurrentCameraCut->ShotInfo.WorkMetrics.OutputFrameIndex++;
 
 			// Now that we've calculated the total range of time we're trying to represent, we can check to see
 			// if this would put us beyond our range of time this shot is supposed to represent.
@@ -172,152 +335,194 @@ void UMovieGraphCoreTimeStep::TickProducingFrames()
 				GetOwningGraph()->TeardownShot(CurrentCameraCut);
 				return;
 			}
+		}
+	}
 
-			// Figure out ranges for how long the shutter is open and closed
-			{
-				TArray<TRange<FFrameTime>> SplitRange = CurrentFrameData.CurrentOutputFrameRange.Split(CurrentFrameData.CurrentOutputFrameRange.GetLowerBoundValue() + CurrentFrameMetrics.FrameTimeWhileShutterOpen);
-				if (ensure(SplitRange.Num() == 2))
-				{
-					CurrentFrameData.RangeShutterOpen = SplitRange[0];
-					CurrentFrameData.RangeShutterClosed = SplitRange[1];
-				}
-			}
-			
-			// Now split the range the shutter is open per temporal sample.
-			{
-				CurrentFrameData.TemporalRanges.Reset();
-				CurrentFrameData.TemporalRanges.Reserve(CurrentFrameData.TemporalSampleCount);
-				FFrameTime SplitStartTime = CurrentFrameData.RangeShutterOpen.GetLowerBoundValue();
-				TRange<FFrameTime> RemainingRange = CurrentFrameData.RangeShutterOpen;
-				for (int32 Index = 0; Index < CurrentFrameData.TemporalSampleCount - 1; Index++)
-				{
-					SplitStartTime += CurrentFrameMetrics.FrameTimePerTemporalSample;
-					TArray<TRange<FFrameTime>> NewRanges = RemainingRange.Split(SplitStartTime);
-					if (ensure(NewRanges.Num() == 2))
-					{
-						CurrentFrameData.TemporalRanges.Add(NewRanges[0]);
-						RemainingRange = NewRanges[1];
-					}
-				}
+	// The delta time for this frame is the difference between the current range index, and the last range index.
+	const TRange<FFrameTime>& PreviousRange = CurrentFrameData.LastSampleRange;
+	const TRange<FFrameTime>& NextRange = CurrentFrameData.TemporalRanges[GetNextTemporalRangeIndex()];
 
-				// Add the remaining range, it's the leftover from the last split.
-				CurrentFrameData.TemporalRanges.Add(RemainingRange);
+	// Because some time-steps may not advance linearly through time, the abs value of the ranges needs to be taken
+	// to avoid negative delta times.
+	FFrameTime FrameDeltaTime = NextRange.GetLowerBoundValue() > PreviousRange.GetLowerBoundValue()
+		? NextRange.GetLowerBoundValue() - PreviousRange.GetLowerBoundValue()
+		: PreviousRange.GetLowerBoundValue() - NextRange.GetLowerBoundValue();
 
-				//for (int32 Index = 0; Index < CurrentFrameData.TemporalRanges.Num(); Index++)
-				//{
-				//	UE_LOG(LogTemp, Warning, TEXT("Range: [%s, %s)"),
-				//		*LexToString(CurrentFrameData.TemporalRanges[Index].GetLowerBoundValue()),
-				//		*LexToString(CurrentFrameData.TemporalRanges[Index].GetUpperBoundValue()));
-				//}
-			}
-		}
-		
-		// The delta time for this frame is the difference between the current range index, and the last range index.
-		const TRange<FFrameTime>& PreviousRange = CurrentFrameData.LastSampleRange;
-		const TRange<FFrameTime>& NextRange = CurrentFrameData.TemporalRanges[GetNextTemporalRangeIndex()];
-
-		// Because some time-steps may not advance linearly through time, the abs value of the ranges needs to be taken
-		// to avoid negative delta times.
-		FFrameTime FrameDeltaTime = NextRange.GetLowerBoundValue() > PreviousRange.GetLowerBoundValue()
-			? NextRange.GetLowerBoundValue() - PreviousRange.GetLowerBoundValue()
-			: PreviousRange.GetLowerBoundValue() - NextRange.GetLowerBoundValue();
-
-		// ToDo: Propagate delta time multipliers to cloth
-
-		// Because we know what time range we're supposed to represent, we can just assign the CurrentTimeInRoot absolutely,
-		// instead of accumulating delta times. // ToDo: shutter timing, motion blur offset?
-		CurrentCameraCut->ShotInfo.CurrentTimeInRoot = NextRange.GetLowerBoundValue();
-
-		FFrameTime FinalEvalTime = CurrentCameraCut->ShotInfo.CurrentTimeInRoot + CurrentFrameMetrics.MotionBlurCenteringOffsetTime + CurrentFrameMetrics.ShutterOffsetFrameTime;
-		UE_LOG(LogMovieRenderPipeline, VeryVerbose, TEXT("FinalEvalTime: %s (Tick: %s SOffset: %s MOffset: %s)"),
-			*LexToString(FinalEvalTime),
-			*LexToString(CurrentCameraCut->ShotInfo.CurrentTimeInRoot), 
-			*LexToString(CurrentFrameMetrics.ShutterOffsetFrameTime),
-			*LexToString(CurrentFrameMetrics.MotionBlurCenteringOffsetTime));
-
-		// Now we need to fill out some of our current timestep data for the renderer portion to use.
-		double FrameDeltaTimeAsSeconds = CurrentFrameMetrics.TickResolution.AsSeconds(FrameDeltaTime);
-		CurrentTimeStepData.FrameDeltaTime = FrameDeltaTimeAsSeconds;
-		CurrentTimeStepData.WorldSeconds = CurrentTimeStepData.WorldSeconds + FrameDeltaTimeAsSeconds;
-		CurrentTimeStepData.WorldTimeDilation = 1.0; // Temp
-
-		// The combination of shutter angle percentage, non-uniform render frame delta times and dividing by sample
-		// count produce the correct length for motion blur in all cases.
-		CurrentTimeStepData.MotionBlurFraction = CurrentFrameMetrics.MotionBlurAmount / CurrentFrameData.TemporalSampleCount;
-		CurrentTimeStepData.bIsFirstTemporalSampleForFrame = IsFirstTemporalSample();
-		CurrentTimeStepData.bIsLastTemporalSampleForFrame = IsLastTemporalSample();
-		CurrentTimeStepData.bRequiresAccumulator = CurrentFrameData.TemporalSampleCount > 1;
-		CurrentTimeStepData.OutputFrameNumber = CurrentFrameData.OutputFrameNumber;
-		CurrentTimeStepData.EvaluatedConfig = TObjectPtr<UMovieGraphEvaluatedConfig>(CurrentFrameData.EvaluatedConfig.Get());
+	// There are valid scenarios where FrameDeltaTime is zero (such as warming up and using motion blur emulation,
+	// the sequence isn't actually progressing forward), but we can't propagate zero delta times into the engine.
+	if (FrameDeltaTime.GetFrame() == 0 && FrameDeltaTime.GetSubFrame() == 0.f)
+	{
+		FrameDeltaTime = CurrentFrameData.TemporalRanges[0].Size<FFrameTime>();
+	}
+	// ToDo: Propagate delta time multipliers to cloth
+
+	// Because we know what time range we're supposed to represent, we can just assign the CurrentTimeInRoot absolutely,
+	// instead of accumulating delta times.
+	CurrentCameraCut->ShotInfo.CurrentTimeInRoot = NextRange.GetLowerBoundValue();
+
+	FFrameTime FinalEvalTime = CurrentCameraCut->ShotInfo.CurrentTimeInRoot + CurrentFrameMetrics.MotionBlurCenteringOffsetTime + CurrentFrameMetrics.ShutterOffsetFrameTime;
+	UE_LOG(LogMovieRenderPipeline, VeryVerbose, TEXT("FinalEvalTime: %s (Tick: %s SOffset: %s MOffset: %s)"),
+		*LexToString(FinalEvalTime),
+		*LexToString(CurrentCameraCut->ShotInfo.CurrentTimeInRoot), 
+		*LexToString(CurrentFrameMetrics.ShutterOffsetFrameTime),
+		*LexToString(CurrentFrameMetrics.MotionBlurCenteringOffsetTime));
+
+	// Now we need to fill out some of our current timestep data for the renderer portion to use.
+ 	double FrameDeltaTimeAsSeconds = CurrentFrameMetrics.TickResolution.AsSeconds(FrameDeltaTime);
+	CurrentTimeStepData.FrameDeltaTime = FrameDeltaTimeAsSeconds;
+	CurrentTimeStepData.WorldSeconds = CurrentTimeStepData.WorldSeconds + FrameDeltaTimeAsSeconds;
+	CurrentTimeStepData.WorldTimeDilation = 1.0; // Temp
+	CurrentTimeStepData.FrameRate = CurrentFrameMetrics.FrameRate;
+
+	// The combination of shutter angle percentage, non-uniform render frame delta times and dividing by sample
+	// count produce the correct length for motion blur in all cases.
+	CurrentTimeStepData.MotionBlurFraction = CurrentFrameMetrics.MotionBlurAmount / CurrentFrameData.TemporalSampleCount;
+	CurrentTimeStepData.bIsFirstTemporalSampleForFrame = bIsFirstTemporalSampleOverride.Get(IsFirstTemporalSample());
+	CurrentTimeStepData.bIsLastTemporalSampleForFrame = bIsLastTemporalSampleOverride.Get(IsLastTemporalSample());
+	CurrentTimeStepData.bRequiresAccumulator = CurrentFrameData.TemporalSampleCount > 1 && !CurrentTimeStepData.bDiscardOutput;
+	CurrentTimeStepData.OutputFrameNumber = CurrentFrameData.OutputFrameNumber;
+	CurrentTimeStepData.RenderedFrameNumber = CurrentFrameData.RenderedFrameNumber;
+	CurrentTimeStepData.EvaluatedConfig = TObjectPtr<UMovieGraphEvaluatedConfig>(CurrentFrameData.EvaluatedConfig.Get());
+
+	//UE_LOG(LogTemp, Warning, TEXT("F# %d bFirst: %d bLast: %d bReqAc: %d"),
+	//	CurrentTimeStepData.OutputFrameNumber, CurrentTimeStepData.bIsFirstTemporalSampleForFrame,
+	//	CurrentTimeStepData.bIsLastTemporalSampleForFrame, CurrentTimeStepData.bRequiresAccumulator);
+
+	// Calculate frame numbers and timecodes for the current sequence (root) and shot
+	{
+		constexpr bool bIncludeCDOs = true;
+		UMovieGraphOutputSettingNode* OutputSetting = CurrentFrameData.EvaluatedConfig->GetSettingForBranch<UMovieGraphOutputSettingNode>(UMovieGraphNode::GlobalsPinName, bIncludeCDOs);
 
-		//UE_LOG(LogTemp, Warning, TEXT("F# %d bFirst: %d bLast: %d bReqAc: %d"),
-		//	CurrentTimeStepData.OutputFrameNumber, CurrentTimeStepData.bIsFirstTemporalSampleForFrame,
-		//	CurrentTimeStepData.bIsLastTemporalSampleForFrame, CurrentTimeStepData.bRequiresAccumulator);
+		// "Closest" isn't straightforward when using temporal sub-sampling, ie: A large enough shutter angle pushes a sample over the half way point and it rounds to
+		// the wrong one. Because temporal sub-sampling isn't centered around a frame (the centering is done via the final eval time) we can just subtract TSI*TPS to get our centered value.
+		const FFrameTime CenteringOffset = CurrentFrameData.TemporalSampleIndex * CurrentFrameMetrics.FrameTimePerTemporalSample;
+		FFrameTime CenteredFrameTime = CurrentCameraCut->ShotInfo.CurrentTimeInRoot - CenteringOffset;
+	
+		const FFrameRate SourceFrameRate = GetOwningGraph()->GetDataSourceInstance()->GetDisplayRate();
+		const FFrameRate EffectiveFrameRate = UMovieGraphBlueprintLibrary::GetEffectiveFrameRate(OutputSetting, SourceFrameRate);
+		const FFrameRate TickResolution = GetOwningGraph()->GetDataSourceInstance()->GetTickResolution();
 
-		// Calculate frame numbers and timecodes for the current sequence (root) and shot
-		{
-			constexpr bool bIncludeCDOs = true;
-			UMovieGraphOutputSettingNode* OutputSetting = CurrentFrameData.EvaluatedConfig->GetSettingForBranch<UMovieGraphOutputSettingNode>(UMovieGraphNode::GlobalsPinName, bIncludeCDOs);
+		constexpr bool bDropFrame = false;
+		CurrentTimeStepData.RootFrameNumber = FFrameRate::TransformTime(CenteredFrameTime, TickResolution, EffectiveFrameRate).RoundToFrame();
+		CurrentTimeStepData.RootTimeCode = FTimecode::FromFrameNumber(CurrentTimeStepData.RootFrameNumber, EffectiveFrameRate, bDropFrame);
 
-			// "Closest" isn't straightforward when using temporal sub-sampling, ie: A large enough shutter angle pushes a sample over the half way point and it rounds to
-			// the wrong one. Because temporal sub-sampling isn't centered around a frame (the centering is done via the final eval time) we can just subtract TSI*TPS to get our centered value.
-			const FFrameTime CenteringOffset = CurrentFrameData.TemporalSampleIndex * CurrentFrameMetrics.FrameTimePerTemporalSample;
-			FFrameTime CenteredFrameTime = CurrentCameraCut->ShotInfo.CurrentTimeInRoot - CenteringOffset;
-	
-			const FFrameRate SourceFrameRate = GetOwningGraph()->GetDataSourceInstance()->GetDisplayRate();
-			const FFrameRate EffectiveFrameRate = UMovieGraphBlueprintLibrary::GetEffectiveFrameRate(OutputSetting, SourceFrameRate);
-			const FFrameRate TickResolution = GetOwningGraph()->GetDataSourceInstance()->GetTickResolution();
-
-			constexpr bool bDropFrame = false;
-			CurrentTimeStepData.RootFrameNumber = FFrameRate::TransformTime(CenteredFrameTime, TickResolution, EffectiveFrameRate).RoundToFrame();
-			CurrentTimeStepData.RootTimeCode = FTimecode::FromFrameNumber(CurrentTimeStepData.RootFrameNumber, EffectiveFrameRate, bDropFrame);
-
-			// Calculate metrics for the shot as well
-			CenteredFrameTime = CenteredFrameTime * CurrentCameraCut->ShotInfo.OuterToInnerTransform;
-			CurrentTimeStepData.ShotFrameNumber = FFrameRate::TransformTime(CenteredFrameTime, TickResolution, EffectiveFrameRate).RoundToFrame();
-			CurrentTimeStepData.ShotTimeCode = FTimecode::FromFrameNumber(CurrentTimeStepData.ShotFrameNumber, EffectiveFrameRate, bDropFrame);
-		}
+		// Calculate metrics for the shot as well
+		CenteredFrameTime = CenteredFrameTime * CurrentCameraCut->ShotInfo.OuterToInnerTransform;
+		CurrentTimeStepData.ShotFrameNumber = FFrameRate::TransformTime(CenteredFrameTime, TickResolution, EffectiveFrameRate).RoundToFrame();
+		CurrentTimeStepData.ShotTimeCode = FTimecode::FromFrameNumber(CurrentTimeStepData.ShotFrameNumber, EffectiveFrameRate, bDropFrame);
+	}
 
-		// Set our time step for the next frame. We use the undilated delta time for the Custom Timestep as the engine will
-		// apply the time dilation to the world tick for us, so we don't want to double up time dilation.
-		double UndilatedDeltaTime = CurrentFrameMetrics.TickResolution.AsSeconds(FrameDeltaTime);
-		CustomTimeStep->SetCachedFrameTiming(UMovieGraphEngineTimeStep::FTimeStepCache(UndilatedDeltaTime));
-		GetOwningGraph()->GetDataSourceInstance()->SyncDataSourceTime(FinalEvalTime);
+	// Set our time step for the next frame. We use the undilated delta time for the Custom Timestep as the engine will
+	// apply the time dilation to the world tick for us, so we don't want to double up time dilation.
+	double UndilatedDeltaTime = CurrentFrameMetrics.TickResolution.AsSeconds(FrameDeltaTime);
+	CustomTimeStep->SetCachedFrameTiming(UMovieGraphEngineTimeStep::FTimeStepCache(UndilatedDeltaTime));
+	GetOwningGraph()->GetDataSourceInstance()->SyncDataSourceTime(FinalEvalTime);
+	if (bShouldJump)
+	{
+		GetOwningGraph()->GetDataSourceInstance()->JumpDataSource(FinalEvalTime);
+	}
+
+	// ToDo: This should be converted back to an 'effective' frame number (source frame in external data asset)
+	// so you can line up profiling with the acutal content on screen.
+	TRACE_BOOKMARK(TEXT("MRQ Frame %d [TS: %d]"), CurrentTimeStepData.OutputFrameNumber, CurrentFrameData.TemporalSampleIndex);
 
-		// ToDo: This should be converted back to an 'effective' frame number (source frame in external data asset)
-		// so you can line up profiling with the acutal content on screen.
-		TRACE_BOOKMARK(TEXT("MRQ Frame %d [TS: %d]"), CurrentTimeStepData.OutputFrameNumber, CurrentFrameData.TemporalSampleIndex);
+	// Increment various post-frame counters to set them up for the next frame. This is okay
+	// because the rest of the Movie Graph Pipeline system is based on CurrentTimeStepData which accurately
+	// reflects which frame we're on.
+	{
+		// Update the last sample range to the one we just rendered.
+		CurrentFrameData.LastSampleRange = NextRange;
 
-		// Increment various post-frame counters to set them up for the next frame. This is okay
-		// because the rest of the Movie Graph Pipeline system is based on CurrentTimeStepData which accurately
-		// reflects which frame we're on.
+		if ( bIsLastTemporalSampleOverride.Get(IsLastTemporalSample()))
 		{
-			// Update the last sample range to the one we just rendered.
-			CurrentFrameData.LastSampleRange = NextRange;
-			if (IsLastTemporalSample())
-			{
-				CurrentFrameData.LastOutputFrameRange = CurrentFrameData.CurrentOutputFrameRange;
+			CurrentFrameData.RenderedFrameNumber++;
+			CurrentFrameData.LastOutputFrameRange = CurrentFrameData.CurrentOutputFrameRange;
 				
-				// Increment the output frame number only on the last temporal sample.
+			// Increment the output frame number only on the last temporal sample, and only if
+			// we're actually rendering frames to disk.
+			if (CurrentCameraCut->ShotInfo.State == EMovieRenderShotState::Rendering)
+			{
 				CurrentFrameData.OutputFrameNumber++;
-				
-				// If we've rendered the last temporal sub-sample, we've started a new output frame
-				// and we need to reset our temporal sample index.
-				CurrentFrameData.TemporalSampleIndex = 0;
 			}
-			else
+				
+			// If we've rendered the last temporal sub-sample, we've started a new output frame
+			// and we need to reset our temporal sample index.
+			CurrentFrameData.TemporalSampleIndex = 0;
+		}
+		else
+		{
+			// We skip this on the motion blur state (which should only be true here on the one 
+			// motion blur emulation frame), as the next frame we're going to reset back to
+			// the rendering moe, and we want the renderer to start on the correct temporal
+			// sample index to start a frame (0).
+			if (CurrentCameraCut->ShotInfo.State != EMovieRenderShotState::MotionBlur) // ToDo: maybe not needed?
 			{
 				// Each tick we increment the temporal sample we're on.
 				CurrentFrameData.TemporalSampleIndex++;
 			}
 		}
+	}
 
-		if (!ensure(CurrentCameraCut->ShotInfo.CurrentTickInRoot < CurrentCameraCut->ShotInfo.TotalOutputRangeRoot.GetUpperBoundValue()))
+	if (!ensure(CurrentCameraCut->ShotInfo.CurrentTickInRoot < CurrentCameraCut->ShotInfo.TotalOutputRangeRoot.GetUpperBoundValue()))
+	{
+		UE_LOG(LogMovieRenderPipeline, Error, TEXT("Shot ran past evaluation range, this shouldn't be possible."));
+	}
+
+	if (!ensureMsgf(!FMath::IsNearlyZero(CustomTimeStep->TimeCache.UndilatedDeltaTime), TEXT("An incorrect or uninitialized time step was used!")))
+	{
+		UE_LOG(LogMovieRenderPipeline, Error, TEXT("An incorrect or uninitialized time step was used!"));
+	}
+}
+
+void UMovieGraphCoreTimeStep::UpdateShutterRanges()
+{
+	// Figure out ranges for how long the shutter is open and closed
+	TArray<TRange<FFrameTime>> SplitRange = CurrentFrameData.CurrentOutputFrameRange.Split(CurrentFrameData.CurrentOutputFrameRange.GetLowerBoundValue() + CurrentFrameMetrics.FrameTimeWhileShutterOpen);
+
+	// We will have two ranges when the shutter angle isn't 360/360
+	if (ensure(SplitRange.Num() > 0))
+	{
+		if (SplitRange.Num() == 2)
+		{
+			CurrentFrameData.RangeShutterOpen = SplitRange[0];
+			CurrentFrameData.RangeShutterClosed = SplitRange[1];
+		}
+		else if (SplitRange.Num() == 1)
 		{
-			UE_LOG(LogMovieRenderPipeline, Error, TEXT("Shot ran past evaluation range, this shouldn't be possible."));
+			// If the shutter is fully open, we create an empty range for closed as it is never actually closed.
+			CurrentFrameData.RangeShutterOpen = SplitRange[0];
+			CurrentFrameData.RangeShutterClosed = TRange<FFrameTime>::Empty();
 		}
 	}
+}
 
+void UMovieGraphCoreTimeStep::UpdateTemporalRanges()
+{
+	// Split the range the shutter is open per temporal sample.
+	CurrentFrameData.TemporalRanges.Reset();
+	CurrentFrameData.TemporalRanges.Reserve(CurrentFrameData.TemporalSampleCount);
+	FFrameTime SplitStartTime = CurrentFrameData.RangeShutterOpen.GetLowerBoundValue();
+	TRange<FFrameTime> RemainingRange = CurrentFrameData.RangeShutterOpen;
+	for (int32 Index = 0; Index < CurrentFrameData.TemporalSampleCount - 1; Index++)
+	{
+		SplitStartTime += CurrentFrameMetrics.FrameTimePerTemporalSample;
+		TArray<TRange<FFrameTime>> NewRanges = RemainingRange.Split(SplitStartTime);
+		if (ensure(NewRanges.Num() == 2))
+		{
+			CurrentFrameData.TemporalRanges.Add(NewRanges[0]);
+			RemainingRange = NewRanges[1];
+		}
+	}
+
+	// Add the remaining range, it's the leftover from the last split.
+	CurrentFrameData.TemporalRanges.Add(RemainingRange);
+
+	//for (int32 Index = 0; Index < CurrentFrameData.TemporalRanges.Num(); Index++)
+	//{
+	//	UE_LOG(LogTemp, Warning, TEXT("Range: [%s, %s)"),
+	//		*LexToString(CurrentFrameData.TemporalRanges[Index].GetLowerBoundValue()),
+	//		*LexToString(CurrentFrameData.TemporalRanges[Index].GetUpperBoundValue()));
+	//}
 }
 
 bool UMovieGraphCoreTimeStep::IsFirstTemporalSample() const
@@ -545,3 +750,4 @@ bool UMovieGraphEngineTimeStep::UpdateTimeStep(UEngine* /*InEngine*/)
 	// Return false so the engine doesn't run its own logic to overwrite FApp timings.
 	return false;
 }
+
diff --git a/Engine/Plugins/MovieScene/MovieRenderPipeline/Source/MovieRenderPipelineCore/Private/Graph/MovieGraphDefaultRenderer.cpp b/Engine/Plugins/MovieScene/MovieRenderPipeline/Source/MovieRenderPipelineCore/Private/Graph/MovieGraphDefaultRenderer.cpp
--- a/Engine/Plugins/MovieScene/MovieRenderPipeline/Source/MovieRenderPipelineCore/Private/Graph/MovieGraphDefaultRenderer.cpp
+++ b/Engine/Plugins/MovieScene/MovieRenderPipeline/Source/MovieRenderPipelineCore/Private/Graph/MovieGraphDefaultRenderer.cpp
@@ -194,7 +194,7 @@ void UMovieGraphDefaultRenderer::Render(const FMovieGraphTimeStepData& InTimeSte
 		// If this is the first sample for this output frame, then we need to 
 		// talk to all of our render passes and ask them for what data they will
 		// produce, and set the Output Merger up with that knowledge.
-		UE::MovieGraph::FMovieGraphOutputMergerFrame& NewOutputFrame = GetOwningGraph()->GetOutputMerger()->AllocateNewOutputFrame_GameThread(InTimeStepData.OutputFrameNumber);
+		UE::MovieGraph::FMovieGraphOutputMergerFrame& NewOutputFrame = GetOwningGraph()->GetOutputMerger()->AllocateNewOutputFrame_GameThread(InTimeStepData.RenderedFrameNumber);
 
 		// Get the Traversal Context (not specific to any render pass) at the first sample. This is so
 		// we can easily fetch things that are shared between all render layers later.
@@ -208,7 +208,7 @@ void UMovieGraphDefaultRenderer::Render(const FMovieGraphTimeStepData& InTimeSte
 
 		// Register the frame with our render statistics as being worked on
 
-		UE::MovieGraph::FRenderTimeStatistics* TimeStats = GetRenderTimeStatistics(NewOutputFrame.TraversalContext.Time.OutputFrameNumber);
+		UE::MovieGraph::FRenderTimeStatistics* TimeStats = GetRenderTimeStatistics(NewOutputFrame.TraversalContext.Time.RenderedFrameNumber);
 		if (ensure(TimeStats))
 		{
 			TimeStats->StartTime = FDateTime::UtcNow();
@@ -218,7 +218,7 @@ void UMovieGraphDefaultRenderer::Render(const FMovieGraphTimeStepData& InTimeSte
 	for (const TObjectPtr<UMovieGraphRenderPassNode>& RenderPass : RenderPassesInUse)
 	{
 		// Pass in a copy of the traversal context so the renderer can decide what to do with it.
-		UE::MovieGraph::FMovieGraphOutputMergerFrame& OutputFrame = GetOwningGraph()->GetOutputMerger()->GetOutputFrame_GameThread(InTimeStepData.OutputFrameNumber);
+		UE::MovieGraph::FMovieGraphOutputMergerFrame& OutputFrame = GetOwningGraph()->GetOutputMerger()->GetOutputFrame_GameThread(InTimeStepData.RenderedFrameNumber);
 		RenderPass->Render(OutputFrame.TraversalContext, InTimeStepData);
 	}
 
diff --git a/Engine/Plugins/MovieScene/MovieRenderPipeline/Source/MovieRenderPipelineCore/Private/Graph/MovieGraphPipeline.cpp b/Engine/Plugins/MovieScene/MovieRenderPipeline/Source/MovieRenderPipelineCore/Private/Graph/MovieGraphPipeline.cpp
--- a/Engine/Plugins/MovieScene/MovieRenderPipeline/Source/MovieRenderPipelineCore/Private/Graph/MovieGraphPipeline.cpp
+++ b/Engine/Plugins/MovieScene/MovieRenderPipeline/Source/MovieRenderPipelineCore/Private/Graph/MovieGraphPipeline.cpp
@@ -321,6 +321,8 @@ void UMovieGraphPipeline::BuildShotListFromDataSource()
 		ExpandShot(Shot, OutputNode->HandleFrameCount, bExpandForTemporalSubSample, bPrePass, FinalFrameRate, TickResolution, WarmUpNode->NumWarmUpFrames);
 
 		Shot->ShotInfo.CurrentTimeInRoot = Shot->ShotInfo.TotalOutputRangeRoot.GetLowerBoundValue();
+		Shot->ShotInfo.NumEngineWarmUpFramesRemaining = WarmUpNode->NumWarmUpFrames;
+		Shot->ShotInfo.bEmulateFirstFrameMotionBlur = WarmUpNode->bEmulateMotionBlur;
 		Shot->ShotInfo.CalculateWorkMetrics(GetDataSourceInstance());
 	}
 
@@ -1019,12 +1021,12 @@ void UMovieGraphPipeline::ProcessOutstandingFinishedFrames()
 		UE::MovieGraph::FMovieGraphOutputMergerFrame OutputFrame;
 		OutputMerger->GetFinishedFrames().Dequeue(OutputFrame);
 
-		UE::MovieGraph::FRenderTimeStatistics* TimeStats = GetRendererInstance()->GetRenderTimeStatistics(OutputFrame.TraversalContext.Time.OutputFrameNumber);
+		UE::MovieGraph::FRenderTimeStatistics* TimeStats = GetRendererInstance()->GetRenderTimeStatistics(OutputFrame.TraversalContext.Time.RenderedFrameNumber);
 		if (ensure(TimeStats))
 		{
 			TimeStats->EndTime = FDateTime::UtcNow();
 
-			int32 FrameNumber = OutputFrame.TraversalContext.Time.OutputFrameNumber;
+			int32 FrameNumber = OutputFrame.TraversalContext.Time.RenderedFrameNumber;
 			FString DurationTimeStr = (TimeStats->EndTime - TimeStats->StartTime).ToString();
 			// UE_LOG(LogTemp, Log, TEXT("Frame: %d Duration: %s"), FrameNumber, *DurationTimeStr);
 		}
diff --git a/Engine/Plugins/MovieScene/MovieRenderPipeline/Source/MovieRenderPipelineCore/Private/Graph/MovieGraphSequenceDataSource.cpp b/Engine/Plugins/MovieScene/MovieRenderPipeline/Source/MovieRenderPipelineCore/Private/Graph/MovieGraphSequenceDataSource.cpp
--- a/Engine/Plugins/MovieScene/MovieRenderPipeline/Source/MovieRenderPipelineCore/Private/Graph/MovieGraphSequenceDataSource.cpp
+++ b/Engine/Plugins/MovieScene/MovieRenderPipeline/Source/MovieRenderPipelineCore/Private/Graph/MovieGraphSequenceDataSource.cpp
@@ -131,19 +131,6 @@ FFrameRate UMovieGraphSequenceDataSource::GetDisplayRate() const
 	return FFrameRate(24, 1);
 }
 
-void UMovieGraphSequenceDataSource::InitializeShot(const TObjectPtr<UMoviePipelineExecutorShot>& InShot)
-{
-	//if (!InShot->ShotInfo.bEmulateFirstFrameMotionBlur)
-	{
-		// Real warm up frames walk through the Sequence
-		LevelSequenceActor->GetSequencePlayer()->Play();
-	}
-	//else
-	//{
-	//	// Ensure we don't try to evaluate as we want to sit and wait during warm up and motion blur frames.
-	//	LevelSequenceActor->GetSequencePlayer()->Pause();
-	//}
-}
 void UMovieGraphSequenceDataSource::CacheHierarchyForShot(const TObjectPtr<UMoviePipelineExecutorShot>& InShot)
 {
 	// Save the whole hierarchy for this shot, then set it to be inactive.
@@ -272,6 +259,21 @@ void UMovieGraphSequenceDataSource::SyncDataSourceTime(const FFrameTime& InTime)
 	FFrameRate TickResolution = LevelSequenceActor->GetSequence()->GetMovieScene()->GetTickResolution();
 	CustomSequenceTimeController->SetCachedFrameTiming(FQualifiedFrameTime(InTime, TickResolution));
 }
+void UMovieGraphSequenceDataSource::PlayDataSource()
+{
+	LevelSequenceActor->GetSequencePlayer()->Play();
+}
+
+void UMovieGraphSequenceDataSource::PauseDataSource() 
+{
+	LevelSequenceActor->GetSequencePlayer()->Pause();
+}
+
+void UMovieGraphSequenceDataSource::JumpDataSource(const FFrameTime& InTimeToJumpTo) 
+{
+	LevelSequenceActor->GetSequencePlayer()->SetPlaybackPosition(FMovieSceneSequencePlaybackParams(InTimeToJumpTo, EUpdatePositionMethod::Jump));
+
+}
 
 namespace UE::MovieGraph
 {
diff --git a/Engine/Plugins/MovieScene/MovieRenderPipeline/Source/MovieRenderPipelineCore/Private/Graph/Nodes/MovieGraphCoreRenderPassNode.cpp b/Engine/Plugins/MovieScene/MovieRenderPipeline/Source/MovieRenderPipelineCore/Private/Graph/Nodes/MovieGraphCoreRenderPassNode.cpp
--- a/Engine/Plugins/MovieScene/MovieRenderPipeline/Source/MovieRenderPipelineCore/Private/Graph/Nodes/MovieGraphCoreRenderPassNode.cpp
+++ b/Engine/Plugins/MovieScene/MovieRenderPipeline/Source/MovieRenderPipelineCore/Private/Graph/Nodes/MovieGraphCoreRenderPassNode.cpp
@@ -363,10 +363,10 @@ void UMovieGraphCoreRenderPassNode::FMovieGraphRenderPass::Render(const FMovieGr
 	GetRendererModule().BeginRenderingViewFamily(&Canvas, ViewFamily.ToSharedPtr().Get());
 
 	// If this was just to contribute to the history buffer, no need to go any further.
-	//if (InSampleState.bDiscardResult)
-	//{
-	//	return;
-	//}
+	if (InTimeData.bDiscardOutput)
+	{
+		return;
+	}
 
 	// Take our per-frame Traversal Context and update it with context specific to this sample.
 	FMovieGraphTraversalContext UpdatedTraversalContext = InFrameTraversalContext;
@@ -401,7 +401,7 @@ void UMovieGraphCoreRenderPassNode::FMovieGraphRenderPass::PostRendererSubmissio
 	UE::MovieGraph::DefaultRenderer::FSurfaceAccumulatorPool::FInstancePtr AccumulatorInstance = nullptr;
 	{
 		// SCOPE_CYCLE_COUNTER(STAT_MoviePipeline_WaitForAvailableAccumulator);
-		AccumulatorInstance = SampleAccumulatorPool->BlockAndGetAccumulator_GameThread(InSampleState.TraversalContext.Time.OutputFrameNumber, InSampleState.TraversalContext.RenderDataIdentifier);
+		AccumulatorInstance = SampleAccumulatorPool->BlockAndGetAccumulator_GameThread(InSampleState.TraversalContext.Time.RenderedFrameNumber, InSampleState.TraversalContext.RenderDataIdentifier);
 	}
 
 	FMoviePipelineSurfaceQueuePtr LocalSurfaceQueue = Renderer->GetOrCreateSurfaceQueue(InRenderTargetInitParams);
@@ -547,16 +547,16 @@ FSceneView* UMovieGraphCoreRenderPassNode::FMovieGraphRenderPass::AllocateSceneV
 
 void UMovieGraphCoreRenderPassNode::FMovieGraphRenderPass::ApplyMoviePipelineOverridesToViewFamily(TSharedRef<FSceneViewFamilyContext> InOutFamily, const FViewFamilyContextInitData& InInitData)
 {
-	/*// A third set of overrides required to properly configure views to match the given showflags/etc.
-	// ToDo: There's now five(!) identical implementations of this, we should unify them.
+	// A third set of overrides required to properly configure views to match the given showflags/etc.
 	// SetupViewForViewModeOverride(View);
 
 	// Override the view's FrameIndex to be based on our progress through the sequence. This greatly increases
 	// determinism with things like TAA.
-	InInitData.View->OverrideFrameIndexValue = InInitData.FrameIndex;
-	InInitData.View->bCameraCut = InInitData.bCameraCut;
-	InInitData.View->bIsOfflineRender = true;
-	InInitData.View->AntiAliasingMethod = InInitData.AntiAliasingMethod;
+	FSceneView* View = const_cast<FSceneView*>(InOutFamily->Views[0]);
+	View->OverrideFrameIndexValue = InInitData.FrameIndex;
+	View->bCameraCut = InInitData.bCameraCut;
+	View->bIsOfflineRender = true;
+	View->AntiAliasingMethod = InInitData.AntiAliasingMethod;
 
 	// Add any view extensions that were added to the scene to this View too
 	// OutViewFamily->ViewExtensions.Append(GEngine->ViewExtensions->GatherActiveExtensions(FSceneViewExtensionContext(GetWorld()->Scene)));
@@ -576,12 +576,9 @@ void UMovieGraphCoreRenderPassNode::FMovieGraphRenderPass::ApplyMoviePipelineOve
 
 	// Override the Motion Blur settings since these are controlled by the movie pipeline.
 	{
-		// FFrameRate OutputFrameRate = GetPipeline()->GetPipelinePrimaryConfig()->GetEffectiveFrameRate(GetPipeline()->GetTargetSequence());
-		FFrameRate OutputFrameRate = FFrameRate(24, 1); // ToDo, get this from config.
-
 		// We need to inversly scale the target FPS by time dilation to counteract slowmo. If scaling isn't applied then motion blur length
 		// stays the same length despite the smaller delta time and the blur ends up too long.
-		View->FinalPostProcessSettings.MotionBlurTargetFPS = FMath::RoundToInt(OutputFrameRate.AsDecimal() / FMath::Max(SMALL_NUMBER, InInitData.TimeData.TimeDilation));
+		View->FinalPostProcessSettings.MotionBlurTargetFPS = FMath::RoundToInt(InInitData.TimeData.FrameRate.AsDecimal() / FMath::Max(SMALL_NUMBER, InInitData.TimeData.WorldTimeDilation));
 		View->FinalPostProcessSettings.MotionBlurAmount = InInitData.TimeData.MotionBlurFraction;
 		View->FinalPostProcessSettings.MotionBlurMax = 100.f;
 		View->FinalPostProcessSettings.bOverride_MotionBlurAmount = true;
@@ -591,12 +588,12 @@ void UMovieGraphCoreRenderPassNode::FMovieGraphRenderPass::ApplyMoviePipelineOve
 		// Skip the whole pass if they don't want motion blur.
 		if (FMath::IsNearlyZero(InInitData.TimeData.MotionBlurFraction))
 		{
-			OutViewFamily->EngineShowFlags.SetMotionBlur(false);
+			InOutFamily->EngineShowFlags.SetMotionBlur(false);
 		}
 	}
 
 	// Warn the user for invalid setting combinations / enforce hardware limitations
-	{
+	/*{
 
 		// Locked Exposure
 		const bool bAutoExposureAllowed = IsAutoExposureAllowed(InOutSampleState);
diff --git a/Engine/Plugins/MovieScene/MovieRenderPipeline/Source/MovieRenderPipelineCore/Public/Graph/MovieGraphCoreTimeStep.h b/Engine/Plugins/MovieScene/MovieRenderPipeline/Source/MovieRenderPipelineCore/Public/Graph/MovieGraphCoreTimeStep.h
--- a/Engine/Plugins/MovieScene/MovieRenderPipeline/Source/MovieRenderPipelineCore/Public/Graph/MovieGraphCoreTimeStep.h
+++ b/Engine/Plugins/MovieScene/MovieRenderPipeline/Source/MovieRenderPipelineCore/Public/Graph/MovieGraphCoreTimeStep.h
@@ -35,6 +35,12 @@ public:
 	// ~UMovieGraphTimeStepBase Interface
 
 protected:
+	/** Update CurrentFrameData.RangeShutterOpen and CurrentFrameData.RangeShutterClosed. */
+	virtual void UpdateShutterRanges();
+	
+	/** Update CurrentFrameData.TemporalRanges. */
+	virtual void UpdateTemporalRanges();
+
 	/**
 	 * Gets the index of the next temporal range. The index returned should not be the index returned by the prior call
 	 * (this would result in a frame delta time of zero, causing the engine to crash).
@@ -101,6 +107,7 @@ protected:
 			: TemporalSampleIndex(0)
 			, TemporalSampleCount(0)
 			, OutputFrameNumber(0)
+			, RenderedFrameNumber(0)
 		{
 		}
 
@@ -113,6 +120,12 @@ protected:
 		/** Which output frame are we working on, relative to zero.*/
 		int32 OutputFrameNumber;
 
+		/** 
+		* Index of which frames we've submitted for rendering. Doesn't line up with OutputFrameNumber when using warm-up frames.
+		* Used internally by the rendering engine to keep track of which frames need to be read back.
+		*/
+		int32 RenderedFrameNumber;
+
 		/** A range of time (in Tick Resolution) for the last output frame being worked on. Updated before first TS of next frame. */
 		TRange<FFrameTime> LastOutputFrameRange;
 
diff --git a/Engine/Plugins/MovieScene/MovieRenderPipeline/Source/MovieRenderPipelineCore/Public/Graph/MovieGraphDataTypes.h b/Engine/Plugins/MovieScene/MovieRenderPipeline/Source/MovieRenderPipelineCore/Public/Graph/MovieGraphDataTypes.h
--- a/Engine/Plugins/MovieScene/MovieRenderPipeline/Source/MovieRenderPipelineCore/Public/Graph/MovieGraphDataTypes.h
+++ b/Engine/Plugins/MovieScene/MovieRenderPipeline/Source/MovieRenderPipelineCore/Public/Graph/MovieGraphDataTypes.h
@@ -151,8 +151,16 @@ public:
 
 	/** Called by the Time Step system when it wants the external data source to update. Time is in TickResolution scale. */
 	virtual void SyncDataSourceTime(const FFrameTime& InTime) {}
+	
+	/** Called by the Time Step system when the external data source should start playback (time values will have been set by SyncDataSourceTime */
+	virtual void PlayDataSource() {}
+
+	/** Called by the Time Step system when the external data source should pause playback. */
+	virtual void PauseDataSource() {}
+
+	/** Called by the Time Step system when the external data source should jump to the given time. Time is in TickResolution scale. */
+	virtual void JumpDataSource(const FFrameTime& InTimeToJumpTo) {}
 
-	virtual void BuildTimeRanges() { }
 	/** 
 	* Called when the Movie Graph Pipeline starts before anything has happened, allowing you to 
 	* cache your datasource before making any modifications to it as a result of rendering.
diff --git a/Engine/Plugins/MovieScene/MovieRenderPipeline/Source/MovieRenderPipelineCore/Public/Graph/MovieGraphSequenceDataSource.h b/Engine/Plugins/MovieScene/MovieRenderPipeline/Source/MovieRenderPipelineCore/Public/Graph/MovieGraphSequenceDataSource.h
--- a/Engine/Plugins/MovieScene/MovieRenderPipeline/Source/MovieRenderPipelineCore/Public/Graph/MovieGraphSequenceDataSource.h
+++ b/Engine/Plugins/MovieScene/MovieRenderPipeline/Source/MovieRenderPipelineCore/Public/Graph/MovieGraphSequenceDataSource.h
@@ -44,7 +44,9 @@ public:
 	virtual FFrameRate GetTickResolution() const override;
 	virtual FFrameRate GetDisplayRate() const override;
 	virtual void SyncDataSourceTime(const FFrameTime& InTime) override;
-	virtual void InitializeShot(const TObjectPtr<UMoviePipelineExecutorShot>& InShot) override;
+	virtual void PlayDataSource() override;
+	virtual void PauseDataSource() override;
+	virtual void JumpDataSource(const FFrameTime& InTimeToJumpTo) override;
 	virtual void CacheHierarchyForShot(const TObjectPtr<UMoviePipelineExecutorShot>& InShot) override;
 	virtual void RestoreHierarchyForShot(const TObjectPtr<UMoviePipelineExecutorShot>& InShot) override;
 	virtual void MuteShot(const TObjectPtr<UMoviePipelineExecutorShot>& InShot) override;
diff --git a/Engine/Plugins/MovieScene/MovieRenderPipeline/Source/MovieRenderPipelineCore/Public/Graph/MovieGraphTimeStepData.h b/Engine/Plugins/MovieScene/MovieRenderPipeline/Source/MovieRenderPipelineCore/Public/Graph/MovieGraphTimeStepData.h
--- a/Engine/Plugins/MovieScene/MovieRenderPipeline/Source/MovieRenderPipelineCore/Public/Graph/MovieGraphTimeStepData.h
+++ b/Engine/Plugins/MovieScene/MovieRenderPipeline/Source/MovieRenderPipelineCore/Public/Graph/MovieGraphTimeStepData.h
@@ -20,18 +20,24 @@ struct MOVIERENDERPIPELINECORE_API FMovieGraphTimeStepData
 public:
 	FMovieGraphTimeStepData()
 		: OutputFrameNumber(0)
+		, RenderedFrameNumber(0)
 		, FrameDeltaTime(0.f)
 		, WorldTimeDilation(0.f)
 		, WorldSeconds(0.f)
 		, MotionBlurFraction(0.f)
+		, FrameRate(0, 0)
 		, bIsFirstTemporalSampleForFrame(false)
 		, bIsLastTemporalSampleForFrame(false)
+		, bDiscardOutput(false)
 		, bRequiresAccumulator(false)
 	{}
 
 	UPROPERTY(BlueprintReadWrite, EditAnywhere, Category = "Movie Graph")
 	int32 OutputFrameNumber;
 
+	UPROPERTY(BlueprintReadWrite, EditAnywhere, Category = "Movie Graph")
+	int32 RenderedFrameNumber;
+
 	UPROPERTY(BlueprintReadWrite, EditAnywhere, Category = "Movie Graph")
 	float FrameDeltaTime;
 	
@@ -43,6 +49,9 @@ public:
 	
 	UPROPERTY(BlueprintReadWrite, EditAnywhere, Category = "Movie Graph")
 	float MotionBlurFraction;
+	
+	UPROPERTY(BlueprintReadWrite, EditAnywhere, Category = "Movie Graph")
+	FFrameRate FrameRate;
 
 	/** 
 	* Should be set to true for the first sample of each output frame. Used to determine
@@ -59,6 +68,13 @@ public:
 	UPROPERTY(BlueprintReadWrite, EditAnywhere, Category = "Movie Graph")
 	bool bIsLastTemporalSampleForFrame;
 
+	/**
+	* Should the rendered result be discarded after a render? This will skip any
+	* accumulators or readback and is used for frames that are only produced to
+	* warm up the renderer.
+	*/
+	UPROPERTY(BlueprintReadWrite, EditAnywhere, Category = "Movie Graph")
+	bool bDiscardOutput;
 
 	/**
 	* Should be set to true for every sample if there is more than one temporal sample
diff --git a/Engine/Plugins/MovieScene/MovieRenderPipeline/Source/MovieRenderPipelineCore/Public/Graph/Nodes/MovieGraphCoreRenderPassNode.h b/Engine/Plugins/MovieScene/MovieRenderPipeline/Source/MovieRenderPipelineCore/Public/Graph/Nodes/MovieGraphCoreRenderPassNode.h
--- a/Engine/Plugins/MovieScene/MovieRenderPipeline/Source/MovieRenderPipelineCore/Public/Graph/Nodes/MovieGraphCoreRenderPassNode.h
+++ b/Engine/Plugins/MovieScene/MovieRenderPipeline/Source/MovieRenderPipelineCore/Public/Graph/Nodes/MovieGraphCoreRenderPassNode.h
@@ -191,7 +191,6 @@ protected:
 			, FrameIndex(-1)
 			, bCameraCut(false)
 			, AntiAliasingMethod(EAntiAliasingMethod::AAM_None)
-			, View(nullptr)
 			, SceneViewStateReference(nullptr)
 		{
 		}
@@ -207,9 +206,6 @@ protected:
 		bool bCameraCut;
 		EAntiAliasingMethod AntiAliasingMethod;
 
-		// Ownership of this will be passed to the FSceneViewFamilyContext
-		FSceneView* View;
-
 		FSceneViewStateInterface* SceneViewStateReference;
 
 		// Camera Setup
diff --git a/Engine/Plugins/MovieScene/MovieRenderPipeline/Source/MovieRenderPipelineCore/Public/Graph/Nodes/MovieGraphWarmUpSettingNode.h b/Engine/Plugins/MovieScene/MovieRenderPipeline/Source/MovieRenderPipelineCore/Public/Graph/Nodes/MovieGraphWarmUpSettingNode.h
--- a/Engine/Plugins/MovieScene/MovieRenderPipeline/Source/MovieRenderPipelineCore/Public/Graph/Nodes/MovieGraphWarmUpSettingNode.h
+++ b/Engine/Plugins/MovieScene/MovieRenderPipeline/Source/MovieRenderPipelineCore/Public/Graph/Nodes/MovieGraphWarmUpSettingNode.h
@@ -28,7 +28,7 @@ public:
 	uint8 bOverride_NumWarmUpFrames : 1;
 
 	UPROPERTY(EditAnywhere, BlueprintReadWrite, Category = Overrides, meta = (InlineEditConditionToggle))
-	uint8 bOverride_EmulateMotionBlur : 1;
+	uint8 bOverride_bEmulateMotionBlur : 1;
 
 	/** At the start of each shot, how many frames should we run the engine (without writing renders to disk) to warm up various systems. */
 	UPROPERTY(EditAnywhere, BlueprintReadWrite, Category = "Settings", meta=(EditCondition="bOverride_NumWarmUpFrames"))
@@ -37,6 +37,6 @@ public:
 	/** 
 	If true, we will evaluate frame 0, and then wait NumWarmUpFrames frames, rendering each one as we go. Then we will evaluate frame 1, and then frame 0 again. This emulates motion blur on the first frame (which normally needs data before frame 0). If false, we will "walk" towards the first frame of the shot, starting NumWarmUpFrames before the shot normally starts. 
 	*/
-	UPROPERTY(EditAnywhere, BlueprintReadWrite, Category = "Settings", meta = (EditCondition = "bOverride_EmulateMotionBlur"))
+	UPROPERTY(EditAnywhere, BlueprintReadWrite, Category = "Settings", meta = (EditCondition = "bOverride_bEmulateMotionBlur"))
 	bool bEmulateMotionBlur;
 };
