MoviePipeline: Added ".wav Output" node to the graph. Integrates properly with the Command Line Encoder.

#rnx

diff --git a/Engine/Plugins/MovieScene/MovieRenderPipeline/Source/MovieRenderPipelineCore/Private/Graph/MovieGraphDataTypes.cpp b/Engine/Plugins/MovieScene/MovieRenderPipeline/Source/MovieRenderPipelineCore/Private/Graph/MovieGraphDataTypes.cpp
--- a/Engine/Plugins/MovieScene/MovieRenderPipeline/Source/MovieRenderPipelineCore/Private/Graph/MovieGraphDataTypes.cpp
+++ b/Engine/Plugins/MovieScene/MovieRenderPipeline/Source/MovieRenderPipelineCore/Private/Graph/MovieGraphDataTypes.cpp
@@ -4,11 +4,13 @@
 #include "Graph/MovieGraphDefaultRenderer.h"
 #include "Graph/MovieGraphSequenceDataSource.h"
 #include "Graph/MovieGraphPipeline.h"
+#include "Graph/MovieGraphDefaultAudioRenderer.h"
 
 FMovieGraphInitConfig::FMovieGraphInitConfig()
 {
 	RendererClass = UMovieGraphDefaultRenderer::StaticClass();
 	DataSourceClass = UMovieGraphSequenceDataSource::StaticClass();
+	AudioRendererClass = UMovieGraphDefaultAudioRenderer::StaticClass();
 	bRenderViewport = false;
 }
 
@@ -26,3 +28,13 @@ UMovieGraphPipeline* UMovieGraphDataSourceBase::GetOwningGraph() const
 {
 	return GetTypedOuter<UMovieGraphPipeline>();
 }
+
+UMovieGraphPipeline* UMovieGraphAudioRendererBase::GetOwningGraph() const
+{
+	return GetTypedOuter<UMovieGraphPipeline>();
+}
+
+const MoviePipeline::FAudioState& UMovieGraphAudioRendererBase::GetAudioState() const
+{
+	return AudioState;
+}
diff --git a/Engine/Plugins/MovieScene/MovieRenderPipeline/Source/MovieRenderPipelineCore/Private/Graph/MovieGraphDefaultAudioRenderer.cpp b/Engine/Plugins/MovieScene/MovieRenderPipeline/Source/MovieRenderPipelineCore/Private/Graph/MovieGraphDefaultAudioRenderer.cpp
new file mode 100644
--- /dev/null
+++ b/Engine/Plugins/MovieScene/MovieRenderPipeline/Source/MovieRenderPipelineCore/Private/Graph/MovieGraphDefaultAudioRenderer.cpp
@@ -0,0 +1,180 @@
+﻿// Copyright Epic Games, Inc. All Rights Reserved.
+
+#include "Graph/MovieGraphDefaultAudioRenderer.h"
+
+#include "AudioDevice.h"
+#include "AudioMixerDevice.h"
+#include "AudioMixerPlatformNonRealtime.h"
+#include "AudioMixerSubmix.h"
+#include "DSP/BufferVectorOperations.h"
+#include "Engine/Engine.h"
+#include "Graph/MovieGraphBlueprintLibrary.h"
+#include "Graph/MovieGraphUtils.h"
+#include "Graph/Nodes/MovieGraphOutputSettingNode.h"
+#include "HAL/IConsoleManager.h"
+#include "MoviePipelineQueue.h"
+#include "MovieRenderPipelineDataTypes.h"
+#include "Graph/MovieGraphPipeline.h"
+
+void UMovieGraphDefaultAudioRenderer::StartAudioRecording()
+{
+	AudioState.bIsRecordingAudio = true;
+
+	if (Audio::FMixerDevice* MixerDevice = UE::MoviePipeline::RenderGraph::Audio::GetAudioMixerDeviceFromWorldContext(this))
+	{
+		const TWeakPtr<Audio::FMixerSubmix> MasterSubmix = MixerDevice->GetMasterSubmix();
+		if (MasterSubmix.Pin())
+		{
+			AudioState.ActiveSubmixes.Add(MasterSubmix);
+		}
+
+		for (TWeakPtr<Audio::FMixerSubmix>& WeakSubmix : AudioState.ActiveSubmixes)
+		{
+			constexpr float ExpectedDuration = 30.f;
+			WeakSubmix.Pin()->OnStartRecordingOutput(ExpectedDuration);
+		}
+	}
+}
+
+void UMovieGraphDefaultAudioRenderer::StopAudioRecording()
+{
+	AudioState.bIsRecordingAudio = false;
+
+	for (TWeakPtr<Audio::FMixerSubmix>& WeakSubmix : AudioState.ActiveSubmixes)
+	{
+		if (const TSharedPtr<Audio::FMixerSubmix> Submix = WeakSubmix.Pin())
+		{
+			MoviePipeline::FAudioState::FAudioSegment& NewSegment = AudioState.FinishedSegments.AddDefaulted_GetRef();
+
+			// Copy the data returned by the submix
+			const Audio::AlignedFloatBuffer& AudioRecording = Submix->OnStopRecordingOutput(NewSegment.NumChannels, NewSegment.SampleRate);
+			NewSegment.SegmentData = AudioRecording; 
+		}
+	}
+
+	AudioState.ActiveSubmixes.Reset();
+}
+
+void UMovieGraphDefaultAudioRenderer::ProcessAudioTick()
+{
+	// Only supported on the new audio mixer (with the non-realtime device, windows only).
+	Audio::FMixerDevice* MixerDevice = UE::MoviePipeline::RenderGraph::Audio::GetAudioMixerDeviceFromWorldContext(this);
+	if (!MixerDevice)
+	{
+		return;
+	}
+
+	const Audio::IAudioMixerPlatformInterface* AudioMixerPlatform = MixerDevice->GetAudioMixerPlatform();
+	Audio::FMixerPlatformNonRealtime* NRTPlatform = nullptr;
+	if (AudioMixerPlatform && AudioMixerPlatform->IsNonRealtime())
+	{
+		// There is only one non-realtime audio platform at this time, so we can safely static_cast this due to the IsNonRealtime check.
+		NRTPlatform = static_cast<Audio::FMixerPlatformNonRealtime*>(MixerDevice->GetAudioMixerPlatform());
+	}
+
+	if (!NRTPlatform)
+	{
+		return;
+	}
+
+	const TArray<TObjectPtr<UMoviePipelineExecutorShot>>& ActiveShotList = GetOwningGraph()->GetActiveShotList();
+	const int32 CurrentShotIndex = GetOwningGraph()->GetCurrentShotIndex();
+	const TObjectPtr<UMoviePipelineExecutorShot> CurrentShot = ActiveShotList[CurrentShotIndex];
+
+	// Start capturing any produced samples on the same frame we start submitting samples that will make it to disk.
+	// This comes before we process samples for this frame (below).
+	if ((CurrentShot->ShotInfo.State == EMovieRenderShotState::Rendering) && !AudioState.bIsRecordingAudio)
+	{
+		// There can be time in the Unreal world that passes when MRQ is not yet active. When MRQ isn't active,
+		// then the NRT audio platform doesn't process anything and it builds up for when we do start to actually
+		// process it for a shot. This causes our actual audio to become offset from the sequence once lined up
+		// afterwards. So for now we'll just flush out anything that needed to be processed before we start capturing
+			
+		MixerDevice->Update(true);
+		
+		// We don't have a great way of determining how much outstanding audio there is, so we're just going to 
+		// process an arbitrary amount of time.
+		NRTPlatform->RenderAudio(30.f);
+		
+		StartAudioRecording();
+	}
+
+	// This needs to tick every frame so we process and clear audio that happens between shots while the world is running.
+	// We have special ticking logic while rendering due to temporal sampling, and otherwise just fall back to normal.
+	bool bCanRenderAudio = true;
+	double AudioDeltaTime = FApp::GetDeltaTime();
+
+	if (CurrentShot->ShotInfo.State == EMovieRenderShotState::Rendering)
+	{
+		const UMovieGraphTimeStepBase* TimeStepInstance = GetOwningGraph()->GetTimeStepInstance();
+		const UMovieGraphDataSourceBase* DataSourceInstance = GetOwningGraph()->GetDataSourceInstance();
+		
+		constexpr bool bIncludeCDOs = true;
+		constexpr bool bExactMatch = true;
+		const TObjectPtr<UMovieGraphEvaluatedConfig> EvaluatedGraph = TimeStepInstance->GetCalculatedTimeData().EvaluatedConfig;
+		UMovieGraphOutputSettingNode* OutputSettingNode = EvaluatedGraph->GetSettingForBranch<UMovieGraphOutputSettingNode>(UMovieGraphNode::GlobalsPinName, bIncludeCDOs, bExactMatch);
+
+		const FFrameRate SourceFrameRate = DataSourceInstance->GetDisplayRate();
+		const FFrameRate EffectiveFrameRate = UMovieGraphBlueprintLibrary::GetEffectiveFrameRate(OutputSettingNode, SourceFrameRate);
+		
+		// The non-real time audio renderer desires even engine time steps. Unfortunately, when using temporal sampling
+		// we don't have an even time-step. However, because it's non-real time, and we're accumulating the results into
+		// a single frame anyways, we can bunch up the audio work and then process it when we've reached the end of a frame.
+		bCanRenderAudio = TimeStepInstance->GetCalculatedTimeData().bIsLastTemporalSampleForFrame;
+		
+		// Process work that has been submitted from the game thread to the audio thread over the temporal samples of this frame.
+		AudioDeltaTime = EffectiveFrameRate.AsInterval();
+	}
+
+	// Handle any game logic that changed Audio State.
+	MixerDevice->Update(true);
+
+	// Process work that has been submitted from the game thread to the audio thread over the temporal samples of this frame.
+	if (bCanRenderAudio)
+	{
+		NRTPlatform->RenderAudio(AudioDeltaTime);
+	}
+}
+
+void UMovieGraphDefaultAudioRenderer::SetupAudioRendering()
+{
+	// Ensure that we try to play audio at full volume, even if we're unfocused.
+	AudioState.PrevUnfocusedAudioMultiplier = FApp::GetUnfocusedVolumeMultiplier();
+	FApp::SetUnfocusedVolumeMultiplier(1.f);
+
+	// Will be null if the NRT module wasn't loaded
+	if (IConsoleVariable* AudioRenderEveryTickCvar = IConsoleManager::Get().FindConsoleVariable(TEXT("au.nrt.RenderEveryTick")))
+	{
+		AudioState.PrevRenderEveryTickValue = AudioRenderEveryTickCvar->GetInt();
+		
+		// Override it to prevent it from automatically ticking, we'll control this below
+		AudioRenderEveryTickCvar->SetWithCurrentPriority(0);
+	}
+
+	// Ensure that the NRT audio doesn't get muted
+	if (IConsoleVariable* NeverMuteNRTAudioCvar = IConsoleManager::Get().FindConsoleVariable(TEXT("au.NeverMuteNonRealtimeAudioDevices")))
+	{
+		AudioState.PrevNeverMuteNRTAudioValue = NeverMuteNRTAudioCvar->GetInt();
+		
+		NeverMuteNRTAudioCvar->SetWithCurrentPriority(1);
+	}
+}
+
+void UMovieGraphDefaultAudioRenderer::TeardownAudioRendering() const
+{
+	// Restore previous unfocused audio multiplier, to no longer force audio when unfocused
+	FApp::SetUnfocusedVolumeMultiplier(AudioState.PrevUnfocusedAudioMultiplier);
+
+	// Restore our cached CVar values
+	
+	// This will be null if the NRT wasn't used (module not loaded)
+	if (IConsoleVariable* AudioRenderEveryTickCvar = IConsoleManager::Get().FindConsoleVariable(TEXT("au.nrt.RenderEveryTick")))
+	{
+		AudioRenderEveryTickCvar->Set(AudioState.PrevRenderEveryTickValue, ECVF_SetByConstructor);
+	}
+
+	if (IConsoleVariable* NeverMuteNRTAudioCvar = IConsoleManager::Get().FindConsoleVariable(TEXT("au.NeverMuteNonRealtimeAudioDevices")))
+	{
+		NeverMuteNRTAudioCvar->Set(AudioState.PrevNeverMuteNRTAudioValue, ECVF_SetByConstructor);
+	}
+}
diff --git a/Engine/Plugins/MovieScene/MovieRenderPipeline/Source/MovieRenderPipelineCore/Private/Graph/MovieGraphPipeline.cpp b/Engine/Plugins/MovieScene/MovieRenderPipeline/Source/MovieRenderPipelineCore/Private/Graph/MovieGraphPipeline.cpp
--- a/Engine/Plugins/MovieScene/MovieRenderPipeline/Source/MovieRenderPipelineCore/Private/Graph/MovieGraphPipeline.cpp
+++ b/Engine/Plugins/MovieScene/MovieRenderPipeline/Source/MovieRenderPipelineCore/Private/Graph/MovieGraphPipeline.cpp
@@ -84,6 +84,7 @@ void UMovieGraphPipeline::Initialize(UMoviePipelineExecutorJob* InJob, const FMo
 	// Create instances of our different classes from the InitConfig
 	GraphRendererInstance = NewObject<UMovieGraphRendererBase>(this, InitConfig.RendererClass);
 	GraphDataSourceInstance = NewObject<UMovieGraphDataSourceBase>(this, InitConfig.DataSourceClass);
+	GraphAudioRendererInstance = NewObject<UMovieGraphAudioRendererBase>(this, InitConfig.AudioRendererClass);
 	
 	CurrentJob = InJob;
 	CurrentShotIndex = 0;
@@ -97,7 +98,8 @@ void UMovieGraphPipeline::Initialize(UMoviePipelineExecutorJob* InJob, const FMo
 
 	// Construct the viewport preview UI and bind it to this instance.
 	LoadPreviewWidget();
-	// SetupAudioRendering();
+
+	GraphAudioRendererInstance->SetupAudioRendering();
 
 	// Update our list of shots from our data source, and then
 	// create our list of active shots, so we don't try to render
@@ -564,6 +566,11 @@ void UMovieGraphPipeline::SetupShot(const TObjectPtr<UMoviePipelineExecutorShot>
 
 void UMovieGraphPipeline::TeardownShot(const TObjectPtr<UMoviePipelineExecutorShot>& InShot)
 {
+	// Teardown happens at the start of the first frame the shot is finished so we'll stop recording
+	// audio, which will prevent it from capturing any samples for this frame. We don't do a similar
+	// start in InitializeShot() because we don't want to record samples during warm up/motion blur.
+	GraphAudioRendererInstance->StopAudioRecording();
+	
 	// Teardown any rendering architecture for this shot. This needs to happen first because it'll flush outstanding rendering commands
 	GraphRendererInstance->TeardownRenderingPipelineForShot(InShot);
 
@@ -574,18 +581,18 @@ void UMovieGraphPipeline::TeardownShot(const TObjectPtr<UMoviePipelineExecutorSh
 
 	// some other stuff
 
-	CurrentShotIndex++;
-
 	// Revert the cvar values that were initially applied for the shot
 	CVarManager->RevertAllCVars();
 
 	// Check to see if this was the last shot in the Pipeline, otherwise on the next
 	// tick the new shot will be initialized and processed.
-	if (CurrentShotIndex >= ActiveShotList.Num())
+	if (CurrentShotIndex >= (ActiveShotList.Num() - 1))
 	{
 		UE_LOG(LogMovieRenderPipeline, Log, TEXT("MovieGraph Finished rendering last shot. Moving to Finalize to finish writing items to disk."));
 		TransitionToState(EMovieRenderPipelineState::Finalize);
 	}
+
+	CurrentShotIndex++;
 }
 
 void UMovieGraphPipeline::SetSoloShot(const TObjectPtr<UMoviePipelineExecutorShot>& InShot)
@@ -723,7 +730,8 @@ void UMovieGraphPipeline::OnEngineTickEndFrame()
 
 	UE_LOG(LogMovieRenderPipeline, VeryVerbose, TEXT("MovieGraph OnEngineTickEndFrame (Start)"));
 
-	// ProcessAudioTick();
+	GraphAudioRendererInstance->ProcessAudioTick();
+	
 	RenderFrame();
 
 	UE_LOG(LogMovieRenderPipeline, VeryVerbose, TEXT("MovieGraph OnEngineTickEndFrame (End)"));
@@ -847,7 +855,7 @@ void UMovieGraphPipeline::TransitionToState(const EMovieRenderPipelineState InNe
 
 			// If we had naturally finished the last shot before doing this transition it will have
 			// already been torn down, so this only catches mid-shot transitions to ensure teardown.
-			if (CurrentShotIndex < ActiveShotList.Num())
+			if (CurrentShotIndex < (ActiveShotList.Num() - 1))
 			{
 				// Ensures all in-flight work for that shot is handled.
 				TeardownShot(ActiveShotList[CurrentShotIndex]);
@@ -931,7 +939,7 @@ void UMovieGraphPipeline::TransitionToState(const EMovieRenderPipelineState InNe
 			//	Setting->OnPipelineFinished();
 			//}
 			//
-			//TeardownAudioRendering();
+			GraphAudioRendererInstance->TeardownAudioRendering();
 			//LevelSequenceActor->GetSequencePlayer()->Stop();
 			//RestoreTargetSequenceToOriginalState();
 			//
diff --git a/Engine/Plugins/MovieScene/MovieRenderPipeline/Source/MovieRenderPipelineCore/Private/Graph/MovieGraphUtils.cpp b/Engine/Plugins/MovieScene/MovieRenderPipeline/Source/MovieRenderPipelineCore/Private/Graph/MovieGraphUtils.cpp
--- a/Engine/Plugins/MovieScene/MovieRenderPipeline/Source/MovieRenderPipelineCore/Private/Graph/MovieGraphUtils.cpp
+++ b/Engine/Plugins/MovieScene/MovieRenderPipeline/Source/MovieRenderPipelineCore/Private/Graph/MovieGraphUtils.cpp
@@ -2,6 +2,10 @@
 
 #include "MovieGraphUtils.h"
 
+#include "AudioMixerDevice.h"
+#include "AudioThread.h"
+#include "Engine/Engine.h"
+
 namespace UE::MoviePipeline::RenderGraph
 {
 	FString GetUniqueName(const TArray<FString>& InExistingNames, const FString& InBaseName)
@@ -17,4 +21,49 @@ namespace UE::MoviePipeline::RenderGraph
 
 		return NewName;
 	}
-}
\ No newline at end of file
+
+	namespace Audio
+	{
+		FAudioDevice* GetAudioDeviceFromWorldContext(const UObject* InWorldContextObject)
+		{
+			const UWorld* ThisWorld = GEngine->GetWorldFromContextObject(InWorldContextObject, EGetWorldErrorMode::LogAndReturnNull);
+			if (!ThisWorld || !ThisWorld->bAllowAudioPlayback || (ThisWorld->GetNetMode() == NM_DedicatedServer))
+			{
+				return nullptr;
+			}
+
+			return ThisWorld->GetAudioDeviceRaw();
+		}
+
+		::Audio::FMixerDevice* GetAudioMixerDeviceFromWorldContext(const UObject* InWorldContextObject)
+		{
+			if (FAudioDevice* AudioDevice = GetAudioDeviceFromWorldContext(InWorldContextObject))
+			{
+				return static_cast<::Audio::FMixerDevice*>(AudioDevice);
+			}
+	
+			return nullptr;
+		}
+
+		bool IsMoviePipelineAudioOutputSupported(const UObject* InWorldContextObject)
+		{
+			const ::Audio::FMixerDevice* MixerDevice = GetAudioMixerDeviceFromWorldContext(InWorldContextObject);
+			const ::Audio::IAudioMixerPlatformInterface* AudioMixerPlatform = MixerDevice ? MixerDevice->GetAudioMixerPlatform() : nullptr;
+
+			// If the current audio mixer is non-realtime, audio output is supported
+			if (AudioMixerPlatform && AudioMixerPlatform->IsNonRealtime())
+			{
+				return true;
+			}
+
+			// If there is no async audio processing (e.g. we're in the editor), it's possible to create a new non-realtime audio mixer
+			if (!FAudioThread::IsUsingThreadedAudio())
+			{
+				return true;
+			}
+
+			// Otherwise, we can't support audio output
+			return false;
+		}
+	}
+}
diff --git a/Engine/Plugins/MovieScene/MovieRenderPipeline/Source/MovieRenderPipelineCore/Private/Graph/MovieGraphUtils.h b/Engine/Plugins/MovieScene/MovieRenderPipeline/Source/MovieRenderPipelineCore/Private/Graph/MovieGraphUtils.h
--- a/Engine/Plugins/MovieScene/MovieRenderPipeline/Source/MovieRenderPipelineCore/Private/Graph/MovieGraphUtils.h
+++ b/Engine/Plugins/MovieScene/MovieRenderPipeline/Source/MovieRenderPipelineCore/Private/Graph/MovieGraphUtils.h
@@ -4,6 +4,13 @@
 
 #include "CoreMinimal.h"
 
+// Forward declares
+class FAudioDevice;
+namespace Audio
+{
+	class FMixerDevice;
+}
+
 namespace UE::MoviePipeline::RenderGraph
 {
 	/**
@@ -12,4 +19,16 @@ namespace UE::MoviePipeline::RenderGraph
 	 * there's already an existing name "Foo", the generated name would be "Foo 1").
 	 */
 	FString GetUniqueName(const TArray<FString>& InExistingNames, const FString& InBaseName);
+
+	namespace Audio
+	{
+		/** Gets the audio device from the supplied world context (or nullptr if it could not be determined). */
+		FAudioDevice* GetAudioDeviceFromWorldContext(const UObject* InWorldContextObject);
+
+		/** Gets the audio mixer from the supplied world context (or nullptr if it could not be determined). */
+		::Audio::FMixerDevice* GetAudioMixerDeviceFromWorldContext(const UObject* InWorldContextObject);
+
+		/** Determines if the pipeline can generate audio. */
+		bool IsMoviePipelineAudioOutputSupported(const UObject* InWorldContextObject);
+	}
 }
diff --git a/Engine/Plugins/MovieScene/MovieRenderPipeline/Source/MovieRenderPipelineCore/Private/Graph/Nodes/MovieGraphAudioOutputNode.cpp b/Engine/Plugins/MovieScene/MovieRenderPipeline/Source/MovieRenderPipelineCore/Private/Graph/Nodes/MovieGraphAudioOutputNode.cpp
new file mode 100644
--- /dev/null
+++ b/Engine/Plugins/MovieScene/MovieRenderPipeline/Source/MovieRenderPipelineCore/Private/Graph/Nodes/MovieGraphAudioOutputNode.cpp
@@ -0,0 +1,210 @@
+// Copyright Epic Games, Inc. All Rights Reserved.
+
+#include "Graph/Nodes/MovieGraphAudioOutputNode.h"
+
+#include "AudioMixerDevice.h"
+#include "Graph/MovieGraphBlueprintLibrary.h"
+#include "Graph/MovieGraphPipeline.h"
+#include "Graph/MovieGraphUtils.h"
+#include "Graph/Nodes/MovieGraphOutputSettingNode.h"
+#include "MoviePipelineUtils.h"
+#include "MovieRenderPipelineCoreModule.h"
+#include "Styling/AppStyle.h"
+
+static TAutoConsoleVariable<float> CVarWaveOutputDelay(
+	TEXT("MovieRenderPipeline.WaveOutput.WriteDelay"),
+	0.5f,
+	TEXT("How long (in seconds) should the .wav writer stall the main thread to wait for the async write to finish\n")
+	TEXT("before moving on? If your .wav files take a long time to write and they're not finished by the time the\n")
+	TEXT("encoder runs, the encoder may fail.\n"),
+	ECVF_Default);
+
+#if WITH_EDITOR
+FText UMovieGraphAudioOutputNode::GetNodeTitle(const bool bGetDescriptive) const
+{
+	static const FText AudioOutputNodeName = NSLOCTEXT("MovieGraphNodes", "NodeName_AudioOutput", ".wav Output");
+	return AudioOutputNodeName;
+}
+
+FLinearColor UMovieGraphAudioOutputNode::GetNodeTitleColor() const
+{
+	static const FLinearColor AudioOutputNodeColor = FLinearColor(0.04f, 0.22f, 0.36f);
+	return AudioOutputNodeColor;
+}
+
+FSlateIcon UMovieGraphAudioOutputNode::GetIconAndTint(FLinearColor& OutColor) const
+{
+	static const FSlateIcon AudioOutputPresetIcon = FSlateIcon(FAppStyle::GetAppStyleSetName(), "GenericPlay");
+
+	OutColor = FLinearColor::White;
+	return AudioOutputPresetIcon;
+}
+#endif // WITH_EDITOR
+
+EMovieGraphBranchRestriction UMovieGraphAudioOutputNode::GetBranchRestriction() const
+{
+	return EMovieGraphBranchRestriction::Globals;
+}
+
+void UMovieGraphAudioOutputNode::OnReceiveImageDataImpl(UMovieGraphPipeline* InPipeline, UE::MovieGraph::FMovieGraphOutputMergerFrame* InRawFrameData, const TSet<FMovieGraphRenderDataIdentifier>& InMask)
+{
+	CachedPipeline = InPipeline;
+}
+
+void UMovieGraphAudioOutputNode::OnAllFramesSubmittedImpl()
+{
+	if (!UE::MoviePipeline::RenderGraph::Audio::IsMoviePipelineAudioOutputSupported(CachedPipeline.Get()))
+	{
+		return;
+	}
+
+	ActiveWriters.Reset();
+
+	const MoviePipeline::FAudioState& AudioState = CachedPipeline->GetAudioRendererInstance()->GetAudioState();
+		
+	// There should be no active submixes by the time we finalize - they should all have been converted to recorded samples.
+	check(AudioState.ActiveSubmixes.IsEmpty());
+
+	// If we didn't end up recording audio, don't try outputting any files.
+	if (AudioState.FinishedSegments.IsEmpty())
+	{
+		return;
+	}
+
+	TArray<FFinalAudioData> FinalAudioData;
+	GenerateFinalAudioData(FinalAudioData);
+
+	// Write audio files to disk
+	for (FFinalAudioData& AudioData : FinalAudioData)
+	{
+		UE::MovieGraph::FMovieGraphOutputFutureData OutputData;
+		OutputData.Shot = CachedPipeline->GetActiveShotList()[AudioData.ShotIndex];
+		OutputData.DataIdentifier = AudioData.RenderIdentifier;
+
+		// Do this before we start manipulating the filepath for the audio API
+		OutputData.FilePath = AudioData.FilePath;
+
+		const FString FileName = FPaths::GetBaseFilename(AudioData.FilePath);
+		FString FileFolder = FPaths::GetPath(AudioData.FilePath);
+
+		TPromise<bool> Completed;
+		CachedPipeline->AddOutputFuture(Completed.GetFuture(), OutputData);
+
+		TUniquePtr<Audio::FSoundWavePCMWriter> Writer = MakeUnique<Audio::FSoundWavePCMWriter>();
+		const bool bSuccess = Writer->BeginWriteToWavFile(AudioData.SampleBuffer, FileName, FileFolder);
+
+		Completed.SetValue(bSuccess);
+		ActiveWriters.Add(MoveTemp(Writer));
+	}
+
+	// The FSoundWavePCMWriter is unfortunately async, and the completion callbacks don't work unless the main thread
+	// can be spun (as it enqueues a callback onto the main thread). We're going to just cheat here and stall the main thread
+	// for 0.5s to give it a chance to write to disk. It'll only potentially be an issue with command line encoding if it takes
+	// longer than 0.5s to write to disk.
+	if (const TConsoleVariableData<float>* CVar = IConsoleManager::Get().FindTConsoleVariableDataFloat(TEXT("MovieRenderPipeline.WaveOutput.WriteDelay")))
+	{
+		UE_LOG(LogMovieRenderPipeline, Log, TEXT("Delaying main thread for %f seconds while audio writes to disk."), CVar->GetValueOnGameThread());
+		FPlatformProcess::Sleep(CVar->GetValueOnGameThread());
+	}
+}
+
+bool UMovieGraphAudioOutputNode::IsFinishedWritingToDiskImpl() const
+{
+	// Files are written in OnAllFramesSubmittedImpl(), so writing should be finished by the time this method is called
+	return true;
+}
+
+FString UMovieGraphAudioOutputNode::GenerateOutputPath(const FMovieGraphRenderDataIdentifier& InRenderIdentifier, const TObjectPtr<UMoviePipelineExecutorShot>& InShot) const
+{
+	constexpr bool bIncludeCDOs = true;
+	constexpr bool bExactMatch = true;
+	const TObjectPtr<UMovieGraphEvaluatedConfig> EvaluatedConfig = CachedPipeline->GetTimeStepInstance()->GetCalculatedTimeData().EvaluatedConfig;
+	const UMovieGraphOutputSettingNode* OutputNode = EvaluatedConfig->GetSettingForBranch<UMovieGraphOutputSettingNode>(GlobalsPinName, bIncludeCDOs, bExactMatch);
+	const FString OutputFilename = FileNameFormatOverride.Len() > 0 ? FileNameFormatOverride : OutputNode->FileNameFormat;
+	FString FileNameFormatString = OutputNode->OutputDirectory.Path / OutputFilename;
+
+	constexpr bool bIncludeRenderPass = false;
+	constexpr bool bTestFrameNumber = false;
+	UE::MoviePipeline::ValidateOutputFormatString(FileNameFormatString, bIncludeRenderPass, bTestFrameNumber);
+	
+	// Strip any frame number tags so we don't get one audio file per frame.
+	UE::MoviePipeline::RemoveFrameNumberFormatStrings(FileNameFormatString, true);
+
+	FMovieGraphFilenameResolveParams ResolveParams;
+	ResolveParams.InitializationTime = CachedPipeline->GetInitializationTime();
+	ResolveParams.Job = CachedPipeline->GetCurrentJob();
+	ResolveParams.Shot = InShot;
+	ResolveParams.FileNameFormatOverrides = {
+		{TEXT("render_pass"), RendererName},
+		{TEXT("ext"), OutputExtension}
+	};
+	ResolveParams.EvaluatedConfig = EvaluatedConfig;
+	ResolveParams.RenderDataIdentifier = InRenderIdentifier;
+	// int32 VersionNumber = 0; // TODO: Track versions
+	// FileNameFormatString.ReplaceInline(TEXT("{version}"), *FString::Printf(TEXT("v%0*d"), 3, VersionNumber));
+
+	// Generate a filename for this output file
+	FMovieGraphResolveArgs FinalFormatArgs;
+	FString FinalFilePath = UMovieGraphBlueprintLibrary::ResolveFilenameFormatArguments(FileNameFormatString, ResolveParams, FinalFormatArgs);
+
+	if (FPaths::IsRelative(FinalFilePath))
+	{
+		FinalFilePath = FPaths::ConvertRelativePathToFull(FinalFilePath);
+	}
+
+	return FinalFilePath;
+}
+
+void UMovieGraphAudioOutputNode::GenerateFinalAudioData(TArray<FFinalAudioData>& OutFinalAudioData)
+{
+	// Generate the "render" identifier for the audio
+	const int32 ShotIndex = CachedPipeline->GetCurrentShotIndex();
+	const TObjectPtr<UMoviePipelineExecutorShot>& Shot = CachedPipeline->GetActiveShotList()[ShotIndex];
+	const FString SubResourceName = FString();
+	const FString CameraName = Shot->InnerName;
+	const FMovieGraphRenderDataIdentifier RenderDataIdentifier(GlobalsPinName, RendererName, SubResourceName, CameraName);
+	
+	const FString FinalFilePath = GenerateOutputPath(RenderDataIdentifier, Shot);
+	
+	// Move each finished audio segment to its associated final audio data struct (the one with a matching file name)
+	const TArray<MoviePipeline::FAudioState::FAudioSegment>& FinishedSegments = CachedPipeline->GetAudioRendererInstance()->GetAudioState().FinishedSegments;
+	for (int32 SegmentIndex = 0; SegmentIndex < FinishedSegments.Num(); SegmentIndex++)
+	{
+		const MoviePipeline::FAudioState::FAudioSegment& Segment = FinishedSegments[SegmentIndex];
+		if (AlreadyWrittenSegments.Contains(Segment.Id))
+		{
+			continue;
+		}
+
+		// This function may get called many times over the course of a rendering (if flushing to disk). But because we don't own the audio
+		// data we can't remove it after we've written it to disk. So we keep track of which segments we've already written to disk for and
+		// skip them. This avoids us generating extra output futures for previous shots when we render a new shot.
+		AlreadyWrittenSegments.Add(Segment.Id);
+
+		// Look to see if we already have a output file to append this to.
+		FFinalAudioData* OutputSegment = nullptr;
+		for (int32 Index = 0; Index < OutFinalAudioData.Num(); Index++)
+		{
+			if (OutFinalAudioData[Index].FilePath == FinalFilePath)
+			{
+				OutputSegment = &OutFinalAudioData[Index];
+				break;
+			}
+		}
+
+		if (!OutputSegment)
+		{
+			OutFinalAudioData.AddDefaulted();
+			OutputSegment = &OutFinalAudioData[OutFinalAudioData.Num() - 1];
+			OutputSegment->FilePath = FinalFilePath;
+			OutputSegment->ShotIndex = CachedPipeline->GetCurrentShotIndex();
+			OutputSegment->RenderIdentifier = RenderDataIdentifier;
+		}
+
+		// Convert our samples and append them to the existing array
+		const double StartTime = FPlatformTime::Seconds();
+		Audio::TSampleBuffer<int16> SampleBuffer = Audio::TSampleBuffer<int16>(Segment.SegmentData.GetData(), Segment.SegmentData.Num(), Segment.NumChannels, Segment.SampleRate);
+		OutputSegment->SampleBuffer.Append(SampleBuffer.GetData(), Segment.SegmentData.Num(), Segment.NumChannels, Segment.SampleRate);
+		UE_LOG(LogMovieRenderPipeline, Log, TEXT("Audio Segment took %f seconds to convert to a sample buffer."), (FPlatformTime::Seconds() - StartTime));
+	}
+}
diff --git a/Engine/Plugins/MovieScene/MovieRenderPipeline/Source/MovieRenderPipelineCore/Private/Graph/Nodes/MovieGraphCommandLineEncoderNode.cpp b/Engine/Plugins/MovieScene/MovieRenderPipeline/Source/MovieRenderPipelineCore/Private/Graph/Nodes/MovieGraphCommandLineEncoderNode.cpp
--- a/Engine/Plugins/MovieScene/MovieRenderPipeline/Source/MovieRenderPipelineCore/Private/Graph/Nodes/MovieGraphCommandLineEncoderNode.cpp
+++ b/Engine/Plugins/MovieScene/MovieRenderPipeline/Source/MovieRenderPipelineCore/Private/Graph/Nodes/MovieGraphCommandLineEncoderNode.cpp
@@ -4,6 +4,7 @@
 
 #include "Graph/MovieGraphBlueprintLibrary.h"
 #include "Graph/MovieGraphPipeline.h"
+#include "Graph/Nodes/MovieGraphAudioOutputNode.h"
 #include "Graph/Nodes/MovieGraphOutputSettingNode.h"
 #include "HAL/FileManager.h"
 #include "HAL/PlatformFileManager.h"
@@ -279,7 +280,7 @@ TMap<FMovieGraphRenderDataIdentifier, UMovieGraphCommandLineEncoderNode::FEncode
 		TMap<FString, TArray<FString>> AudioPathsByExtension;
 		for (const TPair<FMovieGraphRenderDataIdentifier, FMovieGraphRenderLayerOutputData>& InnerRenderLayer : GeneratedRenderData.RenderLayerData)
 		{
-			if (InnerRenderLayer.Key.RendererName == FString(TEXT("Audio")))
+			if (InnerRenderLayer.Key.RendererName == UMovieGraphAudioOutputNode::RendererName)
 			{
 				for (const FString& FilePath : InnerRenderLayer.Value.FilePaths)
 				{
@@ -296,8 +297,7 @@ TMap<FMovieGraphRenderDataIdentifier, UMovieGraphCommandLineEncoderNode::FEncode
 			const FMovieGraphRenderDataIdentifier& RenderIdentifier = RenderLayer.Key;
 			const FMovieGraphRenderLayerOutputData& RenderOutputData = RenderLayer.Value;
 			
-			// TODO: Is "Audio" set on the renderer name or something else? We don't have audio yet...
-			if (RenderIdentifier.RendererName == FString(TEXT("Audio")))
+			if (RenderIdentifier.RendererName == UMovieGraphAudioOutputNode::RendererName)
 			{
 				continue;
 			}
@@ -423,8 +423,9 @@ void UMovieGraphCommandLineEncoderNode::GenerateTemporaryEncoderInputFiles(const
 		{
 			StringBuilder.Appendf(TEXT("file 'file:%s'%s"), *Path, LINE_TERMINATOR);
 
-			// Some encoders require the duration of each file to be listed after the file.
-			if (Pair.Key != TEXT("wav"))
+			// Some video encoders require the duration of each file to be listed after the file. Write duration for all
+			// video encoders; providing the extra information doesn't hurt.
+			if (Pair.Key != UMovieGraphAudioOutputNode::OutputExtension)
 			{
 				StringBuilder.Appendf(TEXT("duration %f%s"), FrameRateAsDuration, LINE_TERMINATOR);
 			}
@@ -433,8 +434,8 @@ void UMovieGraphCommandLineEncoderNode::GenerateTemporaryEncoderInputFiles(const
 		// Save this to disk.
 		FFileHelper::SaveStringToFile(StringBuilder.ToString(), *FinalFilePath);
 
-		// Not a great solution but best we've got right now
-		if (Pair.Key == TEXT("wav"))
+		// Separate audio and video files
+		if (Pair.Key == UMovieGraphAudioOutputNode::OutputExtension)
 		{
 			OutAudioInputFilePaths.Add(FinalFilePath);
 		}
diff --git a/Engine/Plugins/MovieScene/MovieRenderPipeline/Source/MovieRenderPipelineCore/Public/Graph/MovieGraphDataTypes.h b/Engine/Plugins/MovieScene/MovieRenderPipeline/Source/MovieRenderPipelineCore/Public/Graph/MovieGraphDataTypes.h
--- a/Engine/Plugins/MovieScene/MovieRenderPipeline/Source/MovieRenderPipelineCore/Public/Graph/MovieGraphDataTypes.h
+++ b/Engine/Plugins/MovieScene/MovieRenderPipeline/Source/MovieRenderPipelineCore/Public/Graph/MovieGraphDataTypes.h
@@ -9,6 +9,8 @@
 #include "ImagePixelData.h"
 #include "Containers/Queue.h"
 #include "Misc/FrameRate.h"
+#include "MovieRenderPipelineDataTypes.h"
+
 #include "MovieGraphDataTypes.generated.h"
 
 // Forward Declares
@@ -49,6 +51,13 @@ struct MOVIERENDERPIPELINECORE_API FMovieGraphInitConfig
 	UPROPERTY(BlueprintReadWrite, EditAnywhere, Category = "Movie Graph")
 	TSubclassOf<UMovieGraphDataSourceBase> DataSourceClass;
 
+	/**
+	 * Which class should the UMovieGraphPipeline use to generate audio. Defaults to
+	 * UMovieGraphDefaultAudioOutput.
+	 */
+	UPROPERTY(BlueprintReadWrite, EditAnywhere, Category = "Movie Graph")
+	TSubclassOf<UMovieGraphAudioRendererBase> AudioRendererClass;
+
 	/**
 	* Should the UMovieGraphPipeline render the full player viewport? Defaults
 	* to false (so no 3d content is rendered) so we can display the UMG widgets
@@ -179,6 +188,38 @@ public:
 	UMovieGraphPipeline* GetOwningGraph() const;
 };
 
+/** Base class for generating audio while the pipeline is running. */
+UCLASS(BlueprintType, Abstract)
+class MOVIERENDERPIPELINECORE_API UMovieGraphAudioRendererBase : public UObject
+{
+	GENERATED_BODY()
+
+public:
+	/** Tell our submixes to start capturing the data they are generating. Should only be called once output frames are being produced. */
+	virtual void StartAudioRecording() {}
+
+	/** Tell our submixes to stop capturing the data, and then store a copy of it. */
+	virtual void StopAudioRecording() {}
+
+	/** Attempt to process the audio thread work. This is complicated by non-linear time steps. */
+	virtual void ProcessAudioTick() {}
+
+	/** Prepares for audio rendering (ensuring volume is set correctly, the needed cvars are set, etc). */
+	virtual void SetupAudioRendering() {}
+
+	/** Undoes the work done in SetupAudioRendering(). */
+	virtual void TeardownAudioRendering() const {}
+
+	/** Gets the pipeline that owns this audio output instance. */
+	UMovieGraphPipeline* GetOwningGraph() const;
+
+	/** Gets the current state of the audio renderer. This is the main data source that audio-related nodes can reference. */
+	const MoviePipeline::FAudioState& GetAudioState() const;
+
+protected:
+	MoviePipeline::FAudioState AudioState;
+};
+
 // ToDo: Both of these can probably go into the Default Renderer implementation.
 struct FMovieGraphRenderPassLayerData
 {
diff --git a/Engine/Plugins/MovieScene/MovieRenderPipeline/Source/MovieRenderPipelineCore/Public/Graph/MovieGraphDefaultAudioRenderer.h b/Engine/Plugins/MovieScene/MovieRenderPipeline/Source/MovieRenderPipelineCore/Public/Graph/MovieGraphDefaultAudioRenderer.h
new file mode 100644
--- /dev/null
+++ b/Engine/Plugins/MovieScene/MovieRenderPipeline/Source/MovieRenderPipelineCore/Public/Graph/MovieGraphDefaultAudioRenderer.h
@@ -0,0 +1,28 @@
+﻿// Copyright Epic Games, Inc. All Rights Reserved.
+
+#pragma once
+
+#include "MovieGraphDataTypes.h"
+
+#include "MovieGraphDefaultAudioRenderer.generated.h"
+
+/**
+ * Provides default audio rendering for the pipeline.
+ */
+UCLASS(BlueprintType, meta = (DisplayName = "Default Audio Renderer"))
+class MOVIERENDERPIPELINECORE_API UMovieGraphDefaultAudioRenderer : public UMovieGraphAudioRendererBase
+{
+	GENERATED_BODY()
+
+public:
+	UMovieGraphDefaultAudioRenderer() = default;
+
+protected:
+	// UMovieGraphAudioOutputBase interface
+	virtual void StartAudioRecording() override;
+	virtual void StopAudioRecording() override;
+	virtual void ProcessAudioTick() override;
+	virtual void SetupAudioRendering() override;
+	virtual void TeardownAudioRendering() const override;
+	// ~UMovieGraphAudioOutputBase
+};
diff --git a/Engine/Plugins/MovieScene/MovieRenderPipeline/Source/MovieRenderPipelineCore/Public/Graph/MovieGraphPipeline.h b/Engine/Plugins/MovieScene/MovieRenderPipeline/Source/MovieRenderPipelineCore/Public/Graph/MovieGraphPipeline.h
--- a/Engine/Plugins/MovieScene/MovieRenderPipeline/Source/MovieRenderPipelineCore/Public/Graph/MovieGraphPipeline.h
+++ b/Engine/Plugins/MovieScene/MovieRenderPipeline/Source/MovieRenderPipelineCore/Public/Graph/MovieGraphPipeline.h
@@ -83,6 +83,9 @@ public:
 	
 	/** Used occasionally to cross-reference other components. Don't call this unless you know what you're doing. */
 	UMovieGraphDataSourceBase* GetDataSourceInstance() const { return GraphDataSourceInstance; }
+
+	/** Used occasionally to cross-reference other components. Don't call this unless you know what you're doing. */
+	UMovieGraphAudioRendererBase* GetAudioRendererInstance() const { return GraphAudioRendererInstance; }
 	
 	/** Gets the Output Merger for this Movie Pipeline which is responsible for gathering all of the data coming in for a given output frame, before making it available to the MovieGraphPipeline. */
 	TSharedPtr<UE::MovieGraph::IMovieGraphOutputMerger> GetOutputMerger() const { return OutputMerger; }
@@ -154,7 +157,6 @@ protected:
 	virtual EMovieRenderPipelineState GetPipelineStateImpl() const override { return PipelineState; }
 	// ~UMoviePipelineBase Interface
 
-
 protected:
 	/** Time step instances for each shot, where the index into the array corresponds to the shot index. */
 	UPROPERTY(Transient)
@@ -174,6 +176,9 @@ protected:
 	UPROPERTY(Transient)
 	TObjectPtr<UMovieGraphDataSourceBase> GraphDataSourceInstance;
 
+	UPROPERTY(Transient)
+	TObjectPtr<UMovieGraphAudioRendererBase> GraphAudioRendererInstance;
+
 protected:
 	UPROPERTY(Transient)
 	TObjectPtr<UMoviePipelineExecutorJob> CurrentJob;
diff --git a/Engine/Plugins/MovieScene/MovieRenderPipeline/Source/MovieRenderPipelineCore/Public/Graph/Nodes/MovieGraphAudioOutputNode.h b/Engine/Plugins/MovieScene/MovieRenderPipeline/Source/MovieRenderPipelineCore/Public/Graph/Nodes/MovieGraphAudioOutputNode.h
new file mode 100644
--- /dev/null
+++ b/Engine/Plugins/MovieScene/MovieRenderPipeline/Source/MovieRenderPipelineCore/Public/Graph/Nodes/MovieGraphAudioOutputNode.h
@@ -0,0 +1,82 @@
+﻿// Copyright Epic Games, Inc. All Rights Reserved.
+
+#pragma once
+
+#include "Graph/Nodes/MovieGraphFileOutputNode.h"
+#include "Sound/SampleBufferIO.h"
+
+#include "MovieGraphAudioOutputNode.generated.h"
+
+/**
+ * A node which outputs .wav audio after all renders have completed.
+ */
+UCLASS()
+class UMovieGraphAudioOutputNode : public UMovieGraphFileOutputNode
+{
+	GENERATED_BODY()
+	
+public:
+	UMovieGraphAudioOutputNode() = default;
+
+	// UMovieGraphSettingNode Interface
+#if WITH_EDITOR
+	virtual FText GetNodeTitle(const bool bGetDescriptive) const override;
+	virtual FLinearColor GetNodeTitleColor() const override;
+	virtual FSlateIcon GetIconAndTint(FLinearColor& OutColor) const override;
+#endif
+	// ~UMovieGraphSettingNode Interface
+
+	// UMovieGraphNode Interface
+	virtual EMovieGraphBranchRestriction GetBranchRestriction() const override;
+	// ~UMovieGraphNode Interface
+
+	/**
+	 * Although this node does not represent a proper "renderer", it will sometimes be identified in data structures
+	 * as a renderer. If that's the case, this is the name of the audio "renderer".
+	 */
+	inline static const FString RendererName = TEXT("Audio");
+
+	/** The extension of the audio files generated by this node. */
+	inline static const FString OutputExtension = TEXT("wav");
+
+public:
+	UPROPERTY(EditAnywhere, BlueprintReadWrite, Category = Overrides, meta = (InlineEditConditionToggle))
+	uint8 bOverride_FileNameFormatOverride : 1;
+	
+	/* File name format string override. If specified it will override the FileNameFormat from the Output Setting node. */
+	UPROPERTY(EditAnywhere, BlueprintReadWrite, Category = "File Output", meta=(EditCondition="bOverride_FileNameFormatOverride"))
+	FString FileNameFormatOverride;
+
+protected:
+	// UMovieGraphFileOutputNode Interface
+	virtual void OnReceiveImageDataImpl(UMovieGraphPipeline* InPipeline, UE::MovieGraph::FMovieGraphOutputMergerFrame* InRawFrameData, const TSet<FMovieGraphRenderDataIdentifier>& InMask) override;
+	virtual void OnAllFramesSubmittedImpl() override;
+	virtual bool IsFinishedWritingToDiskImpl() const override;
+	// ~UMovieGraphFileOutputNode Interface
+
+private:
+	/** Represents an audio file, with its associated audio data. */
+	struct FFinalAudioData
+	{
+		FString FilePath;
+		Audio::TSampleBuffer<int16> SampleBuffer;
+		int32 ShotIndex;
+		FMovieGraphRenderDataIdentifier RenderIdentifier;
+	};
+
+	/** Generate the final output path for the audio file. */
+	FString GenerateOutputPath(const FMovieGraphRenderDataIdentifier& InRenderIdentifier, const TObjectPtr<UMoviePipelineExecutorShot>& InShot) const;
+
+	/** Generates final audio data from all of the audio segments that were collected during render time. */
+	void GenerateFinalAudioData(TArray<FFinalAudioData>& OutFinalAudioData);
+
+private:
+	/** Kept alive during finalization because the writer writes async to disk but doesn't expect to fall out of scope. */
+	TArray<TUniquePtr<Audio::FSoundWavePCMWriter>> ActiveWriters;
+
+	/** Keep track of segments that have already been written to disk to avoid re-writing them (and generating new output futures). */
+	TSet<FGuid> AlreadyWrittenSegments;
+
+	/** The pipeline that is running this node. */
+	TWeakObjectPtr<UMovieGraphPipeline> CachedPipeline;
+};
diff --git a/Engine/Plugins/MovieScene/MovieRenderPipeline/Source/MovieRenderPipelineCore/Public/MovieRenderPipelineDataTypes.h b/Engine/Plugins/MovieScene/MovieRenderPipeline/Source/MovieRenderPipelineCore/Public/MovieRenderPipelineDataTypes.h
--- a/Engine/Plugins/MovieScene/MovieRenderPipeline/Source/MovieRenderPipelineCore/Public/MovieRenderPipelineDataTypes.h
+++ b/Engine/Plugins/MovieScene/MovieRenderPipeline/Source/MovieRenderPipelineCore/Public/MovieRenderPipelineDataTypes.h
@@ -1220,12 +1220,14 @@ namespace MoviePipeline
 			: bIsRecordingAudio(false)
 			, PrevUnfocusedAudioMultiplier(1.f)
 			, PrevRenderEveryTickValue(1)
+			, PrevNeverMuteNRTAudioValue(0)
 		{}
 
 		bool bIsRecordingAudio;
 
 		float PrevUnfocusedAudioMultiplier;
 		int32 PrevRenderEveryTickValue;
+		int32 PrevNeverMuteNRTAudioValue;
 
 		/** Float Buffers for Movie Segments we've finished rendering. Stored until shutdown. Fully available during Finalize stage. */
 		TArray<FAudioSegment> FinishedSegments;
diff --git a/Engine/Plugins/MovieScene/MovieRenderPipeline/Source/MovieRenderPipelineRenderPasses/Private/MovieGraphImageSequenceOutputNode.cpp b/Engine/Plugins/MovieScene/MovieRenderPipeline/Source/MovieRenderPipelineRenderPasses/Private/MovieGraphImageSequenceOutputNode.cpp
--- a/Engine/Plugins/MovieScene/MovieRenderPipeline/Source/MovieRenderPipelineRenderPasses/Private/MovieGraphImageSequenceOutputNode.cpp
+++ b/Engine/Plugins/MovieScene/MovieRenderPipeline/Source/MovieRenderPipelineRenderPasses/Private/MovieGraphImageSequenceOutputNode.cpp
@@ -163,6 +163,7 @@ void UMovieGraphImageSequenceOutputNode::OnReceiveImageDataImpl(UMovieGraphPipel
 		}
 
 		UE::MovieGraph::FMovieGraphSampleState* Payload = RenderData.Value->GetPayload<UE::MovieGraph::FMovieGraphSampleState>();
+		const TObjectPtr<UMoviePipelineExecutorShot>& Shot = InPipeline->GetActiveShotList()[Payload->TraversalContext.ShotIndex];
 
 		const bool bIncludeCDOs = true;
 		UMovieGraphOutputSettingNode* OutputSettingNode = InRawFrameData->EvaluatedConfig->GetSettingForBranch<UMovieGraphOutputSettingNode>(RenderData.Key.RootBranchName, bIncludeCDOs);
@@ -224,6 +225,7 @@ void UMovieGraphImageSequenceOutputNode::OnReceiveImageDataImpl(UMovieGraphPipel
 		Params.bEnsureAbsolutePath = true;
 		Params.FileNameFormatOverrides = AdditionalFormatArgs;
 		Params.InitializationTime = InPipeline->GetInitializationTime();
+		Params.Shot = Shot;
 		Params.Job = InPipeline->GetCurrentJob();
 		Params.EvaluatedConfig = InRawFrameData->EvaluatedConfig.Get();
 
@@ -276,7 +278,7 @@ void UMovieGraphImageSequenceOutputNode::OnReceiveImageDataImpl(UMovieGraphPipel
 		}
 
 		UE::MovieGraph::FMovieGraphOutputFutureData OutputData;
-		OutputData.Shot = nullptr;
+		OutputData.Shot = Shot;
 		OutputData.FilePath = FileName;
 		OutputData.DataIdentifier = RenderData.Key;
 
@@ -540,6 +542,7 @@ FString UMovieGraphImageSequenceOutputNode_EXR::ResolveOutputFilename(
 	Params.FileNameFormatOverrides = FormatOverrides;
 	Params.InitializationTime = InPipeline->GetInitializationTime();
 	Params.Job = InPipeline->GetCurrentJob();
+	Params.Shot = InPipeline->GetActiveShotList()[InRawFrameData->TraversalContext.ShotIndex];
 	Params.EvaluatedConfig = InRawFrameData->EvaluatedConfig.Get();
 	
 	FString FinalFilePath = UMovieGraphBlueprintLibrary::ResolveFilenameFormatArguments(FilePathFormatString, Params, OutResolveArgs);
diff --git a/Engine/Plugins/MovieScene/MovieRenderPipeline/Source/MovieRenderPipelineRenderPasses/Private/MoviePipelineWaveOutput.cpp b/Engine/Plugins/MovieScene/MovieRenderPipeline/Source/MovieRenderPipelineRenderPasses/Private/MoviePipelineWaveOutput.cpp
--- a/Engine/Plugins/MovieScene/MovieRenderPipeline/Source/MovieRenderPipelineRenderPasses/Private/MoviePipelineWaveOutput.cpp
+++ b/Engine/Plugins/MovieScene/MovieRenderPipeline/Source/MovieRenderPipelineRenderPasses/Private/MoviePipelineWaveOutput.cpp
@@ -17,14 +17,6 @@
 
 #include UE_INLINE_GENERATED_CPP_BY_NAME(MoviePipelineWaveOutput)
 
-static TAutoConsoleVariable<float> CVarWaveOutputDelay(
-	TEXT("MovieRenderPipeline.WaveOutput.WriteDelay"),
-	0.5f,
-	TEXT("How long (in seconds) should the .wav writer stall the main thread to wait for the async write to finish\n")
-	TEXT("before moving on? If your .wav files take a long time to write and they're not finished by the time the\n")
-	TEXT("encoder runs, the encoder may fail.\n"),
-	ECVF_Default);
-
 static FAudioDevice* GetAudioDeviceFromWorldContext(const UObject* WorldContextObject)
 {
 	UWorld* ThisWorld = GEngine->GetWorldFromContextObject(WorldContextObject, EGetWorldErrorMode::LogAndReturnNull);
@@ -183,8 +175,11 @@ void UMoviePipelineWaveOutput::BeginFinalizeImpl()
 	// can be spun (as it enqueues a callback onto the main thread). We're going to just cheat here and stall the main thread
 	// for 0.5s to give it a chance to write to disk. It'll only potentially be an issue with command line encoding if it takes
 	// longer than 0.5s to write to disk.
-	UE_LOG(LogMovieRenderPipeline, Log, TEXT("Delaying main thread for %f seconds while audio writes to disk."), CVarWaveOutputDelay.GetValueOnGameThread());
-	FPlatformProcess::Sleep(CVarWaveOutputDelay.GetValueOnGameThread());
+	if (const TConsoleVariableData<float>* CVar = IConsoleManager::Get().FindTConsoleVariableDataFloat(TEXT("MovieRenderPipeline.WaveOutput.WriteDelay")))
+	{
+		UE_LOG(LogMovieRenderPipeline, Log, TEXT("Delaying main thread for %f seconds while audio writes to disk."), CVar->GetValueOnGameThread());
+		FPlatformProcess::Sleep(CVar->GetValueOnGameThread());
+	}
 }
 
 void UMoviePipelineWaveOutput::OnShotFinishedImpl(const UMoviePipelineExecutorShot* InShot, const bool bFlushToDisk)
