MoviePipeline: Updated the Select node to include a Default branch. Additionally, users can now select almost any data type to use with the node -- it's no longer restricted to just strings. Many updates to the supporting infrastructure were required to make this possible.

The performance impact of the change to using arbitrary data types on the node should be minimal in most cases. However, the case where the Select node is driven by a variable node that is in turn driven by a variable assignment on the job is the slowest and is a potential target for optimization in the future.

Note that the default graph asset was updated in order to primarily reflect the changes made to UMovieGraphValueContainer (default property name and data type updates).

#rnx

diff --git a/Engine/Plugins/MovieScene/MovieRenderPipeline/Source/MovieRenderPipelineCore/Private/Graph/MovieGraphNode.cpp b/Engine/Plugins/MovieScene/MovieRenderPipeline/Source/MovieRenderPipelineCore/Private/Graph/MovieGraphNode.cpp
--- a/Engine/Plugins/MovieScene/MovieRenderPipeline/Source/MovieRenderPipelineCore/Private/Graph/MovieGraphNode.cpp
+++ b/Engine/Plugins/MovieScene/MovieRenderPipeline/Source/MovieRenderPipelineCore/Private/Graph/MovieGraphNode.cpp
@@ -286,11 +286,11 @@ UMovieGraphConfig* UMovieGraphNode::GetGraph() const
 	return Cast<UMovieGraphConfig>(GetOuter());
 }
 
-UMovieGraphPin* UMovieGraphNode::GetInputPin(const FName& Label) const
+UMovieGraphPin* UMovieGraphNode::GetInputPin(const FName& Label, const bool bIsBuiltInPin) const
 {
 	for (UMovieGraphPin* InputPin : InputPins)
 	{
-		if (InputPin->Properties.Label == Label)
+		if ((InputPin->Properties.Label == Label) && (InputPin->Properties.bIsBuiltIn == bIsBuiltInPin))
 		{
 			return InputPin;
 		}
diff --git a/Engine/Plugins/MovieScene/MovieRenderPipeline/Source/MovieRenderPipelineCore/Private/Graph/Nodes/MovieGraphSelectNode.cpp b/Engine/Plugins/MovieScene/MovieRenderPipeline/Source/MovieRenderPipelineCore/Private/Graph/Nodes/MovieGraphSelectNode.cpp
--- a/Engine/Plugins/MovieScene/MovieRenderPipeline/Source/MovieRenderPipelineCore/Private/Graph/Nodes/MovieGraphSelectNode.cpp
+++ b/Engine/Plugins/MovieScene/MovieRenderPipeline/Source/MovieRenderPipelineCore/Private/Graph/Nodes/MovieGraphSelectNode.cpp
@@ -2,25 +2,157 @@
 
 #include "Graph/Nodes/MovieGraphSelectNode.h"
 
+#include "MovieRenderPipelineCoreModule.h"
 #include "Graph/MovieGraphConfig.h"
 #include "Styling/AppStyle.h"
 
 namespace UE::MovieGraph::SelectNode
 {
 	static const FName SelectedOption("Selected Option");
+	static const FName DefaultBranch("Default");
+}
+
+UMovieGraphSelectNode::UMovieGraphSelectNode()
+{
+	if (!HasAnyFlags(RF_ClassDefaultObject))
+	{
+		SelectOptions = CreateDefaultSubobject<UMovieGraphValueContainer>(TEXT("SelectOptions"));
+		SelectOptions->SetValueContainerType(EMovieGraphContainerType::Array);
+		SelectOptions->SetValueType(EMovieGraphValueType::String);
+		SelectOptions->SetPropertyName(FName(TEXT("Options")));
+
+		SelectedOption = CreateDefaultSubobject<UMovieGraphValueContainer>(TEXT("SelectedOption"));
+		SelectedOption->SetValueType(EMovieGraphValueType::String);
+		SelectedOption->SetPropertyName(FName(TEXT("Selection")));
+	}
 }
 
 TArray<FMovieGraphPinProperties> UMovieGraphSelectNode::GetInputPinProperties() const
 {
 	TArray<FMovieGraphPinProperties> Properties;
 
+	// Ensure the select options array is valid
+	TValueOrError<FPropertyBagArrayRef, EPropertyBagResult> SelectOptionsArrayResult = SelectOptions->GetArrayRef();
+	if (!SelectOptionsArrayResult.HasValue() || SelectOptionsArrayResult.HasError())
+	{
+		UE_LOG(LogMovieRenderPipeline, Warning, TEXT("When generating pins on the Select node, found an invalid options array. Select node will not work as intended."));
+		return Properties;
+	}
+	
+	const FPropertyBagArrayRef SelectOptionsArray = SelectOptionsArrayResult.GetValue();
+	
 	// Generate branch pins for each option
-	for (const FString& SelectOption : SelectOptions)
+	for (int32 Index = 0; Index < SelectOptionsArray.Num(); ++Index)
 	{
-		Properties.Add(FMovieGraphPinProperties::MakeBranchProperties(FName(SelectOption)));
+		switch (const EMovieGraphValueType ValueType = SelectOptions->GetValueType())
+		{
+		case EMovieGraphValueType::Bool:
+			{
+				TValueOrError<bool, EPropertyBagResult> Result = SelectOptionsArray.GetValueBool(Index);
+				if (Result.IsValid())
+				{
+					FMovieGraphPinProperties BranchProperties = FMovieGraphPinProperties::MakeBranchProperties(FName(LexToString(Result.GetValue())));
+					BranchProperties.bIsBuiltIn = false;
+					Properties.Add(MoveTemp(BranchProperties));
+				}
+			}
+			break;
+
+		case EMovieGraphValueType::Byte:
+			{
+				TValueOrError<uint8, EPropertyBagResult> Result = SelectOptionsArray.GetValueByte(Index);
+				if (Result.IsValid())
+				{
+					FMovieGraphPinProperties BranchProperties = FMovieGraphPinProperties::MakeBranchProperties(FName(LexToString(Result.GetValue())));
+					BranchProperties.bIsBuiltIn = false;
+					Properties.Add(MoveTemp(BranchProperties));
+				}
+			}
+			break;
+			
+		case EMovieGraphValueType::String:
+			{
+				TValueOrError<FString, EPropertyBagResult> Result = SelectOptionsArray.GetValueString(Index);
+				if (Result.IsValid())
+				{
+					FMovieGraphPinProperties BranchProperties = FMovieGraphPinProperties::MakeBranchProperties(FName(Result.GetValue()));
+					BranchProperties.bIsBuiltIn = false;
+					Properties.Add(MoveTemp(BranchProperties));
+				}
+			}
+			break;
+
+		case EMovieGraphValueType::Int32:
+			{
+				TValueOrError<int32, EPropertyBagResult> Result = SelectOptionsArray.GetValueInt32(Index);
+				if (Result.IsValid())
+				{
+					FMovieGraphPinProperties BranchProperties = FMovieGraphPinProperties::MakeBranchProperties(FName(LexToString(Result.GetValue())));
+					BranchProperties.bIsBuiltIn = false;
+					Properties.Add(MoveTemp(BranchProperties));
+				}
+			}
+			break;
+
+		case EMovieGraphValueType::Int64:
+			{
+				TValueOrError<int64, EPropertyBagResult> Result = SelectOptionsArray.GetValueInt64(Index);
+				if (Result.IsValid())
+				{
+					FMovieGraphPinProperties BranchProperties = FMovieGraphPinProperties::MakeBranchProperties(FName(LexToString(Result.GetValue())));
+					BranchProperties.bIsBuiltIn = false;
+					Properties.Add(MoveTemp(BranchProperties));
+				}
+			}
+			break;
+
+		case EMovieGraphValueType::Name:
+			{
+				TValueOrError<FName, EPropertyBagResult> Result = SelectOptionsArray.GetValueName(Index);
+				if (Result.IsValid())
+				{
+					FMovieGraphPinProperties BranchProperties = FMovieGraphPinProperties::MakeBranchProperties(Result.GetValue());
+					BranchProperties.bIsBuiltIn = false;
+					Properties.Add(MoveTemp(BranchProperties));
+				}
+			}
+			break;
+
+		case EMovieGraphValueType::Text:
+			{
+				TValueOrError<FText, EPropertyBagResult> Result = SelectOptionsArray.GetValueText(Index);
+				if (Result.IsValid())
+				{
+					FMovieGraphPinProperties BranchProperties = FMovieGraphPinProperties::MakeBranchProperties(FName(Result.GetValue().ToString()));
+					BranchProperties.bIsBuiltIn = false;
+					Properties.Add(MoveTemp(BranchProperties));
+				}
+			}
+			break;
+
+		case EMovieGraphValueType::Enum:
+			{
+				const UEnum* Enum = Cast<UEnum>(SelectOptions->GetValueTypeObject());
+				TValueOrError<uint8, EPropertyBagResult> Result = SelectOptionsArray.GetValueEnum(Index, Enum);
+				if (Result.IsValid())
+				{
+					const FText DisplayName = Enum->GetDisplayNameTextByValue(Result.GetValue());
+					FMovieGraphPinProperties BranchProperties = FMovieGraphPinProperties::MakeBranchProperties(FName(DisplayName.ToString()));
+					BranchProperties.bIsBuiltIn = false;
+					Properties.Add(MoveTemp(BranchProperties));
+				}
+			}
+			break;
+			
+		default:
+			const UEnum* EnumType = StaticEnum<EMovieGraphValueType>();
+			UE_LOG(LogMovieRenderPipeline, Warning, TEXT("Encountered an unexpected type in the Select node when generating pins: %s"), *EnumType->GetNameStringByValue(static_cast<uint8>(ValueType)));
+			break;
+		}
 	}
 
-	Properties.Add(FMovieGraphPinProperties(UE::MovieGraph::SelectNode::SelectedOption, EMovieGraphValueType::String, false));
+	Properties.Add(FMovieGraphPinProperties::MakeBranchProperties(UE::MovieGraph::SelectNode::DefaultBranch));
+	Properties.Add(FMovieGraphPinProperties(UE::MovieGraph::SelectNode::SelectedOption, SelectOptions->GetValueType(), false));
 	
 	return Properties;
 }
@@ -48,7 +180,7 @@ TArray<UMovieGraphPin*> UMovieGraphSelectNode::EvaluatePinsToFollow(FMovieGraphE
 	}
 
 	// The resolved value of the "Selected Option" property. May come from a connection or a value specified on the node.
-	FString ResolvedSelectValue;
+	TObjectPtr<UMovieGraphValueContainer> ResolvedSelectedOption;
 
 	// Try getting the value from a connection first
 	bool bGotValueFromConnection = false;
@@ -58,38 +190,182 @@ TArray<UMovieGraphPin*> UMovieGraphSelectNode::EvaluatePinsToFollow(FMovieGraphE
 		{
 			if (const UMovieGraphNode* ConnectedNode = OtherPin->Node)
 			{
-				bGotValueFromConnection = true;
-				ResolvedSelectValue = ConnectedNode->GetResolvedValueForOutputPin(OtherPin->Properties.Label, &InContext.UserContext);
+				bGotValueFromConnection = ConnectedNode->GetResolvedValueForOutputPin(OtherPin->Properties.Label, &InContext.UserContext, ResolvedSelectedOption);
 			}
 		}
 	}
 
 	if (!bGotValueFromConnection)
 	{
-		ResolvedSelectValue = SelectedOption;
+		ResolvedSelectedOption = SelectedOption;
 	}
 
-	for (const FString& SelectOption : SelectOptions)
+	// Ensure the select options array is valid
+	TValueOrError<FPropertyBagArrayRef, EPropertyBagResult> SelectOptionsArrayResult = SelectOptions->GetArrayRef();
+	if (!SelectOptionsArrayResult.HasValue() || SelectOptionsArrayResult.HasError())
 	{
-		if (SelectOption == ResolvedSelectValue)
+		UE_LOG(LogMovieRenderPipeline, Warning, TEXT("When evaluating pins on the Select node, found an invalid options array. Select node will not work as intended."));
+		return PinsToFollow;
+	}
+	
+	const FPropertyBagArrayRef SelectOptionsArray = SelectOptionsArrayResult.GetValue();
+
+	// Iterate over all select options, and follow the pin that matches the selected option
+	for (int32 Index = 0; Index < SelectOptionsArray.Num(); ++Index)
+	{
+		if (!ensureMsgf(InputPins.IsValidIndex(Index), TEXT("Found more select options than pins on the node")))
+		{
+			continue;
+		}
+		
+		const TObjectPtr<UMovieGraphPin>& InputPin = InputPins[Index];
+
+		// If the pin isn't connected, there's no point in determining if it's a match
+		if (!InputPin->IsConnected())
 		{
-			PinsToFollow.Add(GetInputPin(FName(SelectOption)));
+			continue;
+		}
+
+		switch(SelectOptions->GetValueType())
+		{
+		case EMovieGraphValueType::Bool:
+			{
+				TValueOrError<bool, EPropertyBagResult> PinValue = SelectOptionsArray.GetValueBool(Index);
+				bool bSelectedOption;
+				if (PinValue.HasValue() && !PinValue.HasError() && ResolvedSelectedOption->GetValueBool(bSelectedOption) && (PinValue.GetValue() == bSelectedOption))
+				{
+					PinsToFollow.Add(InputPin);
+				}
+			}
+			break;
+
+		case EMovieGraphValueType::Byte:
+			{
+				TValueOrError<uint8, EPropertyBagResult> PinValue = SelectOptionsArray.GetValueByte(Index);
+				uint8 SelectedValue;
+				if (PinValue.HasValue() && !PinValue.HasError() && ResolvedSelectedOption->GetValueByte(SelectedValue) && (PinValue.GetValue() == SelectedValue))
+				{
+					PinsToFollow.Add(InputPin);
+				}
+			}
+			break;
+
+		case EMovieGraphValueType::String:
+			{
+				TValueOrError<FString, EPropertyBagResult> PinValue = SelectOptionsArray.GetValueString(Index);
+				FString SelectedValue;
+				if (PinValue.HasValue() && !PinValue.HasError() && ResolvedSelectedOption->GetValueString(SelectedValue) && (PinValue.GetValue() == SelectedValue))
+				{
+					PinsToFollow.Add(InputPin);
+				}
+			}
+			break;
+		
+		case EMovieGraphValueType::Int32:
+			{
+				TValueOrError<int32, EPropertyBagResult> PinValue = SelectOptionsArray.GetValueInt32(Index);
+				int32 SelectedValue;
+				if (PinValue.HasValue() && !PinValue.HasError() && ResolvedSelectedOption->GetValueInt32(SelectedValue) && (PinValue.GetValue() == SelectedValue))
+				{
+					PinsToFollow.Add(InputPin);
+				}
+			}
+			break;
+
+		case EMovieGraphValueType::Int64:
+			{
+				TValueOrError<int64, EPropertyBagResult> PinValue = SelectOptionsArray.GetValueInt64(Index);
+				int64 SelectedValue;
+				if (PinValue.HasValue() && !PinValue.HasError() && ResolvedSelectedOption->GetValueInt64(SelectedValue) && (PinValue.GetValue() == SelectedValue))
+				{
+					PinsToFollow.Add(InputPin);
+				}
+			}
+			break;
+
+		case EMovieGraphValueType::Name:
+			{
+				TValueOrError<FName, EPropertyBagResult> PinValue = SelectOptionsArray.GetValueName(Index);
+				FName SelectedValue;
+				if (PinValue.HasValue() && !PinValue.HasError() && ResolvedSelectedOption->GetValueName(SelectedValue) && (PinValue.GetValue() == SelectedValue))
+				{
+					PinsToFollow.Add(InputPin);
+				}
+			}
+			break;
+
+		case EMovieGraphValueType::Text:
+			{
+				TValueOrError<FText, EPropertyBagResult> PinValue = SelectOptionsArray.GetValueText(Index);
+				FText SelectedValue;
+				if (PinValue.HasValue() && !PinValue.HasError() && ResolvedSelectedOption->GetValueText(SelectedValue) && PinValue.GetValue().EqualTo(SelectedValue))
+				{
+					PinsToFollow.Add(InputPin);
+				}
+			}
+			break;
+			
+		case EMovieGraphValueType::Enum:
+			{
+				const UEnum* OptionsEnum = Cast<UEnum>(SelectOptions->GetValueTypeObject());
+				const UEnum* SelectedEnum = Cast<UEnum>(ResolvedSelectedOption->GetValueTypeObject());
+				
+				TValueOrError<uint8, EPropertyBagResult> PinValue = SelectOptionsArray.GetValueEnum(Index, OptionsEnum);
+				uint8 SelectedValue;
+				
+				if (OptionsEnum && SelectedEnum && PinValue.HasValue() && !PinValue.HasError() && ResolvedSelectedOption->GetValueEnum(SelectedValue, SelectedEnum) && (PinValue.GetValue() == SelectedValue))
+				{
+					PinsToFollow.Add(InputPin);
+				}
+			}
+			break;
+
+		default:
+			const UEnum* EnumType = StaticEnum<EMovieGraphValueType>();
+			UE_LOG(LogMovieRenderPipeline, Warning, TEXT("Encountered an unexpected type in the Select node when following pins: %s"), *EnumType->GetNameStringByValue(static_cast<uint8>(SelectOptions->GetValueType())));
 			break;
 		}
 	}
 
+	// Follow the Default branch if no pins match
+	if (PinsToFollow.IsEmpty())
+	{
+		PinsToFollow.Add(GetInputPin(UE::MovieGraph::SelectNode::DefaultBranch));
+	}	
+
 	return PinsToFollow;
 }
 
+void UMovieGraphSelectNode::SetDataType(const EMovieGraphValueType ValueType, UObject* InValueTypeObject)
+{
+	SelectOptions->SetValueType(ValueType, InValueTypeObject);
+	SelectedOption->SetValueType(ValueType, InValueTypeObject);
+
+	UpdatePins();
+}
+
+EMovieGraphValueType UMovieGraphSelectNode::GetValueType() const
+{
+	return SelectOptions->GetValueType();
+}
+
+const UObject* UMovieGraphSelectNode::GetValueTypeObject() const
+{
+	return SelectOptions->GetValueTypeObject();
+}
+
 #if WITH_EDITOR
 FText UMovieGraphSelectNode::GetNodeTitle(const bool bGetDescriptive) const
 {
 	static const FText SelectNodeName = NSLOCTEXT("MovieGraphNodes", "NodeName_Select", "Select");
-	static const FText SelectNodeDescription = NSLOCTEXT("MovieGraphNodes", "NodeDescription_Select", "Select\n{0}");
+	static const FText SelectNodeNameDescription = NSLOCTEXT("MovieGraphNodes", "NodeDescription_Select", "Select\n{0}");
 
-	if (bGetDescriptive && !Description.IsEmpty())
+	if (bGetDescriptive)
 	{
-		return FText::Format(SelectNodeDescription, FText::FromString(Description));
+		if (const UEnum* Enum = Cast<UEnum>(SelectOptions->GetValueTypeObject()))
+		{
+			return FText::Format(SelectNodeNameDescription, FText::FromString(Enum->GetName()));
+		}
 	}
 
 	return SelectNodeName;
@@ -123,10 +399,5 @@ void UMovieGraphSelectNode::PostEditChangeProperty(FPropertyChangedEvent& Proper
 	{
 		UpdatePins();
 	}
-	
-	if (PropertyChangedEvent.GetPropertyName() == GET_MEMBER_NAME_CHECKED(UMovieGraphSelectNode, Description))
-	{
-		OnNodeChangedDelegate.Broadcast(this);
-	}
 }
 #endif // WITH_EDITOR
\ No newline at end of file
diff --git a/Engine/Plugins/MovieScene/MovieRenderPipeline/Source/MovieRenderPipelineCore/Private/Graph/Nodes/MovieGraphVariableNode.cpp b/Engine/Plugins/MovieScene/MovieRenderPipeline/Source/MovieRenderPipelineCore/Private/Graph/Nodes/MovieGraphVariableNode.cpp
--- a/Engine/Plugins/MovieScene/MovieRenderPipeline/Source/MovieRenderPipelineCore/Private/Graph/Nodes/MovieGraphVariableNode.cpp
+++ b/Engine/Plugins/MovieScene/MovieRenderPipeline/Source/MovieRenderPipelineCore/Private/Graph/Nodes/MovieGraphVariableNode.cpp
@@ -34,27 +34,35 @@ FString UMovieGraphVariableNode::GetResolvedValueForOutputPin(const FName& InPin
 {
 	if (GraphVariable && (GraphVariable->GetMemberName() == InPinName))
 	{
-		// If there's a valid job in the traversal context, and the job has an enabled variable assignment, use that
-		// instead of the variable value set in the graph
-		if (InContext && InContext->Job)
+		if (ContextHasEnabledAssignmentForVariable(InContext))
 		{
-			bool bIsEnabled = false;
-			if (InContext->Job->VariableAssignments->GetVariableAssignmentEnableState(GraphVariable, bIsEnabled))
-			{
-				if (bIsEnabled)
-				{
-					return InContext->Job->VariableAssignments->GetValueSerializedString(FName(GraphVariable->GetMemberName()));
-				}
-			}
+			return InContext->Job->VariableAssignments->GetValueSerializedString(FName(GraphVariable->GetMemberName()));
 		}
 
-		// No valid job context: just get the value from the variable
+		// No valid variable assignment: just get the value from the variable
 		return GraphVariable->GetValueSerializedString();
 	}
 	
 	return FString();
 }
 
+bool UMovieGraphVariableNode::GetResolvedValueForOutputPin(const FName& InPinName, const FMovieGraphTraversalContext* InContext, TObjectPtr<UMovieGraphValueContainer>& OutValueContainer) const
+{
+	if (GraphVariable && (GraphVariable->GetMemberName() == InPinName))
+	{
+		if (ContextHasEnabledAssignmentForVariable(InContext))
+		{
+			return InContext->Job->VariableAssignments->GetValueContainer(InPinName, OutValueContainer);
+		}
+
+		// No valid variable assignment: just get the value from the variable
+		OutValueContainer = GraphVariable;
+		return true;
+	}
+	
+	return false;
+}
+
 void UMovieGraphVariableNode::SetVariable(UMovieGraphVariable* InVariable)
 {
 	if (InVariable)
@@ -117,4 +125,21 @@ void UMovieGraphVariableNode::UpdateOutputPin(UMovieGraphMember* ChangedVariable
 	}
 
 	OnNodeChangedDelegate.Broadcast(this);
-}
\ No newline at end of file
+}
+
+bool UMovieGraphVariableNode::ContextHasEnabledAssignmentForVariable(const FMovieGraphTraversalContext* InContext) const
+{
+	if (InContext && InContext->Job)
+	{
+		bool bIsEnabled = false;
+		if (InContext->Job->VariableAssignments->GetVariableAssignmentEnableState(GraphVariable, bIsEnabled))
+		{
+			if (bIsEnabled)
+			{
+				return true;
+			}
+		}
+	}
+
+	return false;
+}
diff --git a/Engine/Plugins/MovieScene/MovieRenderPipeline/Source/MovieRenderPipelineCore/Private/MovieJobVariableAssignmentContainer.cpp b/Engine/Plugins/MovieScene/MovieRenderPipeline/Source/MovieRenderPipelineCore/Private/MovieJobVariableAssignmentContainer.cpp
--- a/Engine/Plugins/MovieScene/MovieRenderPipeline/Source/MovieRenderPipelineCore/Private/MovieJobVariableAssignmentContainer.cpp
+++ b/Engine/Plugins/MovieScene/MovieRenderPipeline/Source/MovieRenderPipelineCore/Private/MovieJobVariableAssignmentContainer.cpp
@@ -104,7 +104,7 @@ bool UMovieJobVariableAssignmentContainer::GetValueObject(const FName& PropertyN
 	return UE::MovieGraph::Private::GetOptionalValue<UObject*>(Result, OutValue);
 }
 
-bool UMovieJobVariableAssignmentContainer::GetValueClass(const FName& PropertyName, UClass* OutValue) const
+bool UMovieJobVariableAssignmentContainer::GetValueClass(const FName& PropertyName, UClass*& OutValue) const
 {
 	TValueOrError<UClass*, EPropertyBagResult> Result = Value.GetValueClass(PropertyName);
 	return UE::MovieGraph::Private::GetOptionalValue<UClass*>(Result, OutValue);
@@ -218,6 +218,19 @@ EMovieGraphContainerType UMovieJobVariableAssignmentContainer::GetValueContainer
 	return EMovieGraphContainerType::None;
 }
 
+bool UMovieJobVariableAssignmentContainer::GetValueContainer(const FName& PropertyName, TObjectPtr<UMovieGraphValueContainer>& OutValueContainer)
+{
+	if (const FPropertyBagPropertyDesc* Desc = Value.FindPropertyDescByName(PropertyName))
+	{
+		// TODO: This object allocation is unfortunate -- it would be nice to find a way to avoid it
+		OutValueContainer = NewObject<UMovieGraphValueContainer>();
+		OutValueContainer->SetFromDesc(Desc, GetValueSerializedString(PropertyName));
+		return true;
+	}
+
+	return false;
+}
+
 bool UMovieJobVariableAssignmentContainer::FindOrGenerateVariableOverride(
 	const UMovieGraphVariable* InGraphVariable, FPropertyBagPropertyDesc* OutPropDesc, FPropertyBagPropertyDesc* OutEditConditionPropDesc, bool bGenerateIfNotExists)
 {
diff --git a/Engine/Plugins/MovieScene/MovieRenderPipeline/Source/MovieRenderPipelineCore/Public/Graph/MovieGraphNode.h b/Engine/Plugins/MovieScene/MovieRenderPipeline/Source/MovieRenderPipelineCore/Public/Graph/MovieGraphNode.h
--- a/Engine/Plugins/MovieScene/MovieRenderPipeline/Source/MovieRenderPipelineCore/Public/Graph/MovieGraphNode.h
+++ b/Engine/Plugins/MovieScene/MovieRenderPipeline/Source/MovieRenderPipelineCore/Public/Graph/MovieGraphNode.h
@@ -123,6 +123,15 @@ public:
 		return FString();
 	}
 
+	/**
+	 * The same functionality as the other method that returns a string. However, this version returns true if a
+	 * resolved value could be determined, and provides the value via a UMovieGraphValueContainer.
+	 */
+	virtual bool GetResolvedValueForOutputPin(const FName& InPinName, const FMovieGraphTraversalContext* InContext, TObjectPtr<UMovieGraphValueContainer>& OutValueContainer) const
+	{
+		return false;
+	}
+
 	/**
 	 * Gets the descriptions of properties which can be dynamically added to the node. These types of properties
 	 * do not correspond to a UPROPERTY defined on the node itself.
@@ -203,7 +212,15 @@ public:
 	void UpdatePins();
 	void UpdateDynamicProperties();
 	class UMovieGraphConfig* GetGraph() const;
-	UMovieGraphPin* GetInputPin(const FName& InPinLabel) const;
+
+	/**
+	 * Gets the input pin with the specified name, or nullptr if one could not be found. Most pins on a node are
+	 * "built-in", meaning they ship with the node. Dynamic pins (pins which are not built-in) can potentially have the
+	 * same name as a built-in (eg, the option pins on the Select node). To disambiguate between built-in and dynamic
+	 * pins, specify bIsBuiltInPin = false if trying to fetch a pin that is not built-in. */
+	UMovieGraphPin* GetInputPin(const FName& InPinLabel, const bool bIsBuiltInPin = true) const;
+
+	/** Gets the output pin with the specified name, or nullptr if one could not be found. */
 	UMovieGraphPin* GetOutputPin(const FName& InPinLabel) const;
 
 	/** Gets the first input pin on the node which has a connection, or nullptr if no pins are connected. */
diff --git a/Engine/Plugins/MovieScene/MovieRenderPipeline/Source/MovieRenderPipelineCore/Public/Graph/MovieGraphPin.h b/Engine/Plugins/MovieScene/MovieRenderPipeline/Source/MovieRenderPipelineCore/Public/Graph/MovieGraphPin.h
--- a/Engine/Plugins/MovieScene/MovieRenderPipeline/Source/MovieRenderPipelineCore/Public/Graph/MovieGraphPin.h
+++ b/Engine/Plugins/MovieScene/MovieRenderPipeline/Source/MovieRenderPipelineCore/Public/Graph/MovieGraphPin.h
@@ -44,12 +44,20 @@ struct FMovieGraphPinProperties
 	UPROPERTY(BlueprintReadWrite, EditAnywhere, Category = "Settings")
 	bool bIsBranch = false;
 
+	/**
+	 * Whether this pin is built-in (ie, the pin ships with the node and cannot be removed). Option pins on the Select
+	 * node would be an example of pins which are not built-in (they can be added and removed dynamically).
+	 */
+	UPROPERTY(BlueprintReadWrite, EditAnywhere, Category = "Settings")
+	bool bIsBuiltIn = true;
+
 	bool operator==(const FMovieGraphPinProperties& Other) const
 	{
 		return Label == Other.Label
 			&& Type == Other.Type
 			&& bAllowMultipleConnections == Other.bAllowMultipleConnections
-			&& bIsBranch == Other.bIsBranch;
+			&& bIsBranch == Other.bIsBranch
+			&& bIsBuiltIn == Other.bIsBuiltIn;
 	}
 
 	bool operator !=(const FMovieGraphPinProperties& Other) const
diff --git a/Engine/Plugins/MovieScene/MovieRenderPipeline/Source/MovieRenderPipelineCore/Public/Graph/MovieGraphValueContainer.cpp b/Engine/Plugins/MovieScene/MovieRenderPipeline/Source/MovieRenderPipelineCore/Public/Graph/MovieGraphValueContainer.cpp
--- a/Engine/Plugins/MovieScene/MovieRenderPipeline/Source/MovieRenderPipelineCore/Public/Graph/MovieGraphValueContainer.cpp
+++ b/Engine/Plugins/MovieScene/MovieRenderPipeline/Source/MovieRenderPipelineCore/Public/Graph/MovieGraphValueContainer.cpp
@@ -5,101 +5,128 @@
 #include "MovieGraphCommon.h"
 
 // The property bag has one "default" property in it
-const FName UMovieGraphValueContainer::PropertyBagDefaultPropertyName("Default");
+const FName UMovieGraphValueContainer::PropertyBagDefaultPropertyName("Value");
 
 UMovieGraphValueContainer::UMovieGraphValueContainer()
 {
+	PropertyName = PropertyBagDefaultPropertyName;
+	
 	if (!HasAnyFlags(RF_ClassDefaultObject))
 	{
-		// Add a default float property if one does not exist already
+		// Add a default double property if one does not exist already
 		if (Value.GetNumPropertiesInBag() == 0)
 		{
-			Value.AddProperty(PropertyBagDefaultPropertyName, EPropertyBagPropertyType::Float);
+			Value.AddProperty(PropertyName, EPropertyBagPropertyType::Double);
 		}
 	}
 }
 
+void UMovieGraphValueContainer::SetPropertyName(const FName& InName)
+{
+	if (const FPropertyBagPropertyDesc* Desc = Value.FindPropertyDescByName(PropertyName))
+	{
+		if (Desc->Name == InName)
+		{
+			return;
+		}
+
+		PropertyName = InName;
+
+		// Changing the property name requires a new desc and a migration
+		FPropertyBagPropertyDesc NewDesc(*Desc);
+		NewDesc.Name = InName;
+
+		const UPropertyBag* NewPropBag = UPropertyBag::GetOrCreateFromDescs({NewDesc});
+		Value.MigrateToNewBagStruct(NewPropBag);
+	}
+}
+
+FName UMovieGraphValueContainer::GetPropertyName() const
+{
+	return PropertyName;
+}
+
 bool UMovieGraphValueContainer::GetValueBool(bool& bOutValue) const
 {
-	TValueOrError<bool, EPropertyBagResult> Result = Value.GetValueBool(PropertyBagDefaultPropertyName);
+	TValueOrError<bool, EPropertyBagResult> Result = Value.GetValueBool(PropertyName);
 	return UE::MovieGraph::Private::GetOptionalValue<bool>(Result, bOutValue);
 }
 
 bool UMovieGraphValueContainer::GetValueByte(uint8& OutValue) const
 {
-	TValueOrError<uint8, EPropertyBagResult> Result = Value.GetValueByte(PropertyBagDefaultPropertyName);
+	TValueOrError<uint8, EPropertyBagResult> Result = Value.GetValueByte(PropertyName);
 	return UE::MovieGraph::Private::GetOptionalValue<uint8>(Result, OutValue);
 }
 
 bool UMovieGraphValueContainer::GetValueInt32(int32& OutValue) const
 {
-	TValueOrError<int32, EPropertyBagResult> Result = Value.GetValueInt32(PropertyBagDefaultPropertyName);
+	TValueOrError<int32, EPropertyBagResult> Result = Value.GetValueInt32(PropertyName);
 	return UE::MovieGraph::Private::GetOptionalValue<int32>(Result, OutValue);
 }
 
 bool UMovieGraphValueContainer::GetValueInt64(int64& OutValue) const
 {
-	TValueOrError<int64, EPropertyBagResult> Result = Value.GetValueInt64(PropertyBagDefaultPropertyName);
+	TValueOrError<int64, EPropertyBagResult> Result = Value.GetValueInt64(PropertyName);
 	return UE::MovieGraph::Private::GetOptionalValue<int64>(Result, OutValue);
 }
 
 bool UMovieGraphValueContainer::GetValueFloat(float& OutValue) const
 {
-	TValueOrError<float, EPropertyBagResult> Result = Value.GetValueFloat(PropertyBagDefaultPropertyName);
+	TValueOrError<float, EPropertyBagResult> Result = Value.GetValueFloat(PropertyName);
 	return UE::MovieGraph::Private::GetOptionalValue<float>(Result, OutValue);
 }
 
 bool UMovieGraphValueContainer::GetValueDouble(double& OutValue) const
 {
-	TValueOrError<double, EPropertyBagResult> Result = Value.GetValueDouble(PropertyBagDefaultPropertyName);
+	TValueOrError<double, EPropertyBagResult> Result = Value.GetValueDouble(PropertyName);
 	return UE::MovieGraph::Private::GetOptionalValue<double>(Result, OutValue);
 }
 
 bool UMovieGraphValueContainer::GetValueName(FName& OutValue) const
 {
-	TValueOrError<FName, EPropertyBagResult> Result = Value.GetValueName(PropertyBagDefaultPropertyName);
+	TValueOrError<FName, EPropertyBagResult> Result = Value.GetValueName(PropertyName);
 	return UE::MovieGraph::Private::GetOptionalValue<FName>(Result, OutValue);
 }
 
 bool UMovieGraphValueContainer::GetValueString(FString& OutValue) const
 {
-	TValueOrError<FString, EPropertyBagResult> Result = Value.GetValueString(PropertyBagDefaultPropertyName);
+	TValueOrError<FString, EPropertyBagResult> Result = Value.GetValueString(PropertyName);
 	return UE::MovieGraph::Private::GetOptionalValue<FString>(Result, OutValue);
 }
 
 bool UMovieGraphValueContainer::GetValueText(FText& OutValue) const
 {
-	TValueOrError<FText, EPropertyBagResult> Result = Value.GetValueText(PropertyBagDefaultPropertyName);
+	TValueOrError<FText, EPropertyBagResult> Result = Value.GetValueText(PropertyName);
 	return UE::MovieGraph::Private::GetOptionalValue<FText>(Result, OutValue);
 }
 
 bool UMovieGraphValueContainer::GetValueEnum(uint8& OutValue, const UEnum* RequestedEnum) const
 {
-	TValueOrError<uint8, EPropertyBagResult> Result = Value.GetValueEnum(PropertyBagDefaultPropertyName, RequestedEnum);
+	TValueOrError<uint8, EPropertyBagResult> Result = Value.GetValueEnum(PropertyName, RequestedEnum);
 	return UE::MovieGraph::Private::GetOptionalValue<uint8>(Result, OutValue);
 }
 
 bool UMovieGraphValueContainer::GetValueStruct(FStructView& OutValue, const UScriptStruct* RequestedStruct) const
 {
-	TValueOrError<FStructView, EPropertyBagResult> Result = Value.GetValueStruct(PropertyBagDefaultPropertyName, RequestedStruct);
+	TValueOrError<FStructView, EPropertyBagResult> Result = Value.GetValueStruct(PropertyName, RequestedStruct);
 	return UE::MovieGraph::Private::GetOptionalValue<FStructView>(Result, OutValue);
 }
 
 bool UMovieGraphValueContainer::GetValueObject(UObject* OutValue, const UClass* RequestedClass) const
 {
-	TValueOrError<UObject*, EPropertyBagResult> Result = Value.GetValueObject(PropertyBagDefaultPropertyName, RequestedClass);
+	TValueOrError<UObject*, EPropertyBagResult> Result = Value.GetValueObject(PropertyName, RequestedClass);
 	return UE::MovieGraph::Private::GetOptionalValue<UObject*>(Result, OutValue);
 }
 
-bool UMovieGraphValueContainer::GetValueClass(UClass* OutValue) const
+bool UMovieGraphValueContainer::GetValueClass(UClass*& OutValue) const
 {
-	TValueOrError<UClass*, EPropertyBagResult> Result = Value.GetValueClass(PropertyBagDefaultPropertyName);
+	TValueOrError<UClass*, EPropertyBagResult> Result = Value.GetValueClass(PropertyName);
 	return UE::MovieGraph::Private::GetOptionalValue<UClass*>(Result, OutValue);
 }
 
 FString UMovieGraphValueContainer::GetValueSerializedString()
 {
-	TValueOrError<FString, EPropertyBagResult> Result = Value.GetValueSerializedString(PropertyBagDefaultPropertyName);
+	TValueOrError<FString, EPropertyBagResult> Result = Value.GetValueSerializedString(PropertyName);
 	FString ResultString;
 	UE::MovieGraph::Private::GetOptionalValue<FString>(Result, ResultString);
 	return ResultString;
@@ -107,77 +134,77 @@ FString UMovieGraphValueContainer::GetValueSerializedString()
 
 bool UMovieGraphValueContainer::SetValueBool(const bool bInValue)
 {
-	return Value.SetValueBool(PropertyBagDefaultPropertyName, bInValue) == EPropertyBagResult::Success;
+	return Value.SetValueBool(PropertyName, bInValue) == EPropertyBagResult::Success;
 }
 
 bool UMovieGraphValueContainer::SetValueByte(const uint8 InValue)
 {
-	return Value.SetValueByte(PropertyBagDefaultPropertyName, InValue) == EPropertyBagResult::Success;
+	return Value.SetValueByte(PropertyName, InValue) == EPropertyBagResult::Success;
 }
 
 bool UMovieGraphValueContainer::SetValueInt32(const int32 InValue)
 {
-	return Value.SetValueInt32(PropertyBagDefaultPropertyName, InValue) == EPropertyBagResult::Success;
+	return Value.SetValueInt32(PropertyName, InValue) == EPropertyBagResult::Success;
 }
 
 bool UMovieGraphValueContainer::SetValueInt64(const int64 InValue)
 {
-	return Value.SetValueInt64(PropertyBagDefaultPropertyName, InValue) == EPropertyBagResult::Success;
+	return Value.SetValueInt64(PropertyName, InValue) == EPropertyBagResult::Success;
 }
 
 bool UMovieGraphValueContainer::SetValueFloat(const float InValue)
 {
-	return Value.SetValueFloat(PropertyBagDefaultPropertyName, InValue) == EPropertyBagResult::Success;
+	return Value.SetValueFloat(PropertyName, InValue) == EPropertyBagResult::Success;
 }
 
 bool UMovieGraphValueContainer::SetValueDouble(const double InValue)
 {
-	return Value.SetValueDouble(PropertyBagDefaultPropertyName, InValue) == EPropertyBagResult::Success;
+	return Value.SetValueDouble(PropertyName, InValue) == EPropertyBagResult::Success;
 }
 
 bool UMovieGraphValueContainer::SetValueName(const FName InValue)
 {
-	return Value.SetValueName(PropertyBagDefaultPropertyName, InValue) == EPropertyBagResult::Success;
+	return Value.SetValueName(PropertyName, InValue) == EPropertyBagResult::Success;
 }
 
 bool UMovieGraphValueContainer::SetValueString(const FString& InValue)
 {
-	return Value.SetValueString(PropertyBagDefaultPropertyName, InValue) == EPropertyBagResult::Success;
+	return Value.SetValueString(PropertyName, InValue) == EPropertyBagResult::Success;
 }
 
 bool UMovieGraphValueContainer::SetValueText(const FText& InValue)
 {
-	return Value.SetValueText(PropertyBagDefaultPropertyName, InValue) == EPropertyBagResult::Success;
+	return Value.SetValueText(PropertyName, InValue) == EPropertyBagResult::Success;
 }
 
 bool UMovieGraphValueContainer::SetValueEnum(const uint8 InValue, const UEnum* Enum)
 {
-	return Value.SetValueEnum(PropertyBagDefaultPropertyName, InValue, Enum) == EPropertyBagResult::Success;
+	return Value.SetValueEnum(PropertyName, InValue, Enum) == EPropertyBagResult::Success;
 }
 
 bool UMovieGraphValueContainer::SetValueStruct(FConstStructView InValue)
 {
-	return Value.SetValueStruct(PropertyBagDefaultPropertyName, InValue) == EPropertyBagResult::Success;
+	return Value.SetValueStruct(PropertyName, InValue) == EPropertyBagResult::Success;
 }
 
 bool UMovieGraphValueContainer::SetValueObject(UObject* InValue)
 {
-	return Value.SetValueObject(PropertyBagDefaultPropertyName, InValue) == EPropertyBagResult::Success;
+	return Value.SetValueObject(PropertyName, InValue) == EPropertyBagResult::Success;
 }
 
 bool UMovieGraphValueContainer::SetValueClass(UClass* InValue)
 {
-	return Value.SetValueClass(PropertyBagDefaultPropertyName, InValue) == EPropertyBagResult::Success;
+	return Value.SetValueClass(PropertyName, InValue) == EPropertyBagResult::Success;
 }
 
 bool UMovieGraphValueContainer::SetValueSerializedString(const FString& NewValue)
 {
-	return Value.SetValueSerializedString(PropertyBagDefaultPropertyName, NewValue) == EPropertyBagResult::Success;
+	return Value.SetValueSerializedString(PropertyName, NewValue) == EPropertyBagResult::Success;
 }
 
 EMovieGraphValueType UMovieGraphValueContainer::GetValueType() const
 {
-	if (const FPropertyBagPropertyDesc* Desc = Value.FindPropertyDescByName(PropertyBagDefaultPropertyName))
+	if (const FPropertyBagPropertyDesc* Desc = Value.FindPropertyDescByName(PropertyName))
 	{
 		return static_cast<EMovieGraphValueType>(Desc->ValueType);
 	}
@@ -185,21 +212,31 @@ EMovieGraphValueType UMovieGraphValueContainer::GetValueType() const
 	return EMovieGraphValueType::None;
 }
 
-void UMovieGraphValueContainer::SetValueType(EMovieGraphValueType ValueType)
+void UMovieGraphValueContainer::SetValueType(EMovieGraphValueType ValueType, UObject* InValueTypeObject)
 {
-	if (const FPropertyBagPropertyDesc* Desc = Value.FindPropertyDescByName(PropertyBagDefaultPropertyName))
+	if (const FPropertyBagPropertyDesc* Desc = Value.FindPropertyDescByName(PropertyName))
 	{
-		FPropertyBagPropertyDesc NewDesc(*Desc);
+		FPropertyBagPropertyDesc NewDesc;
+		NewDesc.Name = Desc->Name;
 		NewDesc.ValueType = static_cast<EPropertyBagPropertyType>(ValueType);
-
-		const UPropertyBag* NewPropBag = UPropertyBag::GetOrCreateFromDescs({NewDesc});
-		Value.MigrateToNewBagStruct(NewPropBag);
+		NewDesc.ContainerTypes = Desc->ContainerTypes;
+		NewDesc.ValueTypeObject = InValueTypeObject;
+
+		Value.Reset();
+		Value.AddProperties({NewDesc});
+
+#if WITH_EDITOR
+		// Send a property change event manually since the property bag doesn't seem to generate one in this scenario
+		FProperty* ValueProperty = FindFProperty<FProperty>(StaticClass(), GET_MEMBER_NAME_CHECKED(UMovieGraphValueContainer, Value));
+		FPropertyChangedEvent PropertyEvent(ValueProperty, EPropertyChangeType::ValueSet);
+		PostEditChangeProperty(PropertyEvent);
+#endif // WITH_EDITOR
 	}
 }
 
 const UObject* UMovieGraphValueContainer::GetValueTypeObject() const
 {
-	if (const FPropertyBagPropertyDesc* Desc = Value.FindPropertyDescByName(PropertyBagDefaultPropertyName))
+	if (const FPropertyBagPropertyDesc* Desc = Value.FindPropertyDescByName(PropertyName))
 	{
 		return Desc->ValueTypeObject;
 	}
@@ -209,7 +246,7 @@ const UObject* UMovieGraphValueContainer::GetValueTypeObject() const
 
 void UMovieGraphValueContainer::SetValueTypeObject(const UObject* ValueTypeObject)
 {
-	if (const FPropertyBagPropertyDesc* Desc = Value.FindPropertyDescByName(PropertyBagDefaultPropertyName))
+	if (const FPropertyBagPropertyDesc* Desc = Value.FindPropertyDescByName(PropertyName))
 	{
 		FPropertyBagPropertyDesc NewDesc(*Desc);
 		NewDesc.ValueTypeObject = ValueTypeObject;
@@ -221,7 +258,7 @@ void UMovieGraphValueContainer::SetValueTypeObject(const UObject* ValueTypeObjec
 
 EMovieGraphContainerType UMovieGraphValueContainer::GetValueContainerType() const
 {
-	if (const FPropertyBagPropertyDesc* Desc = Value.FindPropertyDescByName(PropertyBagDefaultPropertyName))
+	if (const FPropertyBagPropertyDesc* Desc = Value.FindPropertyDescByName(PropertyName))
 	{
 		return static_cast<EMovieGraphContainerType>(Desc->ContainerTypes.GetFirstContainerType());
 	}
@@ -231,7 +268,7 @@ EMovieGraphContainerType UMovieGraphValueContainer::GetValueContainerType() cons
 
 void UMovieGraphValueContainer::SetValueContainerType(EMovieGraphContainerType ContainerType)
 {
-	if (const FPropertyBagPropertyDesc* Desc = Value.FindPropertyDescByName(PropertyBagDefaultPropertyName))
+	if (const FPropertyBagPropertyDesc* Desc = Value.FindPropertyDescByName(PropertyName))
 	{
 		FPropertyBagPropertyDesc NewDesc(*Desc);
 		NewDesc.ContainerTypes = { static_cast<EPropertyBagContainerType>(ContainerType) };
@@ -239,4 +276,28 @@ void UMovieGraphValueContainer::SetValueContainerType(EMovieGraphContainerType C
 		const UPropertyBag* NewPropBag = UPropertyBag::GetOrCreateFromDescs({NewDesc});
 		Value.MigrateToNewBagStruct(NewPropBag);
 	}
-}
\ No newline at end of file
+}
+
+TValueOrError<FPropertyBagArrayRef, EPropertyBagResult> UMovieGraphValueContainer::GetArrayRef()
+{
+	return Value.GetMutableArrayRef(PropertyName);
+}
+
+void UMovieGraphValueContainer::SetFromDesc(const FPropertyBagPropertyDesc* InDesc, const FString& InString)
+{
+	if (InDesc)
+	{
+		FPropertyBagPropertyDesc NewDesc;
+		NewDesc.Name = InDesc->Name;
+		NewDesc.ValueType = InDesc->ValueType;
+		NewDesc.ContainerTypes = InDesc->ContainerTypes;
+		NewDesc.ValueTypeObject = InDesc->ValueTypeObject;
+
+		PropertyName = InDesc->Name;
+		
+		Value.Reset();
+		Value.AddProperties({NewDesc});
+
+		SetValueSerializedString(InString);
+	}
+}
diff --git a/Engine/Plugins/MovieScene/MovieRenderPipeline/Source/MovieRenderPipelineCore/Public/Graph/MovieGraphValueContainer.h b/Engine/Plugins/MovieScene/MovieRenderPipeline/Source/MovieRenderPipelineCore/Public/Graph/MovieGraphValueContainer.h
--- a/Engine/Plugins/MovieScene/MovieRenderPipeline/Source/MovieRenderPipelineCore/Public/Graph/MovieGraphValueContainer.h
+++ b/Engine/Plugins/MovieScene/MovieRenderPipeline/Source/MovieRenderPipelineCore/Public/Graph/MovieGraphValueContainer.h
@@ -5,6 +5,7 @@
 #include "CoreMinimal.h"
 #include "MovieGraphCommon.h"
 #include "PropertyBag.h"
+#include "Templates/ValueOrError.h"
 
 #include "MovieGraphValueContainer.generated.h"
 
@@ -29,7 +30,7 @@ namespace UE::MovieGraph::Private
  * Holds a generic value, with an API for getting/setting the value, as well as getting/setting its type
  * and container (eg, array).
  */
-UCLASS(Abstract)
+UCLASS()
 class MOVIERENDERPIPELINECORE_API UMovieGraphValueContainer : public UObject
 {
 	GENERATED_BODY()
@@ -86,7 +87,7 @@ public:
 
 	/** Gets the UClass value of the held data. Returns true on success, else false. */
 	UFUNCTION(BlueprintCallable, Category="Config")
-	bool GetValueClass(UClass* OutValue) const;
+	bool GetValueClass(UClass*& OutValue) const;
 
 	/** Gets the serialized string value of the held data. */
 	UFUNCTION(BlueprintCallable, Category="Config")
@@ -196,9 +197,9 @@ public:
 	UFUNCTION(BlueprintCallable, Category="Config")
 	EMovieGraphValueType GetValueType() const;
 
-	/** Sets the type of the stored data. */
+	/** Sets the type of the stored data. Enums, structs, and classes must specify a value type object. */
 	UFUNCTION(BlueprintCallable, Category="Config")
-	void SetValueType(EMovieGraphValueType ValueType);
+	void SetValueType(EMovieGraphValueType ValueType, UObject* InValueTypeObject = nullptr);
 
 	/** Gets the object that defines the enum, struct, or class. */
 	UFUNCTION(BlueprintCallable, Category="Config")
@@ -216,9 +217,37 @@ public:
 	UFUNCTION(BlueprintCallable, Category="Config")
 	void SetValueContainerType(EMovieGraphContainerType ContainerType);
 
+	/**
+	 * Gets a reference to the array backing the value, if any. GetValueContainerType() will return
+	 * EMovieGraphContainerType::Array if the value is holding an array.
+	 */
+	TValueOrError<FPropertyBagArrayRef, EPropertyBagResult> GetArrayRef();
+
+	/**
+	 * Sets the name of the property that the value container holds. This is mostly unneeded unless the display name
+	 * of the property is important (eg, in the details panel).
+	 */
+	void SetPropertyName(const FName& InName);
+
+	/** Gets the name of the property that the value container holds. */
+	FName GetPropertyName() const;
+
 private:
-	/** The name of the single property stored in the property bag. The name in the property bag does not change. */
+	/**
+	 * Sets the configuration of this container from a property desc. This is less safe than the strongly-typed methods,
+	 * and is reserved only for UMovieJobVariableAssignmentContainer to call. The serialized representation of the value
+	 * needs to be supplied via InString.
+	 */
+	friend class UMovieJobVariableAssignmentContainer;
+	void SetFromDesc(const FPropertyBagPropertyDesc* InDesc, const FString& InString);
+
+private:
+	/** The default name of the single property stored in the property bag. */
 	static const FName PropertyBagDefaultPropertyName;
+
+	/** The name of the single property stored in the property bag. */
+	UPROPERTY()
+	FName PropertyName;
 	
 	// Note: The property bag only stores one property, since the object only needs to store one value. This is an odd
 	// use of a property bag, but the property bag solves a number of issues that would be very difficult to solve
@@ -226,6 +255,6 @@ private:
 	// of a property which can have its type changed, 3) the ability to set the value of the property from both Python
 	// and C++, and 4) the ability to change the property at runtime.
 	/** The value held by this object. */
-	UPROPERTY(EditAnywhere, meta=(ShowOnlyInnerProperties), Category = "Value")
+	UPROPERTY(EditAnywhere, meta=(ShowOnlyInnerProperties, FixedLayout), Category = "Value")
 	FInstancedPropertyBag Value;
 };
\ No newline at end of file
diff --git a/Engine/Plugins/MovieScene/MovieRenderPipeline/Source/MovieRenderPipelineCore/Public/Graph/Nodes/MovieGraphSelectNode.h b/Engine/Plugins/MovieScene/MovieRenderPipeline/Source/MovieRenderPipelineCore/Public/Graph/Nodes/MovieGraphSelectNode.h
--- a/Engine/Plugins/MovieScene/MovieRenderPipeline/Source/MovieRenderPipelineCore/Public/Graph/Nodes/MovieGraphSelectNode.h
+++ b/Engine/Plugins/MovieScene/MovieRenderPipeline/Source/MovieRenderPipelineCore/Public/Graph/Nodes/MovieGraphSelectNode.h
@@ -6,7 +6,6 @@
 
 #include "MovieGraphSelectNode.generated.h"
 
-// TODO: Currently this node is restricted to just accepting string-based options.
 /** A node which creates a condition that selects from a set of input branches. */
 UCLASS()
 class MOVIERENDERPIPELINECORE_API UMovieGraphSelectNode : public UMovieGraphNode
@@ -14,12 +13,24 @@ class MOVIERENDERPIPELINECORE_API UMovieGraphSelectNode : public UMovieGraphNode
 	GENERATED_BODY()
 
 public:
-	UMovieGraphSelectNode() = default;
+	UMovieGraphSelectNode();
 
 	virtual TArray<FMovieGraphPinProperties> GetInputPinProperties() const override;
 	virtual TArray<FMovieGraphPinProperties> GetOutputPinProperties() const override;
 	virtual TArray<UMovieGraphPin*> EvaluatePinsToFollow(FMovieGraphEvaluationContext& InContext) const override;
 
+	/**
+	 * Sets the data type for the select options and selected option (note that their data will be reset when this
+	 * is called). For structs, enums, and objects, the value type object needs to be provided.
+	 */
+	void SetDataType(const EMovieGraphValueType ValueType, UObject* InValueTypeObject = nullptr);
+
+	/** Gets the data type that the select node is currently using. */
+	EMovieGraphValueType GetValueType() const;
+
+	/** Gets the value type object associated with the value type currently set (if any). */
+	const UObject* GetValueTypeObject() const;
+
 #if WITH_EDITOR
 	virtual FText GetNodeTitle(const bool bGetDescriptive = false) const override;
 	virtual FText GetMenuCategory() const override;
@@ -29,17 +40,14 @@ public:
 	//~ Begin UObject interface
 	virtual void PostEditChangeProperty(FPropertyChangedEvent& PropertyChangedEvent) override;
 	//~ End UObject interface
-	#endif
-
-	/** The values of options which can be selected. */
-	UPROPERTY(EditAnywhere, Category = "General")
-	TArray<FString> SelectOptions;
+#endif
 
-	/** The value of the option which has been selected. */
-	UPROPERTY(EditAnywhere, Category = "General")
-	FString SelectedOption;
+private:
+	/** The options that are available on the node. */
+	UPROPERTY(EditAnywhere, Instanced, Category = "Options", meta=(ShowInnerProperties, FullyExpand="true"))
+	TObjectPtr<UMovieGraphValueContainer> SelectOptions;
 
-	/** A description of what this select is doing. */
-	UPROPERTY(EditAnywhere, Category = "General")
-	FString Description;
+	/** The currently selected option. */
+	UPROPERTY(EditAnywhere, Instanced, Category = "Options", meta=(ShowInnerProperties, FullyExpand="true"))
+	TObjectPtr<UMovieGraphValueContainer> SelectedOption;
 };
\ No newline at end of file
diff --git a/Engine/Plugins/MovieScene/MovieRenderPipeline/Source/MovieRenderPipelineCore/Public/Graph/Nodes/MovieGraphVariableNode.h b/Engine/Plugins/MovieScene/MovieRenderPipeline/Source/MovieRenderPipelineCore/Public/Graph/Nodes/MovieGraphVariableNode.h
--- a/Engine/Plugins/MovieScene/MovieRenderPipeline/Source/MovieRenderPipelineCore/Public/Graph/Nodes/MovieGraphVariableNode.h
+++ b/Engine/Plugins/MovieScene/MovieRenderPipeline/Source/MovieRenderPipelineCore/Public/Graph/Nodes/MovieGraphVariableNode.h
@@ -18,6 +18,7 @@ public:
 
 	virtual TArray<FMovieGraphPinProperties> GetOutputPinProperties() const override;
 	virtual FString GetResolvedValueForOutputPin(const FName& InPinName, const FMovieGraphTraversalContext* InContext) const override;
+	virtual bool GetResolvedValueForOutputPin(const FName& InPinName, const FMovieGraphTraversalContext* InContext, TObjectPtr<UMovieGraphValueContainer>& OutValueContainer) const override;
 
 	/** Gets the variable that this node represents. */
 	UMovieGraphVariable* GetVariable() const { return GraphVariable; }
@@ -41,6 +42,9 @@ private:
 	/** Updates the output pin on the node to match the provided variable. */
 	void UpdateOutputPin(UMovieGraphMember* ChangedVariable) const;
 
+	/** Determines if the job in the given context has a valid, enabled variable assignment for the variable this node represents. */
+	bool ContextHasEnabledAssignmentForVariable(const FMovieGraphTraversalContext* InContext) const;
+
 private:
 	/** The underlying graph variable this node represents. */
 	UPROPERTY()
diff --git a/Engine/Plugins/MovieScene/MovieRenderPipeline/Source/MovieRenderPipelineCore/Public/MovieJobVariableAssignmentContainer.h b/Engine/Plugins/MovieScene/MovieRenderPipeline/Source/MovieRenderPipelineCore/Public/MovieJobVariableAssignmentContainer.h
--- a/Engine/Plugins/MovieScene/MovieRenderPipeline/Source/MovieRenderPipelineCore/Public/MovieJobVariableAssignmentContainer.h
+++ b/Engine/Plugins/MovieScene/MovieRenderPipeline/Source/MovieRenderPipelineCore/Public/MovieJobVariableAssignmentContainer.h
@@ -84,7 +84,7 @@ public:
 
 	/** Gets the UClass value of the specified property. Returns true on success, else false. */
 	UFUNCTION(BlueprintCallable, Category="Experimental")
-	bool GetValueClass(const FName& PropertyName, UClass* OutValue) const;
+	bool GetValueClass(const FName& PropertyName, UClass*& OutValue) const;
 
 	/** Gets the serialized string value of the specified property. */
 	UFUNCTION(BlueprintCallable, Category="Experimental")
@@ -202,6 +202,12 @@ public:
 	UFUNCTION(BlueprintCallable, Category="Experimental")
 	EMovieGraphContainerType GetValueContainerType(const FName& PropertyName) const;
 
+	/**
+	 * Gets a value container object rather than a strongly-typed value. Useful if the type of the value is not known
+	 * ahead of time. If no property with the specified name exists, OutValueContainer will not be modified.
+	 */
+	bool GetValueContainer(const FName& PropertyName, TObjectPtr<UMovieGraphValueContainer>& OutValueContainer);
+
 	/**
 	 * Updates an existing variable assignment for the provided graph variable to a new enable state, or adds a new
 	 * assignment and updates its enable state. Returns true on success, else false.
diff --git a/Engine/Plugins/MovieScene/MovieRenderPipeline/Source/MovieRenderPipelineEditor/MovieRenderPipelineEditor.Build.cs b/Engine/Plugins/MovieScene/MovieRenderPipeline/Source/MovieRenderPipelineEditor/MovieRenderPipelineEditor.Build.cs
--- a/Engine/Plugins/MovieScene/MovieRenderPipeline/Source/MovieRenderPipelineEditor/MovieRenderPipelineEditor.Build.cs
+++ b/Engine/Plugins/MovieScene/MovieRenderPipeline/Source/MovieRenderPipelineEditor/MovieRenderPipelineEditor.Build.cs
@@ -70,6 +70,8 @@ public class MovieRenderPipelineEditor : ModuleRules
 				"GraphEditor",
 				"StructUtilsEditor",
 				"ApplicationCore",
+				"KismetWidgets",
+				"BlueprintGraph",
 			}
         );
 
diff --git a/Engine/Plugins/MovieScene/MovieRenderPipeline/Source/MovieRenderPipelineEditor/Private/Customizations/Graph/MovieGraphMemberCustomization.h b/Engine/Plugins/MovieScene/MovieRenderPipeline/Source/MovieRenderPipelineEditor/Private/Customizations/Graph/MovieGraphMemberCustomization.h
--- a/Engine/Plugins/MovieScene/MovieRenderPipeline/Source/MovieRenderPipelineEditor/Private/Customizations/Graph/MovieGraphMemberCustomization.h
+++ b/Engine/Plugins/MovieScene/MovieRenderPipeline/Source/MovieRenderPipelineEditor/Private/Customizations/Graph/MovieGraphMemberCustomization.h
@@ -6,8 +6,12 @@
 #include "DetailLayoutBuilder.h"
 #include "DetailWidgetRow.h"
 #include "IDetailCustomization.h"
+#include "IPropertyUtilities.h"
+#include "Graph/MovieEdGraphNode.h"
 #include "Graph/MovieGraphConfig.h"
+#include "PropertyBagDetails.h"
 #include "PropertyHandle.h"
+#include "SPinTypeSelector.h"
 #include "Widgets/Input/SEditableTextBox.h"
 #include "Widgets/Text/STextBlock.h"
 
@@ -50,6 +54,28 @@ protected:
 
 	virtual void CustomizeDetails(IDetailLayoutBuilder& DetailBuilder) override
 	{
+		TSharedRef<IPropertyUtilities> PropUtils = DetailBuilder.GetPropertyUtilities();
+		
+		auto GetFilteredVariableTypeTree = [this](TArray<TSharedPtr<UEdGraphSchema_K2::FPinTypeTreeInfo>>& TypeTree, ETypeTreeFilter TypeTreeFilter)
+		{
+			// All types from the schema are allowed
+			check(GetDefault<UEdGraphSchema_K2>());
+			GetDefault<UPropertyBagSchema>()->GetVariableTypeTree(TypeTree, TypeTreeFilter);
+		};
+
+		auto PinInfoChanged = [PropUtils](const FEdGraphPinType& PinType, const TWeakObjectPtr<UMovieGraphMember>& GraphMember)
+		{
+			// The SPinTypeSelector popup might outlive this details view, so the member could be invalid
+			if (GraphMember.IsValid())
+			{
+				GraphMember->SetValueType(UMoviePipelineEdGraphNodeBase::GetValueTypeFromPinType(PinType), PinType.PinSubCategoryObject.Get());
+
+				// Need the ForceRefresh to make sure the details panel refreshes immediately after the data type change.
+				// Can result in a crash without it.
+				PropUtils->ForceRefresh();
+			}
+		};
+		
 		TArray<TWeakObjectPtr<UObject>> ObjectsBeingCustomized;
 		DetailBuilder.GetObjectsBeingCustomized(ObjectsBeingCustomized);
 
@@ -59,7 +85,8 @@ protected:
 			// be driven by UPROPERTY metadata. Hence why this needs to be done w/ a details customization.
 			
 			// Enable/disable the value property for inputs/outputs based on whether it is specified as a branch or not
-			if (const UMovieGraphInterfaceBase* InterfaceBase = Cast<UMovieGraphInterfaceBase>(CustomizedObject))
+			const UMovieGraphInterfaceBase* InterfaceBase = Cast<UMovieGraphInterfaceBase>(CustomizedObject);
+			if (InterfaceBase)
 			{
 				const TSharedRef<IPropertyHandle> ValueProperty = DetailBuilder.GetProperty("Value", UMovieGraphValueContainer::StaticClass());
 				if (ValueProperty->IsValidHandle())
@@ -84,9 +111,12 @@ protected:
 				continue;
 			}
 
+			const bool bIsNameRowEnabled = MemberObject->IsEditable();
+			
 			// Add a custom row for the Name property (to allow for proper validation)
 			IDetailCategoryBuilder& GeneralCategory = DetailBuilder.EditCategory("General");
 			GeneralCategory.AddCustomRow(FText::GetEmpty())
+			.IsEnabled(bIsNameRowEnabled)
 			.NameContent()
 			[
 				SNew(STextBlock)
@@ -99,10 +129,46 @@ protected:
 				.Text_Lambda([MemberObject]() { return FText::FromString(MemberObject->GetMemberName()); })
 				.OnTextChanged(this, &FMovieGraphMemberCustomization::OnNameChanged, MemberObject)
 				.OnTextCommitted(this, &FMovieGraphMemberCustomization::OnNameCommitted, MemberObject)
-				.IsReadOnly_Lambda([MemberObject]() { return !MemberObject->IsEditable(); })
 				.SelectAllTextWhenFocused(true)
 				.Font(IDetailLayoutBuilder::GetDetailFont())
 			];
+
+			// If this is an interface (eg, input/output) only enable the type selector if it's not a branch and editable.
+			// Otherwise, editability is the only factor in enable state.
+			const bool bIsTypeRowEnabled = InterfaceBase ? (!InterfaceBase->bIsBranch && MemberObject->IsEditable()) : MemberObject->IsEditable();
+
+			TWeakObjectPtr<UMovieGraphMember> WeakMemberObject = MakeWeakObjectPtr(MemberObject);
+
+			// Add a PinTypeSelector widget to pick the data type the member uses
+			IDetailCategoryBuilder& ValueCategory = DetailBuilder.EditCategory("Value");
+			ValueCategory.AddCustomRow(FText::GetEmpty())
+			.IsEnabled(bIsTypeRowEnabled)
+			.NameContent()
+			[
+				SNew(STextBlock)
+				.Text(LOCTEXT("MemberPropertyLabel_Type", "Type"))
+				.Font(DetailBuilder.GetDetailFont())
+			]
+			.ValueContent()
+			[
+				SNew(SPinTypeSelector, FGetPinTypeTree::CreateLambda(GetFilteredVariableTypeTree))
+				.TargetPinType_Lambda([WeakMemberObject]()
+				{
+					// The SPinTypeSelector popup might outlive this details view, so the member could be invalid
+					if (!WeakMemberObject.IsValid())
+					{
+						return FEdGraphPinType();
+					}
+
+					constexpr bool bIsBranch = false;
+					return UMoviePipelineEdGraphNodeBase::GetPinType(WeakMemberObject->GetValueType(), bIsBranch, WeakMemberObject->GetValueTypeObject());
+				})
+				.OnPinTypeChanged_Lambda(PinInfoChanged, WeakMemberObject)
+				.Schema(GetDefault<UPropertyBagSchema>())
+				.bAllowArrays(false)
+				.TypeTreeFilter(ETypeTreeFilter::None)
+				.Font(IDetailLayoutBuilder::GetDetailFont())
+			];
 		}
 	}
 	//~ End IDetailCustomization interface
diff --git a/Engine/Plugins/MovieScene/MovieRenderPipeline/Source/MovieRenderPipelineEditor/Private/Customizations/Graph/MovieGraphSelectNodeCustomization.h b/Engine/Plugins/MovieScene/MovieRenderPipeline/Source/MovieRenderPipelineEditor/Private/Customizations/Graph/MovieGraphSelectNodeCustomization.h
new file mode 100644
--- /dev/null
+++ b/Engine/Plugins/MovieScene/MovieRenderPipeline/Source/MovieRenderPipelineEditor/Private/Customizations/Graph/MovieGraphSelectNodeCustomization.h
@@ -0,0 +1,170 @@
+// Copyright Epic Games, Inc. All Rights Reserved.
+
+#pragma once
+
+#include "DetailCategoryBuilder.h"
+#include "DetailLayoutBuilder.h"
+#include "DetailWidgetRow.h"
+#include "Graph/MovieEdGraphNode.h"
+#include "Graph/MovieGraphSchema.h"
+#include "Graph/Nodes/MovieGraphSelectNode.h"
+#include "IDetailCustomization.h"
+#include "IPropertyUtilities.h"
+#include "PropertyBagDetails.h"
+#include "PropertyHandle.h"
+#include "SPinTypeSelector.h"
+#include "Widgets/Input/SEditableTextBox.h"
+#include "Widgets/Text/STextBlock.h"
+
+#define LOCTEXT_NAMESPACE "MoviePipelineEditor"
+
+/** Customize how the Select node appears in the details panel. */
+class FMovieGraphSelectNodeCustomization : public IDetailCustomization
+{
+public:
+	static TSharedRef<IDetailCustomization> MakeInstance()
+	{
+		return MakeShared<FMovieGraphSelectNodeCustomization>();
+	}
+
+protected:
+	//~ Begin IDetailCustomization interface
+	virtual void CustomizeDetails(const TSharedPtr<IDetailLayoutBuilder>& DetailBuilder) override
+	{
+		CustomizeDetails(*DetailBuilder);
+	}
+
+	bool IsTypeAllowed(const FEdGraphPinType& PinType) const
+	{
+		const FName PinCategory = PinType.PinCategory;
+		
+		// Only a subset of reliably-compared types are available to be chosen
+		return
+			(PinCategory == UMovieGraphSchema::PC_Boolean) ||
+			(PinCategory == UMovieGraphSchema::PC_Byte) ||
+			(PinCategory == UMovieGraphSchema::PC_Integer) ||
+			(PinCategory == UMovieGraphSchema::PC_Int64) ||
+			(PinCategory == UMovieGraphSchema::PC_Name) ||
+			(PinCategory == UMovieGraphSchema::PC_String) ||
+			(PinCategory == UMovieGraphSchema::PC_Text) ||
+			(PinCategory == UMovieGraphSchema::PC_Enum) ||
+			(PinCategory == UMovieGraphSchema::PC_Class) ||
+			(PinCategory == UMovieGraphSchema::PC_SoftClass);
+	}
+	
+	virtual void CustomizeDetails(IDetailLayoutBuilder& DetailBuilder) override
+	{
+		TSharedRef<IPropertyUtilities> PropUtils = DetailBuilder.GetPropertyUtilities();
+		
+		auto GetFilteredVariableTypeTree = [this](TArray<TSharedPtr<UEdGraphSchema_K2::FPinTypeTreeInfo>>& TypeTree, ETypeTreeFilter TypeTreeFilter)
+		{
+			check(GetDefault<UEdGraphSchema_K2>());
+			GetDefault<UPropertyBagSchema>()->GetVariableTypeTree(TypeTree, TypeTreeFilter);
+
+			// Filter items in the tree down to those that are explicitly allowed
+			constexpr bool bForceLoadSubCategoryObject = false;
+			for (int32 Index = 0; Index < TypeTree.Num(); )
+			{
+				TSharedPtr<UEdGraphSchema_K2::FPinTypeTreeInfo>& PinType = TypeTree[Index];
+				if (!PinType.IsValid())
+				{
+					return;
+				}
+
+				if (!IsTypeAllowed(PinType->GetPinType(bForceLoadSubCategoryObject)))
+				{
+					TypeTree.RemoveAt(Index);
+					continue;
+				}
+
+				for (int32 ChildIndex = 0; ChildIndex < PinType->Children.Num(); )
+				{
+					TSharedPtr<UEdGraphSchema_K2::FPinTypeTreeInfo> Child = PinType->Children[ChildIndex];
+					if (Child.IsValid())
+					{
+						if (!IsTypeAllowed(PinType->GetPinType(bForceLoadSubCategoryObject)))
+						{
+							PinType->Children.RemoveAt(ChildIndex);
+							continue;
+						}
+					}
+					++ChildIndex;
+				}
+
+				++Index;
+			}
+		};
+		
+		for (TWeakObjectPtr<UMovieGraphSelectNode>& SelectNodeObject : DetailBuilder.GetObjectsOfTypeBeingCustomized<UMovieGraphSelectNode>())
+		{
+			if (!SelectNodeObject.IsValid())
+			{
+				continue;
+			}
+			
+			IDetailCategoryBuilder& Category = DetailBuilder.EditCategory(FName("General"));
+
+			// Add a PinTypeSelector widget to pick the data type the select node uses
+			Category.AddCustomRow(LOCTEXT("DataTypeFilterString", "Data Type"))
+			.NameContent()
+			[
+				SNew(STextBlock)
+				.Text(LOCTEXT("DataType", "Data Type"))
+				.Font(IDetailLayoutBuilder::GetDetailFont())
+			]
+			.ValueContent()
+			[
+				SNew(SPinTypeSelector, FGetPinTypeTree::CreateLambda(GetFilteredVariableTypeTree))
+				.TargetPinType_Lambda([SelectNodeObject]()
+				{
+					// The SPinTypeSelector popup might outlive this details view, so the node could be invalid
+					if (!SelectNodeObject.IsValid())
+					{
+						return FEdGraphPinType();
+					}
+
+					constexpr bool bIsBranch = false;
+					return UMoviePipelineEdGraphNodeBase::GetPinType(SelectNodeObject->GetValueType(), bIsBranch, SelectNodeObject->GetValueTypeObject());
+				})
+				.OnPinTypeChanged_Lambda([PropUtils, SelectNodeObject](const FEdGraphPinType& InPinType)
+				{
+					// The SPinTypeSelector popup might outlive this details view, so the node could be invalid
+					if (SelectNodeObject.IsValid())
+					{
+						SelectNodeObject->SetDataType(UMoviePipelineEdGraphNodeBase::GetValueTypeFromPinType(InPinType), InPinType.PinSubCategoryObject.Get());
+
+						// Need the ForceRefresh to make sure the details panel refreshes immediately after the data type change.
+						// Can result in a crash without it.
+						PropUtils->ForceRefresh();
+					}
+				})
+				.Schema(GetDefault<UPropertyBagSchema>())
+				.bAllowArrays(false)
+				.TypeTreeFilter(ETypeTreeFilter::None)
+				.Font(IDetailLayoutBuilder::GetDetailFont())
+			];
+
+			const TSharedRef<IPropertyHandle> SelectOptions = DetailBuilder.GetProperty("SelectOptions");
+			const TSharedRef<IPropertyHandle> SelectedOption = DetailBuilder.GetProperty("SelectedOption");
+
+			// Add the property bag property ("Value") for both SelectOptions and SelectedOption
+			if (SelectOptions->IsValidHandle() && SelectedOption->IsValidHandle())
+			{
+				if (const TSharedPtr<IPropertyHandle> ValueProperty = SelectOptions->GetChildHandle("Value"))
+				{
+					Category.AddProperty(ValueProperty);
+					SelectOptions->MarkHiddenByCustomization();
+				}
+
+				if (const TSharedPtr<IPropertyHandle> ValueProperty = SelectedOption->GetChildHandle("Value"))
+				{
+					Category.AddProperty(ValueProperty);
+					SelectedOption->MarkHiddenByCustomization();
+				}
+			}
+		}
+	}
+	//~ End IDetailCustomization interface
+};
+
+#undef LOCTEXT_NAMESPACE
\ No newline at end of file
diff --git a/Engine/Plugins/MovieScene/MovieRenderPipeline/Source/MovieRenderPipelineEditor/Private/Graph/MovieEdGraphNode.cpp b/Engine/Plugins/MovieScene/MovieRenderPipeline/Source/MovieRenderPipelineEditor/Private/Graph/MovieEdGraphNode.cpp
--- a/Engine/Plugins/MovieScene/MovieRenderPipeline/Source/MovieRenderPipelineEditor/Private/Graph/MovieEdGraphNode.cpp
+++ b/Engine/Plugins/MovieScene/MovieRenderPipeline/Source/MovieRenderPipelineEditor/Private/Graph/MovieEdGraphNode.cpp
@@ -13,6 +13,7 @@
 #include "ToolMenu.h"
 #include "EdGraph/EdGraphSchema.h"
 #include "GraphEditorActions.h"
+#include "MovieRenderPipelineCoreModule.h"
 #include "Framework/Commands/GenericCommands.h"
 
 #define LOCTEXT_NAMESPACE "MoviePipelineEdGraphNodeBase"
@@ -55,7 +56,7 @@ void UMoviePipelineEdGraphNodeBase::PostTransacted(const FTransactionObjectEvent
 	}
 }
 
-FEdGraphPinType UMoviePipelineEdGraphNodeBase::GetPinType(EMovieGraphValueType ValueType, bool bIsBranch)
+FEdGraphPinType UMoviePipelineEdGraphNodeBase::GetPinType(EMovieGraphValueType ValueType, bool bIsBranch, const UObject* InValueTypeObject)
 {
 	FEdGraphPinType EdPinType;
 	EdPinType.ResetToDefaults();
@@ -85,10 +86,12 @@ FEdGraphPinType UMoviePipelineEdGraphNodeBase::GetPinType(EMovieGraphValueType V
 		EdPinType.PinCategory = UMovieGraphSchema::PC_Int64;
 		break;
 	case EMovieGraphValueType::Float:
-		EdPinType.PinCategory = UMovieGraphSchema::PC_Float;
+		EdPinType.PinCategory = UMovieGraphSchema::PC_Real;
+		EdPinType.PinSubCategory = UMovieGraphSchema::PC_Float;
 		break;
 	case EMovieGraphValueType::Double:
-		EdPinType.PinCategory = UMovieGraphSchema::PC_Double;
+		EdPinType.PinCategory = UMovieGraphSchema::PC_Real;
+		EdPinType.PinSubCategory = UMovieGraphSchema::PC_Double;
 		break;
 	case EMovieGraphValueType::Name:
 		EdPinType.PinCategory = UMovieGraphSchema::PC_Name;
@@ -101,21 +104,27 @@ FEdGraphPinType UMoviePipelineEdGraphNodeBase::GetPinType(EMovieGraphValueType V
 		break;
 	case EMovieGraphValueType::Enum:
 		EdPinType.PinCategory = UMovieGraphSchema::PC_Enum;
+		EdPinType.PinSubCategoryObject = MakeWeakObjectPtr(const_cast<UObject*>(InValueTypeObject));
 		break;
 	case EMovieGraphValueType::Struct:
 		EdPinType.PinCategory = UMovieGraphSchema::PC_Struct;
+		EdPinType.PinSubCategoryObject = MakeWeakObjectPtr(const_cast<UObject*>(InValueTypeObject));
 		break;
 	case EMovieGraphValueType::Object:
 		EdPinType.PinCategory = UMovieGraphSchema::PC_Object;
+		EdPinType.PinSubCategoryObject = MakeWeakObjectPtr(const_cast<UObject*>(InValueTypeObject));
 		break;
 	case EMovieGraphValueType::SoftObject:
 		EdPinType.PinCategory = UMovieGraphSchema::PC_SoftObject;
+		EdPinType.PinSubCategoryObject = MakeWeakObjectPtr(const_cast<UObject*>(InValueTypeObject));
 		break;
 	case EMovieGraphValueType::Class:
 		EdPinType.PinCategory = UMovieGraphSchema::PC_Class;
+		EdPinType.PinSubCategoryObject = MakeWeakObjectPtr(const_cast<UObject*>(InValueTypeObject));
 		break;
 	case EMovieGraphValueType::SoftClass:
 		EdPinType.PinCategory = UMovieGraphSchema::PC_SoftClass;
+		EdPinType.PinSubCategoryObject = MakeWeakObjectPtr(const_cast<UObject*>(InValueTypeObject));
 		break;
 	default:
 		EdPinType.PinCategory = UMovieGraphSchema::PC_Float;
@@ -130,6 +139,56 @@ FEdGraphPinType UMoviePipelineEdGraphNodeBase::GetPinType(const UMovieGraphPin*
 	return GetPinType(InPin->Properties.Type, InPin->Properties.bIsBranch);
 }
 
+EMovieGraphValueType UMoviePipelineEdGraphNodeBase::GetValueTypeFromPinType(const FEdGraphPinType& InPinType)
+{
+	static const TMap<FName, EMovieGraphValueType> PinCategoryToValueType =
+	{
+		{UMovieGraphSchema::PC_Boolean, EMovieGraphValueType::Bool},
+		{UMovieGraphSchema::PC_Byte, EMovieGraphValueType::Byte},
+		{UMovieGraphSchema::PC_Integer, EMovieGraphValueType::Int32},
+		{UMovieGraphSchema::PC_Int64, EMovieGraphValueType::Int64},
+		{UMovieGraphSchema::PC_Float, EMovieGraphValueType::Float},
+		{UMovieGraphSchema::PC_Double, EMovieGraphValueType::Double},
+		{UMovieGraphSchema::PC_Name, EMovieGraphValueType::Name},
+		{UMovieGraphSchema::PC_String, EMovieGraphValueType::String},
+		{UMovieGraphSchema::PC_Text, EMovieGraphValueType::Text},
+		{UMovieGraphSchema::PC_Enum, EMovieGraphValueType::Enum},
+		{UMovieGraphSchema::PC_Struct, EMovieGraphValueType::Struct},
+		{UMovieGraphSchema::PC_Object, EMovieGraphValueType::Object},
+		{UMovieGraphSchema::PC_SoftObject, EMovieGraphValueType::SoftObject},
+		{UMovieGraphSchema::PC_Class, EMovieGraphValueType::Class},
+		{UMovieGraphSchema::PC_SoftClass, EMovieGraphValueType::SoftClass}
+	};
+
+	// Enums can be reported as bytes with a pin sub-category set to the enum
+	if (Cast<UEnum>(InPinType.PinSubCategoryObject))
+	{
+		return EMovieGraphValueType::Enum;
+	}
+
+	// Double/float are a bit special: they're reported as a "real" w/ a float/double sub-type
+	if (InPinType.PinCategory == UMovieGraphSchema::PC_Real)
+	{
+		if (InPinType.PinSubCategory == UMovieGraphSchema::PC_Float)
+		{
+			return EMovieGraphValueType::Float;
+		}
+
+		if (InPinType.PinSubCategory == UMovieGraphSchema::PC_Double)
+		{
+			return EMovieGraphValueType::Double;
+		}
+	}
+	
+	if (const EMovieGraphValueType* FoundValueType = PinCategoryToValueType.Find(InPinType.PinCategory))
+	{
+		return *FoundValueType;
+	}
+
+	UE_LOG(LogMovieRenderPipeline, Warning, TEXT("Unable to convert pin type: category [%s], sub-category [%s]"), *InPinType.PinCategory.ToString(), *InPinType.PinSubCategory.ToString());
+	return EMovieGraphValueType::None;
+}
+
 void UMoviePipelineEdGraphNodeBase::UpdatePosition() const
 {
 	if (RuntimeNode)
diff --git a/Engine/Plugins/MovieScene/MovieRenderPipeline/Source/MovieRenderPipelineEditor/Private/Graph/MovieEdGraphNode.h b/Engine/Plugins/MovieScene/MovieRenderPipeline/Source/MovieRenderPipelineEditor/Private/Graph/MovieEdGraphNode.h
--- a/Engine/Plugins/MovieScene/MovieRenderPipeline/Source/MovieRenderPipelineEditor/Private/Graph/MovieEdGraphNode.h
+++ b/Engine/Plugins/MovieScene/MovieRenderPipeline/Source/MovieRenderPipelineEditor/Private/Graph/MovieEdGraphNode.h
@@ -33,8 +33,9 @@ public:
 	virtual void GetNodeContextMenuActions(UToolMenu* Menu, UGraphNodeContextMenuContext* Context) const override;
 	//~ End UEdGraphNode Interface
 
-	static FEdGraphPinType GetPinType(EMovieGraphValueType ValueType, bool bIsBranch);
+	static FEdGraphPinType GetPinType(EMovieGraphValueType ValueType, bool bIsBranch, const UObject* InValueTypeObject = nullptr);
 	static FEdGraphPinType GetPinType(const UMovieGraphPin* InPin);
+	static EMovieGraphValueType GetValueTypeFromPinType(const FEdGraphPinType& InPinType);
 
 	// Called after PrepareForCopying to restore changes made during preparation. 
 	virtual void PostCopy();
diff --git a/Engine/Plugins/MovieScene/MovieRenderPipeline/Source/MovieRenderPipelineEditor/Private/Graph/MovieGraphAssetToolkit.cpp b/Engine/Plugins/MovieScene/MovieRenderPipeline/Source/MovieRenderPipelineEditor/Private/Graph/MovieGraphAssetToolkit.cpp
--- a/Engine/Plugins/MovieScene/MovieRenderPipeline/Source/MovieRenderPipelineEditor/Private/Graph/MovieGraphAssetToolkit.cpp
+++ b/Engine/Plugins/MovieScene/MovieRenderPipeline/Source/MovieRenderPipelineEditor/Private/Graph/MovieGraphAssetToolkit.cpp
@@ -4,10 +4,12 @@
 
 #include "Customizations/Graph/MovieGraphMemberCustomization.h"
 #include "Customizations/Graph/MovieGraphNodeCustomization.h"
+#include "Customizations/Graph/MovieGraphSelectNodeCustomization.h"
 #include "Customizations/Graph/MovieGraphShowFlagsCustomization.h"
 #include "Framework/Commands/GenericCommands.h"
 #include "Graph/MovieGraphConfig.h"
 #include "Graph/Nodes/MovieGraphCoreRenderPassNode.h"
+#include "Graph/Nodes/MovieGraphSelectNode.h"
 #include "MovieEdGraphNode.h"
 #include "MovieGraphSchema.h"
 #include "PropertyEditorModule.h"
@@ -233,6 +235,10 @@ TSharedRef<SDockTab> FMovieGraphAssetToolkit::SpawnTab_RenderGraphDetails(const
 		UMovieGraphNode::StaticClass(),
 		FOnGetDetailCustomizationInstance::CreateStatic(&FMovieGraphNodeCustomization::MakeInstance));
 
+	SelectedGraphObjectsDetailsWidget->RegisterInstancedCustomPropertyLayout(
+		UMovieGraphSelectNode::StaticClass(),
+		FOnGetDetailCustomizationInstance::CreateStatic(&FMovieGraphSelectNodeCustomization::MakeInstance));
+
 	SelectedGraphObjectsDetailsWidget->RegisterInstancedCustomPropertyTypeLayout(
 		UMovieGraphShowFlags::StaticClass()->GetFName(),
 		FOnGetPropertyTypeCustomizationInstance::CreateStatic(&FMovieGraphShowFlagsCustomization::MakeInstance));
diff --git a/Engine/Plugins/MovieScene/MovieRenderPipeline/Source/MovieRenderPipelineEditor/Private/Graph/MovieGraphSchema.cpp b/Engine/Plugins/MovieScene/MovieRenderPipeline/Source/MovieRenderPipelineEditor/Private/Graph/MovieGraphSchema.cpp
--- a/Engine/Plugins/MovieScene/MovieRenderPipeline/Source/MovieRenderPipelineEditor/Private/Graph/MovieGraphSchema.cpp
+++ b/Engine/Plugins/MovieScene/MovieRenderPipeline/Source/MovieRenderPipelineEditor/Private/Graph/MovieGraphSchema.cpp
@@ -1,6 +1,8 @@
 // Copyright Epic Games, Inc. All Rights Reserved.
 
 #include "MovieGraphSchema.h"
+
+#include "EdGraphSchema_K2.h"
 #include "Graph/MovieGraphConfig.h"
 #include "Graph/Nodes/MovieGraphInputNode.h"
 #include "Graph/Nodes/MovieGraphOutputNode.h"
@@ -20,22 +22,23 @@ TArray<UClass*> UMovieGraphSchema::MoviePipelineNodeClasses;
 
 #define LOCTEXT_NAMESPACE "MoviePipelineGraphSchema"
 
-const FName UMovieGraphSchema::PC_Branch(TEXT("branch"));
-const FName UMovieGraphSchema::PC_Boolean(TEXT("boolean"));
-const FName UMovieGraphSchema::PC_Byte(TEXT("byte"));
-const FName UMovieGraphSchema::PC_Integer(TEXT("integer"));
-const FName UMovieGraphSchema::PC_Int64(TEXT("int64"));
-const FName UMovieGraphSchema::PC_Float(TEXT("float"));
-const FName UMovieGraphSchema::PC_Double(TEXT("double"));
-const FName UMovieGraphSchema::PC_Name(TEXT("name"));
-const FName UMovieGraphSchema::PC_String(TEXT("string"));
-const FName UMovieGraphSchema::PC_Text(TEXT("text"));
-const FName UMovieGraphSchema::PC_Enum(TEXT("enum"));
-const FName UMovieGraphSchema::PC_Struct(TEXT("struct"));
-const FName UMovieGraphSchema::PC_Object(TEXT("object"));
-const FName UMovieGraphSchema::PC_SoftObject(TEXT("softobject"));
-const FName UMovieGraphSchema::PC_Class(TEXT("class"));
-const FName UMovieGraphSchema::PC_SoftClass(TEXT("softclass"));
+const FName UMovieGraphSchema::PC_Branch(TEXT("branch"));	// The branch looks like an Exec pin, but isn't the same thing, so we don't use the BP Exec type
+const FName UMovieGraphSchema::PC_Boolean(UEdGraphSchema_K2::PC_Boolean);
+const FName UMovieGraphSchema::PC_Byte(UEdGraphSchema_K2::PC_Byte);
+const FName UMovieGraphSchema::PC_Integer(UEdGraphSchema_K2::PC_Int);
+const FName UMovieGraphSchema::PC_Int64(UEdGraphSchema_K2::PC_Int64);
+const FName UMovieGraphSchema::PC_Real(UEdGraphSchema_K2::PC_Real);
+const FName UMovieGraphSchema::PC_Float(UEdGraphSchema_K2::PC_Float);
+const FName UMovieGraphSchema::PC_Double(UEdGraphSchema_K2::PC_Double);
+const FName UMovieGraphSchema::PC_Name(UEdGraphSchema_K2::PC_Name);
+const FName UMovieGraphSchema::PC_String(UEdGraphSchema_K2::PC_String);
+const FName UMovieGraphSchema::PC_Text(UEdGraphSchema_K2::PC_Text);
+const FName UMovieGraphSchema::PC_Enum(UEdGraphSchema_K2::PC_Enum);
+const FName UMovieGraphSchema::PC_Struct(UEdGraphSchema_K2::PC_Struct);
+const FName UMovieGraphSchema::PC_Object(UEdGraphSchema_K2::PC_Object);
+const FName UMovieGraphSchema::PC_SoftObject(UEdGraphSchema_K2::PC_SoftObject);
+const FName UMovieGraphSchema::PC_Class(UEdGraphSchema_K2::PC_Class);
+const FName UMovieGraphSchema::PC_SoftClass(UEdGraphSchema_K2::PC_SoftClass);
 
 namespace UE::MovieGraph::Private
 {
@@ -99,6 +102,18 @@ void UMovieGraphSchema::CreateDefaultNodesForGraph(UEdGraph& Graph) const
 	}*/
 }
 
+bool UMovieGraphSchema::SupportsPinTypeContainer(TWeakPtr<const FEdGraphSchemaAction> SchemaAction, const FEdGraphPinType& PinType, const EPinContainerType& ContainerType) const
+{
+	// No maps, sets, or arrays
+	return ContainerType == EPinContainerType::None;
+}
+
+bool UMovieGraphSchema::ShouldHidePinDefaultValue(UEdGraphPin* Pin) const
+{
+	// The graph doesn't support editing default values for pins yet
+	return true;
+}
+
 void UMovieGraphSchema::InitMoviePipelineNodeClasses()
 {
 	if (MoviePipelineNodeClasses.Num() > 0)
@@ -374,86 +389,99 @@ void UMovieGraphSchema::BreakSinglePinLink(UEdGraphPin* SourcePin, UEdGraphPin*
 	RuntimeGraph->RemoveEdge(SourceRuntimeNode, SourcePin->PinName, TargetRuntimeNode, TargetPin->PinName);
 }
 
-FLinearColor UMovieGraphSchema::GetTypeColor(const FName& InType)
+FLinearColor UMovieGraphSchema::GetTypeColor(const FName& InPinCategory, const FName& InPinSubCategory)
 {
 	const UGraphEditorSettings* Settings = GetDefault<UGraphEditorSettings>();
 
-	if (InType == PC_Branch)
+	if (InPinCategory == PC_Branch)
 	{
 		return Settings->ExecutionPinTypeColor;
 	}
 
-	if (InType == PC_Boolean)
+	if (InPinCategory == PC_Boolean)
 	{
 		return Settings->BooleanPinTypeColor;
 	}
 	
-	if (InType == PC_Byte)
+	if (InPinCategory == PC_Byte)
 	{
 		return Settings->BytePinTypeColor;
 	}
 
-	if (InType == PC_Integer)
+	if (InPinCategory == PC_Integer)
 	{
 		return Settings->IntPinTypeColor;
 	}
 
-	if (InType == PC_Int64)
+	if (InPinCategory == PC_Int64)
 	{
 		return Settings->Int64PinTypeColor;
 	}
 
-	if (InType == PC_Float)
+	if (InPinCategory == PC_Float)
 	{
 		return Settings->FloatPinTypeColor;
 	}
 
-	if (InType == PC_Double)
+	if (InPinCategory == PC_Double)
 	{
 		return Settings->DoublePinTypeColor;
 	}
 
-	if (InType == PC_Name)
+	if (InPinCategory == PC_Real)
+	{
+		if (InPinSubCategory == PC_Float)
+		{
+			return Settings->FloatPinTypeColor;
+		}
+
+		if (InPinSubCategory == PC_Double)
+		{
+			return Settings->DoublePinTypeColor;
+		}
+	}
+
+	if (InPinCategory == PC_Name)
 	{
 		return Settings->NamePinTypeColor;
 	}
 
-	if (InType == PC_String)
+	if (InPinCategory == PC_String)
 	{
 		return Settings->StringPinTypeColor;
 	}
 
-	if (InType == PC_Text)
+	if (InPinCategory == PC_Text)
 	{
 		return Settings->TextPinTypeColor;
 	}
 
-	if (InType == PC_Enum)
+	if (InPinCategory == PC_Enum)
 	{
 		return Settings->BytePinTypeColor;
 	}
 
-	if (InType == PC_Struct)
+	if (InPinCategory == PC_Struct)
 	{
 		return Settings->StructPinTypeColor;
 	}
 	
-	if (InType == PC_Object)
+	if (InPinCategory == PC_Object)
 	{
 		return Settings->ObjectPinTypeColor;
 	}
 
-	if (InType == PC_SoftObject)
+	if (InPinCategory == PC_SoftObject)
 	{
 		return Settings->SoftObjectPinTypeColor;
 	}
 	
-	if (InType == PC_Class)
+	if (InPinCategory == PC_Class)
     {
     	return Settings->ClassPinTypeColor;
     }
 
-	if (InType == PC_SoftClass)
+	if (InPinCategory == PC_SoftClass)
 	{
 		return Settings->SoftClassPinTypeColor;
 	}
@@ -463,7 +491,7 @@ FLinearColor UMovieGraphSchema::GetTypeColor(const FName& InType)
 
 FLinearColor UMovieGraphSchema::GetPinTypeColor(const FEdGraphPinType& PinType) const
 {
-	return GetTypeColor(PinType.PinCategory);
+	return GetTypeColor(PinType.PinCategory, PinType.PinSubCategory);
 }
 
 FConnectionDrawingPolicy* UMovieGraphSchema::CreateConnectionDrawingPolicy(int32 InBackLayerID, int32 InFrontLayerID,
diff --git a/Engine/Plugins/MovieScene/MovieRenderPipeline/Source/MovieRenderPipelineEditor/Private/Graph/MovieGraphSchema.h b/Engine/Plugins/MovieScene/MovieRenderPipeline/Source/MovieRenderPipelineEditor/Private/Graph/MovieGraphSchema.h
--- a/Engine/Plugins/MovieScene/MovieRenderPipeline/Source/MovieRenderPipelineEditor/Private/Graph/MovieGraphSchema.h
+++ b/Engine/Plugins/MovieScene/MovieRenderPipeline/Source/MovieRenderPipelineEditor/Private/Graph/MovieGraphSchema.h
@@ -22,7 +22,7 @@ public:
 	virtual const FPinConnectionResponse CanCreateConnection(const UEdGraphPin* A, const UEdGraphPin* B) const override;
 	virtual bool TryCreateConnection(UEdGraphPin* InA, UEdGraphPin* InB) const override;
 	//virtual const FPinConnectionResponse CanMergeNodes(const UEdGraphNode* A, const UEdGraphNode* B) const override;
-	static FLinearColor GetTypeColor(const FName& InType);
+	static FLinearColor GetTypeColor(const FName& InPinCategory, const FName& InPinSubCategory);
 	virtual FLinearColor GetPinTypeColor(const FEdGraphPinType& PinType) const override;
 	virtual FConnectionDrawingPolicy* CreateConnectionDrawingPolicy(int32 InBackLayerID, int32 InFrontLayerID, float InZoomFactor, const FSlateRect& InClippingRect, FSlateWindowElementList& InDrawElements, UEdGraph* InGraphObj) const override;
 
@@ -33,6 +33,8 @@ public:
 	//virtual TSharedPtr<FEdGraphSchemaAction> GetCreateCommentAction() const override;
 	virtual void BreakPinLinks(UEdGraphPin& TargetPin, bool bSendsNodeNotification) const override;
 	virtual void BreakSinglePinLink(UEdGraphPin* SourcePin, UEdGraphPin* TargetPin) const override;
+	virtual bool SupportsPinTypeContainer(TWeakPtr<const FEdGraphSchemaAction> SchemaAction, const FEdGraphPinType& PinType, const EPinContainerType& ContainerType) const override;
+	virtual bool ShouldHidePinDefaultValue(UEdGraphPin* Pin) const override;
 	//~ End EdGraphSchema Interface
 	static void InitMoviePipelineNodeClasses();
 
@@ -50,6 +52,7 @@ public:
 	static const FName PC_Byte;
 	static const FName PC_Integer;
 	static const FName PC_Int64;
+	static const FName PC_Real;
 	static const FName PC_Float;
 	static const FName PC_Double;
 	static const FName PC_Name;
diff --git a/Engine/Plugins/MovieScene/MovieRenderPipeline/Source/MovieRenderPipelineEditor/Private/Graph/SMovieGraphMembersTabContent.cpp b/Engine/Plugins/MovieScene/MovieRenderPipeline/Source/MovieRenderPipelineEditor/Private/Graph/SMovieGraphMembersTabContent.cpp
--- a/Engine/Plugins/MovieScene/MovieRenderPipeline/Source/MovieRenderPipelineEditor/Private/Graph/SMovieGraphMembersTabContent.cpp
+++ b/Engine/Plugins/MovieScene/MovieRenderPipeline/Source/MovieRenderPipelineEditor/Private/Graph/SMovieGraphMembersTabContent.cpp
@@ -34,13 +34,18 @@ namespace UE::MovieGraph::Private
 		return nullptr;
 	}
 
-	static bool GetIconAndColorFromDataType(EMovieGraphValueType VariableType, const FSlateBrush*& OutPrimaryBrush, FSlateColor& OutIconColor, const FSlateBrush*& OutSecondaryBrush, FSlateColor& OutSecondaryColor)
+	static bool GetIconAndColorFromDataType(const UMovieGraphVariable* InGraphVariable, const FSlateBrush*& OutPrimaryBrush, FSlateColor& OutIconColor, const FSlateBrush*& OutSecondaryBrush, FSlateColor& OutSecondaryColor)
 	{
+		if (!InGraphVariable)
+		{
+			return false;
+		}
+		
 		constexpr bool bIsBranch = false;
-		const FEdGraphPinType PinType = UMoviePipelineEdGraphNodeBase::GetPinType(VariableType, bIsBranch);
+		const FEdGraphPinType PinType = UMoviePipelineEdGraphNodeBase::GetPinType(InGraphVariable->GetValueType(), bIsBranch, InGraphVariable->GetValueTypeObject());
 
 		OutPrimaryBrush = FAppStyle::GetBrush("Kismet.AllClasses.VariableIcon");
-		OutIconColor = UMovieGraphSchema::GetTypeColor(PinType.PinCategory);
+		OutIconColor = UMovieGraphSchema::GetTypeColor(PinType.PinCategory, PinType.PinSubCategory);
 		OutSecondaryBrush = nullptr;
 		
 		return true;
@@ -107,7 +112,7 @@ TSharedRef<SWidget> SMovieGraphMembersTabContent::CreateActionWidget(FCreateWidg
 			static_cast<FMovieGraphSchemaAction_NewVariableNode*>(InCreateData->Action.Get());
 		const UMovieGraphVariable* Variable = Cast<UMovieGraphVariable>(VariableAction->ActionTarget);
 		const FEdGraphPinType PinType = Variable
-			? UMoviePipelineEdGraphNodeBase::GetPinType(Variable->GetValueType(), false)
+			? UMoviePipelineEdGraphNodeBase::GetPinType(Variable->GetValueType(), false, Variable->GetValueTypeObject())
 			: FEdGraphPinType();
 		const FLinearColor PinColor = CurrentGraph->PipelineEdGraph->GetSchema()->GetPinTypeColor(PinType);
 		
@@ -449,7 +454,7 @@ void FMovieGraphDragAction_Variable::GetDefaultStatusSymbol(
 {
 	const UMovieGraphVariable* VariableMember = WeakVariable.Get();
 	if (!VariableMember ||
-		!UE::MovieGraph::Private::GetIconAndColorFromDataType(VariableMember->GetValueType(), OutPrimaryBrush, OutIconColor, OutSecondaryBrush, OutSecondaryColor))
+		!UE::MovieGraph::Private::GetIconAndColorFromDataType(VariableMember, OutPrimaryBrush, OutIconColor, OutSecondaryBrush, OutSecondaryColor))
 	{
 		return FGraphSchemaActionDragDropAction::GetDefaultStatusSymbol(OutPrimaryBrush, OutIconColor, OutSecondaryBrush, OutSecondaryColor);
 	}
