Fix non-initialized pointers to prevent compilation errorswith LTO and undefined runtime behaviors
Reported in UDN case #622715

diff --git a/Engine/Plugins/MovieScene/MovieRenderPipeline/Source/MovieRenderPipelineRenderPasses/Private/MoviePipelineEXROutput.cpp b/Engine/Plugins/MovieScene/MovieRenderPipeline/Source/MovieRenderPipelineRenderPasses/Private/MoviePipelineEXROutput.cpp
--- a/Engine/Plugins/MovieScene/MovieRenderPipeline/Source/MovieRenderPipelineRenderPasses/Private/MoviePipelineEXROutput.cpp
+++ b/Engine/Plugins/MovieScene/MovieRenderPipeline/Source/MovieRenderPipelineRenderPasses/Private/MoviePipelineEXROutput.cpp
@@ -336,9 +336,15 @@ static int32 GetComponentWidth(const EImagePixelType InPixelType)
 template <Imf::PixelType OutputFormat>
 int64 FEXRImageWriteTask::CompressRaw(Imf::Header& InHeader, Imf::FrameBuffer& InFrameBuffer, FImagePixelData* InLayer)
 {
-	void const* RawDataPtr;
+	void const* RawDataPtr = nullptr;
 	int64 RawDataSize;
-	InLayer->GetRawData(RawDataPtr, RawDataSize);
+
+	if (InLayer->GetRawData(RawDataPtr, RawDataSize) == false)
+	{
+		UE_LOG(LogMovieRenderPipelineIO, Error, TEXT("Failed to retrieve raw data from image data for writing. Bailing."));
+		return 0;
+	}
+		
 
 	// Look up our layer name (if any).
 	FString& LayerName = LayerNames.FindOrAdd(InLayer);
diff --git a/Engine/Plugins/MovieScene/MovieRenderPipeline/Source/MovieRenderPipelineRenderPasses/Private/MoviePipelineImagePassBase.cpp b/Engine/Plugins/MovieScene/MovieRenderPipeline/Source/MovieRenderPipelineRenderPasses/Private/MoviePipelineImagePassBase.cpp
--- a/Engine/Plugins/MovieScene/MovieRenderPipeline/Source/MovieRenderPipelineRenderPasses/Private/MoviePipelineImagePassBase.cpp
+++ b/Engine/Plugins/MovieScene/MovieRenderPipeline/Source/MovieRenderPipelineRenderPasses/Private/MoviePipelineImagePassBase.cpp
@@ -921,23 +921,38 @@ namespace MoviePipeline
 
 				if (SamplePixelData->GetType() == EImagePixelType::Float32)
 				{
-					const void* RawDataPtr;
+					const void* RawDataPtr = nullptr;
 					int64 RawDataSize;
-					SamplePixelData->GetRawData(RawDataPtr, RawDataSize);
-
-					FMemory::Memcpy(FullSizeData.GetData(), RawDataPtr, RawDataSize);
+					
+					if(SamplePixelData->GetRawData(RawDataPtr, RawDataSize) == true)
+					{
+						FMemory::Memcpy(FullSizeData.GetData(), RawDataPtr, RawDataSize);
+					}
+					else
+					{
+						UE_LOG(LogMovieRenderPipelineIO, Error, TEXT("Failed to retrieve raw data from image data for writing. Bailing."));
+						return;
+					}
 				}
 				else if (SamplePixelData->GetType() == EImagePixelType::Float16)
 				{
-					const void* RawDataPtr;
+					const void* RawDataPtr = nullptr;
 					int64 RawDataSize;
-					SamplePixelData->GetRawData(RawDataPtr, RawDataSize);
 
-					const FFloat16Color* DataAsColor = reinterpret_cast<const FFloat16Color*>(RawDataPtr);
-					for (int64 Index = 0; Index < RawSize.X * RawSize.Y; Index++)
+					if(SamplePixelData->GetRawData(RawDataPtr, RawDataSize) == true)
+					{
+						const FFloat16Color* DataAsColor = reinterpret_cast<const FFloat16Color*>(RawDataPtr);
+						for (int64 Index = 0; Index < RawSize.X * RawSize.Y; Index++)
+						{
+							FullSizeData[Index] = FLinearColor(DataAsColor[Index]);
+						}
+					}
+					else
 					{
-						FullSizeData[Index] = FLinearColor(DataAsColor[Index]);
+						UE_LOG(LogMovieRenderPipelineIO, Error, TEXT("Failed to retrieve raw data from image data for writing. Bailing."));
+						return;
 					}
+
 				}
 				else
 				{
