MoviePipeline: Added search aliasing (keywords) for nodes in the graph. Added a first pass at keywords for existing nodes, but the keyword lists will probably need future updates.

Also removed the "Get CVar Value" node. It has no useful implementation currently and is not expected to be needed in practice.

#rnx

diff --git a/Engine/Plugins/MovieScene/MovieRenderPipeline/Source/MovieRenderPipelineCore/Private/Graph/Nodes/MovieGraphApplyCVarPresetNode.cpp b/Engine/Plugins/MovieScene/MovieRenderPipeline/Source/MovieRenderPipelineCore/Private/Graph/Nodes/MovieGraphApplyCVarPresetNode.cpp
--- a/Engine/Plugins/MovieScene/MovieRenderPipeline/Source/MovieRenderPipelineCore/Private/Graph/Nodes/MovieGraphApplyCVarPresetNode.cpp
+++ b/Engine/Plugins/MovieScene/MovieRenderPipeline/Source/MovieRenderPipelineCore/Private/Graph/Nodes/MovieGraphApplyCVarPresetNode.cpp
@@ -40,6 +40,12 @@ FText UMovieGraphApplyCVarPresetNode::GetMenuCategory() const
 	return NSLOCTEXT("MovieGraphNodes", "ApplyCVarPresetGraphNode_Category", "Utility");
 }
 
+FText UMovieGraphApplyCVarPresetNode::GetKeywords() const
+{
+	static const FText Keywords = NSLOCTEXT("MovieGraphNodes", "ApplyCVarPresetGraphNode_Keywords", "cvar console variable preset");
+	return Keywords;
+}
+
 FLinearColor UMovieGraphApplyCVarPresetNode::GetNodeTitleColor() const
 {
 	static const FLinearColor ApplyCVarPresetNodeColor = FLinearColor(0.04f, 0.22f, 0.36f);
diff --git a/Engine/Plugins/MovieScene/MovieRenderPipeline/Source/MovieRenderPipelineCore/Private/Graph/Nodes/MovieGraphAudioOutputNode.cpp b/Engine/Plugins/MovieScene/MovieRenderPipeline/Source/MovieRenderPipelineCore/Private/Graph/Nodes/MovieGraphAudioOutputNode.cpp
--- a/Engine/Plugins/MovieScene/MovieRenderPipeline/Source/MovieRenderPipelineCore/Private/Graph/Nodes/MovieGraphAudioOutputNode.cpp
+++ b/Engine/Plugins/MovieScene/MovieRenderPipeline/Source/MovieRenderPipelineCore/Private/Graph/Nodes/MovieGraphAudioOutputNode.cpp
@@ -26,6 +26,12 @@ FText UMovieGraphAudioOutputNode::GetNodeTitle(const bool bGetDescriptive) const
 	return AudioOutputNodeName;
 }
 
+FText UMovieGraphAudioOutputNode::GetKeywords() const
+{
+	static const FText Keywords = NSLOCTEXT("MovieGraphNodes", "AudioOutputNode_Keywords", "wav audio");
+	return Keywords;
+}
+
 FLinearColor UMovieGraphAudioOutputNode::GetNodeTitleColor() const
 {
 	static const FLinearColor AudioOutputNodeColor = FLinearColor(0.04f, 0.22f, 0.36f);
diff --git a/Engine/Plugins/MovieScene/MovieRenderPipeline/Source/MovieRenderPipelineCore/Private/Graph/Nodes/MovieGraphBranchNode.cpp b/Engine/Plugins/MovieScene/MovieRenderPipeline/Source/MovieRenderPipelineCore/Private/Graph/Nodes/MovieGraphBranchNode.cpp
--- a/Engine/Plugins/MovieScene/MovieRenderPipeline/Source/MovieRenderPipelineCore/Private/Graph/Nodes/MovieGraphBranchNode.cpp
+++ b/Engine/Plugins/MovieScene/MovieRenderPipeline/Source/MovieRenderPipelineCore/Private/Graph/Nodes/MovieGraphBranchNode.cpp
@@ -109,6 +109,12 @@ FText UMovieGraphBranchNode::GetMenuCategory() const
 	return LOCTEXT("NodeCategory_Conditionals", "Conditionals");
 }
 
+FText UMovieGraphBranchNode::GetKeywords() const
+{
+	static const FText Keywords = LOCTEXT("BranchNode_Keywords", "branch if logic conditional");
+    return Keywords;
+}
+
 FLinearColor UMovieGraphBranchNode::GetNodeTitleColor() const
 {
 	static const FLinearColor BranchNodeColor = FLinearColor(0.266f, 0.266f, 0.266f);
diff --git a/Engine/Plugins/MovieScene/MovieRenderPipeline/Source/MovieRenderPipelineCore/Private/Graph/Nodes/MovieGraphGetCVarValueNode.cpp b/Engine/Plugins/MovieScene/MovieRenderPipeline/Source/MovieRenderPipelineCore/Private/Graph/Nodes/MovieGraphGetCVarValueNode.cpp
deleted file mode 100644
diff --git a/Engine/Plugins/MovieScene/MovieRenderPipeline/Source/MovieRenderPipelineCore/Private/Graph/Nodes/MovieGraphSelectNode.cpp b/Engine/Plugins/MovieScene/MovieRenderPipeline/Source/MovieRenderPipelineCore/Private/Graph/Nodes/MovieGraphSelectNode.cpp
--- a/Engine/Plugins/MovieScene/MovieRenderPipeline/Source/MovieRenderPipelineCore/Private/Graph/Nodes/MovieGraphSelectNode.cpp
+++ b/Engine/Plugins/MovieScene/MovieRenderPipeline/Source/MovieRenderPipelineCore/Private/Graph/Nodes/MovieGraphSelectNode.cpp
@@ -377,6 +377,12 @@ FText UMovieGraphSelectNode::GetMenuCategory() const
 	return NodeCategory_Conditionals;
 }
 
+FText UMovieGraphSelectNode::GetKeywords() const
+{
+	static const FText Keywords = NSLOCTEXT("MovieGraphNodes", "SelectNode_Keywords", "select if logic conditional");
+	return Keywords;
+}
+
 FLinearColor UMovieGraphSelectNode::GetNodeTitleColor() const
 {
 	static const FLinearColor SelectNodeColor = FLinearColor(0.266f, 0.266f, 0.266f);
diff --git a/Engine/Plugins/MovieScene/MovieRenderPipeline/Source/MovieRenderPipelineCore/Private/Graph/Nodes/MovieGraphSetCVarValueNode.cpp b/Engine/Plugins/MovieScene/MovieRenderPipeline/Source/MovieRenderPipelineCore/Private/Graph/Nodes/MovieGraphSetCVarValueNode.cpp
--- a/Engine/Plugins/MovieScene/MovieRenderPipeline/Source/MovieRenderPipelineCore/Private/Graph/Nodes/MovieGraphSetCVarValueNode.cpp
+++ b/Engine/Plugins/MovieScene/MovieRenderPipeline/Source/MovieRenderPipelineCore/Private/Graph/Nodes/MovieGraphSetCVarValueNode.cpp
@@ -34,6 +34,12 @@ FText UMovieGraphSetCVarValueNode::GetMenuCategory() const
 	return NSLOCTEXT("MovieGraphNodes", "SetCVarGraphNode_Category", "Utility");
 }
 
+FText UMovieGraphSetCVarValueNode::GetKeywords() const
+{
+	static const FText Keywords = NSLOCTEXT("MovieGraphNodes", "SetCVarGraphNode_Keywords", "cvar console variable");
+	return Keywords;
+}
+
 FLinearColor UMovieGraphSetCVarValueNode::GetNodeTitleColor() const
 {
 	static const FLinearColor SetCVarNodeColor = FLinearColor(0.04f, 0.22f, 0.36f);
diff --git a/Engine/Plugins/MovieScene/MovieRenderPipeline/Source/MovieRenderPipelineCore/Public/Graph/MovieGraphNode.h b/Engine/Plugins/MovieScene/MovieRenderPipeline/Source/MovieRenderPipelineCore/Public/Graph/MovieGraphNode.h
--- a/Engine/Plugins/MovieScene/MovieRenderPipeline/Source/MovieRenderPipelineCore/Public/Graph/MovieGraphNode.h
+++ b/Engine/Plugins/MovieScene/MovieRenderPipeline/Source/MovieRenderPipelineCore/Public/Graph/MovieGraphNode.h
@@ -288,6 +288,9 @@ public:
 
 	/** Gets the category that the node belongs under. */
 	virtual FText GetMenuCategory() const PURE_VIRTUAL(UMovieGraphNode::GetMenuCategory, return FText(); );
+
+	/** Gets the keywords (space-separated) that will be searched in the node creation context menu. */
+	virtual FText GetKeywords() const { return FText(); }
 #endif
 
 protected:
diff --git a/Engine/Plugins/MovieScene/MovieRenderPipeline/Source/MovieRenderPipelineCore/Public/Graph/Nodes/MovieGraphApplyCVarPresetNode.h b/Engine/Plugins/MovieScene/MovieRenderPipeline/Source/MovieRenderPipelineCore/Public/Graph/Nodes/MovieGraphApplyCVarPresetNode.h
--- a/Engine/Plugins/MovieScene/MovieRenderPipeline/Source/MovieRenderPipelineCore/Public/Graph/Nodes/MovieGraphApplyCVarPresetNode.h
+++ b/Engine/Plugins/MovieScene/MovieRenderPipeline/Source/MovieRenderPipelineCore/Public/Graph/Nodes/MovieGraphApplyCVarPresetNode.h
@@ -23,6 +23,7 @@ public:
 #if WITH_EDITOR
 	virtual FText GetNodeTitle(const bool bGetDescriptive = false) const override;
 	virtual FText GetMenuCategory() const override;
+	virtual FText GetKeywords() const override;
 	virtual FLinearColor GetNodeTitleColor() const override;
 	virtual FSlateIcon GetIconAndTint(FLinearColor& OutColor) const override;
 
diff --git a/Engine/Plugins/MovieScene/MovieRenderPipeline/Source/MovieRenderPipelineCore/Public/Graph/Nodes/MovieGraphAudioOutputNode.h b/Engine/Plugins/MovieScene/MovieRenderPipeline/Source/MovieRenderPipelineCore/Public/Graph/Nodes/MovieGraphAudioOutputNode.h
--- a/Engine/Plugins/MovieScene/MovieRenderPipeline/Source/MovieRenderPipelineCore/Public/Graph/Nodes/MovieGraphAudioOutputNode.h
+++ b/Engine/Plugins/MovieScene/MovieRenderPipeline/Source/MovieRenderPipelineCore/Public/Graph/Nodes/MovieGraphAudioOutputNode.h
@@ -21,6 +21,7 @@ public:
 	// UMovieGraphSettingNode Interface
 #if WITH_EDITOR
 	virtual FText GetNodeTitle(const bool bGetDescriptive) const override;
+	virtual FText GetKeywords() const override;
 	virtual FLinearColor GetNodeTitleColor() const override;
 	virtual FSlateIcon GetIconAndTint(FLinearColor& OutColor) const override;
 #endif
diff --git a/Engine/Plugins/MovieScene/MovieRenderPipeline/Source/MovieRenderPipelineCore/Public/Graph/Nodes/MovieGraphBranchNode.h b/Engine/Plugins/MovieScene/MovieRenderPipeline/Source/MovieRenderPipelineCore/Public/Graph/Nodes/MovieGraphBranchNode.h
--- a/Engine/Plugins/MovieScene/MovieRenderPipeline/Source/MovieRenderPipelineCore/Public/Graph/Nodes/MovieGraphBranchNode.h
+++ b/Engine/Plugins/MovieScene/MovieRenderPipeline/Source/MovieRenderPipelineCore/Public/Graph/Nodes/MovieGraphBranchNode.h
@@ -26,6 +26,7 @@ public:
 #if WITH_EDITOR
 	virtual FText GetNodeTitle(const bool bGetDescriptive = false) const override;
 	virtual FText GetMenuCategory() const override;
+	virtual FText GetKeywords() const override;
 	virtual FLinearColor GetNodeTitleColor() const override;
 	virtual FSlateIcon GetIconAndTint(FLinearColor& OutColor) const override;
 #endif
diff --git a/Engine/Plugins/MovieScene/MovieRenderPipeline/Source/MovieRenderPipelineCore/Public/Graph/Nodes/MovieGraphGetCVarValueNode.h b/Engine/Plugins/MovieScene/MovieRenderPipeline/Source/MovieRenderPipelineCore/Public/Graph/Nodes/MovieGraphGetCVarValueNode.h
deleted file mode 100644
diff --git a/Engine/Plugins/MovieScene/MovieRenderPipeline/Source/MovieRenderPipelineCore/Public/Graph/Nodes/MovieGraphSelectNode.h b/Engine/Plugins/MovieScene/MovieRenderPipeline/Source/MovieRenderPipelineCore/Public/Graph/Nodes/MovieGraphSelectNode.h
--- a/Engine/Plugins/MovieScene/MovieRenderPipeline/Source/MovieRenderPipelineCore/Public/Graph/Nodes/MovieGraphSelectNode.h
+++ b/Engine/Plugins/MovieScene/MovieRenderPipeline/Source/MovieRenderPipelineCore/Public/Graph/Nodes/MovieGraphSelectNode.h
@@ -34,6 +34,7 @@ public:
 #if WITH_EDITOR
 	virtual FText GetNodeTitle(const bool bGetDescriptive = false) const override;
 	virtual FText GetMenuCategory() const override;
+	virtual FText GetKeywords() const override;
 	virtual FLinearColor GetNodeTitleColor() const override;
 	virtual FSlateIcon GetIconAndTint(FLinearColor& OutColor) const override;
 
diff --git a/Engine/Plugins/MovieScene/MovieRenderPipeline/Source/MovieRenderPipelineCore/Public/Graph/Nodes/MovieGraphSetCVarValueNode.h b/Engine/Plugins/MovieScene/MovieRenderPipeline/Source/MovieRenderPipelineCore/Public/Graph/Nodes/MovieGraphSetCVarValueNode.h
--- a/Engine/Plugins/MovieScene/MovieRenderPipeline/Source/MovieRenderPipelineCore/Public/Graph/Nodes/MovieGraphSetCVarValueNode.h
+++ b/Engine/Plugins/MovieScene/MovieRenderPipeline/Source/MovieRenderPipelineCore/Public/Graph/Nodes/MovieGraphSetCVarValueNode.h
@@ -20,6 +20,7 @@ public:
 #if WITH_EDITOR
 	virtual FText GetNodeTitle(const bool bGetDescriptive = false) const override;
 	virtual FText GetMenuCategory() const override;
+	virtual FText GetKeywords() const override;
 	virtual FLinearColor GetNodeTitleColor() const override;
 	virtual FSlateIcon GetIconAndTint(FLinearColor& OutColor) const override;
 
diff --git a/Engine/Plugins/MovieScene/MovieRenderPipeline/Source/MovieRenderPipelineEditor/Private/Graph/MovieGraphSchema.cpp b/Engine/Plugins/MovieScene/MovieRenderPipeline/Source/MovieRenderPipelineEditor/Private/Graph/MovieGraphSchema.cpp
--- a/Engine/Plugins/MovieScene/MovieRenderPipeline/Source/MovieRenderPipelineEditor/Private/Graph/MovieGraphSchema.cpp
+++ b/Engine/Plugins/MovieScene/MovieRenderPipeline/Source/MovieRenderPipelineEditor/Private/Graph/MovieGraphSchema.cpp
@@ -246,8 +246,10 @@ void UMovieGraphSchema::GetGraphContextActions(FGraphContextMenuBuilder& Context
 			const FText Name = PipelineNode->GetNodeTitle();
 			const FText Category = PipelineNode->GetMenuCategory();
 			const FText Tooltip = LOCTEXT("CreateNode_Tooltip", "Create a node of this type.");
+			constexpr int32 Grouping = 0;
+			const FText Keywords = PipelineNode->GetKeywords();
 			
-			TSharedPtr<FMovieGraphSchemaAction> NewAction = MakeShared<FMovieGraphSchemaAction_NewNode>(Category, Name, Tooltip); 
+			TSharedPtr<FMovieGraphSchemaAction> NewAction = MakeShared<FMovieGraphSchemaAction_NewNode>(Category, Name, Tooltip, Grouping, Keywords); 
 			NewAction->NodeClass = PipelineNodeClass;
 
 			ContextMenuBuilder.AddAction(NewAction);
@@ -501,8 +503,8 @@ FConnectionDrawingPolicy* UMovieGraphSchema::CreateConnectionDrawingPolicy(int32
 	return new FMovieEdGraphConnectionDrawingPolicy(InBackLayerID, InFrontLayerID, InZoomFactor, InClippingRect, InDrawElements, InGraphObj);
 }
 
-FMovieGraphSchemaAction_NewNode::FMovieGraphSchemaAction_NewNode(FText InNodeCategory, FText InDisplayName, FText InToolTip)
-	: FMovieGraphSchemaAction(MoveTemp(InNodeCategory), MoveTemp(InDisplayName), MoveTemp(InToolTip), 0)
+FMovieGraphSchemaAction_NewNode::FMovieGraphSchemaAction_NewNode(FText InNodeCategory, FText InDisplayName, FText InToolTip, int32 InGrouping, FText InKeywords)
+	: FMovieGraphSchemaAction(MoveTemp(InNodeCategory), MoveTemp(InDisplayName), MoveTemp(InToolTip), InGrouping, MoveTemp(InKeywords))
 {
 	
 }
diff --git a/Engine/Plugins/MovieScene/MovieRenderPipeline/Source/MovieRenderPipelineEditor/Private/Graph/MovieGraphSchema.h b/Engine/Plugins/MovieScene/MovieRenderPipeline/Source/MovieRenderPipelineEditor/Private/Graph/MovieGraphSchema.h
--- a/Engine/Plugins/MovieScene/MovieRenderPipeline/Source/MovieRenderPipelineEditor/Private/Graph/MovieGraphSchema.h
+++ b/Engine/Plugins/MovieScene/MovieRenderPipeline/Source/MovieRenderPipelineEditor/Private/Graph/MovieGraphSchema.h
@@ -104,7 +104,7 @@ struct FMovieGraphSchemaAction_NewNode : public FMovieGraphSchemaAction
 		: FMovieGraphSchemaAction()
 	{}
 
-	FMovieGraphSchemaAction_NewNode(FText InNodeCategory, FText InDisplayName, FText InToolTip);
+	FMovieGraphSchemaAction_NewNode(FText InNodeCategory, FText InDisplayName, FText InToolTip, int32 InGrouping, FText InKeywords);
 
 	virtual ~FMovieGraphSchemaAction_NewNode() override = default;
 	
diff --git a/Engine/Plugins/MovieScene/MovieRenderPipeline/Source/MovieRenderPipelineRenderPasses/Public/MovieGraphImageSequenceOutputNode.h b/Engine/Plugins/MovieScene/MovieRenderPipeline/Source/MovieRenderPipelineRenderPasses/Public/MovieGraphImageSequenceOutputNode.h
--- a/Engine/Plugins/MovieScene/MovieRenderPipeline/Source/MovieRenderPipelineRenderPasses/Public/MovieGraphImageSequenceOutputNode.h
+++ b/Engine/Plugins/MovieScene/MovieRenderPipeline/Source/MovieRenderPipelineRenderPasses/Public/MovieGraphImageSequenceOutputNode.h
@@ -67,6 +67,12 @@ public:
 		static const FText EXRSequenceNodeName = NSLOCTEXT("MovieGraph", "NodeName_EXRSequence", ".exr Sequence");
 		return EXRSequenceNodeName;
 	}
+
+	virtual FText GetKeywords() const override
+	{
+		static const FText Keywords = NSLOCTEXT("MovieGraph", "ImageSequenceOutputNode_EXR_Keywords", "exr image");
+		return Keywords;
+	}
 	
 	virtual FLinearColor GetNodeTitleColor() const override
 	{
@@ -148,6 +154,12 @@ public:
 		}
 		return NSLOCTEXT("MovieGraph", "ImgSequenceBMPSetting_NodeTitleShort", ".bmp Sequence");
 	}
+
+	virtual FText GetKeywords() const override
+	{
+		static const FText Keywords = NSLOCTEXT("MovieGraph", "ImageSequenceOutputNode_BMP_Keywords", "bmp image");
+		return Keywords;
+	}
 	
 	virtual FLinearColor GetNodeTitleColor() const override
 	{
@@ -190,6 +202,12 @@ public:
 		}
 		return NSLOCTEXT("MovieGraph", "ImgSequenceJPGSetting_NodeTitleShort", ".jpg Sequence");
 	}
+
+	virtual FText GetKeywords() const override
+	{
+		static const FText Keywords = NSLOCTEXT("MovieGraph", "ImageSequenceOutputNode_JPG_Keywords", "jpg jpeg image");
+		return Keywords;
+	}
 	
 	virtual FLinearColor GetNodeTitleColor() const override
 	{
@@ -231,6 +249,12 @@ public:
 		}
 		return NSLOCTEXT("MovieGraph", "ImgSequencePNGSetting_NodeTitleShort", ".png Sequence");
 	}
+
+	virtual FText GetKeywords() const override
+	{
+		static const FText Keywords = NSLOCTEXT("MovieGraph", "ImageSequenceOutputNode_PNG_Keywords", "png image");
+		return Keywords;
+	}
 	
 	virtual FLinearColor GetNodeTitleColor() const override
 	{
