MoviePipeline: Added versioning support to file output in the graph.

#rnx

diff --git a/Engine/Plugins/MovieScene/MovieRenderPipeline/Source/MovieRenderPipelineCore/Private/Graph/MovieGraphBlueprintLibrary.cpp b/Engine/Plugins/MovieScene/MovieRenderPipeline/Source/MovieRenderPipelineCore/Private/Graph/MovieGraphBlueprintLibrary.cpp
--- a/Engine/Plugins/MovieScene/MovieRenderPipeline/Source/MovieRenderPipelineCore/Private/Graph/MovieGraphBlueprintLibrary.cpp
+++ b/Engine/Plugins/MovieScene/MovieRenderPipeline/Source/MovieRenderPipelineCore/Private/Graph/MovieGraphBlueprintLibrary.cpp
@@ -5,6 +5,8 @@
 #include "Graph/MovieGraphPipeline.h"
 #include "Graph/Nodes/MovieGraphOutputSettingNode.h"
 #include "Graph/Nodes/MovieGraphRenderLayerNode.h"
+#include "HAL/FileManager.h"
+#include "Internationalization/Regex.h"
 #include "MoviePipelineBlueprintLibrary.h"
 #include "MoviePipelineUtils.h"
 
@@ -144,6 +146,9 @@ FString UMovieGraphBlueprintLibrary::ResolveFilenameFormatArguments(const FStrin
 	// but we don't actually want a value there (unless there's a collision). So we override it with an empty string by default.
 	OutMergedFormatArgs.FilenameArguments.Add(TEXT("file_dup"), FString());
 
+	// Format the {version} token so it's of the form "v###"
+	OutMergedFormatArgs.FilenameArguments.Add(TEXT("version"), FString::Printf(TEXT("v%0*d"), 3, InParams.Version));
+
 	// Overwrite any {tokens} with the user-supplied overrides if needed. This allows different requesters to share the same variables (ie: filename extension)
 	for (const TPair<FString, FString>& KVP : InParams.FileNameFormatOverrides)
 	{
@@ -204,6 +209,89 @@ FString UMovieGraphBlueprintLibrary::ResolveFilenameFormatArguments(const FStrin
 	}
 }
 
+int32 UMovieGraphBlueprintLibrary::ResolveVersionNumber(FMovieGraphFilenameResolveParams InParams, const bool bGetNextVersion)
+{
+	// Note: InParams is passed by copy rather than const& because it is modified within this method.
+	
+	if (!InParams.EvaluatedConfig)
+	{
+		FFrame::KismetExecutionMessage(TEXT("Cannot resolve version number without a valid evaluated graph to pull settings from."), ELogVerbosity::Error);
+		return -1;
+	}
+
+	constexpr bool bIncludeCDOs = true;
+	constexpr bool bExactMatch = true;
+	const UMovieGraphOutputSettingNode* OutputSettingNode = InParams.EvaluatedConfig->GetSettingForBranch<UMovieGraphOutputSettingNode>(UMovieGraphNode::GlobalsPinName, bIncludeCDOs, bExactMatch);
+	if (!OutputSettingNode->bAutoVersion)
+	{
+		return OutputSettingNode->VersionNumber;
+	}
+
+	// Calculate a version number by looking at the output path and then scanning for a version token.
+	const FString FileNameFormatString = InParams.FileNameOverride.Len() > 0
+		? InParams.FileNameOverride
+		: OutputSettingNode->OutputDirectory.Path / OutputSettingNode->FileNameFormat;
+
+	// Force the Version string to stay as {version} so we can substring based on it later.
+	InParams.FileNameFormatOverrides.Add(TEXT("version"), TEXT("{version}"));
+
+	FMovieGraphResolveArgs FinalFormatArgs;
+	FString FinalPath = ResolveFilenameFormatArguments(FileNameFormatString, InParams, FinalFormatArgs);
+	FinalPath = FPaths::ConvertRelativePathToFull(FinalPath);
+	FPaths::NormalizeFilename(FinalPath);
+
+	// Can't resolve a version if it's not clear from the path where the version number will be used.
+	if (!FinalPath.Contains(TEXT("{version}")))
+	{
+		return -1;
+	}
+	
+	int32 HighestVersion = 0;
+
+	// FinalPath can have {version} either in a folder name or in a file name. We need to find the 'parent' of either the
+	// file or folder that contains it. We can do this by looking for {version} and then finding the last "/" character,
+	// which will be the containing folder.
+	const int32 VersionStringIndex = FinalPath.Find(TEXT("{version}"), ESearchCase::Type::IgnoreCase, ESearchDir::Type::FromStart);
+	if (VersionStringIndex >= 0)
+	{
+		const int32 LastParentFolder = FinalPath.Find(TEXT("/"), ESearchCase::Type::IgnoreCase, ESearchDir::Type::FromEnd, VersionStringIndex);
+		FinalPath.LeftInline(LastParentFolder + 1);
+
+		// Now that we have the parent folder of either the folder with the version token, or the file with the version
+		// token, we will look through all immediate children and scan for version tokens so we can find the highest one.
+		const FRegexPattern VersionSearchPattern(TEXT("v([0-9]{3})"));
+		constexpr bool bFindFiles = true;
+		constexpr bool bFindDirectories = true;
+		const FString SearchString = FinalPath / TEXT("*.*");
+		TArray<FString> FoundFilesAndFoldersInDirectory;
+		IFileManager& FileManager = IFileManager::Get();
+		FileManager.FindFiles(FoundFilesAndFoldersInDirectory, *SearchString, bFindFiles, bFindDirectories);
+
+		for (const FString& Path : FoundFilesAndFoldersInDirectory)
+		{
+			FRegexMatcher Regex(VersionSearchPattern, *Path);
+			if (Regex.FindNext())
+			{
+				FString Result = Regex.GetCaptureGroup(0);
+				if (Result.Len() > 0)
+				{
+					// Strip the "v" token off, expected pattern is vXXX
+					Result.RightChopInline(1);
+				}
+
+				int32 VersionNumber = 0;
+				LexFromString(VersionNumber, *Result);
+				if (VersionNumber > HighestVersion)
+				{
+					HighestVersion = VersionNumber;
+				}
+			}
+		}
+	}
+
+	return HighestVersion + (bGetNextVersion ? 1 : 0);
+}
+
 FIntPoint UMovieGraphBlueprintLibrary::GetEffectiveOutputResolution(UMovieGraphEvaluatedConfig* InEvaluatedGraph, const FName& InBranchName)
 {
 	if (!InEvaluatedGraph)
diff --git a/Engine/Plugins/MovieScene/MovieRenderPipeline/Source/MovieRenderPipelineCore/Private/Graph/MovieGraphPipeline.cpp b/Engine/Plugins/MovieScene/MovieRenderPipeline/Source/MovieRenderPipelineCore/Private/Graph/MovieGraphPipeline.cpp
--- a/Engine/Plugins/MovieScene/MovieRenderPipeline/Source/MovieRenderPipelineCore/Private/Graph/MovieGraphPipeline.cpp
+++ b/Engine/Plugins/MovieScene/MovieRenderPipeline/Source/MovieRenderPipelineCore/Private/Graph/MovieGraphPipeline.cpp
@@ -315,7 +315,6 @@ void UMovieGraphPipeline::BuildShotListFromDataSource()
 		const FFrameRate FinalFrameRate = UMovieGraphBlueprintLibrary::GetEffectiveFrameRate(OutputNode, SourceFrameRate);
 		const FFrameRate TickResolution = GetDataSourceInstance()->GetTickResolution();
 
-
 		UMovieGraphWarmUpSettingNode* WarmUpNode = EvaluatedConfig->GetSettingForBranch<UMovieGraphWarmUpSettingNode>(UMovieGraphSettingNode::GlobalsPinName);
 		
 		const bool bPrePass = true;
@@ -326,6 +325,7 @@ void UMovieGraphPipeline::BuildShotListFromDataSource()
 		Shot->ShotInfo.NumEngineWarmUpFramesRemaining = WarmUpNode->NumWarmUpFrames;
 		Shot->ShotInfo.bEmulateFirstFrameMotionBlur = WarmUpNode->bEmulateMotionBlur;
 		Shot->ShotInfo.CalculateWorkMetrics(GetDataSourceInstance());
+		Shot->ShotInfo.VersionNumber = ResolveVersionForShot(Shot, EvaluatedConfig);
 	}
 
 	// Initialize the time step instance to be the first one. This is usually set up by the shot, but some logic may attempt
@@ -746,6 +746,54 @@ void UMovieGraphPipeline::RenderFrame()
 	GraphRendererInstance->Render(TimeStepData);
 }
 
+int32 UMovieGraphPipeline::ResolveVersionForShot(const TObjectPtr<UMoviePipelineExecutorShot>& Shot, const TObjectPtr<UMovieGraphEvaluatedConfig>& EvaluatedConfig)
+{	
+	int32 HighestVersionFound = 1;
+	
+	FMovieGraphFilenameResolveParams ResolveParams;
+	ResolveParams.EvaluatedConfig = EvaluatedConfig;
+	ResolveParams.Shot = Shot;
+	ResolveParams.Job = GetCurrentJob();
+
+	const FMovieGraphTimeStepData& ShotTimeStepInstance = GraphTimeStepInstances.Last()->GetCalculatedTimeData();
+
+	for (const FName& BranchName : EvaluatedConfig->GetBranchNames())
+	{
+		constexpr bool bIncludeCDOs = true;
+		constexpr bool bExactMatch = true;
+		const UMovieGraphOutputSettingNode* BranchOutputSettingNode = EvaluatedConfig->GetSettingForBranch<UMovieGraphOutputSettingNode>(BranchName, bIncludeCDOs, bExactMatch);
+
+		// TODO: The RenderDataIdentifier needs to be fully populated in order to resolve the version reliably. For
+		// example, {renderer_name} may be used as a directory, and discovering files correctly depends on that token
+		// being resolved. We can loop over all renderer and file output nodes here since the branches are already being
+		// iterated, and resolve a version for each (still using the highest version found as the final resolved
+		// version). Note that to do this correctly, we need a way of asking these nodes for their renderer name and all
+		// possible sub-resource names.
+		ResolveParams.InitializationTime = GetInitializationTime();
+		ResolveParams.DefaultFrameRate = ShotTimeStepInstance.FrameRate;
+		ResolveParams.FrameNumberOffset = BranchOutputSettingNode->FrameNumberOffset;
+		ResolveParams.RenderDataIdentifier.RootBranchName = BranchName;
+		ResolveParams.RenderDataIdentifier.CameraName = Shot->InnerName;
+		// ResolveParams.RenderDataIdentifier.RendererName = ?
+		// ResolveParams.RenderDataIdentifier.SubResourceName = ?
+		ResolveParams.RootFrameNumber = ShotTimeStepInstance.RootFrameNumber.Value;
+		ResolveParams.ShotFrameNumber = ShotTimeStepInstance.ShotFrameNumber.Value;
+		ResolveParams.bForceRelativeFrameNumbers = false;	// TODO: This should not be hardcoded
+		// ResolveParams.FileNameFormatOverrides = ?
+		ResolveParams.RootFrameNumberRel = ShotTimeStepInstance.OutputFrameNumber;
+		// ResolveParams.ShotFrameNumberRel = ?
+		ResolveParams.ZeroPadFrameNumberCount = BranchOutputSettingNode->ZeroPadFrameNumbers;
+
+		const int32 BranchVersion = UMovieGraphBlueprintLibrary::ResolveVersionNumber(ResolveParams);
+		if (BranchVersion > HighestVersionFound)
+		{
+			HighestVersionFound = BranchVersion;
+		}
+	}
+	
+	return HighestVersionFound;
+}
+
 void UMovieGraphPipeline::RequestShutdownImpl(bool bIsError)
 {
 	// It's possible for a previous call to RequestionShutdown to have set an error before this call that may not
diff --git a/Engine/Plugins/MovieScene/MovieRenderPipeline/Source/MovieRenderPipelineCore/Private/Graph/Nodes/MovieGraphAudioOutputNode.cpp b/Engine/Plugins/MovieScene/MovieRenderPipeline/Source/MovieRenderPipelineCore/Private/Graph/Nodes/MovieGraphAudioOutputNode.cpp
--- a/Engine/Plugins/MovieScene/MovieRenderPipeline/Source/MovieRenderPipelineCore/Private/Graph/Nodes/MovieGraphAudioOutputNode.cpp
+++ b/Engine/Plugins/MovieScene/MovieRenderPipeline/Source/MovieRenderPipelineCore/Private/Graph/Nodes/MovieGraphAudioOutputNode.cpp
@@ -146,8 +146,7 @@ FString UMovieGraphAudioOutputNode::GenerateOutputPath(const FMovieGraphRenderDa
 	};
 	ResolveParams.EvaluatedConfig = EvaluatedConfig;
 	ResolveParams.RenderDataIdentifier = InRenderIdentifier;
-	// int32 VersionNumber = 0; // TODO: Track versions
-	// FileNameFormatString.ReplaceInline(TEXT("{version}"), *FString::Printf(TEXT("v%0*d"), 3, VersionNumber));
+	ResolveParams.Version = InShot->ShotInfo.VersionNumber;
 
 	// Generate a filename for this output file
 	FMovieGraphResolveArgs FinalFormatArgs;
diff --git a/Engine/Plugins/MovieScene/MovieRenderPipeline/Source/MovieRenderPipelineCore/Private/Graph/Nodes/MovieGraphCommandLineEncoderNode.cpp b/Engine/Plugins/MovieScene/MovieRenderPipeline/Source/MovieRenderPipelineCore/Private/Graph/Nodes/MovieGraphCommandLineEncoderNode.cpp
--- a/Engine/Plugins/MovieScene/MovieRenderPipeline/Source/MovieRenderPipelineCore/Private/Graph/Nodes/MovieGraphCommandLineEncoderNode.cpp
+++ b/Engine/Plugins/MovieScene/MovieRenderPipeline/Source/MovieRenderPipelineCore/Private/Graph/Nodes/MovieGraphCommandLineEncoderNode.cpp
@@ -357,12 +357,11 @@ FString UMovieGraphCommandLineEncoderNode::GetResolvedOutputFilename(const FMovi
 	FMovieGraphFilenameResolveParams ResolveParams;
 	ResolveParams.InitializationTime = CachedPipeline->GetInitializationTime();
 	ResolveParams.Job = CachedPipeline->GetCurrentJob();
-	ResolveParams.Shot = Shot.IsValid() ? Shot.Get() : nullptr;
+	ResolveParams.Shot = Shot.Get();
 	ResolveParams.FileNameFormatOverrides = FormatOverrides;
 	ResolveParams.EvaluatedConfig = CachedPipeline->GetTimeStepInstance()->GetCalculatedTimeData().EvaluatedConfig;
 	ResolveParams.RenderDataIdentifier = RenderIdentifier;
-	// int32 VersionNumber = 0; // TODO: Track versions
-	// FileNameFormatString.ReplaceInline(TEXT("{version}"), *FString::Printf(TEXT("v%0*d"), 3, VersionNumber));
+	ResolveParams.Version = Shot.IsValid() ? Shot->ShotInfo.VersionNumber : UMovieGraphBlueprintLibrary::ResolveVersionNumber(ResolveParams);
 
 	FMovieGraphResolveArgs FinalFormatArgs;
 	FString FinalFilePath = UMovieGraphBlueprintLibrary::ResolveFilenameFormatArguments(FileNameFormatString, ResolveParams, FinalFormatArgs);
diff --git a/Engine/Plugins/MovieScene/MovieRenderPipeline/Source/MovieRenderPipelineCore/Private/Graph/Nodes/MovieGraphOutputSettingNode.cpp b/Engine/Plugins/MovieScene/MovieRenderPipeline/Source/MovieRenderPipelineCore/Private/Graph/Nodes/MovieGraphOutputSettingNode.cpp
--- a/Engine/Plugins/MovieScene/MovieRenderPipeline/Source/MovieRenderPipelineCore/Private/Graph/Nodes/MovieGraphOutputSettingNode.cpp
+++ b/Engine/Plugins/MovieScene/MovieRenderPipeline/Source/MovieRenderPipelineCore/Private/Graph/Nodes/MovieGraphOutputSettingNode.cpp
@@ -9,8 +9,9 @@ UMovieGraphOutputSettingNode::UMovieGraphOutputSettingNode()
 	, bOverwriteExistingOutput(true)
 	, ZeroPadFrameNumbers(4)
 	, FrameNumberOffset(0)
+	, bAutoVersion(true)
+	, VersionNumber(1)
 {
-
 	FileNameFormat = TEXT("{sequence_name}.{frame_number}");
 	OutputDirectory.Path = TEXT("{project_dir}/Saved/MovieRenders/");
 }
diff --git a/Engine/Plugins/MovieScene/MovieRenderPipeline/Source/MovieRenderPipelineCore/Public/Graph/MovieGraphBlueprintLibrary.h b/Engine/Plugins/MovieScene/MovieRenderPipeline/Source/MovieRenderPipelineCore/Public/Graph/MovieGraphBlueprintLibrary.h
--- a/Engine/Plugins/MovieScene/MovieRenderPipeline/Source/MovieRenderPipelineCore/Public/Graph/MovieGraphBlueprintLibrary.h
+++ b/Engine/Plugins/MovieScene/MovieRenderPipeline/Source/MovieRenderPipelineCore/Public/Graph/MovieGraphBlueprintLibrary.h
@@ -43,6 +43,14 @@ public:
 	UFUNCTION(BlueprintCallable, Category = "Movie Graph")
 	static FString ResolveFilenameFormatArguments(const FString& InFormatString, const FMovieGraphFilenameResolveParams& InParams, FMovieGraphResolveArgs& OutMergedFormatArgs);
 
+	/**
+	 * If the version number is explicitly specified on the Output Setting node, returns that. Otherwise searches the
+	 * output directory for the highest version that already exists (and increments it by one if bGetNextVersion is
+	 * true). Returns -1 if the version could not be resolved.
+	 */
+	UFUNCTION(BlueprintCallable, Category = "Movie Graph")
+	static int32 ResolveVersionNumber(FMovieGraphFilenameResolveParams InParams, const bool bGetNextVersion = true);
+
 	/**
 	* In case of overscan percentage being higher than 0, additional pixels are rendered. This function returns the resolution with overscan taken into account.
 	* @param	InEvaluatedGraph	- The evaluated graph that will provide context for resolving the resolution
diff --git a/Engine/Plugins/MovieScene/MovieRenderPipeline/Source/MovieRenderPipelineCore/Public/Graph/MovieGraphFilenameResolveParams.h b/Engine/Plugins/MovieScene/MovieRenderPipeline/Source/MovieRenderPipelineCore/Public/Graph/MovieGraphFilenameResolveParams.h
--- a/Engine/Plugins/MovieScene/MovieRenderPipeline/Source/MovieRenderPipelineCore/Public/Graph/MovieGraphFilenameResolveParams.h
+++ b/Engine/Plugins/MovieScene/MovieRenderPipeline/Source/MovieRenderPipelineCore/Public/Graph/MovieGraphFilenameResolveParams.h
@@ -110,6 +110,10 @@ struct FMovieGraphFilenameResolveParams
 	UPROPERTY(BlueprintReadWrite, EditAnywhere, Category = "Movie Graph")
 	bool bForceRelativeFrameNumbers;
 
+	/** If specified, this is the filename that will be used instead of building it from the Output Setting node. */
+	UPROPERTY(BlueprintReadWrite, EditAnywhere, Category = "Movie Graph")
+	FString FileNameOverride;
+
 	/** 
 	* If true, If the format string is a relative path, then the resulting resolved path will be converted to an absolute path.
 	* Assumes that the relative path should be resolved relative to the engine/editor executable.
diff --git a/Engine/Plugins/MovieScene/MovieRenderPipeline/Source/MovieRenderPipelineCore/Public/Graph/MovieGraphPipeline.h b/Engine/Plugins/MovieScene/MovieRenderPipeline/Source/MovieRenderPipelineCore/Public/Graph/MovieGraphPipeline.h
--- a/Engine/Plugins/MovieScene/MovieRenderPipeline/Source/MovieRenderPipelineCore/Public/Graph/MovieGraphPipeline.h
+++ b/Engine/Plugins/MovieScene/MovieRenderPipeline/Source/MovieRenderPipelineCore/Public/Graph/MovieGraphPipeline.h
@@ -150,6 +150,9 @@ protected:
 	template<typename T>
 	TArray<TPair<FName, T*>> GetSettingForActiveRenderLayers(const bool bIncludeCDOs, const bool bExactMatch);
 
+	/** Resolve the version number that should be used for the specified shot (in {version} tokens). */
+	int32 ResolveVersionForShot(const TObjectPtr<UMoviePipelineExecutorShot>& Shot, const TObjectPtr<UMovieGraphEvaluatedConfig>& EvaluatedConfig);
+
 	// UMoviePipelineBase Interface
 	virtual void RequestShutdownImpl(bool bIsError) override;
 	virtual void ShutdownImpl(bool bIsError) override;
diff --git a/Engine/Plugins/MovieScene/MovieRenderPipeline/Source/MovieRenderPipelineCore/Public/Graph/Nodes/MovieGraphOutputSettingNode.h b/Engine/Plugins/MovieScene/MovieRenderPipeline/Source/MovieRenderPipelineCore/Public/Graph/Nodes/MovieGraphOutputSettingNode.h
--- a/Engine/Plugins/MovieScene/MovieRenderPipeline/Source/MovieRenderPipelineCore/Public/Graph/Nodes/MovieGraphOutputSettingNode.h
+++ b/Engine/Plugins/MovieScene/MovieRenderPipeline/Source/MovieRenderPipelineCore/Public/Graph/Nodes/MovieGraphOutputSettingNode.h
@@ -50,6 +50,12 @@ public:
 	UPROPERTY(EditAnywhere, BlueprintReadWrite, Category = Overrides, meta = (InlineEditConditionToggle))
 	uint8 bOverride_HandleFrameCount : 1;
 
+	UPROPERTY(EditAnywhere, BlueprintReadWrite, Category = Overrides, meta = (InlineEditConditionToggle))
+	uint8 bOverride_bAutoVersion : 1;
+
+	UPROPERTY(EditAnywhere, BlueprintReadWrite, Category = Overrides, meta = (InlineEditConditionToggle))
+	uint8 bOverride_VersionNumber : 1;
+
 #if WITH_EDITORONLY_DATA
 	UPROPERTY(EditAnywhere, BlueprintReadWrite, Category = Overrides, meta = (InlineEditConditionToggle))
 	uint8 bOverride_OCIOConfiguration : 1;
@@ -95,6 +101,22 @@ public:
 	UPROPERTY(EditAnywhere, BlueprintReadWrite, Category = "Frames", meta = (UIMin = 0, ClampMin = 0, EditCondition = "bOverride_HandleFrameCount"))
 	int32 HandleFrameCount;
 
+	// TODO: This versioning property should be on the "global" node
+	/**
+	 * If true, {version} tokens specified in the Output Directory and File Name Format properties will automatically
+	 * be incremented with each local render. If false, the version specified in Version Number will be used instead.
+	 *
+	 * Auto-versioning will search across all render branches and use the highest version found as the basis for the
+	 * next version used.
+	 */
+	UPROPERTY(EditAnywhere, BlueprintReadWrite, Category = "Versioning", meta = (EditCondition = "bOverride_bAutoVersion"))
+	bool bAutoVersion;
+	
+	// TODO: This versioning property should be on the "global" node
+	/** The value to use for the version token if versions are not automatically incremented (Auto Version is off). */
+	UPROPERTY(EditAnywhere, BlueprintReadWrite, Category = "Versioning", meta = (EditCondition = "bOverride_VersionNumber", UIMin = 1, UIMax = 50, ClampMin = 1))
+	int32 VersionNumber;
+
 #if WITH_EDITORONLY_DATA
 	/**
 	* OCIO configuration/transform settings.
diff --git a/Engine/Plugins/MovieScene/MovieRenderPipeline/Source/MovieRenderPipelineCore/Public/MovieRenderPipelineDataTypes.h b/Engine/Plugins/MovieScene/MovieRenderPipeline/Source/MovieRenderPipelineCore/Public/MovieRenderPipelineDataTypes.h
--- a/Engine/Plugins/MovieScene/MovieRenderPipeline/Source/MovieRenderPipelineCore/Public/MovieRenderPipelineDataTypes.h
+++ b/Engine/Plugins/MovieScene/MovieRenderPipeline/Source/MovieRenderPipelineCore/Public/MovieRenderPipelineDataTypes.h
@@ -525,7 +525,7 @@ public:
 	
 	/** Cached Tick Resolution the movie scene this range was generated for is in. Can be different than the root due to mixed tick resolutions. */
 	FFrameRate CachedShotTickResolution;
-public:
+
 	/** The current state of processing this Shot is in. Not all states will be passed through. */
 	EMovieRenderShotState State;
 
@@ -547,14 +547,13 @@ public:
 	/** Metrics - How much work has been done for this particular shot and how much do we estimate we have to do? */
 	FMoviePipelineSegmentWorkMetrics WorkMetrics;
 
-public:
 	/** Have we evaluated the motion blur frame? Only used if bEvaluateMotionBlurOnFirstFrame is set */
 	bool bHasEvaluatedMotionBlurFrame;
 
 	/** How many engine warm up frames are left to process for this shot. May be zero. */
 	int32 NumEngineWarmUpFramesRemaining;
 
-	/** What version number should this shot use when resolving format arguments. */
+	/** What version number should this shot use when resolving format arguments. This is the highest version number found across all branches. */
 	int32 VersionNumber;
 };
 
diff --git a/Engine/Plugins/MovieScene/MovieRenderPipeline/Source/MovieRenderPipelineRenderPasses/Private/MovieGraphImageSequenceOutputNode.cpp b/Engine/Plugins/MovieScene/MovieRenderPipeline/Source/MovieRenderPipelineRenderPasses/Private/MovieGraphImageSequenceOutputNode.cpp
--- a/Engine/Plugins/MovieScene/MovieRenderPipeline/Source/MovieRenderPipelineRenderPasses/Private/MovieGraphImageSequenceOutputNode.cpp
+++ b/Engine/Plugins/MovieScene/MovieRenderPipeline/Source/MovieRenderPipelineRenderPasses/Private/MovieGraphImageSequenceOutputNode.cpp
@@ -209,9 +209,10 @@ void UMovieGraphImageSequenceOutputNode::OnReceiveImageDataImpl(UMovieGraphPipel
 		Params.RootFrameNumberRel = Payload->TraversalContext.Time.OutputFrameNumber;
 		//Params.ShotFrameNumberRel = Payload->TraversalCOntext.Time.ShotFrameNumberRel
 		//Params.FileMetadata = ToDo: Track File Metadata
-		Params.Version = 1; // ToDo: Track versions
 		Params.ZeroPadFrameNumberCount = OutputSettingNode->ZeroPadFrameNumbers;
 		Params.FrameNumberOffset = OutputSettingNode->FrameNumberOffset;
+		Params.EvaluatedConfig = InRawFrameData->EvaluatedConfig.Get();
+		Params.Version = Shot->ShotInfo.VersionNumber;
 
 		// If time dilation is in effect, RootFrameNumber and ShotFrameNumber will contain duplicates and the files will overwrite each other, 
 		// so we force them into relative mode and then warn users we did that (as their numbers will jump from say 1001 -> 0000).
@@ -544,6 +545,7 @@ FString UMovieGraphImageSequenceOutputNode_EXR::ResolveOutputFilename(
 	Params.Job = InPipeline->GetCurrentJob();
 	Params.Shot = InPipeline->GetActiveShotList()[InRawFrameData->TraversalContext.ShotIndex];
 	Params.EvaluatedConfig = InRawFrameData->EvaluatedConfig.Get();
+	Params.Version = Params.Shot->ShotInfo.VersionNumber;
 	
 	FString FinalFilePath = UMovieGraphBlueprintLibrary::ResolveFilenameFormatArguments(FilePathFormatString, Params, OutResolveArgs);
 
