MoviePipeline: Updated the Game Mode Overrides node to be functional, and updated its property set based on recent discussions (removing shadow-related settings, making scalability configurable, clean-up).

Also updated the default graph asset to have sensible defaults for Game Mode Overrides.

#rnx

diff --git a/Engine/Plugins/MovieScene/MovieRenderPipeline/Source/MovieRenderPipelineCore/Private/Graph/MovieGraphDefaultRenderer.cpp b/Engine/Plugins/MovieScene/MovieRenderPipeline/Source/MovieRenderPipelineCore/Private/Graph/MovieGraphDefaultRenderer.cpp
--- a/Engine/Plugins/MovieScene/MovieRenderPipeline/Source/MovieRenderPipelineCore/Private/Graph/MovieGraphDefaultRenderer.cpp
+++ b/Engine/Plugins/MovieScene/MovieRenderPipeline/Source/MovieRenderPipelineCore/Private/Graph/MovieGraphDefaultRenderer.cpp
@@ -3,6 +3,7 @@
 #include "Graph/MovieGraphDefaultRenderer.h"
 #include "Graph/MovieGraphPipeline.h"
 #include "Graph/MovieGraphConfig.h"
+#include "Graph/Nodes/MovieGraphGlobalGameOverrides.h"
 #include "Graph/Nodes/MovieGraphRenderLayerNode.h"
 #include "Graph/Nodes/MovieGraphRenderPassNode.h"
 #include "MovieRenderPipelineCoreModule.h"
@@ -11,6 +12,14 @@
 #include "UObject/Package.h"
 #include "MoviePipelineSurfaceReader.h"
 
+// For flushing async systems
+#include "AssetCompilingManager.h"
+#include "ContentStreaming.h"
+#include "EngineModule.h"
+#include "LandscapeSubsystem.h"
+#include "Materials/MaterialInterface.h"
+#include "RendererInterface.h"
+
 void UMovieGraphDefaultRenderer::SetupRenderingPipelineForShot(UMoviePipelineExecutorShot* InShot)
 {
 	// Iterate through the graph config and look for Render Layers.
@@ -165,6 +174,10 @@ void UMovieGraphDefaultRenderer::AddReferencedObjects(UObject* InThis, FReferenc
 
 void UMovieGraphDefaultRenderer::Render(const FMovieGraphTimeStepData& InTimeStepData)
 {
+	// Flush built-in systems before we render anything. This maximizes the likelihood that the data is prepared for when
+	// the render thread uses it.
+	FlushAsyncEngineSystems(InTimeStepData.EvaluatedConfig);
+
 	// Housekeeping: Clean up any tasks that were completed since the last frame. This lets us have a 
 	// better idea of how much work we're concurrently doing. 
 	{
@@ -226,6 +239,90 @@ void UMovieGraphDefaultRenderer::Render(const FMovieGraphTimeStepData& InTimeSte
 	GetOwningGraph()->SetPreviewWidgetVisible(true);
 }
 
+void UMovieGraphDefaultRenderer::FlushAsyncEngineSystems(const TObjectPtr<UMovieGraphEvaluatedConfig>& InConfig) const
+{
+	const UMovieGraphGlobalGameOverridesNode* GameOverrides = InConfig->GetSettingForBranch<UMovieGraphGlobalGameOverridesNode>(UMovieGraphNode::GlobalsPinName);
+	if (!GameOverrides)
+	{
+		return;
+	}
+
+	// Block until level streaming is completed, we do this at the end of the frame
+	// so that level streaming requests made by Sequencer level visibility tracks are
+	// accounted for.
+	if (GameOverrides->bFlushLevelStreaming && GetOwningGraph()->GetWorld())
+	{
+		GetOwningGraph()->GetWorld()->BlockTillLevelStreamingCompleted();
+	}
+
+	// Flush all assets still being compiled asynchronously.
+	// A progress bar is already in place so the user can get feedback while waiting for everything to settle.
+	if (GameOverrides->bFlushAssetCompiler)
+	{
+		FAssetCompilingManager::Get().FinishAllCompilation();
+	}
+
+	// Ensure we have complete shader maps for all materials used by primitives in the world.
+	// This way we will never render with the default material.
+	if (GameOverrides->bFlushShaderCompiler)
+	{
+		UMaterialInterface::SubmitRemainingJobsForWorld(GetOwningGraph()->GetWorld());
+	}
+
+	// Flush virtual texture tile calculations.
+	// In its own scope just to minimize the duration FSyncScope has a lock.
+	{
+		UE::RenderCommandPipe::FSyncScope SyncScope;
+		
+		ERHIFeatureLevel::Type FeatureLevel = GetWorld()->GetFeatureLevel();
+		ENQUEUE_RENDER_COMMAND(VirtualTextureSystemFlushCommand)(
+			[FeatureLevel](FRHICommandListImmediate& RHICmdList)
+			{
+				GetRendererModule().LoadPendingVirtualTextureTiles(RHICmdList, FeatureLevel);
+			});
+	}
+
+	// Flush any outstanding work waiting in Streaming Manager implementations (texture streaming, nanite, etc.)
+	// Note: This isn't a magic fix for gpu-based feedback systems, if the work hasn't made it to the streaming
+	// manager, it can't flush it. This just ensures that work that has been requested is done before we render.
+	if (GameOverrides->bFlushStreamingManagers)
+	{
+		FStreamingManagerCollection& StreamingManagers = IStreamingManager::Get();
+		constexpr bool bProcessEverything = true;
+		StreamingManagers.UpdateResourceStreaming(GetOwningGraph()->GetWorld()->GetDeltaSeconds(), bProcessEverything);
+		StreamingManagers.BlockTillAllRequestsFinished();
+	}
+
+	// If there are async tasks to build more grass, wait for them to finish so there aren't missing patches
+	// of grass. If you have way too dense grass this option can cause you to OOM.
+	if (GameOverrides->bFlushGrassStreaming)
+	{
+		if (ULandscapeSubsystem* LandscapeSubsystem = GetWorld()->GetSubsystem<ULandscapeSubsystem>())
+		{
+			constexpr bool bFlushGrass = false; // Flush means a different thing to grass system
+			constexpr bool bInForceSync = true;
+
+			TArray<FVector> CameraLocations;
+			GetCameraLocationsForFrame(CameraLocations);
+
+			LandscapeSubsystem->RegenerateGrass(bFlushGrass, bInForceSync, MakeArrayView(CameraLocations));
+		}
+	}
+}
+
+void UMovieGraphDefaultRenderer::GetCameraLocationsForFrame(TArray<FVector>& OutLocations) const
+{
+	// ToDo: Multi-camera support
+	if (const APlayerController* LocalPlayerController = GetOwningGraph()->GetWorld()->GetFirstPlayerController())
+	{
+		FVector PrimaryCameraLoc;
+		FRotator PrimaryCameraRot;
+
+		LocalPlayerController->GetPlayerViewPoint(PrimaryCameraLoc, PrimaryCameraRot);
+		OutLocations.Add(PrimaryCameraLoc);
+	}
+}
+
 void UMovieGraphDefaultRenderer::AddOutstandingRenderTask_AnyThread(UE::Tasks::FTask InTask)
 {
 	// We might be looping through the array to remove previously completed tasks,
diff --git a/Engine/Plugins/MovieScene/MovieRenderPipeline/Source/MovieRenderPipelineCore/Private/Graph/MovieGraphPipeline.cpp b/Engine/Plugins/MovieScene/MovieRenderPipeline/Source/MovieRenderPipelineCore/Private/Graph/MovieGraphPipeline.cpp
--- a/Engine/Plugins/MovieScene/MovieRenderPipeline/Source/MovieRenderPipelineCore/Private/Graph/MovieGraphPipeline.cpp
+++ b/Engine/Plugins/MovieScene/MovieRenderPipeline/Source/MovieRenderPipelineCore/Private/Graph/MovieGraphPipeline.cpp
@@ -9,6 +9,7 @@
 #include "Graph/MoviePipelineRenderLayerSubsystem.h"
 #include "Graph/Nodes/MovieGraphCollectionNode.h"
 #include "Graph/Nodes/MovieGraphFileOutputNode.h"
+#include "Graph/Nodes/MovieGraphGlobalGameOverrides.h"
 #include "Graph/Nodes/MovieGraphModifierNode.h"
 #include "Graph/Nodes/MovieGraphRenderLayerNode.h"
 #include "Graph/Nodes/MovieGraphSamplingMethodNode.h"
@@ -556,6 +557,16 @@ void UMovieGraphPipeline::SetupShot(const TObjectPtr<UMoviePipelineExecutorShot>
 	const FMovieGraphTimeStepData& TimeStepData = GetTimeStepInstance()->GetCalculatedTimeData();
 	const UMovieGraphEvaluatedConfig* EvaluatedConfig = TimeStepData.EvaluatedConfig;
 
+	// Apply any global game overrides, which includes cvars. This needs to be done before the CVarManager sets cvars
+	// so any user-specified cvars can override cvars set via the global game overrides.
+	constexpr bool bIncludeCDOs = true;
+	constexpr bool bExactMatch = true;
+	if (UMovieGraphGlobalGameOverridesNode* GlobalGameOverridesNode = EvaluatedConfig->GetSettingForBranch<UMovieGraphGlobalGameOverridesNode>(UMovieGraphNode::GlobalsPinName, bIncludeCDOs, bExactMatch))
+	{
+		constexpr bool bOverrideValues = true;
+		GlobalGameOverridesNode->ApplySettings(bOverrideValues, GetWorld());
+	}
+
 	// Apply cvars for the shot
 	CVarManager->AddEvaluatedGraph(EvaluatedConfig);
 	CVarManager->ApplyAllCVars();
@@ -581,9 +592,22 @@ void UMovieGraphPipeline::TeardownShot(const TObjectPtr<UMoviePipelineExecutorSh
 
 	// some other stuff
 
+	const FMovieGraphTimeStepData& TimeStepData = GetTimeStepInstance()->GetCalculatedTimeData();
+	const UMovieGraphEvaluatedConfig* EvaluatedConfig = TimeStepData.EvaluatedConfig;
+
 	// Revert the cvar values that were initially applied for the shot
 	CVarManager->RevertAllCVars();
 
+	// Revert cvars set by the global game overrides. Needs to be done after the CVarManager reverts (since the global
+	// game overrides are applied first in SetupShot).
+	constexpr bool bIncludeCDOs = true;
+	constexpr bool bExactMatch = true;
+	if (UMovieGraphGlobalGameOverridesNode* GlobalGameOverridesNode = EvaluatedConfig->GetSettingForBranch<UMovieGraphGlobalGameOverridesNode>(UMovieGraphNode::GlobalsPinName, bIncludeCDOs, bExactMatch))
+	{
+		constexpr bool bOverrideValues = false;
+		GlobalGameOverridesNode->ApplySettings(bOverrideValues, GetWorld());
+	}
+
 	// Check to see if this was the last shot in the Pipeline, otherwise on the next
 	// tick the new shot will be initialized and processed.
 	if (CurrentShotIndex >= (ActiveShotList.Num() - 1))
diff --git a/Engine/Plugins/MovieScene/MovieRenderPipeline/Source/MovieRenderPipelineCore/Private/Graph/Nodes/MovieGraphCoreRenderPassNode.cpp b/Engine/Plugins/MovieScene/MovieRenderPipeline/Source/MovieRenderPipelineCore/Private/Graph/Nodes/MovieGraphCoreRenderPassNode.cpp
--- a/Engine/Plugins/MovieScene/MovieRenderPipeline/Source/MovieRenderPipelineCore/Private/Graph/Nodes/MovieGraphCoreRenderPassNode.cpp
+++ b/Engine/Plugins/MovieScene/MovieRenderPipeline/Source/MovieRenderPipelineCore/Private/Graph/Nodes/MovieGraphCoreRenderPassNode.cpp
@@ -2,11 +2,12 @@
 
 #include "Graph/Nodes/MovieGraphCoreRenderPassNode.h"
 
-#include "Graph/Nodes/MovieGraphOutputSettingNode.h"
 #include "Graph/MovieGraphDataTypes.h"
 #include "Graph/MovieGraphDefaultRenderer.h"
 #include "Graph/MovieGraphPipeline.h"
 #include "Graph/MoviePipelineRenderLayerSubsystem.h"
+#include "Graph/Nodes/MovieGraphGlobalGameOverrides.h"
+#include "Graph/Nodes/MovieGraphOutputSettingNode.h"
 #include "MovieRenderOverlappedImage.h"
 #include "MoviePipelineSurfaceReader.h"
 #include "EngineModule.h"
@@ -480,6 +481,16 @@ TSharedRef<FSceneViewFamilyContext> UMovieGraphCoreRenderPassNode::FMovieGraphRe
 		.SetTime(FGameTime::CreateUndilated(InInitData.TimeData.WorldSeconds, InInitData.TimeData.FrameDeltaTime))
 		.SetRealtimeUpdate(true));
 
+	// Set the virtual texture feedback factor if explicitly overridden
+	constexpr bool bIncludeCDOs = true;
+	constexpr bool bExactMatch = true;
+	const UMovieGraphGlobalGameOverridesNode* GameOverridesNode =
+		InInitData.TimeData.EvaluatedConfig->GetSettingForBranch<UMovieGraphGlobalGameOverridesNode>(GlobalsPinName, bIncludeCDOs, bExactMatch);
+	if (GameOverridesNode->bOverride_VirtualTextureFeedbackFactor)
+	{
+		OutViewFamily->VirtualTextureFeedbackFactor = GameOverridesNode->VirtualTextureFeedbackFactor;
+	}
+
 	// Used to specify if the Tone Curve is being applied or not to our Linear Output data
 	OutViewFamily->SceneCaptureSource = InInitData.SceneCaptureSource;
 	OutViewFamily->bWorldIsPaused = InInitData.bWorldIsPaused;
diff --git a/Engine/Plugins/MovieScene/MovieRenderPipeline/Source/MovieRenderPipelineCore/Private/Graph/Nodes/MovieGraphGlobalGameOverrides.cpp b/Engine/Plugins/MovieScene/MovieRenderPipeline/Source/MovieRenderPipelineCore/Private/Graph/Nodes/MovieGraphGlobalGameOverrides.cpp
--- a/Engine/Plugins/MovieScene/MovieRenderPipeline/Source/MovieRenderPipelineCore/Private/Graph/Nodes/MovieGraphGlobalGameOverrides.cpp
+++ b/Engine/Plugins/MovieScene/MovieRenderPipeline/Source/MovieRenderPipelineCore/Private/Graph/Nodes/MovieGraphGlobalGameOverrides.cpp
@@ -2,8 +2,158 @@
 
 #include "Graph/Nodes/MovieGraphGlobalGameOverrides.h"
 
+#include "Graph/MovieGraphPipeline.h"
+#include "MoviePipelineUtils.h"
 #include "Styling/AppStyle.h"
 
+UMovieGraphGlobalGameOverridesNode::UMovieGraphGlobalGameOverridesNode()
+	: GameModeOverride(AMoviePipelineGameMode::StaticClass())
+	, ScalabilityQualityLevel(EMovieGraphScalabilityQualityLevel::Cinematic)
+	, bDisableTextureStreaming(false)
+	, bDisableLODs(false)
+	, bDisableHLODs(false)
+	, bFlushLevelStreaming(false)
+	, bFlushAssetCompiler(false)
+	, bFlushShaderCompiler(false)
+	, bFlushGrassStreaming(false)
+	, bFlushStreamingManagers(false)
+	, VirtualTextureFeedbackFactor(1)
+{
+}
+
+void UMovieGraphGlobalGameOverridesNode::ApplySettings(const bool bOverrideValues, UWorld* InWorld)
+{
+	// Apply the scalability settings
+	if (bOverrideValues)
+	{
+		// Store the current scalability settings so we can revert back to them
+		PreviousQualityLevels = Scalability::GetQualityLevels();
+
+		// Set to the chosen scalability quality level
+		Scalability::FQualityLevels QualityLevels;
+		QualityLevels.SetFromSingleQualityLevel(static_cast<int32>(ScalabilityQualityLevel));
+
+		// Apply
+		Scalability::SetQualityLevels(QualityLevels);
+	}
+	else
+	{
+		// We re-apply old scalability settings at the end of the function during teardown
+		// so that any values that are also specified in Scalability don't get overwritten
+		// with the wrong values from the ones below restoring.
+	}
+
+	if (bDisableTextureStreaming)
+	{
+		MOVIEPIPELINE_STORE_AND_OVERRIDE_CVAR_INT(PreviousTextureStreaming, TEXT("r.TextureStreaming"), 0, bOverrideValues);
+	}
+	
+	if (bDisableLODs)
+	{
+		MOVIEPIPELINE_STORE_AND_OVERRIDE_CVAR_INT(PreviousForceLOD, TEXT("r.ForceLOD"), 0, bOverrideValues);
+		MOVIEPIPELINE_STORE_AND_OVERRIDE_CVAR_INT(PreviousSkeletalMeshBias, TEXT("r.SkeletalMeshLODBias"), -10, bOverrideValues);
+		MOVIEPIPELINE_STORE_AND_OVERRIDE_CVAR_INT(PreviousParticleLODBias, TEXT("r.ParticleLODBias"), -10, bOverrideValues);
+		MOVIEPIPELINE_STORE_AND_OVERRIDE_CVAR_INT(PreviousFoliageDitheredLOD, TEXT("foliage.DitheredLOD"), 0, bOverrideValues);
+		MOVIEPIPELINE_STORE_AND_OVERRIDE_CVAR_INT(PreviousFoliageForceLOD, TEXT("foliage.ForceLOD"), 0, bOverrideValues);
+	}
+
+	if (bDisableHLODs)
+	{
+		// It's a command and not an integer cvar (despite taking 1/0), so we can't cache it 
+		if (GEngine && InWorld)
+		{
+			GEngine->Exec(InWorld, TEXT("r.HLOD 0"));
+		}
+	}
+
+	if (bFlushStreamingManagers)
+	{
+		MOVIEPIPELINE_STORE_AND_OVERRIDE_CVAR_INT(PreviousStreamingManagerSyncState, TEXT("r.Streaming.SyncStatesWhenBlocking"), 1, bOverrideValues);
+	}
+
+	// Set cvars that allow MRQ to run correctly. These are not exposed as properties on the node, but the chance of a
+	// user needing to customize these (or be aware of them) is very low.
+	{
+#if WITH_EDITOR
+		// To make sure the GeometryCache streamer doesn't skip frames and doesn't pop up notification during rendering
+		MOVIEPIPELINE_STORE_AND_OVERRIDE_CVAR_INT_IF_EXIST(PreviousGeoCacheStreamerBlockTillFinish, TEXT("GeometryCache.Streamer.BlockTillFinishStreaming"), 1, bOverrideValues);
+		MOVIEPIPELINE_STORE_AND_OVERRIDE_CVAR_INT_IF_EXIST(PreviousGeoCacheStreamerShowNotification, TEXT("GeometryCache.Streamer.ShowNotification"), 0, bOverrideValues);
+#endif
+
+		// Disable systems that try to preserve performance in runtime games.
+		MOVIEPIPELINE_STORE_AND_OVERRIDE_CVAR_INT(PreviousAnimationUROEnabled, TEXT("a.URO.Enable"), 0, bOverrideValues);
+
+		// To make sure that the skylight is always valid and consistent across capture sessions, we enforce a full capture each frame, accepting a small GPU cost.
+		MOVIEPIPELINE_STORE_AND_OVERRIDE_CVAR_INT(PreviousSkyLightRealTimeReflectionCaptureTimeSlice, TEXT("r.SkyLight.RealTimeReflectionCapture.TimeSlice"), 0, bOverrideValues);
+
+		// Cloud are rendered using high quality volumetric render target mode 3: per pixel tracing and composition on screen, while supporting cloud on translucent.
+		MOVIEPIPELINE_STORE_AND_OVERRIDE_CVAR_INT(PreviousVolumetricRenderTarget, TEXT("r.VolumetricRenderTarget"), 1, bOverrideValues);
+		MOVIEPIPELINE_STORE_AND_OVERRIDE_CVAR_INT(PreviousVolumetricRenderTargetMode, TEXT("r.VolumetricRenderTarget.Mode"), 3, bOverrideValues);
+
+		// To make sure that the world partition streaming doesn't end up in critical streaming performances and stops streaming low priority cells.
+		MOVIEPIPELINE_STORE_AND_OVERRIDE_CVAR_INT(PreviousIgnoreStreamingPerformance, TEXT("wp.Runtime.BlockOnSlowStreaming"), 0, bOverrideValues);
+
+		// Remove any minimum delta time requirements from Chaos Physics to ensure accuracy at high Temporal Sample counts
+		MOVIEPIPELINE_STORE_AND_OVERRIDE_CVAR_FLOAT(PreviousChaosImmPhysicsMinStepTime, TEXT("p.Chaos.ImmPhys.MinStepTime"), 0, bOverrideValues);
+
+		// MRQ's 0 -> 0.99 -> 0 evaluation for motion blur emulation can occasionally cause it to be detected as a redundant update and thus never updated
+		// which causes objects to render in the wrong position on the first frame (and without motion blur). This disables an optimization that detects
+		// the redundant updates so the update will get sent through anyways even though it thinks it's a duplicate (but it's not).
+		MOVIEPIPELINE_STORE_AND_OVERRIDE_CVAR_INT(PreviousSkipRedundantTransformUpdate, TEXT("r.SkipRedundantTransformUpdate"), 0, bOverrideValues);
+
+		// Cloth's time step smoothing messes up the change in number of simulation substeps that fixes the cloth simulation behavior when using Temporal Samples.
+		MOVIEPIPELINE_STORE_AND_OVERRIDE_CVAR_INT(PreviousChaosClothUseTimeStepSmoothing, TEXT("p.ChaosCloth.UseTimeStepSmoothing"), 0, bOverrideValues);
+	}
+	
+	// Must come after the above cvars so that if one of those cvars is also specified by the scalability level, then we
+	// restore to the value in the original scalability level, not the value we cached in the scalability level (if applied).
+	if (!bOverrideValues)
+	{
+		Scalability::SetQualityLevels(PreviousQualityLevels);
+	}
+}
+
+TSubclassOf<AGameModeBase> UMovieGraphGlobalGameOverridesNode::GetGameModeOverride(const UMoviePipelineExecutorJob* InJob)
+{
+	UMovieGraphConfig* GraphConfig = InJob->GetGraphPreset();
+	if (InJob->IsUsingGraphConfiguration() && GraphConfig)
+	{
+		// There is most likely no pipeline to fetch a traversal context from at this point, so a temporary context
+		// is generated instead (which will not be fully filled in, but should be enough to generate a correctly flattened
+		// Globals branch).
+		FMovieGraphTraversalContext TraversalContext;
+		TraversalContext.Job = const_cast<UMoviePipelineExecutorJob*>(InJob);
+		
+		if (const UMovieGraphEvaluatedConfig* EvaluatedGraph = GraphConfig->CreateFlattenedGraph(TraversalContext))
+		{
+			constexpr bool bIncludeCDOs = true;
+			constexpr bool bExactMatch = true;
+			UMovieGraphGlobalGameOverridesNode* GameOverridesNode =
+				EvaluatedGraph->GetSettingForBranch<UMovieGraphGlobalGameOverridesNode>(GlobalsPinName, bIncludeCDOs, bExactMatch);
+
+			return GameOverridesNode->GameModeOverride;
+		}
+	}
+	else
+	{
+		TArray<UMoviePipelineSetting*> AllSettings = InJob->GetConfiguration()->GetAllSettings();
+		UMoviePipelineSetting** GameOverridesPtr = AllSettings.FindByPredicate([](const UMoviePipelineSetting* InSetting)
+		{
+			return InSetting->GetClass() == UMoviePipelineGameOverrideSetting::StaticClass();
+		});
+		
+		if (GameOverridesPtr)
+		{
+			if (UMoviePipelineSetting* Setting = *GameOverridesPtr)
+			{
+				return CastChecked<UMoviePipelineGameOverrideSetting>(Setting)->GameModeOverride;
+			}
+		}
+	}
+
+	return nullptr;
+}
+
 #if WITH_EDITOR
 FText UMovieGraphGlobalGameOverridesNode::GetNodeTitle(const bool bGetDescriptive) const
 {
diff --git a/Engine/Plugins/MovieScene/MovieRenderPipeline/Source/MovieRenderPipelineCore/Private/MoviePipelineInProcessExecutor.cpp b/Engine/Plugins/MovieScene/MovieRenderPipeline/Source/MovieRenderPipelineCore/Private/MoviePipelineInProcessExecutor.cpp
--- a/Engine/Plugins/MovieScene/MovieRenderPipeline/Source/MovieRenderPipelineCore/Private/MoviePipelineInProcessExecutor.cpp
+++ b/Engine/Plugins/MovieScene/MovieRenderPipeline/Source/MovieRenderPipelineCore/Private/MoviePipelineInProcessExecutor.cpp
@@ -1,4 +1,5 @@
 // Copyright Epic Games, Inc. All Rights Reserved.
+
 #include "MoviePipelineInProcessExecutor.h"
 #include "Engine/GameInstance.h"
 #include "MoviePipeline.h"
@@ -13,6 +14,7 @@
 #include "Engine/GameEngine.h"
 #include "MoviePipelineUtils.h"
 #include "Graph/MovieGraphPipeline.h"
+#include "Graph/Nodes/MovieGraphGlobalGameOverrides.h"
 
 #include UE_INLINE_GENERATED_CPP_BY_NAME(MoviePipelineInProcessExecutor)
 
@@ -68,21 +70,10 @@ void UMoviePipelineInProcessExecutor::Start(const UMoviePipelineExecutorJob* InJ
 		// We were launched into an empty map so we'll look at our job and figure out which map we should load.
 		// Get the next job in the queue
 		FString MapOptions;
-
-		TArray<UMoviePipelineSetting*> AllSettings = InJob->GetConfiguration()->GetAllSettings();
-		UMoviePipelineSetting** GameOverridesPtr = AllSettings.FindByPredicate([](UMoviePipelineSetting* InSetting) { return InSetting->GetClass() == UMoviePipelineGameOverrideSetting::StaticClass(); });
-		if (GameOverridesPtr)
+		
+		if (const TSubclassOf<AGameModeBase> GameModeOverride = UMovieGraphGlobalGameOverridesNode::GetGameModeOverride(InJob))
 		{
-			UMoviePipelineSetting* Setting = *GameOverridesPtr;
-			if (Setting)
-			{
-				UMoviePipelineGameOverrideSetting* GameOverrideSetting = CastChecked<UMoviePipelineGameOverrideSetting>(Setting);
-				if (GameOverrideSetting->GameModeOverride)
-				{
-					MapOptions = TEXT("?game=") + GameOverrideSetting->GameModeOverride->GetPathName();
-				}
-
-			}
+			MapOptions = TEXT("?game=") + GameModeOverride->GetPathName();
 		}
 
 		FString LevelPath = InJob->Map.GetLongPackageName();
diff --git a/Engine/Plugins/MovieScene/MovieRenderPipeline/Source/MovieRenderPipelineCore/Public/Graph/MovieGraphDefaultRenderer.h b/Engine/Plugins/MovieScene/MovieRenderPipeline/Source/MovieRenderPipelineCore/Public/Graph/MovieGraphDefaultRenderer.h
--- a/Engine/Plugins/MovieScene/MovieRenderPipeline/Source/MovieRenderPipelineCore/Public/Graph/MovieGraphDefaultRenderer.h
+++ b/Engine/Plugins/MovieScene/MovieRenderPipeline/Source/MovieRenderPipelineCore/Public/Graph/MovieGraphDefaultRenderer.h
@@ -184,6 +184,8 @@ public:
 protected:
 	TObjectPtr<UTextureRenderTarget2D> CreateViewRenderTarget(const UE::MovieGraph::DefaultRenderer::FRenderTargetInitParams& InInitParams) const;
 	FMoviePipelineSurfaceQueuePtr CreateSurfaceQueue(const UE::MovieGraph::DefaultRenderer::FRenderTargetInitParams& InInitParams) const;
+	void GetCameraLocationsForFrame(TArray<FVector>& OutLocations) const;
+	void FlushAsyncEngineSystems(const TObjectPtr<UMovieGraphEvaluatedConfig>& InConfig) const;
 
 protected:
 	/** A pointer to the CDOs of the Render Pass nodes that are valid for the current shot render. */
diff --git a/Engine/Plugins/MovieScene/MovieRenderPipeline/Source/MovieRenderPipelineCore/Public/Graph/Nodes/MovieGraphGlobalGameOverrides.h b/Engine/Plugins/MovieScene/MovieRenderPipeline/Source/MovieRenderPipelineCore/Public/Graph/Nodes/MovieGraphGlobalGameOverrides.h
--- a/Engine/Plugins/MovieScene/MovieRenderPipeline/Source/MovieRenderPipelineCore/Public/Graph/Nodes/MovieGraphGlobalGameOverrides.h
+++ b/Engine/Plugins/MovieScene/MovieRenderPipeline/Source/MovieRenderPipelineCore/Public/Graph/Nodes/MovieGraphGlobalGameOverrides.h
@@ -4,8 +4,20 @@
 
 #include "Graph/MovieGraphNode.h"
 #include "MoviePipelineGameOverrideSetting.h"
+
 #include "MovieGraphGlobalGameOverrides.generated.h"
 
+/** Scalability quality levels available for use in renders generated from the graph. */
+UENUM(BlueprintType)
+enum class EMovieGraphScalabilityQualityLevel : uint8
+{
+	Low = 0,
+	Medium = 1,
+	High = 2,
+	Epic = 3,
+	Cinematic = 4
+};
+
 /** A node which configures the global game overrides. */
 UCLASS()
 class MOVIERENDERPIPELINECORE_API UMovieGraphGlobalGameOverridesNode : public UMovieGraphSettingNode
@@ -13,10 +25,23 @@ class MOVIERENDERPIPELINECORE_API UMovieGraphGlobalGameOverridesNode : public UM
 	GENERATED_BODY()
 
 public:
-	UMovieGraphGlobalGameOverridesNode() = default;
+	UMovieGraphGlobalGameOverridesNode();
 
 	virtual EMovieGraphBranchRestriction GetBranchRestriction() const override { return EMovieGraphBranchRestriction::Globals; }
 
+	/**
+	 * Applies any cvars, scalability settings, etc. to reflect the properties set on the node. Remembers what the
+	 * settings are before they are set, so bOverrideValues can be set to false to revert them to their original
+	 * values. The world needs to be provided as context for any commands that are executed.
+	 */
+	void ApplySettings(const bool bOverrideValues, UWorld* InWorld);
+
+	/**
+	 * Gets the game mode override if the job specified is using a graph configuration, otherwise gets the game mode
+	 * override from the legacy non-graph config system.
+	 */
+	static TSubclassOf<AGameModeBase> GetGameModeOverride(const UMoviePipelineExecutorJob* InJob);
+
 #if WITH_EDITOR
 	virtual FText GetNodeTitle(const bool bGetDescriptive = false) const override;
 	virtual FText GetMenuCategory() const override;
@@ -29,91 +54,151 @@ public:
 	uint8 bOverride_GameModeOverride : 1;
 
 	UPROPERTY(EditAnywhere, BlueprintReadWrite, Category = Overrides, meta = (InlineEditConditionToggle))
-	uint8 bOverride_CinematicQualitySettings : 1;
-
-	UPROPERTY(EditAnywhere, BlueprintReadWrite, Category = Overrides, meta = (InlineEditConditionToggle))
-	uint8 bOverride_TextureStreaming : 1;
+	uint8 bOverride_ScalabilityQualityLevel : 1;
 
 	UPROPERTY(EditAnywhere, BlueprintReadWrite, Category = Overrides, meta = (InlineEditConditionToggle))
-	uint8 bOverride_UseLODZero : 1;
+	uint8 bOverride_bDisableTextureStreaming : 1;
 
 	UPROPERTY(EditAnywhere, BlueprintReadWrite, Category = Overrides, meta = (InlineEditConditionToggle))
-	uint8 bOverride_DisableHLODs : 1;
+	uint8 bOverride_bDisableLODs : 1;
 
 	UPROPERTY(EditAnywhere, BlueprintReadWrite, Category = Overrides, meta = (InlineEditConditionToggle))
-	uint8 bOverride_UseHighQualityShadows : 1;
+	uint8 bOverride_bDisableHLODs : 1;
 
 	UPROPERTY(EditAnywhere, BlueprintReadWrite, Category = Overrides, meta = (InlineEditConditionToggle))
-	uint8 bOverride_ShadowDistanceScale : 1;
+	uint8 bOverride_bFlushLevelStreaming : 1;
 
 	UPROPERTY(EditAnywhere, BlueprintReadWrite, Category = Overrides, meta = (InlineEditConditionToggle))
-	uint8 bOverride_ShadowRadiusThreshold : 1;
-	
-	UPROPERTY(EditAnywhere, BlueprintReadWrite, Category = Overrides, meta = (InlineEditConditionToggle))
-	uint8 bOverride_OverrideViewDistanceScale : 1;
+	uint8 bOverride_bFlushAssetCompiler: 1;
 
 	UPROPERTY(EditAnywhere, BlueprintReadWrite, Category = Overrides, meta = (InlineEditConditionToggle))
-	uint8 bOverride_ViewDistanceScale : 1;
+	uint8 bOverride_bFlushShaderCompiler : 1;
 
 	UPROPERTY(EditAnywhere, BlueprintReadWrite, Category = Overrides, meta = (InlineEditConditionToggle))
-	uint8 bOverride_FlushGrassStreaming : 1;
+	uint8 bOverride_bFlushGrassStreaming : 1;
 
 	UPROPERTY(EditAnywhere, BlueprintReadWrite, Category = Overrides, meta = (InlineEditConditionToggle))
-	uint8 bOverride_FlushStreamingManagers : 1;
+	uint8 bOverride_bFlushStreamingManagers : 1;
 
 	UPROPERTY(EditAnywhere, BlueprintReadWrite, Category = Overrides, meta = (InlineEditConditionToggle))
 	uint8 bOverride_VirtualTextureFeedbackFactor : 1;
 
-	UPROPERTY(EditAnywhere, BlueprintReadWrite, Category = Overrides, meta = (InlineEditConditionToggle))
-	uint8 bOverride_DisableGPUTimeout : 1;
-	
-	/** Optional Game Mode to override the map's default game mode with. This can be useful if the game's normal mode displays UI elements or loading screens that you don't want captured. */
+	/**
+	 * Optional game mode to override the map's default game mode with. This can be useful if the game's normal mode
+	 * displays UI elements or loading screens that you don't want captured.
+	 */
 	UPROPERTY(EditAnywhere, BlueprintReadWrite, Category = "Game", meta = (EditCondition = "bOverride_GameModeOverride"))
 	TSubclassOf<AGameModeBase> GameModeOverride;
 
-	/** If true, automatically set the engine to the Cinematic Scalability quality settings during render. See the Scalability Reference documentation for information on how to edit cvars to add/change default quality values.*/
-	UPROPERTY(EditAnywhere, BlueprintReadWrite, Category = "Game", meta = (EditCondition = "bOverride_CinematicQualitySettings"))
-	bool bCinematicQualitySettings;
-
-	/** Defines which If true, when using texture streaming fully load the required textures each frame instead of loading them in over time. This solves objects being blurry after camera cuts. */
-	UPROPERTY(EditAnywhere, BlueprintReadWrite, Category = "Rendering", meta = (EditCondition = "bOverride_TextureStreaming"))
-	EMoviePipelineTextureStreamingMethod TextureStreaming;
-
-	/** Should we try to use the highest quality LOD for meshes and particle systems regardless of distance? */
-	UPROPERTY(EditAnywhere, BlueprintReadWrite, Category = "Rendering", meta = (EditCondition = "bOverride_UseLODZero"))
-	bool bUseLODZero;
-
-	/** Should we disable Hierarchical LODs and instead use their real meshes regardless of distance? */
-	UPROPERTY(EditAnywhere, BlueprintReadWrite, Category = "Rendering", meta = (EditCondition = "bOverride_DisableHLODs"))
+	/**
+	 * The scalability quality level that should be used in renders. See the Scalability Reference documentation for
+	 * information on how to edit cvars to add/change default quality values.
+	 */
+	UPROPERTY(EditAnywhere, BlueprintReadWrite, Category = "Game", meta = (EditCondition = "bOverride_ScalabilityQualityLevel"))
+	EMovieGraphScalabilityQualityLevel ScalabilityQualityLevel;
+
+	/**
+	 * Toggles whether texture streaming is disabled. Can solve objects being blurry after camera cuts.
+	 *
+	 * Configures the following cvars:
+	 * - r.TextureStreaming
+	 */
+	UPROPERTY(EditAnywhere, BlueprintReadWrite, Category = "Rendering", meta = (EditCondition = "bOverride_bDisableTextureStreaming"))
+	bool bDisableTextureStreaming;
+
+	/**
+	 * Disabling LODs will use the highest quality LOD for meshes and particle systems, regardless of distance. Note
+	 * that this does not affect Nanite.
+	 *
+	 * Configures the following cvars:
+	 * - r.ForceLOD
+	 * - r.SkeletalMeshLODBias
+	 * - r.ParticleLODBias
+	 * - foliage.DitheredLOD
+	 * - foliage.ForceLOD
+	 */
+	UPROPERTY(EditAnywhere, BlueprintReadWrite, Category = "Rendering", meta = (EditCondition = "bOverride_bDisableLODs"))
+	bool bDisableLODs;
+
+	/**
+	 * Determines if hierarchical LODs should be disabled and their real meshes used instead, regardless of distance.
+	 * Note that this does not affect World Partition HLODs.
+	 */
+	UPROPERTY(EditAnywhere, BlueprintReadWrite, Category = "Rendering", meta = (EditCondition = "bOverride_bDisableHLODs"))
 	bool bDisableHLODs;
 
-	/** Should we override shadow-related CVars with some high quality preset settings? */
-	UPROPERTY(EditAnywhere, BlueprintReadWrite, Category = "Rendering", meta = (EditCondition = "bOverride_UseHighQualityShadows"))
-	bool bUseHighQualityShadows;
-
-	// TODO: This should also be disabled if bUseHighQualityShadows is false, but EditCondition doesn't work with that *plus* the InlineEditConditionToggle
-	/** Scalability option to trade shadow distance versus performance for directional lights  */
-	UPROPERTY(EditAnywhere, BlueprintReadWrite, Category = "Rendering", meta = (EditCondition = "bOverride_ShadowDistanceScale", ClampMin=0.1, UIMin=0.1, UIMax=10))
-	int32 ShadowDistanceScale;
-
-	// TODO: This should also be disabled if bUseHighQualityShadows is false, but EditCondition doesn't work with that *plus* the InlineEditConditionToggle
-	/** Cull shadow casters if they are too small, value is the minimal screen space bounding sphere radius */
-	UPROPERTY(EditAnywhere, BlueprintReadWrite, Category = "Rendering", meta = (EditCondition = "bOverride_ShadowRadiusThreshold", UIMin=0.001, ClampMin=0.001))
-	float ShadowRadiusThreshold;
-
-	/** Controls the view distance scale. A primitive's MaxDrawDistance is scaled by this value. */
-	UPROPERTY(EditAnywhere, BlueprintReadWrite, Category = "Rendering", meta = (EditCondition = "bOverride_ViewDistanceScale"))
-	int32 ViewDistanceScale;
-	
-	/** Flushing grass streaming (combined with override view distance scale) prevents visible pop-in/culling of grace instances. */
-	UPROPERTY(EditAnywhere, BlueprintReadWrite, Category = "Rendering", meta = (EditCondition = "bOverride_FlushGrassStreaming"))
+	/**
+	 * Flushing level streaming ensures that any pending changes to sub-levels or world partition are fully processed before we render
+	 * the frame. This feature generally only adds to render times on the frames that have level visibility state changes, so generally
+	 * safe to leave turned on all the time.
+	 */
+	UPROPERTY(EditAnywhere, BlueprintReadWrite, Category = "Rendering", meta = (EditCondition = "bOverride_bFlushLevelStreaming"))
+	bool bFlushLevelStreaming;
+
+	/**
+	 * This ensures that any asynchronously compiled assets (static meshes, distance fields, etc.) required for rendering the frame
+	 * are completed before rendering the frame. This feature generally only adds to render times on the frames where a new asset
+	 * is introduced (ie, spawned) that may not be fully compiled. Results are stored in the DDC on subsequent uses.
+	 */
+	UPROPERTY(EditAnywhere, BlueprintReadWrite, Category = "Rendering", meta = (EditCondition = "bOverride_bFlushAssetCompiler"))
+	bool bFlushAssetCompiler;
+
+	/**
+	 * This ensures that any asynchronously compiled shader permutations are completed before rendering the frame. When using
+	 * On Demand Shader Compilation the editor will skip compiling currently unneeded permutations for the material graph to
+	 * decrease artist iteration time, but these permutations need to be compiled when rendering. Results are stored in the DDC on
+	 * subsequent uses.
+	 */
+	UPROPERTY(EditAnywhere, BlueprintReadWrite, Category = "Rendering", meta = (EditCondition = "bOverride_bFlushShaderCompiler"))
+	bool bFlushShaderCompiler;
+
+	/**
+	 * Flushing grass streaming prevents visible pop-in/culling of grass instances, but may come at a high GPU memory
+	 * cost (depending on rendering feature set), and grass density. Try turning this off if you are low on GPU memory
+	 * and have dense grass.
+	 */
+	UPROPERTY(EditAnywhere, BlueprintReadWrite, Category = "Rendering", meta = (EditCondition = "bOverride_bFlushGrassStreaming"))
 	bool bFlushGrassStreaming;
 
-	/** Experimental. If true flush the streaming managers (Texture Streaming) each frame. Allows Texture Streaming to not have visible pop-in in final frames. */
-	UPROPERTY(EditAnywhere, BlueprintReadWrite, Category = "Rendering", meta = (EditCondition = "bOverride_FlushStreamingManagers"))
+	/** 
+	 * This ensures that we wait on any streaming managers that may have outstanding work to finish their work before we render
+	 * the frame. Many GPU-based features rely on a feedback loop where a frame is rendered, then results are read back from the GPU
+	 * to the CPU. These results are analyzed, and additional data is loaded to be used in a subsequent frame. This feature ensures
+	 * that the data is fully loaded before we render the next frame, but does NOT solve issues related to generating the those
+	 * feedback requests in the first place. Virtual Textures and Nanite must render the frame to generate the request -- 
+	 * this option cannot solve that; it only helps ensure that the requests that are made are fully processed before
+	 * the frame is rendered.
+	 *
+	 * Configures the following cvars:
+	 * - r.Streaming.SyncStatesWhenBlocking
+	 */
+	UPROPERTY(EditAnywhere, BlueprintReadWrite, Category = "Rendering", meta = (EditCondition = "bOverride_bFlushStreamingManagers"))
 	bool bFlushStreamingManagers;
 
 	/** The virtual texture feedback resolution factor. A lower factor will increase virtual texture feedback resolution. */
 	UPROPERTY(EditAnywhere, BlueprintReadWrite, Category = "Rendering", meta = (EditCondition = "bOverride_VirtualTextureFeedbackFactor"))
 	int32 VirtualTextureFeedbackFactor;
+
+private:
+	Scalability::FQualityLevels PreviousQualityLevels;
+	int32 PreviousTextureStreaming;
+	int32 PreviousForceLOD;
+	int32 PreviousSkeletalMeshBias;
+	int32 PreviousParticleLODBias;
+	int32 PreviousFoliageDitheredLOD;
+	int32 PreviousFoliageForceLOD;
+	int32 PreviousStreamingManagerSyncState;
+	int32 PreviousAnimationUROEnabled;
+	int32 PreviousSkyLightRealTimeReflectionCaptureTimeSlice;
+	int32 PreviousVolumetricRenderTarget;
+	int32 PreviousVolumetricRenderTargetMode;
+	int32 PreviousIgnoreStreamingPerformance;
+	float PreviousChaosImmPhysicsMinStepTime;
+	int32 PreviousSkipRedundantTransformUpdate;
+	int32 PreviousChaosClothUseTimeStepSmoothing;
+
+#if WITH_EDITOR
+	int32 PreviousGeoCacheStreamerShowNotification;
+	int32 PreviousGeoCacheStreamerBlockTillFinish;
+#endif
 };
\ No newline at end of file
diff --git a/Engine/Plugins/MovieScene/MovieRenderPipeline/Source/MovieRenderPipelineEditor/Private/MoviePipelinePIEExecutor.cpp b/Engine/Plugins/MovieScene/MovieRenderPipeline/Source/MovieRenderPipelineEditor/Private/MoviePipelinePIEExecutor.cpp
--- a/Engine/Plugins/MovieScene/MovieRenderPipeline/Source/MovieRenderPipelineEditor/Private/MoviePipelinePIEExecutor.cpp
+++ b/Engine/Plugins/MovieScene/MovieRenderPipeline/Source/MovieRenderPipelineEditor/Private/MoviePipelinePIEExecutor.cpp
@@ -1,4 +1,5 @@
 // Copyright Epic Games, Inc. All Rights Reserved.
+
 #include "MoviePipelinePIEExecutor.h"
 #include "MoviePipelinePrimaryConfig.h"
 #include "MoviePipelineShotConfig.h"
@@ -21,6 +22,7 @@
 #include "MessageLogModule.h"
 #include "Logging/MessageLog.h"
 #include "Graph/MovieGraphPipeline.h"
+#include "Graph/Nodes/MovieGraphGlobalGameOverrides.h"
 
 #include UE_INLINE_GENERATED_CPP_BY_NAME(MoviePipelinePIEExecutor)
 
@@ -161,15 +163,9 @@ void UMoviePipelinePIEExecutor::Start(const UMoviePipelineExecutorJob* InJob)
 	// Initialize the transient settings so that they will exist in time for the GameOverrides check.
 	InJob->GetConfiguration()->InitializeTransientSettings();
 
-	TArray<UMoviePipelineSetting*> AllSettings = InJob->GetConfiguration()->GetAllSettings();
-	UMoviePipelineSetting** GameOverridesPtr = AllSettings.FindByPredicate([](UMoviePipelineSetting* InSetting) { return InSetting->GetClass() == UMoviePipelineGameOverrideSetting::StaticClass(); });
-	if (GameOverridesPtr)
-	{	
-		UMoviePipelineSetting* Setting = *GameOverridesPtr;
-		if (Setting)
-		{
-			Params.GameModeOverride = CastChecked<UMoviePipelineGameOverrideSetting>(Setting)->GameModeOverride;
-		}
+	if (TSubclassOf<AGameModeBase> GameModeOverride = UMovieGraphGlobalGameOverridesNode::GetGameModeOverride(InJob))
+	{
+		Params.GameModeOverride = GameModeOverride;
 	}
 
 	bPreviousUseFixedTimeStep = FApp::UseFixedTimeStep();
