MovieGraph: Add support for spatial samples, anti-aliasing overrides, and disabling tone-curve in the Deferred Renderer.

This change more clearly splits the node representation from the actual rendering instances (which there may be multiple per render). The architecture is designed around a render pass creating all of the needed datastructures at the top level, and then passing those into the base class functions to build up the data needed to create a render. The reason for this decision is that in the legacy renderer, the render function was lightweight, but the base class would call out to many different virtual functions to allow the renderer to adjust the settings. This was confusing and hard to follow and made it difficult to pass renderer specific data to these functions, so now this concept has been flipped and all of the data is calculated upfront and provided, which makes it much easier for individual render passes to create their own logic for determining camera position, render targets, scene histories, etc.

diff --git a/Engine/Plugins/MovieScene/MovieRenderPipeline/Source/MovieRenderPipelineCore/Private/Graph/MovieGraphBlueprintLibrary.cpp b/Engine/Plugins/MovieScene/MovieRenderPipeline/Source/MovieRenderPipelineCore/Private/Graph/MovieGraphBlueprintLibrary.cpp
--- a/Engine/Plugins/MovieScene/MovieRenderPipeline/Source/MovieRenderPipelineCore/Private/Graph/MovieGraphBlueprintLibrary.cpp
+++ b/Engine/Plugins/MovieScene/MovieRenderPipeline/Source/MovieRenderPipelineCore/Private/Graph/MovieGraphBlueprintLibrary.cpp
@@ -5,6 +5,7 @@
 #include "Graph/MovieGraphPipeline.h"
 #include "Graph/Nodes/MovieGraphOutputSettingNode.h"
 #include "Graph/Nodes/MovieGraphRenderLayerNode.h"
+#include "Graph/Nodes/MovieGraphCameraNode.h"
 #include "HAL/FileManager.h"
 #include "Internationalization/Regex.h"
 #include "MoviePipelineBlueprintLibrary.h"
@@ -306,8 +307,15 @@ FIntPoint UMovieGraphBlueprintLibrary::GetEffectiveOutputResolution(UMovieGraphE
 	{
 		return FIntPoint();
 	}
-	
-	return UMoviePipelineBlueprintLibrary::Utility_GetEffectiveOutputResolution(0 /* TODO: Overscan percentage needs to be provided */, OutputSetting->OutputResolution);
+
+	float RescaledOverscan = 0.f;
+	if (UMovieGraphCameraSettingNode* CameraSetting = InEvaluatedGraph->GetSettingForBranch<UMovieGraphCameraSettingNode>(UMovieGraphNode::GlobalsPinName, bIncludeCDOs))
+	{
+		// The old system used [0-1] range for floats, the new system will use [0-100], so we rescale down before calling through.
+		RescaledOverscan = CameraSetting->OverscanPercentage / 100.f;
+	}
+
+	return UMoviePipelineBlueprintLibrary::Utility_GetEffectiveOutputResolution(RescaledOverscan, OutputSetting->OutputResolution);
 }
 
 FText UMovieGraphBlueprintLibrary::GetJobName(const UMovieGraphPipeline* InMovieGraphPipeline)
diff --git a/Engine/Plugins/MovieScene/MovieRenderPipeline/Source/MovieRenderPipelineCore/Private/Graph/MovieGraphCoreTimeStep.cpp b/Engine/Plugins/MovieScene/MovieRenderPipeline/Source/MovieRenderPipelineCore/Private/Graph/MovieGraphCoreTimeStep.cpp
--- a/Engine/Plugins/MovieScene/MovieRenderPipeline/Source/MovieRenderPipelineCore/Private/Graph/MovieGraphCoreTimeStep.cpp
+++ b/Engine/Plugins/MovieScene/MovieRenderPipeline/Source/MovieRenderPipelineCore/Private/Graph/MovieGraphCoreTimeStep.cpp
@@ -382,6 +382,8 @@ void UMovieGraphCoreTimeStep::TickProducingFrames()
 	CurrentTimeStepData.bRequiresAccumulator = CurrentFrameData.TemporalSampleCount > 1 && !CurrentTimeStepData.bDiscardOutput;
 	CurrentTimeStepData.OutputFrameNumber = CurrentFrameData.OutputFrameNumber;
 	CurrentTimeStepData.RenderedFrameNumber = CurrentFrameData.RenderedFrameNumber;
+	CurrentTimeStepData.TemporalSampleCount = CurrentFrameData.TemporalSampleCount;
+	CurrentTimeStepData.TemporalSampleIndex = CurrentFrameData.TemporalSampleIndex;
 	CurrentTimeStepData.EvaluatedConfig = TObjectPtr<UMovieGraphEvaluatedConfig>(CurrentFrameData.EvaluatedConfig.Get());
 
 	//UE_LOG(LogTemp, Warning, TEXT("F# %d bFirst: %d bLast: %d bReqAc: %d"),
diff --git a/Engine/Plugins/MovieScene/MovieRenderPipeline/Source/MovieRenderPipelineCore/Private/Graph/MovieGraphDefaultRenderer.cpp b/Engine/Plugins/MovieScene/MovieRenderPipeline/Source/MovieRenderPipelineCore/Private/Graph/MovieGraphDefaultRenderer.cpp
--- a/Engine/Plugins/MovieScene/MovieRenderPipeline/Source/MovieRenderPipelineCore/Private/Graph/MovieGraphDefaultRenderer.cpp
+++ b/Engine/Plugins/MovieScene/MovieRenderPipeline/Source/MovieRenderPipelineCore/Private/Graph/MovieGraphDefaultRenderer.cpp
@@ -6,13 +6,20 @@
 #include "Graph/Nodes/MovieGraphGlobalGameOverrides.h"
 #include "Graph/Nodes/MovieGraphRenderLayerNode.h"
 #include "Graph/Nodes/MovieGraphRenderPassNode.h"
+#include "Graph/Nodes/MovieGraphGlobalGameOverrides.h"
+#include "Graph/Nodes/MovieGraphDebugNode.h"
 #include "MovieRenderPipelineCoreModule.h"
 #include "RenderingThread.h"
 #include "Engine/TextureRenderTarget2D.h"
 #include "UObject/Package.h"
 #include "MoviePipelineSurfaceReader.h"
+#include "RenderCaptureInterface.h"
 
 // For flushing async systems
+#include "EngineModule.h"
+#include "MeshCardRepresentation.h"
+#include "ShaderCompiler.h"
+#include "EngineUtils.h"
 #include "AssetCompilingManager.h"
 #include "ContentStreaming.h"
 #include "EngineModule.h"
@@ -174,7 +181,6 @@ void UMovieGraphDefaultRenderer::AddReferencedObjects(UObject* InThis, FReferenc
 
 void UMovieGraphDefaultRenderer::Render(const FMovieGraphTimeStepData& InTimeStepData)
 {
-	// Flush built-in systems before we render anything. This maximizes the likelihood that the data is prepared for when
 	// the render thread uses it.
 	FlushAsyncEngineSystems(InTimeStepData.EvaluatedConfig);
 
@@ -228,6 +234,22 @@ void UMovieGraphDefaultRenderer::Render(const FMovieGraphTimeStepData& InTimeSte
 		}
 	}
 
+	// There is some work we need to signal to the renderer for only the first view of a frame,
+	// so we have to track when any of our *FSceneView* render passes submit stuff (UI renderers don't count)
+	bHasRenderedFirstViewThisFrame = false;
+
+	// Support for RenderDoc captures of just the MRQ work
+#if WITH_EDITOR && !UE_BUILD_SHIPPING
+	TUniquePtr<RenderCaptureInterface::FScopedCapture> ScopedGPUCapture;
+	{
+		UMovieGraphDebugSettingNode* DebugSettings = InTimeStepData.EvaluatedConfig->GetSettingForBranch<UMovieGraphDebugSettingNode>(UMovieGraphNode::GlobalsPinName);
+		if (DebugSettings && DebugSettings->bCaptureFramesWithRenderDoc)
+		{
+			ScopedGPUCapture = MakeUnique<RenderCaptureInterface::FScopedCapture>(true, *FString::Printf(TEXT("MRQ Frame: %d"), InTimeStepData.RootFrameNumber.Value));
+		}
+	}
+#endif
+
 	for (const TObjectPtr<UMovieGraphRenderPassNode>& RenderPass : RenderPassesInUse)
 	{
 		// Pass in a copy of the traversal context so the renderer can decide what to do with it.
@@ -341,7 +363,7 @@ UE::MovieGraph::DefaultRenderer::FCameraInfo UMovieGraphDefaultRenderer::GetCame
 	{
 		CameraInfo.ViewInfo = LocalPlayerController->PlayerCameraManager->GetCameraCacheView();
 		CameraInfo.ViewActor = LocalPlayerController->GetViewTarget();
-		CameraInfo.CameraName = TEXT("Unnamed_Camera"); // ToDo: This eventually needs to come from Level Sequences
+		CameraInfo.CameraName = TEXT("Unsupported"); // ToDo: This eventually needs to come from Level Sequences
 	}
 	else
 	{
diff --git a/Engine/Plugins/MovieScene/MovieRenderPipeline/Source/MovieRenderPipelineCore/Private/Graph/Nodes/MovieGraphCameraNode.cpp b/Engine/Plugins/MovieScene/MovieRenderPipeline/Source/MovieRenderPipelineCore/Private/Graph/Nodes/MovieGraphCameraNode.cpp
new file mode 100644
--- /dev/null
+++ b/Engine/Plugins/MovieScene/MovieRenderPipeline/Source/MovieRenderPipelineCore/Private/Graph/Nodes/MovieGraphCameraNode.cpp
@@ -0,0 +1,33 @@
+// Copyright Epic Games, Inc. All Rights Reserved.
+
+#include "Graph/Nodes/MovieGraphCameraNode.h"
+
+#include "Styling/AppStyle.h"
+
+#if WITH_EDITOR
+FText UMovieGraphCameraSettingNode::GetNodeTitle(const bool bGetDescriptive) const
+{
+	static const FText NodeName = NSLOCTEXT("MovieGraphNodes", "NodeName_CameraSettings", "Camera Settings");
+	return NodeName;
+}
+
+FText UMovieGraphCameraSettingNode::GetMenuCategory() const
+{
+	static const FText NodeCategory_Globals = NSLOCTEXT("MovieGraphNodes", "NodeCategory_Globals", "Globals");
+	return NodeCategory_Globals;
+}
+
+FLinearColor UMovieGraphCameraSettingNode::GetNodeTitleColor() const
+{
+	static const FLinearColor NodeColor = FLinearColor(0.65f, 0.60f, 0.75f);
+	return NodeColor;
+}
+
+FSlateIcon UMovieGraphCameraSettingNode::GetIconAndTint(FLinearColor& OutColor) const
+{
+	static const FSlateIcon Icon = FSlateIcon(FAppStyle::GetAppStyleSetName(), "ClassIcon.CameraComponent");
+
+	OutColor = FLinearColor::White;
+	return Icon;
+}
+#endif // WITH_EDITOR
\ No newline at end of file
diff --git a/Engine/Plugins/MovieScene/MovieRenderPipeline/Source/MovieRenderPipelineCore/Private/Graph/Nodes/MovieGraphCoreRenderPassNode.cpp b/Engine/Plugins/MovieScene/MovieRenderPipeline/Source/MovieRenderPipelineCore/Private/Graph/Nodes/MovieGraphCoreRenderPassNode.cpp
deleted file mode 100644
diff --git a/Engine/Plugins/MovieScene/MovieRenderPipeline/Source/MovieRenderPipelineCore/Private/Graph/Nodes/MovieGraphDebugNode.cpp b/Engine/Plugins/MovieScene/MovieRenderPipeline/Source/MovieRenderPipelineCore/Private/Graph/Nodes/MovieGraphDebugNode.cpp
new file mode 100644
--- /dev/null
+++ b/Engine/Plugins/MovieScene/MovieRenderPipeline/Source/MovieRenderPipelineCore/Private/Graph/Nodes/MovieGraphDebugNode.cpp
@@ -0,0 +1,33 @@
+// Copyright Epic Games, Inc. All Rights Reserved.
+
+#include "Graph/Nodes/MovieGraphDebugNode.h"
+
+#include "Styling/AppStyle.h"
+
+#if WITH_EDITOR
+FText UMovieGraphDebugSettingNode::GetNodeTitle(const bool bGetDescriptive) const
+{
+	static const FText NodeName = NSLOCTEXT("MovieGraphNodes", "NodeName_DebugNode", "Debug Settings");
+	return NodeName;
+}
+
+FText UMovieGraphDebugSettingNode::GetMenuCategory() const
+{
+	static const FText NodeCategory_Globals = NSLOCTEXT("MovieGraphNodes", "NodeCategory_Globals", "Globals");
+	return NodeCategory_Globals;
+}
+
+FLinearColor UMovieGraphDebugSettingNode::GetNodeTitleColor() const
+{
+	static const FLinearColor NodeColor = FLinearColor(0.849f, 0.f, 0.350f);
+	return NodeColor;
+}
+
+FSlateIcon UMovieGraphDebugSettingNode::GetIconAndTint(FLinearColor& OutColor) const
+{
+	static const FSlateIcon Icon = FSlateIcon(FAppStyle::GetAppStyleSetName(), "DebugTools.TabIcon");
+
+	OutColor = FLinearColor::White;
+	return Icon;
+}
+#endif // WITH_EDITOR
\ No newline at end of file
diff --git a/Engine/Plugins/MovieScene/MovieRenderPipeline/Source/MovieRenderPipelineCore/Private/Graph/Nodes/MovieGraphDeferredPassNode.cpp b/Engine/Plugins/MovieScene/MovieRenderPipeline/Source/MovieRenderPipelineCore/Private/Graph/Nodes/MovieGraphDeferredPassNode.cpp
new file mode 100644
--- /dev/null
+++ b/Engine/Plugins/MovieScene/MovieRenderPipeline/Source/MovieRenderPipelineCore/Private/Graph/Nodes/MovieGraphDeferredPassNode.cpp
@@ -0,0 +1,34 @@
+﻿// Copyright Epic Games, Inc. All Rights Reserved.
+
+#include "Graph/Nodes/MovieGraphDeferredPassNode.h"
+#include "Graph/Renderers/MovieGraphDeferredPass.h"
+
+TUniquePtr<UE::MovieGraph::Rendering::FMovieGraphImagePassBase> UMovieGraphDeferredRenderPassNode::CreateInstance() const
+{
+	return MakeUnique<UE::MovieGraph::Rendering::FMovieGraphDeferredPass>();
+}
+
+UMovieGraphDeferredRenderPassNode::UMovieGraphDeferredRenderPassNode()
+	: SpatialSampleCount(1)
+	, AntiAliasingMethod(EAntiAliasingMethod::AAM_TSR)
+	, ViewModeIndex(VMI_Lit)
+{
+}
+
+#if WITH_EDITOR
+FText UMovieGraphDeferredRenderPassNode::GetNodeTitle(const bool bGetDescriptive) const
+{
+	return NSLOCTEXT("MovieGraphNodes", "DeferredRenderPassGraphNode_Description", "Deferred Renderer");
+}
+#endif
+
+FString UMovieGraphDeferredRenderPassNode::GetRendererNameImpl() const
+{
+	static const FString RendererNameImpl(TEXT("Deferred"));
+	return RendererNameImpl;
+}
+
+EViewModeIndex UMovieGraphDeferredRenderPassNode::GetViewModeIndex() const
+{
+	return ViewModeIndex;
+}
\ No newline at end of file
diff --git a/Engine/Plugins/MovieScene/MovieRenderPipeline/Source/MovieRenderPipelineCore/Private/Graph/Nodes/MovieGraphDeferredRenderPassNode.cpp b/Engine/Plugins/MovieScene/MovieRenderPipeline/Source/MovieRenderPipelineCore/Private/Graph/Nodes/MovieGraphDeferredRenderPassNode.cpp
deleted file mode 100644
diff --git a/Engine/Plugins/MovieScene/MovieRenderPipeline/Source/MovieRenderPipelineCore/Private/Graph/Nodes/MovieGraphImagePassBaseNode.cpp b/Engine/Plugins/MovieScene/MovieRenderPipeline/Source/MovieRenderPipelineCore/Private/Graph/Nodes/MovieGraphImagePassBaseNode.cpp
new file mode 100644
--- /dev/null
+++ b/Engine/Plugins/MovieScene/MovieRenderPipeline/Source/MovieRenderPipelineCore/Private/Graph/Nodes/MovieGraphImagePassBaseNode.cpp
@@ -0,0 +1,97 @@
+// Copyright Epic Games, Inc. All Rights Reserved.
+
+#include "Graph/Nodes/MovieGraphImagePassBaseNode.h"
+#include "Graph/MoviePipelineRenderLayerSubsystem.h" 
+
+UMovieGraphImagePassBaseNode::UMovieGraphImagePassBaseNode()
+{
+	ShowFlags = CreateDefaultSubobject<UMovieGraphShowFlags>(TEXT("ShowFlags"));
+}
+
+void UMovieGraphImagePassBaseNode::SetupImpl(const FMovieGraphRenderPassSetupData& InSetupData)
+{
+	// To make the implementation simpler, we make one instance of FMovieGraphDeferredRenderPas
+	// per camera, and per render layer. These objects can pull from common pools to share state,
+	// which gives us a better overview of how many resources are being used by MRQ.
+	for (const FMovieGraphRenderPassLayerData& LayerData : InSetupData.Layers)
+	{
+		// Each node can change which renderer instance actually gets created.
+		TUniquePtr<UE::MovieGraph::Rendering::FMovieGraphImagePassBase> RendererInstance = CreateInstance();
+		if(RendererInstance)
+		{
+			RendererInstance->Setup(InSetupData.Renderer, this, LayerData);
+			CurrentInstances.Add(MoveTemp(RendererInstance));
+		}
+	}
+}
+
+void UMovieGraphImagePassBaseNode::TeardownImpl()
+{
+	// We don't need to flush the rendering commands as we assume the MovieGraph
+	// Renderer has already done that once, so all data for all passes should
+	// have been submitted to the GPU (and subsequently read back) by now.
+	for (TUniquePtr<UE::MovieGraph::Rendering::FMovieGraphImagePassBase>& Instance : CurrentInstances)
+	{
+		Instance->Teardown();
+	}
+	CurrentInstances.Reset();
+}
+
+void UMovieGraphImagePassBaseNode::RenderImpl(const FMovieGraphTraversalContext& InFrameTraversalContext, const FMovieGraphTimeStepData& InTimeData)
+{
+	for (const TUniquePtr<UE::MovieGraph::Rendering::FMovieGraphImagePassBase>& Instance : CurrentInstances)
+	{
+		UMoviePipelineRenderLayerSubsystem* LayerSubsystem =
+			Instance->GetRenderer()->GetWorld()->GetSubsystem<UMoviePipelineRenderLayerSubsystem>();
+		
+		// Apply all modifiers in the evaluated graph
+		if (LayerSubsystem)
+		{
+			TRACE_CPUPROFILER_EVENT_SCOPE(MRQ::DeferredRender::ActivateRenderLayer);
+			LayerSubsystem->SetActiveRenderLayerByName(Instance->GetBranchName());
+		}
+		
+		Instance->Render(InFrameTraversalContext, InTimeData);
+
+		// Revert all modifiers
+		if (LayerSubsystem)
+		{
+			TRACE_CPUPROFILER_EVENT_SCOPE(MRQ::DeferredRender::RevertRenderLayer);
+			LayerSubsystem->ClearActiveRenderLayer();
+		}
+	}
+}
+
+void UMovieGraphImagePassBaseNode::GatherOutputPassesImpl(TArray<FMovieGraphRenderDataIdentifier>& OutExpectedPasses) const
+{
+	for (const TUniquePtr<UE::MovieGraph::Rendering::FMovieGraphImagePassBase>& Instance : CurrentInstances)
+	{
+		Instance->GatherOutputPasses(OutExpectedPasses);
+	}
+}
+
+EViewModeIndex UMovieGraphImagePassBaseNode::GetViewModeIndex() const
+{
+	return VMI_Lit;
+}
+
+FEngineShowFlags UMovieGraphImagePassBaseNode::GetShowFlags() const
+{
+	return ShowFlags->GetShowFlags();
+}
+
+TArray<uint32> UMovieGraphImagePassBaseNode::GetDefaultShowFlags() const
+{
+	return TArray<uint32>();
+}
+
+void UMovieGraphImagePassBaseNode::AddReferencedObjects(UObject* InThis, FReferenceCollector& Collector)
+{
+	Super::AddReferencedObjects(InThis, Collector);
+	
+	UMovieGraphImagePassBaseNode* This = CastChecked<UMovieGraphImagePassBaseNode>(InThis);
+	for (TUniquePtr<UE::MovieGraph::Rendering::FMovieGraphImagePassBase>& Instance : This->CurrentInstances)
+	{
+		Instance->AddReferencedObjects(Collector);
+	}
+}
\ No newline at end of file
diff --git a/Engine/Plugins/MovieScene/MovieRenderPipeline/Source/MovieRenderPipelineCore/Private/Graph/Renderers/MovieGraphDeferredPass.cpp b/Engine/Plugins/MovieScene/MovieRenderPipeline/Source/MovieRenderPipelineCore/Private/Graph/Renderers/MovieGraphDeferredPass.cpp
new file mode 100644
--- /dev/null
+++ b/Engine/Plugins/MovieScene/MovieRenderPipeline/Source/MovieRenderPipelineCore/Private/Graph/Renderers/MovieGraphDeferredPass.cpp
@@ -0,0 +1,322 @@
+// Copyright Epic Games, Inc. All Rights Reserved.
+
+#include "Graph/Renderers/MovieGraphDeferredPass.h"
+#include "Graph/Nodes/MovieGraphDeferredPassNode.h"
+#include "Graph/Nodes/MovieGraphOutputSettingNode.h"
+#include "Graph/Nodes/MovieGraphCameraNode.h"
+#include "Graph/MovieGraphDataTypes.h"
+#include "Graph/MovieGraphDefaultRenderer.h"
+#include "Graph/MovieGraphPipeline.h"
+#include "Graph/MoviePipelineRenderLayerSubsystem.h"
+#include "Graph/MovieGraphBlueprintLibrary.h"
+#include "MovieRenderOverlappedImage.h"
+#include "MoviePipelineSurfaceReader.h"
+#include "EngineModule.h"
+#include "SceneManagement.h"
+#include "CanvasTypes.h"
+#include "Engine/EngineBaseTypes.h"
+#include "Engine/TextureRenderTarget2D.h"
+#include "LegacyScreenPercentageDriver.h"
+#include "OpenColorIODisplayExtension.h"
+#include "TextureResource.h"
+#include "Tasks/Task.h"
+#include "SceneView.h"
+
+// For the 1D Weight table for accumulation
+#include "MovieRenderPipelineDataTypes.h"
+// For Sub-Pixel Jitter
+#include "MoviePipelineUtils.h"
+
+namespace UE::MovieGraph::Rendering
+{
+
+void FMovieGraphDeferredPass::Setup(TWeakObjectPtr<UMovieGraphDefaultRenderer> InRenderer, TWeakObjectPtr<UMovieGraphImagePassBaseNode> InRenderPassNode, const FMovieGraphRenderPassLayerData& InLayer)
+{
+	FMovieGraphImagePassBase::Setup(InRenderer, InRenderPassNode, InLayer);
+	
+	LayerData = InLayer;
+
+
+	RenderDataIdentifier.RootBranchName = LayerData.BranchName;
+	RenderDataIdentifier.RendererName = InRenderPassNode->GetRendererName();
+	RenderDataIdentifier.SubResourceName = TEXT("beauty");
+	
+	UE::MovieGraph::DefaultRenderer::FCameraInfo CameraInfo = InRenderer->GetCameraInfo(LayerData.CameraIdentifier);
+	RenderDataIdentifier.CameraName =  CameraInfo.CameraName;
+
+	SceneViewState.Allocate(InRenderer->GetWorld()->GetFeatureLevel());
+}
+
+void FMovieGraphDeferredPass::Teardown()
+{
+	FSceneViewStateInterface* SceneViewRef = SceneViewState.GetReference();
+	if (SceneViewRef)
+	{
+		SceneViewRef->ClearMIDPool();
+	}
+	SceneViewState.Destroy();
+}
+
+void FMovieGraphDeferredPass::GatherOutputPasses(TArray<FMovieGraphRenderDataIdentifier>& OutExpectedPasses) const
+{
+	FMovieGraphImagePassBase::GatherOutputPasses(OutExpectedPasses);
+	
+	// Add our pre-calculated identifier
+	OutExpectedPasses.Add(RenderDataIdentifier);
+}
+
+void FMovieGraphDeferredPass::AddReferencedObjects(FReferenceCollector& Collector)
+{
+	FSceneViewStateInterface* SceneViewRef = SceneViewState.GetReference();
+	if (SceneViewRef)
+	{
+		SceneViewRef->AddReferencedObjects(Collector);
+	}
+}
+
+void FMovieGraphDeferredPass::Render(const FMovieGraphTraversalContext& InFrameTraversalContext, const FMovieGraphTimeStepData& InTimeData)
+{
+	// ToDo: InFrameTraversalContext includes a copy of TimeData, but may be the one cached at the first temporal sample,
+	// maybe we can combine, maybe we can't?
+	TObjectPtr<UMovieGraphDefaultRenderer> GraphRenderer = GetRenderer().Get();
+	if (!GraphRenderer)
+	{
+		return;
+	}
+
+	// ToDo: This means you have to use FMovieGraphDeferredPass with a DeferredNode (ie: implying that the Path Tracer node has to inherit from the Defered Node,
+	// which incorrectly picks up settings not supported by the path tracer. Perhaps these should be spun off to virtual functions that can be overwritten on 
+	// FMovieGraphPathTracerPass to cast the other way?
+	const bool bIncludeCDOs = false;
+	const UMovieGraphDeferredRenderPassNode* ParentNode = InTimeData.EvaluatedConfig->GetSettingForBranch<UMovieGraphDeferredRenderPassNode>(LayerData.BranchName, bIncludeCDOs);
+	if (!ensureMsgf(ParentNode, TEXT("DeferredPass should not exist without parent node in graph.")))
+	{
+		return;
+	}
+
+	const bool bWriteAllSamples = ParentNode->bWriteAllSamples;
+	int32 NumSpatialSamples = FMath::Max(1, ParentNode->SpatialSampleCount);
+	const bool bDisableToneCurve = ParentNode->bDisableToneCurve;
+	const EAntiAliasingMethod AntiAliasingMethod = ParentNode->AntiAliasingMethod;
+	float OverscanPercentage = 0.f;
+	const float TileOverlapPadRatio = 0.0f; // No tiling support right now
+
+	// Camera nodes are optional
+	const UMovieGraphCameraSettingNode* CameraNode = InTimeData.EvaluatedConfig->GetSettingForBranch<UMovieGraphCameraSettingNode>(LayerData.BranchName, bIncludeCDOs);
+	if (CameraNode)
+	{
+		OverscanPercentage = FMath::Clamp(0.f, 1.f, CameraNode->OverscanPercentage / 100.f);
+	}
+	
+	FIntPoint AccumulatorResolution = UMovieGraphBlueprintLibrary::GetEffectiveOutputResolution(InTimeData.EvaluatedConfig, LayerData.BranchName);
+	// ToDo: When tiling is used, this should be the size of the per-tile backbuffer
+	FIntPoint BackbufferResolution = AccumulatorResolution;
+	// ToDo: This math probably needs the per-tile, pre-overlapped size? 
+	FIntPoint OverlappedPad = FIntPoint(FMath::CeilToInt(BackbufferResolution.X * TileOverlapPadRatio), FMath::CeilToInt(BackbufferResolution.Y * TileOverlapPadRatio));
+	// Calculate a backbuffer
+	UE::MovieGraph::DefaultRenderer::FRenderTargetInitParams RenderTargetInitParams;
+	{
+		RenderTargetInitParams.Size = BackbufferResolution;
+
+		// OCIO: Since this is a manually created Render target we don't need Gamma to be applied.
+		// We use this render target to render to via a display extension that utilizes Display Gamma
+		// which has a default value of 2.2 (DefaultDisplayGamma), therefore we need to set Gamma on this render target to 2.2 to cancel out any unwanted effects.
+		RenderTargetInitParams.TargetGamma = FOpenColorIORendering::DefaultDisplayGamma;
+		RenderTargetInitParams.PixelFormat = EPixelFormat::PF_FloatRGBA;
+	}
+
+	UTextureRenderTarget2D* RenderTarget = GraphRenderer->GetOrCreateViewRenderTarget(RenderTargetInitParams);
+	FRenderTarget* RenderTargetResource = RenderTarget->GameThread_GetRenderTargetResource();
+	check(RenderTargetResource);
+
+	// Now we can actually construct our ViewFamily, SceneView, and submit it for Rendering + Readback
+	for(int32 SpatialIndex = 0; SpatialIndex < NumSpatialSamples; SpatialIndex++)
+	{
+		// World should be paused for every spatial sample except the last one, so that
+		// the view doesn't update histories until the end, allowing us to render the same
+		// scene multiple times.
+		const bool bWorldIsPaused = SpatialIndex < (NumSpatialSamples - 1);
+		const int32 FrameIndex = InTimeData.RenderedFrameNumber * ((InTimeData.TemporalSampleCount * NumSpatialSamples) + (InTimeData.TemporalSampleIndex * NumSpatialSamples)) + SpatialIndex;
+
+		// We only allow a spatial jitter if we have more than one sample
+		FVector2f SpatialShiftAmount = FVector2f(0.f, 0.f);
+		const bool bAntiAliasingAllowsJitter = AntiAliasingMethod == EAntiAliasingMethod::AAM_None;
+		const bool bSampleCountsAllowsJitter = NumSpatialSamples > 1 || InTimeData.TemporalSampleCount > 1;
+		if (bAntiAliasingAllowsJitter && bSampleCountsAllowsJitter)
+		{
+			const int32 NumSamplesPerOutputFrame = NumSpatialSamples * InTimeData.TemporalSampleCount;
+			SpatialShiftAmount = UE::MoviePipeline::GetSubPixelJitter(FrameIndex, NumSamplesPerOutputFrame);
+		}
+		 
+		// These are the parameters of our camera 
+		UE::MovieGraph::DefaultRenderer::FCameraInfo CameraInfo;
+
+		// ToDo: Get this from the renderer based on LayerData.CameraIdentifier for eventual multi-camera support
+		APlayerController* LocalPlayerController = GraphRenderer->GetWorld()->GetFirstPlayerController();
+		// CameraAnim override
+		if (LocalPlayerController->PlayerCameraManager)
+		{
+			CameraInfo.ViewInfo = LocalPlayerController->PlayerCameraManager->GetCameraCacheView();
+		}
+
+
+		CameraInfo.bAllowCameraAspectRatio = true;
+		CameraInfo.TilingParams.TileSize = BackbufferResolution;
+		CameraInfo.TilingParams.OverlapPad = FVector2f(0.f, 0.f); // No tile support
+		CameraInfo.TilingParams.TileCount = FIntPoint(1, 1); // No tile support
+		CameraInfo.TilingParams.TileIndexes = FIntPoint(0, 0); // No tile support
+		CameraInfo.OverscanFraction = OverscanPercentage;
+		CameraInfo.ProjectionMatrixJitterAmount = FVector2D((SpatialShiftAmount.X) * 2.0f / (float)BackbufferResolution.X, SpatialShiftAmount.Y * -2.0f / (float)BackbufferResolution.Y);
+
+		// For this particular tile, what is the offset into the output image
+		FIntPoint OverlappedOffset = FIntPoint(CameraInfo.TilingParams.TileIndexes.X * BackbufferResolution.X - OverlappedPad.X, CameraInfo.TilingParams.TileIndexes.Y * BackbufferResolution.Y - OverlappedPad.Y);
+		
+		// Move the final render by this much in the accumulator to counteract the offset put into the view matrix.
+		// Note that when bAllowSpatialJitter is false, SpatialShiftX/Y will always be zero.
+		FVector2D OverlappedSubpixelShift = FVector2D(0.5f - SpatialShiftAmount.X, 0.5f - SpatialShiftAmount.Y);
+		
+		FMatrix ProjectionMatrix = CalculateProjectionMatrix(CameraInfo);
+		float DoFSensorScale = 1.f;
+
+		// Modify the perspective matrix to do an off center projection, with overlap for high-res tiling
+		const bool bOrthographic = CameraInfo.ViewInfo.ProjectionMode == ECameraProjectionMode::Type::Orthographic;
+		ModifyProjectionMatrixForTiling(CameraInfo.TilingParams, bOrthographic, ProjectionMatrix, DoFSensorScale);
+		
+		CameraInfo.ProjectionMatrix = ProjectionMatrix;
+		CameraInfo.DoFSensorScale = DoFSensorScale;
+
+		// The Scene View Family must be constructed first as the FSceneView needs it to be constructed
+		UE::MovieGraph::Rendering::FViewFamilyInitData ViewFamilyInitData;
+		ViewFamilyInitData.RenderTarget = RenderTargetResource;
+		ViewFamilyInitData.World = GraphRenderer->GetWorld();
+		ViewFamilyInitData.TimeData = InTimeData;
+		ViewFamilyInitData.SceneCaptureSource = bDisableToneCurve ? ESceneCaptureSource::SCS_FinalColorHDR : ESceneCaptureSource::SCS_FinalToneCurveHDR;
+		ViewFamilyInitData.bWorldIsPaused = bWorldIsPaused;
+		ViewFamilyInitData.FrameIndex = FrameIndex;
+		ViewFamilyInitData.AntiAliasingMethod = AntiAliasingMethod;
+		
+		TSharedRef<FSceneViewFamilyContext> ViewFamily = CreateSceneViewFamily(ViewFamilyInitData, CameraInfo);
+		 
+		// Now we can construct a View to go within this family.
+		FSceneViewInitOptions SceneViewInitOptions = CreateViewInitOptions(CameraInfo, ViewFamily.ToSharedPtr().Get(), SceneViewState);
+		FSceneView* NewView = CreateSceneView(SceneViewInitOptions, ViewFamily, CameraInfo);
+		
+		// Then apply Movie Render Queue specific overrides to the ViewFamily, and then to the SceneView.
+		ApplyMovieGraphOverridesToViewFamily(ViewFamily, ViewFamilyInitData);
+		
+		// ToDo: This really only needs access to the ViewFamily for path tracer related things,
+		// and would rather just take a FSceneView* 
+		ApplyMovieGraphOverridesToSceneView(ViewFamily, ViewFamilyInitData, CameraInfo);
+		
+		FHitProxyConsumer* HitProxyConsumer = nullptr;
+		const float DPIScale = 1.0f;
+		FCanvas Canvas = FCanvas(RenderTargetResource, HitProxyConsumer, GraphRenderer->GetWorld(), GraphRenderer->GetWorld()->GetFeatureLevel(), FCanvas::CDM_DeferDrawing, DPIScale);
+		
+		// Submit the renderer to be rendered
+		GetRendererModule().BeginRenderingViewFamily(&Canvas, ViewFamily.ToSharedPtr().Get());
+
+		// Construct the sample state that reflects the current render sample
+		UE::MovieGraph::FMovieGraphSampleState SampleState;
+		{
+			// Take our per-frame Traversal Context and update it with context specific to this sample.
+			FMovieGraphTraversalContext UpdatedTraversalContext = InFrameTraversalContext;
+			UpdatedTraversalContext.Time = InTimeData;
+			UpdatedTraversalContext.RenderDataIdentifier = RenderDataIdentifier;
+
+			SampleState.TraversalContext = MoveTemp(UpdatedTraversalContext);
+			SampleState.BackbufferResolution = BackbufferResolution;
+			SampleState.AccumulatorResolution = AccumulatorResolution;
+			SampleState.bWriteSampleToDisk = bWriteAllSamples;
+			SampleState.bRequiresAccumulator = InTimeData.bRequiresAccumulator || (NumSpatialSamples > 1);
+			SampleState.bFetchFromAccumulator = InTimeData.bIsLastTemporalSampleForFrame && (SpatialIndex == (NumSpatialSamples - 1));
+			SampleState.OverlappedPad = OverlappedPad;
+			SampleState.OverlappedOffset = OverlappedOffset;
+			SampleState.OverlappedSubpixelShift = OverlappedSubpixelShift;
+
+		}
+
+		// If this was just to contribute to the history buffer, no need to go any further.
+		if (InTimeData.bDiscardOutput)
+		{
+			return;
+		}
+
+		// Readback + Accumulate.
+		PostRendererSubmission(SampleState, RenderTargetInitParams, Canvas, CameraInfo);
+	}
+}
+
+void FMovieGraphDeferredPass::PostRendererSubmission(
+	const UE::MovieGraph::FMovieGraphSampleState& InSampleState,
+	const UE::MovieGraph::DefaultRenderer::FRenderTargetInitParams& InRenderTargetInitParams, FCanvas& InCanvas, const UE::MovieGraph::DefaultRenderer::FCameraInfo& InCameraInfo)
+{
+	TObjectPtr<UMovieGraphDefaultRenderer> GraphRenderer = GetRenderer().Get();
+	if (!GraphRenderer)
+	{
+		return;
+	}
+
+	// Draw letterboxing
+	// ToDo: Multi-camera support
+	APlayerCameraManager* PlayerCameraManager = GraphRenderer->GetWorld()->GetFirstPlayerController()->PlayerCameraManager;
+	if(PlayerCameraManager && PlayerCameraManager->GetCameraCacheView().bConstrainAspectRatio)
+	{
+		const FMinimalViewInfo CameraCache = PlayerCameraManager->GetCameraCacheView();
+		
+		// Taking overscan into account.
+		const FIntPoint FullOutputSize = InSampleState.AccumulatorResolution;
+	
+		const float OutputSizeAspectRatio = FullOutputSize.X / (float)FullOutputSize.Y;
+		const float CameraAspectRatio = InCameraInfo.bAllowCameraAspectRatio ? CameraCache.AspectRatio : OutputSizeAspectRatio;
+	
+		const FIntPoint ConstrainedFullSize = CameraAspectRatio > OutputSizeAspectRatio ?
+			FIntPoint(FullOutputSize.X, FMath::CeilToInt((double)FullOutputSize.X / (double)CameraAspectRatio)) :
+			FIntPoint(FMath::CeilToInt(CameraAspectRatio * FullOutputSize.Y), FullOutputSize.Y);
+	
+		const FIntPoint TileViewMin = InSampleState.OverlappedOffset;
+		const FIntPoint TileViewMax = TileViewMin + InSampleState.BackbufferResolution;
+	
+		// Camera ratio constrained rect, clipped by the tile rect
+		FIntPoint ConstrainedViewMin = (FullOutputSize - ConstrainedFullSize) / 2;
+		FIntPoint ConstrainedViewMax = ConstrainedViewMin + ConstrainedFullSize;
+		ConstrainedViewMin = FIntPoint(FMath::Clamp(ConstrainedViewMin.X, TileViewMin.X, TileViewMax.X),
+			FMath::Clamp(ConstrainedViewMin.Y, TileViewMin.Y, TileViewMax.Y));
+		ConstrainedViewMax = FIntPoint(FMath::Clamp(ConstrainedViewMax.X, TileViewMin.X, TileViewMax.X),
+			FMath::Clamp(ConstrainedViewMax.Y, TileViewMin.Y, TileViewMax.Y));
+	
+		// Difference between the clipped constrained rect and the tile rect
+		const FIntPoint OffsetMin = ConstrainedViewMin - TileViewMin;
+		const FIntPoint OffsetMax = TileViewMax - ConstrainedViewMax;
+	
+		// Clear left
+		if (OffsetMin.X > 0)
+		{
+			InCanvas.DrawTile(0, 0, OffsetMin.X, InSampleState.BackbufferResolution.Y,
+				0.0f, 0.0f, 1.0f, 1.0f, FLinearColor::Black, nullptr, false);
+		}
+		// Clear right
+		if (OffsetMax.X > 0)
+		{
+			InCanvas.DrawTile(InSampleState.BackbufferResolution.X - OffsetMax.X, 0, InSampleState.BackbufferResolution.X, InSampleState.BackbufferResolution.Y,
+				0.0f, 0.0f, 1.0f, 1.0f, FLinearColor::Black, nullptr, false);
+		}
+		// Clear top
+		if (OffsetMin.Y > 0)
+		{
+			InCanvas.DrawTile(0, 0, InSampleState.BackbufferResolution.X, OffsetMin.Y,
+				0.0f, 0.0f, 1.0f, 1.0f, FLinearColor::Black, nullptr, false);
+		}
+		// Clear bottom
+		if (OffsetMax.Y > 0)
+		{
+			InCanvas.DrawTile(0, InSampleState.BackbufferResolution.Y - OffsetMax.Y, InSampleState.BackbufferResolution.X, InSampleState.BackbufferResolution.Y,
+				0.0f, 0.0f, 1.0f, 1.0f, FLinearColor::Black, nullptr, false);
+		}
+	
+		InCanvas.Flush_GameThread(true);
+	}
+	
+	FMovieGraphImagePassBase::PostRendererSubmission(InSampleState, InRenderTargetInitParams, InCanvas, InCameraInfo);
+}
+} // namespace UE::MovieGraph::Rendering
+
diff --git a/Engine/Plugins/MovieScene/MovieRenderPipeline/Source/MovieRenderPipelineCore/Private/Graph/Renderers/MovieGraphImagePassBase.cpp b/Engine/Plugins/MovieScene/MovieRenderPipeline/Source/MovieRenderPipelineCore/Private/Graph/Renderers/MovieGraphImagePassBase.cpp
new file mode 100644
--- /dev/null
+++ b/Engine/Plugins/MovieScene/MovieRenderPipeline/Source/MovieRenderPipelineCore/Private/Graph/Renderers/MovieGraphImagePassBase.cpp
@@ -0,0 +1,693 @@
+// Copyright Epic Games, Inc. All Rights Reserved.
+
+#include "Graph/Renderers/MovieGraphImagePassBase.h"
+#include "Graph/MovieGraphDataTypes.h"
+#include "Graph/MovieGraphPipeline.h"
+#include "SceneView.h"
+#include "MoviePipelineUtils.h"
+#include "LegacyScreenPercentageDriver.h"
+#include "MoviePipelineSurfaceReader.h"
+#include "CanvasTypes.h"
+#include "MovieRenderOverlappedImage.h"
+#include "Engine/RendererSettings.h"
+#include "UnrealClient.h"
+
+namespace UE::MovieGraph::Rendering
+{
+
+void FMovieGraphImagePassBase::Setup(TWeakObjectPtr<UMovieGraphDefaultRenderer> InRenderer, TWeakObjectPtr<UMovieGraphImagePassBaseNode> InRenderPassNode, const FMovieGraphRenderPassLayerData& InLayer)
+{
+	// This is a pointer to the UMovieGraphPipeline's renderer which is valid throughout the entire render.
+	WeakGraphRenderer = InRenderer;
+}
+
+FSceneViewInitOptions FMovieGraphImagePassBase::CreateViewInitOptions(const UE::MovieGraph::DefaultRenderer::FCameraInfo& InCameraInfo, FSceneViewFamilyContext* InViewFamily, FSceneViewStateReference& InViewStateRef) const
+{
+	check(InViewFamily);
+
+	FIntPoint RenderResolution = InViewFamily->RenderTarget->GetSizeXY();
+
+	FSceneViewInitOptions ViewInitOptions;
+	ViewInitOptions.ViewFamily = InViewFamily;
+	ViewInitOptions.ViewOrigin = InCameraInfo.ViewInfo.Location;
+	ViewInitOptions.ViewLocation = InCameraInfo.ViewInfo.Location;
+	ViewInitOptions.ViewRotation = InCameraInfo.ViewInfo.Rotation;
+	ViewInitOptions.SetViewRectangle(FIntRect(FIntPoint(0, 0), RenderResolution));
+	ViewInitOptions.ViewRotationMatrix = FInverseRotationMatrix(InCameraInfo.ViewInfo.Rotation);
+	ViewInitOptions.ViewActor = InCameraInfo.ViewActor;
+	ViewInitOptions.ProjectionMatrix = InCameraInfo.ProjectionMatrix;
+
+	// Rotate the view 90 degrees to match the rest of the engine.
+	ViewInitOptions.ViewRotationMatrix = ViewInitOptions.ViewRotationMatrix * FMatrix(
+		FPlane(0, 0, 1, 0),
+		FPlane(1, 0, 0, 0),
+		FPlane(0, 1, 0, 0),
+		FPlane(0, 0, 0, 1));
+	
+	float ViewFOV = InCameraInfo.ViewInfo.FOV;
+
+	// Inflate our FOV to support the overscan 
+	// ToDo: This is a duplicate of the logic in CalculateProjectionMatrix, should combine.
+	ViewFOV = 2.0f * FMath::RadiansToDegrees(FMath::Atan((1.0f + InCameraInfo.OverscanFraction) * FMath::Tan(FMath::DegreesToRadians( ViewFOV * 0.5f ))));
+
+	ViewInitOptions.SceneViewStateInterface = InViewStateRef.GetReference();
+	ViewInitOptions.FOV = ViewFOV;
+	ViewInitOptions.DesiredFOV = ViewFOV;
+	ViewInitOptions.bUseFauxOrthoViewPos = InCameraInfo.ViewInfo.ProjectionMode == ECameraProjectionMode::Orthographic;
+	
+	return ViewInitOptions;
+}
+
+
+
+FSceneView* FMovieGraphImagePassBase::CreateSceneView(const FSceneViewInitOptions& InInitOptions, TSharedRef<FSceneViewFamilyContext> InViewFamily, const UE::MovieGraph::DefaultRenderer::FCameraInfo& InCameraInfo) const
+{
+	// Create our view based on the Init Options
+	FSceneView* View = new FSceneView(InInitOptions);
+	InViewFamily->Views.Add(View);
+
+	View->StartFinalPostprocessSettings(InInitOptions.ViewLocation);
+	ApplyCameraManagerPostProcessBlends(View);
+
+	// Scaling sensor size inversely with the the projection matrix [0][0] should physically
+	// cause the circle of confusion to be unchanged.
+	View->FinalPostProcessSettings.DepthOfFieldSensorWidth *= InCameraInfo.DoFSensorScale;
+	// Modify the 'center' of the lens to be offset for high-res tiling, helps some effects (vignette) etc. still work.
+	View->LensPrincipalPointOffsetScale = CalculatePrinciplePointOffsetForTiling(InCameraInfo.TilingParams);
+	View->EndFinalPostprocessSettings(InInitOptions);
+
+	return View;
+}
+
+void FMovieGraphImagePassBase::ApplyCameraManagerPostProcessBlends(FSceneView* InView) const
+{
+	check(InView);
+
+	TObjectPtr<UMovieGraphDefaultRenderer> GraphRenderer = GetRenderer().Get();
+	if (!GraphRenderer)
+	{
+		return;
+	}
+
+	APlayerController* LocalPlayerController = GraphRenderer->GetWorld()->GetFirstPlayerController();
+	// CameraAnim override
+	if (LocalPlayerController->PlayerCameraManager)
+	{
+		TArray<FPostProcessSettings> const* CameraAnimPPSettings;
+		TArray<float> const* CameraAnimPPBlendWeights;
+		LocalPlayerController->PlayerCameraManager->GetCachedPostProcessBlends(CameraAnimPPSettings, CameraAnimPPBlendWeights);
+
+		if (LocalPlayerController->PlayerCameraManager->bEnableFading)
+		{
+			InView->OverlayColor = LocalPlayerController->PlayerCameraManager->FadeColor;
+			InView->OverlayColor.A = FMath::Clamp(LocalPlayerController->PlayerCameraManager->FadeAmount, 0.f, 1.f);
+		}
+
+		if (LocalPlayerController->PlayerCameraManager->bEnableColorScaling)
+		{
+			FVector ColorScale = LocalPlayerController->PlayerCameraManager->ColorScale;
+			InView->ColorScale = FLinearColor(ColorScale.X, ColorScale.Y, ColorScale.Z);
+		}
+
+		FMinimalViewInfo ViewInfo = LocalPlayerController->PlayerCameraManager->GetCameraCacheView();
+		for (int32 PPIdx = 0; PPIdx < CameraAnimPPBlendWeights->Num(); ++PPIdx)
+		{
+			InView->OverridePostProcessSettings((*CameraAnimPPSettings)[PPIdx], (*CameraAnimPPBlendWeights)[PPIdx]);
+		}
+
+		InView->OverridePostProcessSettings(ViewInfo.PostProcessSettings, ViewInfo.PostProcessBlendWeight);
+	}
+}
+
+TSharedRef<FSceneViewFamilyContext> FMovieGraphImagePassBase::CreateSceneViewFamily(const FViewFamilyInitData& InInitData, const UE::MovieGraph::DefaultRenderer::FCameraInfo& InCameraInfo) const
+{
+	// FEngineShowFlags NewShowFlags = Cast<FMovieGraphImagePassBase>(LayerData.RenderPassNode)->GetShowFlags();
+	// const EViewModeIndex ViewModeIndex = Cast<FMovieGraphImagePassBase>(LayerData.RenderPassNode)->GetViewModeIndex();
+	EViewModeIndex ViewModeIndex = VMI_Lit;
+	FEngineShowFlags NewShowFlags = FEngineShowFlags(ESFIM_Game);
+
+	const bool bIsPerspective = InCameraInfo.ViewInfo.ProjectionMode == ECameraProjectionMode::Type::Perspective;
+
+	// Allow the Engine Showflag system to override our engine showflags, based on our view mode index.
+	// This is required for certain debug view modes (to have matching show flags set for rendering).
+	ApplyViewMode(/*In*/ ViewModeIndex, bIsPerspective, /*InOut*/NewShowFlags);
+
+	// And then we have to let another system override them again (based on cvars, etc.)
+	EngineShowFlagOverride(ESFIM_Game, ViewModeIndex, NewShowFlags, false);
+
+	TSharedRef<FSceneViewFamilyContext> OutViewFamily = MakeShared<FSceneViewFamilyContext>(FSceneViewFamily::ConstructionValues(
+		InInitData.RenderTarget,
+		InInitData.World->Scene,
+		NewShowFlags)
+		.SetTime(FGameTime::CreateUndilated(InInitData.TimeData.WorldSeconds, InInitData.TimeData.FrameDeltaTime))
+		.SetRealtimeUpdate(true));
+
+
+	return OutViewFamily;
+}
+
+void FMovieGraphImagePassBase::ApplyMovieGraphOverridesToViewFamily(TSharedRef<FSceneViewFamilyContext> InOutFamily, const FViewFamilyInitData& InInitData) const
+{
+	// Used to specify if the Tone Curve is being applied or not to our Linear Output data
+	InOutFamily->SceneCaptureSource = InInitData.SceneCaptureSource;
+	InOutFamily->bWorldIsPaused = InInitData.bWorldIsPaused;
+	// InOutFamily->ViewMode = ViewModeIndex;
+	InOutFamily->bOverrideVirtualTextureThrottle = true;
+	
+	// This is for a workaround involving GFrameCounter already being incremented by the time MRQ renders are submitted
+	// meaning that it's out of sync with some parts of the renderer's idea of the current frame.
+	PRAGMA_DISABLE_DEPRECATION_WARNINGS
+	InOutFamily->OverrideFrameCounter = UE::MovieRenderPipeline::GetRendererFrameCount();
+	PRAGMA_ENABLE_DEPRECATION_WARNINGS
+	
+	// We need to check if this is the first FSceneView being submitted to the renderer module, and set some flags on the ViewFamily for ensuring some
+	// parts of the renderer only get updated once per frame. Kept as an if/else statement to avoid the confusion with setting all of these values to 
+	// some permutation of !/!!bHasRenderedFirstViewThisFrame.
+	TObjectPtr<UMovieGraphDefaultRenderer> GraphRenderer = GetRenderer().Get();
+	if (GraphRenderer)
+	{
+		if (!GraphRenderer->GetHasRenderedFirstViewThisFrame())
+		{
+			// Update our renderer
+			GraphRenderer->SetHasRenderedFirstViewThisFrame(true);
+
+			InOutFamily->bIsFirstViewInMultipleViewFamily = true;
+			InOutFamily->bAdditionalViewFamily = false;
+		}
+		else
+		{
+			InOutFamily->bIsFirstViewInMultipleViewFamily = false;
+			InOutFamily->bAdditionalViewFamily = true;
+		}
+
+		InOutFamily->bIsMultipleViewFamily = true;
+	}
+	
+	// Skip the whole pass if they don't want motion blur.
+	if (FMath::IsNearlyZero(InInitData.TimeData.MotionBlurFraction))
+	{
+		InOutFamily->EngineShowFlags.SetMotionBlur(false);
+	}
+	
+
+	// ToDo: Let settings modify the ScreenPercentageInterface so third party screen percentages are supported.
+
+	// If UMoviePipelineViewFamilySetting never set a Screen percentage interface we fallback to default.
+	if (InOutFamily->GetScreenPercentageInterface() == nullptr)
+	{
+		InOutFamily->SetScreenPercentageInterface(new FLegacyScreenPercentageDriver(*InOutFamily, FLegacyScreenPercentageDriver::GetCVarResolutionFraction()));
+	}
+}
+
+void FMovieGraphImagePassBase::ApplyMovieGraphOverridesToSceneView(TSharedRef<FSceneViewFamilyContext> InOutFamily, const FViewFamilyInitData& InInitData, const UE::MovieGraph::DefaultRenderer::FCameraInfo& InCameraInfo) const
+{
+	// Override the view's FrameIndex to be based on our progress through the sequence. This greatly increases
+	// determinism with things like TAA.
+	FSceneView* View = const_cast<FSceneView*>(InOutFamily->Views[0]);
+	View->OverrideFrameIndexValue = InInitData.FrameIndex;
+	// Each shot should initialize a scene history from scratch so there should be no need to do an extra camera cut flag.
+	View->bCameraCut = false; 
+	View->AntiAliasingMethod = InInitData.AntiAliasingMethod;
+	View->bIsOfflineRender = true;
+
+	// Override the Motion Blur settings since these are controlled by the movie pipeline.
+	{
+		// We need to inversly scale the target FPS by time dilation to counteract slowmo. If scaling isn't applied then motion blur length
+		// stays the same length despite the smaller delta time and the blur ends up too long.
+		View->FinalPostProcessSettings.MotionBlurTargetFPS = FMath::RoundToInt(InInitData.TimeData.FrameRate.AsDecimal() / FMath::Max(SMALL_NUMBER, InInitData.TimeData.WorldTimeDilation));
+		View->FinalPostProcessSettings.MotionBlurAmount = InInitData.TimeData.MotionBlurFraction;
+		View->FinalPostProcessSettings.MotionBlurMax = 100.f;
+		View->FinalPostProcessSettings.bOverride_MotionBlurAmount = true;
+		View->FinalPostProcessSettings.bOverride_MotionBlurTargetFPS = true;
+		View->FinalPostProcessSettings.bOverride_MotionBlurMax = true;
+
+	}
+	
+	// Path Tracer Sampling
+	if (InOutFamily->EngineShowFlags.PathTracing)
+	{
+		// Override whatever settings came from PostProcessVolume or Camera
+
+		// If motion blur is enabled:
+		//    blend all spatial samples together while leaving the handling of temporal samples up to MRQ
+		//    each temporal sample will include denoising and post-process effects
+		// If motion blur is NOT enabled:
+		//    blend all temporal+spatial samples within the path tracer and only apply denoising on the last temporal sample
+		//    this way we minimize denoising cost and also allow a much higher number of temporal samples to be used which
+		//    can help reduce strobing
+
+		// NOTE: Tiling is not compatible with the reference motion blur mode because it changes the order of the loops over the image.
+		//const bool bHasTiles = InInitData.CameraInfo.TilingParams.TileCount.X > 1 || InInitData.CameraInfo.TilingParams.TileCount.Y;
+		//const bool bAccumulateSpatialSamplesOnly = InOutFamily->EngineShowFlags.MotionBlur || bHasTiles;
+		//
+		//const int32 SampleCount = bAccumulateSpatialSamplesOnly ? InOutSampleState.SpatialSampleCount : InOutSampleState.TemporalSampleCount * InOutSampleState.SpatialSampleCount;
+		//const int32 SampleIndex = bAccumulateSpatialSamplesOnly ? InOutSampleState.SpatialSampleIndex : InOutSampleState.TemporalSampleIndex * InOutSampleState.SpatialSampleCount + InOutSampleState.SpatialSampleIndex;
+		//
+		//// TODO: pass along FrameIndex (which includes SampleIndex) to make sure sampling is fully deterministic
+		//
+		//// Overwrite whatever sampling count came from the PostProcessVolume
+		//View->FinalPostProcessSettings.bOverride_PathTracingSamplesPerPixel = true;
+		//View->FinalPostProcessSettings.PathTracingSamplesPerPixel = SampleCount;
+		//
+		//// reset path tracer's accumulation at the start of each sample
+		//View->bForcePathTracerReset = SampleIndex == 0;
+		//
+		//// discard the result, unless its the last sample
+		//InOutSampleState.bDiscardResult |= !(SampleIndex == SampleCount - 1);
+	}
+
+	// Warn the user for invalid setting combinations / enforce hardware limitations
+	// Locked Exposure
+	const bool bAutoExposureAllowed = true; // IsAutoExposureAllowed(InOutSampleState);
+	{
+		// If the rendering pass doesn't allow autoexposure and they dont' have manual exposure set up, warn.
+		if (!bAutoExposureAllowed && (View->FinalPostProcessSettings.AutoExposureMethod != EAutoExposureMethod::AEM_Manual))
+		{
+			// Skip warning if the project setting is disabled though, as exposure will be forced off in the renderer anyways.
+			const URendererSettings* RenderSettings = GetDefault<URendererSettings>();
+			if (RenderSettings->bDefaultFeatureAutoExposure != false)
+			{
+				UE_LOG(LogMovieRenderPipeline, Warning, TEXT("Camera Auto Exposure Method not supported by one or more render passes. Change the Auto Exposure Method to Manual!"));
+				View->FinalPostProcessSettings.AutoExposureMethod = EAutoExposureMethod::AEM_Manual;
+			}
+		}
+	}
+
+	// Orthographic cameras don't support anti-aliasing outside the path tracer (other than FXAA)
+	const bool bIsOrthographicCamera = !View->IsPerspectiveProjection();
+	if (bIsOrthographicCamera)
+	{
+		bool bIsSupportedAAMethod = View->AntiAliasingMethod == EAntiAliasingMethod::AAM_FXAA;
+		bool bIsPathTracer = InOutFamily->EngineShowFlags.PathTracing;
+		bool bWarnJitters = InCameraInfo.ProjectionMatrixJitterAmount.SquaredLength() > SMALL_NUMBER;
+		if ((!bIsPathTracer && !bIsSupportedAAMethod) || bWarnJitters)
+		{
+			UE_LOG(LogMovieRenderPipeline, Warning, TEXT("Orthographic Cameras are only supported with PathTracer or Deferred with FXAA Anti-Aliasing"));
+		}
+	}
+
+	{
+		bool bMethodWasUnsupported = false;
+		if (View->AntiAliasingMethod == AAM_TemporalAA && !SupportsGen4TAA(View->GetShaderPlatform()))
+		{
+			UE_LOG(LogMovieRenderPipeline, Error, TEXT("TAA was requested but this hardware does not support it."));
+			bMethodWasUnsupported = true;
+		}
+		else if (View->AntiAliasingMethod == AAM_TSR && !SupportsTSR(View->GetShaderPlatform()))
+		{
+			UE_LOG(LogMovieRenderPipeline, Error, TEXT("TSR was requested but this hardware does not support it."));
+			bMethodWasUnsupported = true;
+		}
+
+		if (bMethodWasUnsupported)
+		{
+			View->AntiAliasingMethod = AAM_None;
+		}
+	}
+
+	// Anti Aliasing
+	{
+		// If we're not using TAA, TSR, or Path Tracing we will apply the View Matrix projection jitter. Normally TAA sets this
+		// inside FSceneRenderer::PreVisibilityFrameSetup. Path Tracing does its own anti-aliasing internally.
+		bool bApplyProjectionJitter = !bIsOrthographicCamera
+			&& !InOutFamily->EngineShowFlags.PathTracing
+			&& !IsTemporalAccumulationBasedMethod(View->AntiAliasingMethod);
+		if (bApplyProjectionJitter)
+		{
+			View->ViewMatrices.HackAddTemporalAAProjectionJitter(InCameraInfo.ProjectionMatrixJitterAmount);
+		}
+	}
+}
+
+FMatrix FMovieGraphImagePassBase::CalculateProjectionMatrix(const UE::MovieGraph::DefaultRenderer::FCameraInfo& InCameraInfo) const
+{
+	// Calculate a Projection Matrix. This code unfortunately ends up similar to, but not quite the same as FMinimalViewInfo::CalculateProjectionMatrixGivenView
+	FMatrix BaseProjMatrix;
+	
+	// Potentially modified when we modify projection matrix for tiling.
+	float DofSensorScale = 1.0f;
+
+	// TileSize should respect the actual backbuffer size being used by the render.
+	float ViewRectWidth = InCameraInfo.TilingParams.TileSize.X;
+	float ViewRectHeight = InCameraInfo.TilingParams.TileSize.Y;
+
+	const float DestAspectRatio = ViewRectWidth / ViewRectHeight;
+	const float CameraAspectRatio = InCameraInfo.bAllowCameraAspectRatio ? InCameraInfo.ViewInfo.AspectRatio : DestAspectRatio;
+	
+	float ViewFOV = InCameraInfo.ViewInfo.FOV;
+
+	// Inflate our FOV to support the overscan 
+	ViewFOV = 2.0f * FMath::RadiansToDegrees(FMath::Atan((1.0f + InCameraInfo.OverscanFraction) * FMath::Tan(FMath::DegreesToRadians(ViewFOV * 0.5f))));
+	
+	if (InCameraInfo.ViewInfo.ProjectionMode == ECameraProjectionMode::Orthographic)
+	{
+		const float YScale = 1.0f / InCameraInfo.ViewInfo.AspectRatio;
+		const float OverscanScale = 1.0f + (InCameraInfo.OverscanFraction);
+
+		const float HalfOrthoWidth = (InCameraInfo.ViewInfo.OrthoWidth / 2.0f) * OverscanScale;
+		const float ScaledOrthoHeight = (InCameraInfo.ViewInfo.OrthoWidth / 2.0f) * OverscanScale * YScale;
+
+		const float NearPlane = InCameraInfo.ViewInfo.OrthoNearClipPlane;
+		const float FarPlane = InCameraInfo.ViewInfo.OrthoFarClipPlane;
+
+		const float ZScale = 1.0f / (FarPlane - NearPlane);
+		const float ZOffset = -NearPlane;
+
+		BaseProjMatrix = FReversedZOrthoMatrix(
+			HalfOrthoWidth,
+			ScaledOrthoHeight,
+			ZScale,
+			ZOffset
+		);
+	}
+	else
+	{
+		float XAxisMultiplier;
+		float YAxisMultiplier;
+
+		if (InCameraInfo.ViewInfo.bConstrainAspectRatio)
+		{
+			// If the camera's aspect ratio has a thinner width, then stretch the horizontal fov more than usual to 
+			// account for the extra with of (before constraining - after constraining)
+			if (InCameraInfo.ViewInfo.AspectRatio < DestAspectRatio)
+			{
+				const float ConstrainedWidth = ViewRectHeight * InCameraInfo.ViewInfo.AspectRatio;
+				XAxisMultiplier = ConstrainedWidth / (float)ViewRectWidth;
+				YAxisMultiplier = InCameraInfo.ViewInfo.AspectRatio;
+			}
+			// Simplified some math here but effectively functions similarly to the above, the unsimplified code would look like:
+			// const float ConstrainedHeight = ViewRectWidth / CameraCache.AspectRatio;
+			// YAxisMultiplier = (ConstrainedHeight / ViewInitOptions.GetViewRect.Height()) * CameraCache.AspectRatio;
+			else
+			{
+				XAxisMultiplier = 1.0f;
+				YAxisMultiplier = ViewRectWidth / ViewRectHeight;
+			}
+		}
+		else
+		{
+			const EAspectRatioAxisConstraint AspectRatioAxisConstraint = GetDefault<ULocalPlayer>()->AspectRatioAxisConstraint;
+			if (((ViewRectWidth > ViewRectHeight) && (AspectRatioAxisConstraint == AspectRatio_MajorAxisFOV)) || (AspectRatioAxisConstraint == AspectRatio_MaintainXFOV))
+			{
+				//if the viewport is wider than it is tall
+				XAxisMultiplier = 1.0f;
+				YAxisMultiplier = ViewRectWidth / ViewRectHeight;
+			}
+			else
+			{
+				//if the viewport is taller than it is wide
+				XAxisMultiplier = ViewRectHeight / ViewRectWidth;
+				YAxisMultiplier = 1.0f;
+			}
+		}
+
+		const float MinZ = InCameraInfo.ViewInfo.GetFinalPerspectiveNearClipPlane();
+		const float MaxZ = MinZ;
+		// Avoid zero ViewFOV's which cause divide by zero's in projection matrix
+		const float MatrixFOV = FMath::Max(0.001f, ViewFOV) * (float)PI / 360.0f;
+
+
+		if ((bool)ERHIZBuffer::IsInverted)
+		{
+			BaseProjMatrix = FReversedZPerspectiveMatrix(
+				MatrixFOV,
+				MatrixFOV,
+				XAxisMultiplier,
+				YAxisMultiplier,
+				MinZ,
+				MaxZ
+			);
+		}
+		else
+		{
+			BaseProjMatrix = FPerspectiveMatrix(
+				MatrixFOV,
+				MatrixFOV,
+				XAxisMultiplier,
+				YAxisMultiplier,
+				MinZ,
+				MaxZ
+			);
+		}
+	}
+	
+	return BaseProjMatrix;
+}
+
+FVector4f FMovieGraphImagePassBase::CalculatePrinciplePointOffsetForTiling(const UE::MovieGraph::DefaultRenderer::FMovieGraphTilingParams& InTilingParams) const 
+{
+	// We need our final view parameters to be in the space of [-1,1], including all the tiles.
+	// Starting with a single tile, the middle of the tile in offset screen space is:
+	FVector2f TilePrincipalPointOffset;
+
+	TilePrincipalPointOffset.X = (float(InTilingParams.TileIndexes.X) + 0.5f - (0.5f * float(InTilingParams.TileCount.X))) * 2.0f;
+	TilePrincipalPointOffset.Y = (float(InTilingParams.TileIndexes.Y) + 0.5f - (0.5f * float(InTilingParams.TileCount.Y))) * 2.0f;
+
+	// For the tile size ratio, we have to multiply by (1.0 + overlap) and then divide by tile num
+	FVector2D OverlapScale;
+	OverlapScale.X = (1.0f + float(2 * InTilingParams.OverlapPad.X) / float(InTilingParams.TileSize.X));
+	OverlapScale.Y = (1.0f + float(2 * InTilingParams.OverlapPad.Y) / float(InTilingParams.TileSize.Y));
+
+	TilePrincipalPointOffset.X /= OverlapScale.X;
+	TilePrincipalPointOffset.Y /= OverlapScale.Y;
+
+	FVector2D TilePrincipalPointScale;
+	TilePrincipalPointScale.X = OverlapScale.X / float(InTilingParams.TileCount.X);
+	TilePrincipalPointScale.Y = OverlapScale.Y / float(InTilingParams.TileCount.Y);
+
+	TilePrincipalPointOffset.X *= TilePrincipalPointScale.X;
+	TilePrincipalPointOffset.Y *= TilePrincipalPointScale.Y;
+
+	return FVector4f(TilePrincipalPointOffset.X, -TilePrincipalPointOffset.Y, TilePrincipalPointScale.X, TilePrincipalPointScale.Y);
+}
+
+void FMovieGraphImagePassBase::ModifyProjectionMatrixForTiling(const UE::MovieGraph::DefaultRenderer::FMovieGraphTilingParams& InTilingParams, const bool bInOrthographic, FMatrix& InOutProjectionMatrix, float& OutDoFSensorScale) const
+{
+	float PadRatioX = 1.0f;
+	float PadRatioY = 1.0f;
+
+	if (InTilingParams.OverlapPad.X > 0 && InTilingParams.OverlapPad.Y > 0)
+	{
+		PadRatioX = float(InTilingParams.OverlapPad.X * 2 + InTilingParams.TileSize.X) / float(InTilingParams.TileSize.X);
+		PadRatioY = float(InTilingParams.OverlapPad.Y * 2 + InTilingParams.TileSize.Y) / float(InTilingParams.TileSize.Y);
+	}
+
+	float ScaleX = PadRatioX / float(InTilingParams.TileCount.X);
+	float ScaleY = PadRatioY / float(InTilingParams.TileCount.Y);
+
+	InOutProjectionMatrix.M[0][0] /= ScaleX;
+	InOutProjectionMatrix.M[1][1] /= ScaleY;
+	OutDoFSensorScale = ScaleX;
+
+	// this offset would be correct with no pad
+	float OffsetX = -((float(InTilingParams.TileIndexes.X) + 0.5f - float(InTilingParams.TileCount.X) / 2.0f) * 2.0f);
+	float OffsetY = ((float(InTilingParams.TileIndexes.Y) + 0.5f - float(InTilingParams.TileCount.Y) / 2.0f) * 2.0f);
+
+	if (bInOrthographic)
+	{
+		InOutProjectionMatrix.M[3][0] += OffsetX / PadRatioX;
+		InOutProjectionMatrix.M[3][1] += OffsetY / PadRatioY;
+	}
+	else
+	{
+		InOutProjectionMatrix.M[2][0] += OffsetX / PadRatioX;
+		InOutProjectionMatrix.M[2][1] += OffsetY / PadRatioY;
+	}
+}
+
+
+void FMovieGraphImagePassBase::PostRendererSubmission(
+	const UE::MovieGraph::FMovieGraphSampleState& InSampleState,
+	const UE::MovieGraph::DefaultRenderer::FRenderTargetInitParams& InRenderTargetInitParams, FCanvas& InCanvas, const UE::MovieGraph::DefaultRenderer::FCameraInfo& InCameraInfo)
+{
+	// We have a pool of accumulators - we multi-thread the accumulation on the task graph, and for each frame,
+	// the task has the previous samples as pre-reqs to keep the accumulation in order. However, each accumulator
+	// can only work on one frame at a time, so we create a pool of them to work concurrently. This needs a limit
+	// as large accumulations (16k) can take a lot of system RAM.
+	TObjectPtr<UMovieGraphDefaultRenderer> GraphRenderer = GetRenderer().Get();
+	if (!GraphRenderer)
+	{
+		return;
+	}
+
+	FMoviePipelineAccumulatorPoolPtr SampleAccumulatorPool = GraphRenderer->GetOrCreateAccumulatorPool<FImageOverlappedAccumulator>();
+	UE::MovieGraph::DefaultRenderer::FSurfaceAccumulatorPool::FInstancePtr AccumulatorInstance = nullptr;
+	{
+		// SCOPE_CYCLE_COUNTER(STAT_MoviePipeline_WaitForAvailableAccumulator);
+		AccumulatorInstance = SampleAccumulatorPool->BlockAndGetAccumulator_GameThread(InSampleState.TraversalContext.Time.RenderedFrameNumber, InSampleState.TraversalContext.RenderDataIdentifier);
+	}
+
+	FMoviePipelineSurfaceQueuePtr LocalSurfaceQueue = GraphRenderer->GetOrCreateSurfaceQueue(InRenderTargetInitParams);
+	LocalSurfaceQueue->BlockUntilAnyAvailable();
+
+	FMovieGraphRenderDataAccumulationArgs AccumulationArgs;
+	{
+		AccumulationArgs.OutputMerger = GraphRenderer->GetOwningGraph()->GetOutputMerger();
+		AccumulationArgs.ImageAccumulator = StaticCastSharedPtr<FImageOverlappedAccumulator>(AccumulatorInstance->Accumulator);
+		AccumulationArgs.bIsFirstSample = InSampleState.TraversalContext.Time.bIsFirstTemporalSampleForFrame;
+		AccumulationArgs.bIsLastSample = InSampleState.TraversalContext.Time.bIsLastTemporalSampleForFrame;
+	}
+
+	auto OnSurfaceReadbackFinished = [this, InSampleState, AccumulationArgs, AccumulatorInstance](TUniquePtr<FImagePixelData>&& InPixelData)
+	{
+		UE::Tasks::TTask<void> Task = UE::Tasks::Launch(UE_SOURCE_LOCATION, [PixelData = MoveTemp(InPixelData), InSampleState, AccumulationArgs, AccumulatorInstance]() mutable
+		{
+			// Enqueue a encode for this frame onto our worker thread.
+			AccumulateSample_TaskThread(MoveTemp(PixelData), InSampleState, AccumulationArgs);
+
+			// We have to defer clearing the accumulator until after sample accumulation has finished
+			if (AccumulationArgs.bIsLastSample)
+			{
+				// Final sample has now been executed, free the accumulator for reuse.
+				AccumulatorInstance->bIsActive = false;
+			}
+		}, AccumulatorInstance->TaskPrereq);
+
+		// Make the next accumulation task that uses this accumulator use the task we just created as a pre-req.
+		AccumulatorInstance->TaskPrereq = Task;
+
+		// Because we're run on a separate thread, we need to check validity differently. The standard
+		// TWeakObjectPtr will report non-valid during GC (even if the object it's pointing to isn't being
+		// GC'd).
+		const bool bEvenIfPendingKill = false;
+		const bool bThreadsafeTest = true;
+		const bool bValid = this->WeakGraphRenderer.IsValid(bEvenIfPendingKill, bThreadsafeTest);
+		if (ensureMsgf(bValid, TEXT("Renderer was garbage collected while outstanding tasks existed, outstanding tasks were not flushed properly during shutdown!")))
+		{
+			// The regular Get() will fail during GC so we use the above check to see if it's valid
+			// before ignoring it and directly getting the object.
+			this->WeakGraphRenderer.GetEvenIfUnreachable()->AddOutstandingRenderTask_AnyThread(Task);
+		}
+	};
+
+	FRenderTarget* RenderTarget = InCanvas.GetRenderTarget();
+
+	ENQUEUE_RENDER_COMMAND(CanvasRenderTargetResolveCommand)(
+		[LocalSurfaceQueue, OnSurfaceReadbackFinished, RenderTarget](FRHICommandListImmediate& RHICmdList) mutable
+		{
+			// The legacy surface reader takes the payload just so it can shuffle it into our callback, but we can just include the data
+			// directly in the callback, so this is just a dummy payload.
+			TSharedRef<FImagePixelDataPayload, ESPMode::ThreadSafe> FramePayload = MakeShared<FImagePixelDataPayload, ESPMode::ThreadSafe>();
+			LocalSurfaceQueue->OnRenderTargetReady_RenderThread(RenderTarget->GetRenderTargetTexture(), FramePayload, MoveTemp(OnSurfaceReadbackFinished));
+		});
+}
+
+void AccumulateSample_TaskThread(TUniquePtr<FImagePixelData>&& InPixelData, const ::UE::MovieGraph::FMovieGraphSampleState InSampleState, const FMovieGraphRenderDataAccumulationArgs& InAccumulationParams)
+{
+	TRACE_CPUPROFILER_EVENT_SCOPE(MoviePipeline_AccumulateSample);
+		
+	TUniquePtr<FImagePixelData> SamplePixelData = MoveTemp(InPixelData);
+
+	// Associate the sample state with the image as payload data, this allows downstream systems to fetch the values without us having to store the data
+	// separately and ensure they stay paired the whole way down.
+	TSharedPtr<FMovieGraphSampleState> SampleStatePayload = MakeShared<FMovieGraphSampleState>(InSampleState);
+	SamplePixelData->SetPayload(StaticCastSharedPtr<IImagePixelDataPayload>(SampleStatePayload));
+
+	TSharedPtr<IMovieGraphOutputMerger, ESPMode::ThreadSafe> OutputMergerPin = InAccumulationParams.OutputMerger.Pin();
+	if (!OutputMergerPin.IsValid())
+	{
+		return;
+	}
+
+	const bool bIsWellFormed = SamplePixelData->IsDataWellFormed();
+	check(bIsWellFormed);
+
+	if (SampleStatePayload->bWriteSampleToDisk)
+	{
+		// Debug Feature: Write the raw sample to disk for debugging purposes. We copy the data here,
+		// as we don't want to disturb the memory flow below.
+		TUniquePtr<FImagePixelData> SampleData = SamplePixelData->CopyImageData();
+		OutputMergerPin->OnSingleSampleDataAvailable_AnyThread(MoveTemp(SampleData));
+	}
+
+	// Optimization! If we don't need the accumulator (no tiling, no sub-sampling) then we'll skip it and just send it straight to the output stage.
+	// This reduces memory requirements and improves performance in the baseline case.
+	if (!SampleStatePayload->bRequiresAccumulator)
+	{
+		OutputMergerPin->OnCompleteRenderPassDataAvailable_AnyThread(MoveTemp(SamplePixelData));
+		return;
+	}
+
+	TSharedPtr<FImageOverlappedAccumulator> AccumulatorPin = InAccumulationParams.ImageAccumulator.Pin();
+	if (AccumulatorPin->NumChannels == 0)
+	{
+		LLM_SCOPE_BYNAME(TEXT("MoviePipeline/ImageAccumulatorInitMemory"));
+		const int32 ChannelCount = 4;
+		AccumulatorPin->InitMemory(SampleStatePayload->AccumulatorResolution, ChannelCount);
+		AccumulatorPin->ZeroPlanes();
+		AccumulatorPin->AccumulationGamma = 1.f;
+	}
+
+	// Accumulate the new sample to our target
+	{
+		TRACE_CPUPROFILER_EVENT_SCOPE(MoviePipeline_AccumulatePixelData);
+				
+		const FIntPoint TileSize = SampleStatePayload->BackbufferResolution;
+		const FIntPoint OverlappedPad = SampleStatePayload->OverlappedPad;
+		const FIntPoint OverlappedOffset = SampleStatePayload->OverlappedOffset;
+		const FVector2D OverlappedSubpixelShift = SampleStatePayload->OverlappedSubpixelShift;
+		::MoviePipeline::FTileWeight1D WeightFunctionX;
+		::MoviePipeline::FTileWeight1D WeightFunctionY;
+		WeightFunctionX.InitHelper(OverlappedPad.X, TileSize.X, OverlappedPad.X);
+		WeightFunctionY.InitHelper(OverlappedPad.Y, TileSize.Y, OverlappedPad.Y);
+				
+		AccumulatorPin->AccumulatePixelData(*SamplePixelData, OverlappedOffset, OverlappedSubpixelShift, WeightFunctionX, WeightFunctionY);
+	}
+
+	// Finally on our last sample, we fetch the data out of the accumulator
+	// and move it to the Output Merger.
+	if (SampleStatePayload->bFetchFromAccumulator)
+	{
+		TRACE_CPUPROFILER_EVENT_SCOPE(MoviePipeline_FetchAccumulatedPixelData);
+		int32 FullSizeX = AccumulatorPin->PlaneSize.X;
+		int32 FullSizeY = AccumulatorPin->PlaneSize.Y;
+
+		// Now that a tile is fully built and accumulated we can notify the output builder that the
+		// data is ready so it can pass that onto the output containers (if needed).
+		if (SamplePixelData->GetType() == EImagePixelType::Float32)
+		{
+			// 32 bit FLinearColor
+			TUniquePtr<TImagePixelData<FLinearColor> > FinalPixelData = MakeUnique<TImagePixelData<FLinearColor>>(FIntPoint(FullSizeX, FullSizeY), SampleStatePayload);
+			AccumulatorPin->FetchFinalPixelDataLinearColor(FinalPixelData->Pixels);
+
+			// Send the data to the Output Builder
+			OutputMergerPin->OnCompleteRenderPassDataAvailable_AnyThread(MoveTemp(FinalPixelData));
+		}
+		else if (SamplePixelData->GetType() == EImagePixelType::Float16)
+		{
+			// 16 bit FLinearColor
+			TUniquePtr<TImagePixelData<FFloat16Color> > FinalPixelData = MakeUnique<TImagePixelData<FFloat16Color>>(FIntPoint(FullSizeX, FullSizeY), SampleStatePayload);
+			AccumulatorPin->FetchFinalPixelDataHalfFloat(FinalPixelData->Pixels);
+
+			// Send the data to the Output Builder
+			OutputMergerPin->OnCompleteRenderPassDataAvailable_AnyThread(MoveTemp(FinalPixelData));
+		}
+		else if (SamplePixelData->GetType() == EImagePixelType::Color)
+		{
+			// 8bit FColors
+			TUniquePtr<TImagePixelData<FColor>> FinalPixelData = MakeUnique<TImagePixelData<FColor>>(FIntPoint(FullSizeX, FullSizeY), SampleStatePayload);
+			AccumulatorPin->FetchFinalPixelDataByte(FinalPixelData->Pixels);
+
+			// Send the data to the Output Builder
+			OutputMergerPin->OnCompleteRenderPassDataAvailable_AnyThread(MoveTemp(FinalPixelData));
+		}
+		else
+		{
+			check(0);
+		}
+
+		{
+			TRACE_CPUPROFILER_EVENT_SCOPE(MoviePipeline_FreeAccumulatedPixelData);
+			// Free the memory in the accumulator.
+			AccumulatorPin->Reset();
+		}
+	}
+
+	{
+		TRACE_CPUPROFILER_EVENT_SCOPE(ReleaseSampleData);
+		SamplePixelData.Reset();
+	}
+}
+
+} // namespace UE::MovieGraph::Rendering
+
diff --git a/Engine/Plugins/MovieScene/MovieRenderPipeline/Source/MovieRenderPipelineCore/Private/Graph/Renderers/MovieGraphShowFlags.cpp b/Engine/Plugins/MovieScene/MovieRenderPipeline/Source/MovieRenderPipelineCore/Private/Graph/Renderers/MovieGraphShowFlags.cpp
new file mode 100644
--- /dev/null
+++ b/Engine/Plugins/MovieScene/MovieRenderPipeline/Source/MovieRenderPipelineCore/Private/Graph/Renderers/MovieGraphShowFlags.cpp
@@ -0,0 +1,136 @@
+// Copyright Epic Games, Inc. All Rights Reserved.
+
+#include "Graph/Renderers/MovieGraphShowFlags.h"
+#include "MovieRenderPipelineCoreModule.h"
+
+UMovieGraphShowFlags::UMovieGraphShowFlags()
+	: ShowFlags(ESFIM_Game)
+{
+	// Using ESFIM_Game as a base, update the show flags based on the serialized state
+	CopyOverridesToShowFlags(ShowFlags);
+}
+
+FEngineShowFlags UMovieGraphShowFlags::GetShowFlags() const
+{
+	// A copy is necessary because we keep ShowFlags around as the default (including values set by the renderer as
+	// default), and sparse overrides are stored on top of it. The sparse overrides need to be applied to a copy of the
+	// default.
+
+	FEngineShowFlags Copy(ShowFlags);
+	CopyOverridesToShowFlags(Copy);
+
+	return Copy;
+}
+
+bool UMovieGraphShowFlags::IsShowFlagEnabled(const uint32 ShowFlagIndex) const
+{
+	if (IsShowFlagOverridden(ShowFlagIndex))
+	{
+		if (const bool* EnableState = ShowFlagEnableState.Find(ShowFlagIndex))
+		{
+			return *EnableState;
+		}
+	}
+
+	return ShowFlags.GetSingleFlag(ShowFlagIndex);
+}
+
+void UMovieGraphShowFlags::SetShowFlagEnabled(const uint32 ShowFlagIndex, const bool bIsShowFlagEnabled)
+{
+	if (!IsShowFlagOverridden(ShowFlagIndex))
+	{
+		UE_LOG(LogMovieRenderPipeline, Warning, TEXT("SetShowFlagEnabled() call for '%s' ignored because the flag has not been marked as overridden."), *ShowFlags.FindNameByIndex(ShowFlagIndex));
+		return;
+	}
+
+	bool& EnableState = ShowFlagEnableState.FindOrAdd(ShowFlagIndex);
+	EnableState = bIsShowFlagEnabled;
+}
+
+bool UMovieGraphShowFlags::IsShowFlagOverridden(const uint32 ShowFlagIndex) const
+{
+	return OverriddenShowFlags.Contains(ShowFlagIndex);
+}
+
+void UMovieGraphShowFlags::SetShowFlagOverridden(const uint32 ShowFlagIndex, const bool bIsOverridden)
+{
+	if (bIsOverridden)
+	{
+		OverriddenShowFlags.Add(ShowFlagIndex);
+
+		// If the flag has not been marked as overridden in the past, then set the enable state to the default from
+		// ShowFlags. Otherwise, marking the flag as overridden will pick up the previously-set enable state.
+		if (!ShowFlagEnableState.Contains(ShowFlagIndex))
+		{
+			ShowFlagEnableState.Add(ShowFlagIndex, ShowFlags.GetSingleFlag(ShowFlagIndex));
+		}
+	}
+	else
+	{
+		OverriddenShowFlags.Remove(ShowFlagIndex);
+
+		// State from ShowFlagEnableState isn't removed; this is retained to restore the previously-set state if the
+		// user re-enables the override
+	}
+}
+
+const TSet<uint32>& UMovieGraphShowFlags::GetOverriddenShowFlags() const
+{
+	return OverriddenShowFlags;
+}
+
+void UMovieGraphShowFlags::ApplyDefaultShowFlagValue(const uint32 ShowFlagIndex, const bool bShowFlagState)
+{
+	// If there's no current override on this flag, then the default can be set
+	if (!OverriddenShowFlags.Contains(ShowFlagIndex))
+	{
+		ShowFlags.SetSingleFlag(ShowFlagIndex, bShowFlagState);
+	}
+}
+
+void UMovieGraphShowFlags::Merge(const IMovieGraphTraversableObject* InSourceObject)
+{
+	const UMovieGraphShowFlags* InSourceFlags = Cast<UMovieGraphShowFlags>(InSourceObject);
+	checkf(InSourceFlags, TEXT("UMovieGraphShowFlags cannot merge with null or an object of another type."));
+
+	for (const uint32& ShowFlagIndex : InSourceFlags->GetOverriddenShowFlags())
+	{
+		// The graph is iterated in reverse, from Outputs to Inputs. If this property is already overridden (on a
+		// object upstream, aka this object), then don't merge in the downstream (InSourceObject) override.
+		if (!OverriddenShowFlags.Contains(ShowFlagIndex))
+		{
+			SetShowFlagOverridden(ShowFlagIndex, true);
+			SetShowFlagEnabled(ShowFlagIndex, InSourceFlags->IsShowFlagEnabled(ShowFlagIndex));
+		}
+	}
+}
+
+TArray<TPair<FString, FString>> UMovieGraphShowFlags::GetMergedProperties() const
+{
+	TArray<TPair<FString, FString>> MergedProperties;
+
+	for (const uint32& ShowFlagIndex : OverriddenShowFlags)
+	{
+		MergedProperties.Add({ FEngineShowFlags::FindNameByIndex(ShowFlagIndex), LexToString(IsShowFlagEnabled(ShowFlagIndex)) });
+	}
+
+	// Sort by show flag name (helps w/ keeping order consistent in the UI)
+	MergedProperties.Sort([](const TPair<FString, FString>& Pair1, const TPair<FString, FString>& Pair2)
+		{
+			return Pair1.Key < Pair2.Key;
+		});
+
+	return MergedProperties;
+}
+
+void UMovieGraphShowFlags::CopyOverridesToShowFlags(FEngineShowFlags& InEngineShowFlags) const
+{
+	for (const uint32 ShowFlagIndex : OverriddenShowFlags)
+	{
+		// ShowFlagEnableState should always have an entry if there's an override, but check just in case
+		if (const bool* EnableState = ShowFlagEnableState.Find(ShowFlagIndex))
+		{
+			InEngineShowFlags.SetSingleFlag(ShowFlagIndex, *EnableState);
+		}
+	}
+}
\ No newline at end of file
diff --git a/Engine/Plugins/MovieScene/MovieRenderPipeline/Source/MovieRenderPipelineCore/Private/MoviePipelineRendering.cpp b/Engine/Plugins/MovieScene/MovieRenderPipeline/Source/MovieRenderPipelineCore/Private/MoviePipelineRendering.cpp
--- a/Engine/Plugins/MovieScene/MovieRenderPipeline/Source/MovieRenderPipelineCore/Private/MoviePipelineRendering.cpp
+++ b/Engine/Plugins/MovieScene/MovieRenderPipeline/Source/MovieRenderPipelineCore/Private/MoviePipelineRendering.cpp
@@ -336,39 +336,17 @@ void UMoviePipeline::RenderFrame()
 					FrameIndex = RenderSampleIndex;
 				}
 				
-
-				// Repeat the Halton Offset equally on each output frame so non-moving objects don't have any chance to crawl between frames.
-				int32 HaltonIndex = (FrameIndex % (NumSpatialSamples*NumTemporalSamples)) + 1;
-				float HaltonOffsetX = Halton(HaltonIndex, 2);
-				float HaltonOffsetY = Halton(HaltonIndex, 3);
-
 				// only allow a spatial jitter if we have more than one sample
 				bool bAllowSpatialJitter = !(NumSpatialSamples == 1 && NumTemporalSamples == 1);
 
-				UE_LOG(LogTemp, VeryVerbose, TEXT("FrameIndex: %d HaltonIndex: %d Offset: (%f,%f)"), FrameIndex, HaltonIndex, HaltonOffsetX, HaltonOffsetY);
 				float SpatialShiftX = 0.0f;
 				float SpatialShiftY = 0.0f;
 
 				if (bAllowSpatialJitter)
 				{
-					static auto CVar = IConsoleManager::Get().FindConsoleVariable(TEXT("r.TemporalAAFilterSize"));
-					float FilterSize = CVar->GetFloat();
-
-					// Scale distribution to set non-unit variance
-					// Variance = Sigma^2
-					float Sigma = 0.47f * FilterSize;
-
-					// Window to [-0.5, 0.5] output
-					// Without windowing we could generate samples far away on the infinite tails.
-					float OutWindow = 0.5f;
-					float InWindow = FMath::Exp(-0.5 * FMath::Square(OutWindow / Sigma));
-
-					// Box-Muller transform
-					float Theta = 2.0f * PI * HaltonOffsetY;
-					float r = Sigma * FMath::Sqrt(-2.0f * FMath::Loge((1.0f - HaltonOffsetX) * InWindow + HaltonOffsetX));
-
-					SpatialShiftX = r * FMath::Cos(Theta);
-					SpatialShiftY = r * FMath::Sin(Theta);
+					FVector2f SubPixelJitter = UE::MoviePipeline::GetSubPixelJitter(FrameIndex, NumSpatialSamples * NumTemporalSamples);
+					SpatialShiftX = SubPixelJitter.X;
+					SpatialShiftY = SubPixelJitter.Y;
 				}
 
 				// We take all of the information needed to render a single sample and package it into a struct.
diff --git a/Engine/Plugins/MovieScene/MovieRenderPipeline/Source/MovieRenderPipelineCore/Private/MoviePipelineUtils.cpp b/Engine/Plugins/MovieScene/MovieRenderPipeline/Source/MovieRenderPipelineCore/Private/MoviePipelineUtils.cpp
--- a/Engine/Plugins/MovieScene/MovieRenderPipeline/Source/MovieRenderPipelineCore/Private/MoviePipelineUtils.cpp
+++ b/Engine/Plugins/MovieScene/MovieRenderPipeline/Source/MovieRenderPipelineCore/Private/MoviePipelineUtils.cpp
@@ -1019,5 +1019,37 @@ namespace UE
 
 			return OutRenderPassMetrics;
 		}
+
+		FVector2f GetSubPixelJitter(int32 InFrameIndex, int32 InSamplesPerFrame)
+		{
+			// Repeat the Halton Offset equally on each output frame so non-moving objects don't have any chance to crawl between frames.
+			int32 HaltonIndex = (InFrameIndex % InSamplesPerFrame) + 1;
+			float HaltonOffsetX = Halton(HaltonIndex, 2);
+			float HaltonOffsetY = Halton(HaltonIndex, 3);
+
+			float SpatialShiftX = 0.0f;
+			float SpatialShiftY = 0.0f;
+
+			static auto CVar = IConsoleManager::Get().FindConsoleVariable(TEXT("r.TemporalAAFilterSize"));
+			float FilterSize = CVar->GetFloat();
+
+			// Scale distribution to set non-unit variance
+			// Variance = Sigma^2
+			float Sigma = 0.47f * FilterSize;
+
+			// Window to [-0.5, 0.5] output
+			// Without windowing we could generate samples far away on the infinite tails.
+			float OutWindow = 0.5f;
+			float InWindow = FMath::Exp(-0.5 * FMath::Square(OutWindow / Sigma));
+
+			// Box-Muller transform
+			float Theta = 2.0f * PI * HaltonOffsetY;
+			float r = Sigma * FMath::Sqrt(-2.0f * FMath::Loge((1.0f - HaltonOffsetX) * InWindow + HaltonOffsetX));
+
+			SpatialShiftX = r * FMath::Cos(Theta);
+			SpatialShiftY = r * FMath::Sin(Theta);
+
+			return FVector2f(SpatialShiftX, SpatialShiftY);
+		}
 	}
 }
diff --git a/Engine/Plugins/MovieScene/MovieRenderPipeline/Source/MovieRenderPipelineCore/Public/Graph/MovieGraphDataTypes.h b/Engine/Plugins/MovieScene/MovieRenderPipeline/Source/MovieRenderPipelineCore/Public/Graph/MovieGraphDataTypes.h
--- a/Engine/Plugins/MovieScene/MovieRenderPipeline/Source/MovieRenderPipelineCore/Public/Graph/MovieGraphDataTypes.h
+++ b/Engine/Plugins/MovieScene/MovieRenderPipeline/Source/MovieRenderPipelineCore/Public/Graph/MovieGraphDataTypes.h
@@ -307,6 +307,15 @@ namespace UE::MovieGraph
 		/** Set this to true if this pass should be composited on top of other renders. */
 		bool bCompositeOnOtherRenders;
 
+		/** When using high-res tiling, how many pixels does each tile overlap the adjacent tiles (on each side)? This should be zero if not using tiling. */
+		FIntPoint OverlappedPad;
+
+		/** When using high-res tiling, how many pixels offset into the output image does this accumulation get added to. */
+		FIntPoint OverlappedOffset;
+
+		/** When using spatial jitters, how much do we need to shift the output data during accumulation to counter-act the jitter. */
+		FVector2D OverlappedSubpixelShift;
+
 		/**
 		* If multiple passes are composited on top of a render, the sort order determines the order in which they're composited.
 		* Passes with a low sort order will composite on top of passes with a higher sort order.
diff --git a/Engine/Plugins/MovieScene/MovieRenderPipeline/Source/MovieRenderPipelineCore/Public/Graph/MovieGraphDefaultRenderer.h b/Engine/Plugins/MovieScene/MovieRenderPipeline/Source/MovieRenderPipelineCore/Public/Graph/MovieGraphDefaultRenderer.h
--- a/Engine/Plugins/MovieScene/MovieRenderPipeline/Source/MovieRenderPipelineCore/Public/Graph/MovieGraphDefaultRenderer.h
+++ b/Engine/Plugins/MovieScene/MovieRenderPipeline/Source/MovieRenderPipelineCore/Public/Graph/MovieGraphDefaultRenderer.h
@@ -23,15 +23,50 @@ typedef TSharedPtr<UE::MovieGraph::DefaultRenderer::FSurfaceAccumulatorPool, ESP
 
 namespace UE::MovieGraph::DefaultRenderer
 {
+	struct FMovieGraphTilingParams
+	{
+		// The number of pixels that the tiles overlap each other.
+		FVector2f OverlapPad;
+		// Size of the backbuffer in pixels
+		FIntPoint TileSize;
+		// Which tile index is this 
+		FIntPoint TileIndexes;
+		// How many tiles are there total
+		FIntPoint TileCount;
+	};
+
 	struct FCameraInfo
 	{
 		FCameraInfo()
 			: ViewActor(nullptr)
+			, OverscanFraction(0.f)
+			, bAllowCameraAspectRatio(true)
+			, DoFSensorScale(1.0f)
 		{}
 
 		FMinimalViewInfo ViewInfo;
-		FString CameraName;
 		class AActor* ViewActor;
+
+		/** The name to use for the {camera_name} token. Filled out by the renderer where possible. */
+		FString CameraName;
+
+		/** In the [0-1] Range. */
+		float OverscanFraction;
+
+		/** should we respect the camera's aspect ratio settings. */
+		bool bAllowCameraAspectRatio;
+
+		/** Projection Matrix this camera should use. Used instead of the one calculated by FMinimalViewInfo to handle special cases. */
+		FMatrix ProjectionMatrix;
+
+		// questionable if these are fcamerainfo
+		/** When using tiling, we scale the sensor to counteract the view changes. This value comes from modifying the ProjectionMatrix. */
+		float DoFSensorScale;
+		FMovieGraphTilingParams TilingParams;
+
+		// Sub-pixel jitter this camera should use. Only applied when using no AA.
+		FVector2D ProjectionMatrixJitterAmount;
+
 	};
 
 	struct FRenderTargetInitParams
@@ -154,7 +189,8 @@ public:
 
 	void AddOutstandingRenderTask_AnyThread(UE::Tasks::FTask InTask);
 	UE::MovieGraph::DefaultRenderer::FCameraInfo GetCameraInfo(const FGuid& InCameraIdentifier) const;
-
+	void SetHasRenderedFirstViewThisFrame(bool bInValue) { bHasRenderedFirstViewThisFrame = bInValue; }
+	bool GetHasRenderedFirstViewThisFrame() const { return bHasRenderedFirstViewThisFrame; }
 public:
 	UTextureRenderTarget2D* GetOrCreateViewRenderTarget(const UE::MovieGraph::DefaultRenderer::FRenderTargetInitParams& InInitParams);
 	FMoviePipelineSurfaceQueuePtr GetOrCreateSurfaceQueue(const UE::MovieGraph::DefaultRenderer::FRenderTargetInitParams& InInitParams);
@@ -192,6 +228,9 @@ protected:
 	UPROPERTY(Transient)
 	TArray<TObjectPtr<UMovieGraphRenderPassNode>> RenderPassesInUse;
 
+	/** For this engine tick, has any view been submitted yet? */
+	bool bHasRenderedFirstViewThisFrame = false;
+
 	/** Keep track of some statistics about render frame data for metadata purposes. */
 	TMap<int32, UE::MovieGraph::FRenderTimeStatistics> RenderTimeStatistics;
 
diff --git a/Engine/Plugins/MovieScene/MovieRenderPipeline/Source/MovieRenderPipelineCore/Public/Graph/MovieGraphTimeStepData.h b/Engine/Plugins/MovieScene/MovieRenderPipeline/Source/MovieRenderPipelineCore/Public/Graph/MovieGraphTimeStepData.h
--- a/Engine/Plugins/MovieScene/MovieRenderPipeline/Source/MovieRenderPipelineCore/Public/Graph/MovieGraphTimeStepData.h
+++ b/Engine/Plugins/MovieScene/MovieRenderPipeline/Source/MovieRenderPipelineCore/Public/Graph/MovieGraphTimeStepData.h
@@ -26,6 +26,8 @@ public:
 		, WorldSeconds(0.f)
 		, MotionBlurFraction(0.f)
 		, FrameRate(0, 0)
+		, TemporalSampleIndex(0)
+		, TemporalSampleCount(0)
 		, bIsFirstTemporalSampleForFrame(false)
 		, bIsLastTemporalSampleForFrame(false)
 		, bDiscardOutput(false)
@@ -53,6 +55,19 @@ public:
 	UPROPERTY(BlueprintReadWrite, EditAnywhere, Category = "Movie Graph")
 	FFrameRate FrameRate;
 
+	/*
+	* Index out of TemporalSampleCount we're on. No guarantee that we'll ever reach Index == Count-1,
+	* if bIsLastTemporalSampleForFrame has priority over that (to allow early outs)
+	*/
+	UPROPERTY(BlueprintReadWrite, EditAnywhere, Category = "Movie Graph")
+	int32 TemporalSampleIndex;
+
+	/*
+	* What is the maximum number of Temporal Samples this frame is expected to be able to process.
+	*/
+	UPROPERTY(BlueprintReadWrite, EditAnywhere, Category = "Movie Graph")
+	int32 TemporalSampleCount;
+
 	/** 
 	* Should be set to true for the first sample of each output frame. Used to determine
 	* if various systems need to reset or gather data for a new frame. Can be true at
diff --git a/Engine/Plugins/MovieScene/MovieRenderPipeline/Source/MovieRenderPipelineCore/Public/Graph/Nodes/MovieGraphCameraNode.h b/Engine/Plugins/MovieScene/MovieRenderPipeline/Source/MovieRenderPipelineCore/Public/Graph/Nodes/MovieGraphCameraNode.h
new file mode 100644
--- /dev/null
+++ b/Engine/Plugins/MovieScene/MovieRenderPipeline/Source/MovieRenderPipelineCore/Public/Graph/Nodes/MovieGraphCameraNode.h
@@ -0,0 +1,57 @@
+// Copyright Epic Games, Inc. All Rights Reserved.
+
+#pragma once
+
+#include "Graph/MovieGraphNode.h"
+#include "MovieRenderPipelineDataTypes.h" // For EMoviePipelineShutterTiming
+#include "MovieGraphCameraNode.generated.h"
+
+/** A node which configures global camera settings that are shared among all renders. */
+UCLASS()
+class MOVIERENDERPIPELINECORE_API UMovieGraphCameraSettingNode : public UMovieGraphSettingNode
+{
+	GENERATED_BODY()
+
+public:
+	UMovieGraphCameraSettingNode() = default;
+
+	virtual EMovieGraphBranchRestriction GetBranchRestriction() const override { return EMovieGraphBranchRestriction::Globals; }
+
+#if WITH_EDITOR
+	virtual FText GetNodeTitle(const bool bGetDescriptive = false) const override;
+	virtual FText GetMenuCategory() const override;
+	virtual FLinearColor GetNodeTitleColor() const override;
+	virtual FSlateIcon GetIconAndTint(FLinearColor& OutColor) const override;
+#endif
+
+public:
+	UPROPERTY(EditAnywhere, BlueprintReadWrite, Category = Overrides, meta = (InlineEditConditionToggle))
+	uint8 bOverride_ShutterTiming : 1;
+	
+	UPROPERTY(EditAnywhere, BlueprintReadWrite, Category = Overrides, meta = (InlineEditConditionToggle))
+	uint8 bOverride_OverscanPercentage : 1;
+
+	/**
+	* Shutter Timing allows you to bias the timing of your shutter angle to either be before, during, or after
+	* a frame. When set to FrameClose, it means that the motion gathered up to produce frame N comes from 
+	* before and right up to frame N. When set to FrameCenter, the motion represents half the time before the
+	* frame and half the time after the frame. When set to FrameOpen, the motion represents the time from 
+	* Frame N onwards.
+	*/
+	UPROPERTY(EditAnywhere, BlueprintReadWrite, Category = "Settings", meta = (EditCondition = "bOverride_ShutterTiming"))
+	EMoviePipelineShutterTiming ShutterTiming;
+
+	/**
+	* Overscan percent enables rendering render additional pixels beyond the set resolution and can be used in conjunction 
+	* with EXR file output to add post-processing effects such as lens distortion.
+	* Please note that using this feature might affect the results due to auto-exposure and other camera settings.
+	* On EXR this will produce a 1080p image with extra pixel data hidden around the outside edges for use 
+	* in post production. For all other formats this will increase the final resolution and no pixels will be hidden 
+	* (ie: 1080p with 10.0 overscan will make a 2112x1188 jpg, but a 1080p exr /w 96/54 pixels hidden on each side).
+	*
+	* Note: This uses 0-100 and not 0-1 like the previous system did to bring it in-line with other usages
+	* of overscan in the engine (nDisplay).
+	*/
+	UPROPERTY(EditAnywhere, BlueprintReadWrite, meta = (UIMin = "0", UIMax = "100", ClampMin = "0", ClampMax = "100"), Category = "Settings", meta = (EditCondition = "bOverride_OverscanPercentage"))
+	float OverscanPercentage;
+};
\ No newline at end of file
diff --git a/Engine/Plugins/MovieScene/MovieRenderPipeline/Source/MovieRenderPipelineCore/Public/Graph/Nodes/MovieGraphCoreRenderPassNode.h b/Engine/Plugins/MovieScene/MovieRenderPipeline/Source/MovieRenderPipelineCore/Public/Graph/Nodes/MovieGraphCoreRenderPassNode.h
deleted file mode 100644
diff --git a/Engine/Plugins/MovieScene/MovieRenderPipeline/Source/MovieRenderPipelineCore/Public/Graph/Nodes/MovieGraphDebugNode.h b/Engine/Plugins/MovieScene/MovieRenderPipeline/Source/MovieRenderPipelineCore/Public/Graph/Nodes/MovieGraphDebugNode.h
new file mode 100644
--- /dev/null
+++ b/Engine/Plugins/MovieScene/MovieRenderPipeline/Source/MovieRenderPipelineCore/Public/Graph/Nodes/MovieGraphDebugNode.h
@@ -0,0 +1,46 @@
+// Copyright Epic Games, Inc. All Rights Reserved.
+
+#pragma once
+
+#include "Graph/MovieGraphNode.h"
+#include "MovieGraphDebugNode.generated.h"
+
+/** A node which configures various debug settings that may be useful when debugging an issue. */
+UCLASS()
+class MOVIERENDERPIPELINECORE_API UMovieGraphDebugSettingNode : public UMovieGraphSettingNode
+{
+	GENERATED_BODY()
+
+public:
+	UMovieGraphDebugSettingNode() = default;
+
+	virtual EMovieGraphBranchRestriction GetBranchRestriction() const override { return EMovieGraphBranchRestriction::Globals; }
+
+#if WITH_EDITOR
+	virtual FText GetNodeTitle(const bool bGetDescriptive = false) const override;
+	virtual FText GetMenuCategory() const override;
+	virtual FLinearColor GetNodeTitleColor() const override;
+	virtual FSlateIcon GetIconAndTint(FLinearColor& OutColor) const override;
+#endif
+
+public:
+	UPROPERTY(EditAnywhere, BlueprintReadWrite, Category = Overrides, meta = (InlineEditConditionToggle))
+	uint8 bOverride_bCaptureFramesWithRenderDoc : 1;
+	
+	UPROPERTY(EditAnywhere, BlueprintReadWrite, Category = Overrides, meta = (InlineEditConditionToggle))
+	uint8 bOverride_bCaptureUnrealInsightsTrace : 1;
+
+	/** 
+	* If true, automatically trigger RenderDoc to capture rendering information. RenderDoc plugin must be enabled, 
+	* and the editor must have been launched with -AttachRenderDoc. Resulting capture will be in /Saved/RenderDocCaptures.
+	*/
+	UPROPERTY(EditAnywhere, BlueprintReadWrite, Category = "Settings", meta = (EditCondition = "bOverride_bCaptureFramesWithRenderDoc"))
+	bool bCaptureFramesWithRenderDoc;
+	
+	/** 
+	* If true, automatically capture an Unreal Insights trace file for the duration of the render.
+	* Resulting capture will be in the global Output Directory for the job.
+	*/
+	UPROPERTY(EditAnywhere, BlueprintReadWrite, Category = "Settings", meta = (EditCondition = "bOverride_bCaptureUnrealInsightsTrace"))
+	bool bCaptureUnrealInsightsTrace;
+};
\ No newline at end of file
diff --git a/Engine/Plugins/MovieScene/MovieRenderPipeline/Source/MovieRenderPipelineCore/Public/Graph/Nodes/MovieGraphDeferredPassNode.h b/Engine/Plugins/MovieScene/MovieRenderPipeline/Source/MovieRenderPipelineCore/Public/Graph/Nodes/MovieGraphDeferredPassNode.h
new file mode 100644
--- /dev/null
+++ b/Engine/Plugins/MovieScene/MovieRenderPipeline/Source/MovieRenderPipelineCore/Public/Graph/Nodes/MovieGraphDeferredPassNode.h
@@ -0,0 +1,91 @@
+﻿// Copyright Epic Games, Inc. All Rights Reserved.
+
+#pragma once
+#include "Graph/Nodes/MovieGraphImagePassBaseNode.h"
+#include "MovieGraphDeferredPassNode.generated.h"
+
+/**
+* A render node which uses the Deferred Renderer.
+*/
+UCLASS()
+class MOVIERENDERPIPELINECORE_API UMovieGraphDeferredRenderPassNode : public UMovieGraphImagePassBaseNode 
+{
+	GENERATED_BODY()
+
+public:
+	UMovieGraphDeferredRenderPassNode();
+
+#if WITH_EDITOR
+	virtual FText GetNodeTitle(const bool bGetDescriptive = false) const override;
+#endif
+
+protected:
+	// UMovieGraphRenderPassNode Interface
+	virtual FString GetRendererNameImpl() const override;
+	virtual TUniquePtr<UE::MovieGraph::Rendering::FMovieGraphImagePassBase> CreateInstance() const override;
+	// ~UMovieGraphRenderPassNode Interface
+
+	// UMovieGraphImagePassBaseNode Interface
+	virtual EViewModeIndex GetViewModeIndex() const override;
+	// ~UMovieGraphImagePassBaseNode Interface
+
+public:
+	UPROPERTY(EditAnywhere, BlueprintReadWrite, Category = Overrides, meta = (InlineEditConditionToggle))
+	uint8 bOverride_SpatialSampleCount: 1;
+
+	UPROPERTY(EditAnywhere, BlueprintReadWrite, Category = Overrides, meta = (InlineEditConditionToggle))
+	uint8 bOverride_AntiAliasingMethod : 1;
+	
+	UPROPERTY(EditAnywhere, BlueprintReadWrite, Category = Overrides, meta = (InlineEditConditionToggle))
+	uint8 bOverride_bDisableToneCurve : 1;
+	
+	UPROPERTY(EditAnywhere, BlueprintReadWrite, Category = Overrides, meta = (InlineEditConditionToggle))
+	uint8 bOverride_ViewModeIndex : 1;
+
+	UPROPERTY(EditAnywhere, BlueprintReadWrite, Category = Overrides, meta = (InlineEditConditionToggle))
+	uint8 bOverride_bWriteAllSamples : 1;
+
+	/**
+	* How many sub-pixel jitter renders should we do per temporal sample? This can be used to achieve high
+	* sample counts without Temporal Sub-Sampling (allowing high sample counts without motion blur being enabled),
+	* but we generally recommend using Temporal Sub-Samples when possible. It can also be combined with
+	* temporal samples and you will get SpatialSampleCount many renders per temporal sample.
+	*/
+	UPROPERTY(EditAnywhere, BlueprintReadWrite, meta = (UIMin = 1, ClampMin = 1), Category = "Settings", meta=(EditCondition="bOverride_SpatialSampleCount"))
+	int32 SpatialSampleCount;
+
+	/**
+	* Which anti-aliasing method should this render use. If this is set to None, then Movie Render Graph
+	* will handle anti-aliasing by doing a sub-pixel jitter (one for each temporal/spatial sample). Some
+	* rendering effects rely on TSR or TAA to reduce noise so we reccomend leaving them enabled
+	* where possible. All options work with Spatial and Temporal samples, but TSR/TAA may introduce minor
+	* visual artifacts (such as ghosting). MSAA is not supported in the deferred renderer.
+	*/
+	UPROPERTY(EditAnywhere, BlueprintReadWrite, Category = "Settings", meta = (EditCondition = "bOverride_AntiAliasingMethod"))
+	TEnumAsByte<EAntiAliasingMethod> AntiAliasingMethod;
+
+	/**
+	* If true, the tone curve will be disabled for this render pass. This will result in values greater than 1.0 in final renders
+	* and can optionally be combined with OCIO profiles on the file output nodes to convert from Linear Values in Working Color Space
+	* (which is sRGB  (Rec. 709) by default, unless changed in the project settings).
+	*/
+	UPROPERTY(EditAnywhere, BlueprintReadWrite, Category = "Settings", meta = (EditCondition = "bOverride_bDisableToneCurve"))
+	bool bDisableToneCurve;
+
+	/**
+	* Debug Feature. Can use this to write out each individual Temporal and Spatial sample rendered by this render pass,
+	* which allows you to see which images are being accumulated together. Can be useful for debugging incorrect looking
+	* frames to see which sub-frame evaluations were incorrect.
+	*/
+	UPROPERTY(EditAnywhere, BlueprintReadWrite, Category = "Debug", meta = (EditCondition = "bOverride_bWriteAllsamples"))
+	bool bWriteAllSamples;
+
+	/** 
+	* The view mode index that will be applied to renders. These mirror the View Modes you find in the Viewport,
+	* but most view modes other than Lit are used for debugging so they may not do what you expect, or may
+	* have to be used in combination with certain Show Flags to produce a result similar to what you see in
+	* the viewport.
+	*/
+	UPROPERTY(EditAnywhere, BlueprintReadWrite, Category = "Settings", meta=(EditCondition="bOverride_ViewModeIndex", InvalidEnumValues = "VMI_PathTracing"))
+	TEnumAsByte<EViewModeIndex> ViewModeIndex;
+};
\ No newline at end of file
diff --git a/Engine/Plugins/MovieScene/MovieRenderPipeline/Source/MovieRenderPipelineCore/Public/Graph/Nodes/MovieGraphDeferredRenderPassNode.h b/Engine/Plugins/MovieScene/MovieRenderPipeline/Source/MovieRenderPipelineCore/Public/Graph/Nodes/MovieGraphDeferredRenderPassNode.h
deleted file mode 100644
diff --git a/Engine/Plugins/MovieScene/MovieRenderPipeline/Source/MovieRenderPipelineCore/Public/Graph/Nodes/MovieGraphImagePassBaseNode.h b/Engine/Plugins/MovieScene/MovieRenderPipeline/Source/MovieRenderPipelineCore/Public/Graph/Nodes/MovieGraphImagePassBaseNode.h
new file mode 100644
--- /dev/null
+++ b/Engine/Plugins/MovieScene/MovieRenderPipeline/Source/MovieRenderPipelineCore/Public/Graph/Nodes/MovieGraphImagePassBaseNode.h
@@ -0,0 +1,61 @@
+// Copyright Epic Games, Inc. All Rights Reserved.
+
+#pragma once
+#include "Graph/MovieGraphNode.h"
+#include "Graph/MovieGraphDataTypes.h"
+#include "Graph/MovieGraphConfig.h"
+#include "Graph/Renderers/MovieGraphImagePassBase.h"
+#include "Graph/Renderers/MovieGraphShowFlags.h"
+#include "Graph/Nodes/MovieGraphRenderPassNode.h"
+#include "MovieGraphImagePassBaseNode.generated.h"
+
+// Forward Declare
+struct FMovieGraphRenderPassSetupData;
+struct FMovieGraphTimeStepData;
+
+/**
+* The UMovieGraphImagePassBaseNode is an abstract base-class for render nodes that wish to create
+* renders of the 3d scene. You are not required to inherit from this node (can inherit from
+* UMovieGraphRenderPassNode), but this node provides a helpful set of functions and default values
+* for constructing the required matrices and settings for viewport-like renders.
+*/
+UCLASS(Abstract)
+class MOVIERENDERPIPELINECORE_API UMovieGraphImagePassBaseNode : public UMovieGraphRenderPassNode
+{
+	GENERATED_BODY()
+public:
+	UMovieGraphImagePassBaseNode();
+
+	static void AddReferencedObjects(UObject* InThis, FReferenceCollector& Collector);
+		
+	/** Gets the view mode index that should be active for this renderer. */
+	virtual EViewModeIndex GetViewModeIndex() const;
+
+	/** Gets the show flags that should be active for this renderer. */
+	virtual FEngineShowFlags GetShowFlags() const;
+
+	/** Gets any show flags that should be applied as defaults, before user changes are applied. */
+	virtual TArray<uint32> GetDefaultShowFlags() const;
+
+protected:
+	// Note: Since *individual* show flags are overridden instead of the entire ShowFlags property, manually set to
+	// overridden so the traversal picks the changes up (otherwise they will be ignored).
+	UPROPERTY(EditAnywhere, BlueprintReadWrite, Category = Overrides, meta = (InlineEditConditionToggle))
+	uint8 bOverride_ShowFlags : 1 = 1;
+
+	/** The show flags that should be active during a render for this node. */
+	UPROPERTY(EditAnywhere, BlueprintReadWrite, Category = "Show Flags")
+	TObjectPtr<UMovieGraphShowFlags> ShowFlags;
+
+	TArray<TUniquePtr<UE::MovieGraph::Rendering::FMovieGraphImagePassBase>> CurrentInstances;
+protected:
+	// UMovieGraphRenderPassNode Interface
+	virtual FString GetRendererNameImpl() const { return TEXT("UnnamedImageRenderPass"); }
+	virtual void SetupImpl(const FMovieGraphRenderPassSetupData& InSetupData) override;
+	virtual void TeardownImpl() override;
+	virtual void RenderImpl(const FMovieGraphTraversalContext& InFrameTraversalContext, const FMovieGraphTimeStepData& InTimeData) override;
+	virtual void GatherOutputPassesImpl(TArray<FMovieGraphRenderDataIdentifier>& OutExpectedPasses) const override;
+	// ~UMovieGraphRenderPassNode Interface
+
+	virtual TUniquePtr<UE::MovieGraph::Rendering::FMovieGraphImagePassBase> CreateInstance() const { return nullptr; }
+};
\ No newline at end of file
diff --git a/Engine/Plugins/MovieScene/MovieRenderPipeline/Source/MovieRenderPipelineCore/Public/Graph/Nodes/MovieGraphPathTracedRenderPassNode.h b/Engine/Plugins/MovieScene/MovieRenderPipeline/Source/MovieRenderPipelineCore/Public/Graph/Nodes/MovieGraphPathTracedRenderPassNode.h
--- a/Engine/Plugins/MovieScene/MovieRenderPipeline/Source/MovieRenderPipelineCore/Public/Graph/Nodes/MovieGraphPathTracedRenderPassNode.h
+++ b/Engine/Plugins/MovieScene/MovieRenderPipeline/Source/MovieRenderPipelineCore/Public/Graph/Nodes/MovieGraphPathTracedRenderPassNode.h
@@ -1,14 +1,14 @@
-﻿// Copyright Epic Games, Inc. All Rights Reserved.
+// Copyright Epic Games, Inc. All Rights Reserved.
 
 #pragma once
 
-#include "Graph/Nodes/MovieGraphCoreRenderPassNode.h"
+#include "Graph/Nodes/MovieGraphImagePassBaseNode.h"
 
 #include "MovieGraphPathTracedRenderPassNode.generated.h"
 
 /** A render node which uses the path tracer. */
 UCLASS()
-class MOVIERENDERPIPELINECORE_API UMovieGraphPathTracedRenderPassNode : public UMovieGraphCoreRenderPassNode
+class MOVIERENDERPIPELINECORE_API UMovieGraphPathTracedRenderPassNode : public UMovieGraphImagePassBaseNode
 {
 	GENERATED_BODY()
 
diff --git a/Engine/Plugins/MovieScene/MovieRenderPipeline/Source/MovieRenderPipelineCore/Public/Graph/Renderers/MovieGraphDeferredPass.h b/Engine/Plugins/MovieScene/MovieRenderPipeline/Source/MovieRenderPipelineCore/Public/Graph/Renderers/MovieGraphDeferredPass.h
new file mode 100644
--- /dev/null
+++ b/Engine/Plugins/MovieScene/MovieRenderPipeline/Source/MovieRenderPipelineCore/Public/Graph/Renderers/MovieGraphDeferredPass.h
@@ -0,0 +1,27 @@
+// Copyright Epic Games, Inc. All Rights Reserved.
+#pragma once
+
+#include "Graph/Renderers/MovieGraphImagePassBase.h"
+
+namespace UE::MovieGraph::Rendering
+{
+	struct FMovieGraphDeferredPass : public FMovieGraphImagePassBase
+	{
+		virtual void Setup(TWeakObjectPtr<UMovieGraphDefaultRenderer> InRenderer, TWeakObjectPtr<UMovieGraphImagePassBaseNode> InRenderPassNode, const FMovieGraphRenderPassLayerData& InLayer) override;
+		virtual void Teardown() override;
+		virtual void Render(const FMovieGraphTraversalContext& InFrameTraversalContext, const FMovieGraphTimeStepData& InTimeData) override;
+		virtual void GatherOutputPasses(TArray<FMovieGraphRenderDataIdentifier>& OutExpectedPasses) const override;
+		virtual void AddReferencedObjects(FReferenceCollector& Collector) override;
+	protected:
+		virtual void PostRendererSubmission(const UE::MovieGraph::FMovieGraphSampleState& InSampleState, const UE::MovieGraph::DefaultRenderer::FRenderTargetInitParams& InRenderTargetInitParams, FCanvas& InCanvas, const UE::MovieGraph::DefaultRenderer::FCameraInfo& InCameraInfo) override;
+
+	protected:
+		FMovieGraphRenderPassLayerData LayerData;
+
+		/** Unique identifier passed in GatherOutputPasses and with each render that identifies the data produced by this renderer. */
+		FMovieGraphRenderDataIdentifier RenderDataIdentifier;
+
+		// Scene View history used by the renderer 
+		FSceneViewStateReference SceneViewState;
+	}; 
+}
\ No newline at end of file
diff --git a/Engine/Plugins/MovieScene/MovieRenderPipeline/Source/MovieRenderPipelineCore/Public/Graph/Renderers/MovieGraphImagePassBase.h b/Engine/Plugins/MovieScene/MovieRenderPipeline/Source/MovieRenderPipelineCore/Public/Graph/Renderers/MovieGraphImagePassBase.h
new file mode 100644
--- /dev/null
+++ b/Engine/Plugins/MovieScene/MovieRenderPipeline/Source/MovieRenderPipelineCore/Public/Graph/Renderers/MovieGraphImagePassBase.h
@@ -0,0 +1,85 @@
+// Copyright Epic Games, Inc. All Rights Reserved.
+
+#pragma once
+#include "CoreTypes.h"
+#include "Graph/MovieGraphDefaultRenderer.h"
+#include "SceneView.h"
+
+// Forward Declares
+class UMovieGraphImagePassBaseNode;
+
+namespace UE::MovieGraph::Rendering
+{
+	struct FMovieGraphRenderDataAccumulationArgs
+	{
+	public:
+		TWeakPtr<FImageOverlappedAccumulator, ESPMode::ThreadSafe> ImageAccumulator;
+		TWeakPtr<IMovieGraphOutputMerger, ESPMode::ThreadSafe> OutputMerger;
+
+		// If it's the first sample then we will reset the accumulator to a clean slate before accumulating into it.
+		bool bIsFirstSample;
+		// If it's the last sample, then we will trigger moving the data to the output merger. It can be both the first and last sample at the same time.
+		bool bIsLastSample;
+	};
+
+	// Forward Declare
+	void AccumulateSample_TaskThread(TUniquePtr<FImagePixelData>&& InPixelData, const ::UE::MovieGraph::FMovieGraphSampleState InSampleState, const FMovieGraphRenderDataAccumulationArgs& InAccumulationParams);
+
+	struct FViewFamilyInitData
+	{
+		FViewFamilyInitData()
+			: RenderTarget(nullptr)
+			, World(nullptr)
+			, SceneCaptureSource(ESceneCaptureSource::SCS_MAX)
+			, bWorldIsPaused(false)
+			, FrameIndex(-1)
+			, AntiAliasingMethod(EAntiAliasingMethod::AAM_None)
+		{
+		}
+
+		class FRenderTarget* RenderTarget;
+		class UWorld* World;
+		FMovieGraphTimeStepData TimeData;
+		ESceneCaptureSource SceneCaptureSource;
+		bool bWorldIsPaused;
+		int32 FrameIndex;
+		EAntiAliasingMethod AntiAliasingMethod;
+	};
+	
+
+
+	struct FMovieGraphImagePassBase
+	{
+		FMovieGraphImagePassBase() = default;
+		virtual ~FMovieGraphImagePassBase() = default;
+
+		virtual void Setup(TWeakObjectPtr<UMovieGraphDefaultRenderer> InRenderer, TWeakObjectPtr<UMovieGraphImagePassBaseNode> InRenderPassNode, const FMovieGraphRenderPassLayerData& InLayer);
+		virtual void Teardown() {}
+		virtual void Render(const FMovieGraphTraversalContext& InFrameTraversalContext, const FMovieGraphTimeStepData& InTimeData) {}
+		virtual void GatherOutputPasses(TArray<FMovieGraphRenderDataIdentifier>& OutExpectedPasses) const {}
+		virtual void AddReferencedObjects(FReferenceCollector& Collector) {}
+		virtual FName GetBranchName() const { return NAME_None; }
+		virtual TWeakObjectPtr<UMovieGraphDefaultRenderer> GetRenderer() const { return WeakGraphRenderer; }
+
+	protected:
+		/** Utility function for calculating a Projection Matrix (Orthographic or Perspective), and modify it based on the overscan percentage, aspect ratio, etc. */
+		virtual FMatrix CalculateProjectionMatrix(const UE::MovieGraph::DefaultRenderer::FCameraInfo& InCameraInfo) const;
+		/** Utility function for modifying the projection matrix to match the given tiling parameters. */
+		virtual void ModifyProjectionMatrixForTiling(const UE::MovieGraph::DefaultRenderer::FMovieGraphTilingParams& InTilingParams, const bool bInOrthographic, FMatrix& InOutProjectionMatrix, float& OutDoFSensorScale) const;
+		/** Utility function for calculating the PrinciplePointOffset with the given tiling parameters. Used to make some effects (like vignette) work with tiling. */
+		virtual FVector4f CalculatePrinciplePointOffsetForTiling(const UE::MovieGraph::DefaultRenderer::FMovieGraphTilingParams& InTilingParams) const;
+		/** Utility function for creating FSceneViewInitOptions based on the specified camera info. */
+		virtual FSceneViewInitOptions CreateViewInitOptions(const UE::MovieGraph::DefaultRenderer::FCameraInfo& InCameraInfo, FSceneViewFamilyContext* InViewFamily, FSceneViewStateReference& InViewStateRef) const;
+		/** Utility function for creating a FSceneView for the given InitOptions, Family, and Camera. */
+		virtual FSceneView* CreateSceneView(const FSceneViewInitOptions& InInitOptions, TSharedRef<FSceneViewFamilyContext> InViewFamily, const UE::MovieGraph::DefaultRenderer::FCameraInfo& InCameraInfo) const;
+
+		virtual void ApplyCameraManagerPostProcessBlends(FSceneView* InView) const;
+		virtual TSharedRef<FSceneViewFamilyContext> CreateSceneViewFamily(const FViewFamilyInitData& InInitData, const UE::MovieGraph::DefaultRenderer::FCameraInfo& InCameraInfo) const;
+		virtual void ApplyMovieGraphOverridesToSceneView(TSharedRef<FSceneViewFamilyContext> InOutFamily, const FViewFamilyInitData& InInitData, const UE::MovieGraph::DefaultRenderer::FCameraInfo& InCameraInfo) const;
+		virtual void ApplyMovieGraphOverridesToViewFamily(TSharedRef<FSceneViewFamilyContext> InOutFamily, const FViewFamilyInitData& InInitData) const;
+		virtual void PostRendererSubmission(const UE::MovieGraph::FMovieGraphSampleState& InSampleState, const UE::MovieGraph::DefaultRenderer::FRenderTargetInitParams& InRenderTargetInitParams, FCanvas& InCanvas, const UE::MovieGraph::DefaultRenderer::FCameraInfo& InCameraInfo);
+
+	protected:
+		TWeakObjectPtr<UMovieGraphDefaultRenderer> WeakGraphRenderer;
+	};
+}
\ No newline at end of file
diff --git a/Engine/Plugins/MovieScene/MovieRenderPipeline/Source/MovieRenderPipelineCore/Public/Graph/Renderers/MovieGraphShowFlags.h b/Engine/Plugins/MovieScene/MovieRenderPipeline/Source/MovieRenderPipelineCore/Public/Graph/Renderers/MovieGraphShowFlags.h
new file mode 100644
--- /dev/null
+++ b/Engine/Plugins/MovieScene/MovieRenderPipeline/Source/MovieRenderPipelineCore/Public/Graph/Renderers/MovieGraphShowFlags.h
@@ -0,0 +1,80 @@
+// Copyright Epic Games, Inc. All Rights Reserved.
+
+#pragma once
+#include "CoreTypes.h"
+#include "ShowFlags.h"
+#include "Graph/MovieGraphConfig.h"
+#include "MovieGraphShowFlags.generated.h"
+
+/**
+ * Stores show flag enable/disable state, as well as per-flag override state.
+ *
+ * Here, enable/disable state refers to the actual state of the show flag itself (ie, whether it is turned on or off).
+ * This is the value that will be applied to renders. "Override" state is whether the flag has been marked as overridden
+ * in the UI (ie, whether it is a value that graph traversal should respect). Flags must be marked as overridden in order
+ * for their values to deviate from the default.
+ */
+UCLASS(BlueprintType)
+class MOVIERENDERPIPELINECORE_API UMovieGraphShowFlags : public UObject, public IMovieGraphTraversableObject
+{
+	GENERATED_BODY()
+
+public:
+	UMovieGraphShowFlags();
+
+	/** Gets a copy of the show flags. If they need to be modified, use one of the Set*() methods. */
+	FEngineShowFlags GetShowFlags() const;
+
+	/** Gets the enable/disable state for a specific show flag. */
+	bool IsShowFlagEnabled(const uint32 ShowFlagIndex) const;
+
+	/** Sets the enable/disable state for a specific show flag. The flag must be marked as overridden for this to have an effect. */
+	void SetShowFlagEnabled(const uint32 ShowFlagIndex, const bool bIsShowFlagEnabled);
+
+	/** Gets the override state for a specific show flag. */
+	bool IsShowFlagOverridden(const uint32 ShowFlagIndex) const;
+
+	/** Sets the override state for a specific show flag. */
+	void SetShowFlagOverridden(const uint32 ShowFlagIndex, const bool bIsOverridden);
+
+	/** Gets the indices of all show flags that have been overridden. */
+	const TSet<uint32>& GetOverriddenShowFlags() const;
+
+	/**
+	 * Applies a default value to the show flags. However, the default will not be applied if a user has already
+	 * specified an override for the flag. This should generally only be called from renderers in order for them to
+	 * specify defaults that are relevant.
+	 */
+	void ApplyDefaultShowFlagValue(const uint32 ShowFlagIndex, const bool bShowFlagState);
+
+	// IMovieGraphTraversableObject interface
+	virtual void Merge(const IMovieGraphTraversableObject* InSourceObject) override;
+	virtual TArray<TPair<FString, FString>> GetMergedProperties() const override;
+	// ~IMovieGraphTraversableObject interface
+
+private:
+	/** Copies the overrides from ShowFlagEnableState to InEngineShowFlags. */
+	void CopyOverridesToShowFlags(FEngineShowFlags& InEngineShowFlags) const;
+
+private:
+	/**
+	 * The show flags which have been marked as overridden in the UI (note this is NOT the enabled/disabled state of the
+	 * flag itself).
+	 */
+	UPROPERTY()
+	TSet<uint32> OverriddenShowFlags;
+
+	/**
+	 * If the flag has been marked as overridden, this stores the enable/disable state of the flag. Key = show flag
+	 * index, value = enable/disable state of the flag.
+	 */
+	UPROPERTY()
+	TMap<uint32, bool> ShowFlagEnableState;
+
+	/**
+	 * The default show flag state. Since FEngineShowFlags itself cannot be serialized (it is not a USTRUCT), we need to
+	 * keep track of sparse overrides (ShowFlagEnableState) on top of the base state that these flags are initialized
+	 * with (ESFIM_Game, plus any additional flags that a specialized renderer specifies).
+	 */
+	FEngineShowFlags ShowFlags;
+};
\ No newline at end of file
diff --git a/Engine/Plugins/MovieScene/MovieRenderPipeline/Source/MovieRenderPipelineCore/Public/MoviePipelineUtils.h b/Engine/Plugins/MovieScene/MovieRenderPipeline/Source/MovieRenderPipelineCore/Public/MoviePipelineUtils.h
--- a/Engine/Plugins/MovieScene/MovieRenderPipeline/Source/MovieRenderPipelineCore/Public/MoviePipelineUtils.h
+++ b/Engine/Plugins/MovieScene/MovieRenderPipeline/Source/MovieRenderPipelineCore/Public/MoviePipelineUtils.h
@@ -144,6 +144,10 @@ namespace UE
 		MOVIERENDERPIPELINECORE_API FMoviePipelineRenderPassMetrics GetRenderPassMetrics(UMoviePipelinePrimaryConfig* InPrimaryConfig, UMoviePipelineExecutorShot* InPipelineExecutorShot, const FMoviePipelineRenderPassMetrics& InRenderPassMetrics, const FIntPoint& InEffectiveOutputResolution);
 		MOVIERENDERPIPELINECORE_API bool CanWriteToFile(const TCHAR* InFilename, bool bOverwriteExisting);
 		MOVIERENDERPIPELINECORE_API FString GetPaddingFormatString(int32 InZeroPadCount, const int32 InFrameNumber);
+
+		/** When using spatial/temporal samples without anti-aliasing, get the sub-pixel jitter for the given frame index. FrameIndex is modded by InSamplesPerFrame so that the aa jitter pattern repeats every output frame. */
+		MOVIERENDERPIPELINECORE_API FVector2f GetSubPixelJitter(int32 InFrameIndex, int32 InSamplesPerFrame);
+
 	}
 }
 
diff --git a/Engine/Plugins/MovieScene/MovieRenderPipeline/Source/MovieRenderPipelineEditor/Private/Customizations/Graph/MovieGraphShowFlagsCustomization.h b/Engine/Plugins/MovieScene/MovieRenderPipeline/Source/MovieRenderPipelineEditor/Private/Customizations/Graph/MovieGraphShowFlagsCustomization.h
--- a/Engine/Plugins/MovieScene/MovieRenderPipeline/Source/MovieRenderPipelineEditor/Private/Customizations/Graph/MovieGraphShowFlagsCustomization.h
+++ b/Engine/Plugins/MovieScene/MovieRenderPipeline/Source/MovieRenderPipelineEditor/Private/Customizations/Graph/MovieGraphShowFlagsCustomization.h
@@ -4,11 +4,12 @@
 
 #include "DetailWidgetRow.h"
 #include "EditorShowFlags.h"
-#include "Graph/Nodes/MovieGraphCoreRenderPassNode.h"
+#include "Graph/Renderers/MovieGraphShowFlags.h"
 #include "IDetailChildrenBuilder.h"
 #include "IDetailGroup.h"
 #include "PropertyHandle.h"
 #include "Widgets/Input/SCheckBox.h"
+#include "Widgets/Text/STextBlock.h"
 
 #define LOCTEXT_NAMESPACE "MoviePipelineEditor"
 
diff --git a/Engine/Plugins/MovieScene/MovieRenderPipeline/Source/MovieRenderPipelineEditor/Private/Graph/MovieGraphAssetToolkit.cpp b/Engine/Plugins/MovieScene/MovieRenderPipeline/Source/MovieRenderPipelineEditor/Private/Graph/MovieGraphAssetToolkit.cpp
--- a/Engine/Plugins/MovieScene/MovieRenderPipeline/Source/MovieRenderPipelineEditor/Private/Graph/MovieGraphAssetToolkit.cpp
+++ b/Engine/Plugins/MovieScene/MovieRenderPipeline/Source/MovieRenderPipelineEditor/Private/Graph/MovieGraphAssetToolkit.cpp
@@ -6,9 +6,9 @@
 #include "Customizations/Graph/MovieGraphNodeCustomization.h"
 #include "Customizations/Graph/MovieGraphSelectNodeCustomization.h"
 #include "Customizations/Graph/MovieGraphShowFlagsCustomization.h"
+#include "Graph/Renderers/MovieGraphShowFlags.h"
 #include "Framework/Commands/GenericCommands.h"
 #include "Graph/MovieGraphConfig.h"
-#include "Graph/Nodes/MovieGraphCoreRenderPassNode.h"
 #include "Graph/Nodes/MovieGraphSelectNode.h"
 #include "MovieEdGraphNode.h"
 #include "MovieGraphSchema.h"
diff --git a/Engine/Plugins/MovieScene/MovieRenderPipeline/Source/MovieRenderPipelineEditor/Private/Graph/SMovieGraphActiveRenderSettingsTabContent.cpp b/Engine/Plugins/MovieScene/MovieRenderPipeline/Source/MovieRenderPipelineEditor/Private/Graph/SMovieGraphActiveRenderSettingsTabContent.cpp
--- a/Engine/Plugins/MovieScene/MovieRenderPipeline/Source/MovieRenderPipelineEditor/Private/Graph/SMovieGraphActiveRenderSettingsTabContent.cpp
+++ b/Engine/Plugins/MovieScene/MovieRenderPipeline/Source/MovieRenderPipelineEditor/Private/Graph/SMovieGraphActiveRenderSettingsTabContent.cpp
@@ -12,7 +12,6 @@
 #include "MoviePipelineQueue.h"
 #include "MoviePipelineQueueSubsystem.h"
 #include "SPositiveActionButton.h"
-#include "Graph/Nodes/MovieGraphCoreRenderPassNode.h"
 #include "Widgets/Input/SComboBox.h"
 
 #define LOCTEXT_NAMESPACE "SMoviePipelineActiveRenderSettingsTabContent"
