MoveRenderGraph:

-Unit Test: Nodes. Create Config > add one node of each class inheriting from UMovieGraphNode > verifying all exist > Connect all by one edge > Verify Edges > Remove All Nodes > Verify none exist
-Add MovieRenderPipelineEditorTestUtilities to help facilitate tests in MoveRenderGraph
-Move pin connection gate logic to UMovieGraphPin class
-Add pin direction test to UMovieGraphPin
-Prevent AddNode function in UMovieGraphConfig from being able to add Base Input and Output nodes
-Prevent AddLabeledEdge making pin connections in UMovieGraphConfig without passing CanCreateConnection

#[fyi] Jason.Walter

diff --git a/Engine/Plugins/MovieScene/MovieRenderPipeline/Source/MovieRenderPipelineCore/Private/Graph/MovieGraphConfig.cpp b/Engine/Plugins/MovieScene/MovieRenderPipeline/Source/MovieRenderPipelineCore/Private/Graph/MovieGraphConfig.cpp
--- a/Engine/Plugins/MovieScene/MovieRenderPipeline/Source/MovieRenderPipelineCore/Private/Graph/MovieGraphConfig.cpp
+++ b/Engine/Plugins/MovieScene/MovieRenderPipeline/Source/MovieRenderPipelineCore/Private/Graph/MovieGraphConfig.cpp
@@ -378,21 +378,40 @@ bool UMovieGraphConfig::AddLabeledEdge(UMovieGraphNode* FromNode, const FName& F
 {
 	if (!FromNode || !ToNode)
 	{
-		UE_LOG(LogMovieRenderPipeline, Error, TEXT("AddLabeledEdge: Invalid Edge Nodes"));
+		FFrame::KismetExecutionMessage(
+			*FString::Printf(
+				TEXT("%hs: Invalid Edge Nodes"), __FUNCTION__),
+				ELogVerbosity::Error);
 		return false;
 	}
+	
 
 	UMovieGraphPin* FromPin = FromNode->GetOutputPin(FromPinLabel);
 	if (!FromPin)
 	{
-		UE_LOG(LogMovieRenderPipeline, Error, TEXT("AddLabeledEdge: FromNode: %s does not have a pin with the label: %s"), *FromNode->GetName(), *FromPinLabel.ToString());
+		FFrame::KismetExecutionMessage(
+			*FString::Printf(
+				TEXT("%hs: FromNode: %s does not have a pin with the label: %s"), __FUNCTION__, *FromNode->GetName(), *FromPinLabel.ToString()),
+				ELogVerbosity::Error);
 		return false;
 	}
 
 	UMovieGraphPin* ToPin = ToNode->GetInputPin(ToPinLabel);
 	if (!ToPin)
 	{
-		UE_LOG(LogMovieRenderPipeline, Error, TEXT("AddLabeledEdge: ToNode: %s does not have a pin with the label: %s"), *ToNode->GetName(), *ToPinLabel.ToString());
+		FFrame::KismetExecutionMessage(
+			*FString::Printf(
+				TEXT("%hs: ToNode: %s does not have a pin with the label: %s"), __FUNCTION__, *ToNode->GetName(), *ToPinLabel.ToString()),
+				ELogVerbosity::Error);
+		return false;
+	}
+
+	if (!FromPin->CanCreateConnection(ToPin))
+	{
+		FFrame::KismetExecutionMessage(
+			*FString::Printf(
+				TEXT("%hs: FromNode %s's pin %s cannot be connected to ToNode %s's pin %s "), __FUNCTION__, *FromNode->GetName(), *FromPinLabel.ToString(), *ToNode->GetName(), *ToPinLabel.ToString()),
+				ELogVerbosity::Error);
 		return false;
 	}
 
@@ -421,21 +440,28 @@ bool UMovieGraphConfig::RemoveEdge(UMovieGraphNode* FromNode, const FName& FromP
 {
 	if (!FromNode || !ToNode)
 	{
-		UE_LOG(LogMovieRenderPipeline, Error, TEXT("RemoveEdge: Invalid Edge Nodes"));
+		FFrame::KismetExecutionMessage(
+			*FString::Printf(TEXT("%hs: Invalid Edge Nodes"), __FUNCTION__), ELogVerbosity::Error);
 		return false;
 	}
 
 	UMovieGraphPin* FromPin = FromNode->GetOutputPin(FromPinLabel);
 	if (!FromPin)
 	{
-		UE_LOG(LogMovieRenderPipeline, Error, TEXT("RemoveEdge: FromNode: %s does not have a pin with the label: %s"), *FromNode->GetName(), *FromPinLabel.ToString());
+		FFrame::KismetExecutionMessage(
+			*FString::Printf(
+				TEXT("%hs: FromNode: %s does not have a pin with the label: %s"), __FUNCTION__, *FromNode->GetName(), *FromPinLabel.ToString()),
+				ELogVerbosity::Error);
 		return false;
 	}
 
 	UMovieGraphPin* ToPin = ToNode->GetInputPin(ToPinLabel);
 	if (!ToPin)
 	{
-		UE_LOG(LogMovieRenderPipeline, Error, TEXT("RemoveEdge: ToNode: %s does not have a pin with the label: %s"), *ToNode->GetName(), *ToPinLabel.ToString());
+		FFrame::KismetExecutionMessage(
+			*FString::Printf(
+				TEXT("%hs: ToNode: %s does not have a pin with the label: %s"), __FUNCTION__, *ToNode->GetName(), *ToPinLabel.ToString()),
+				ELogVerbosity::Error);
 		return false;
 	}
 
@@ -454,7 +480,8 @@ bool UMovieGraphConfig::RemoveAllInboundEdges(UMovieGraphNode* InNode)
 {
 	if (!InNode)
 	{
-		UE_LOG(LogMovieRenderPipeline, Error, TEXT("RemoveAllInboundEdges: Invalid Edge Nodes"));
+		FFrame::KismetExecutionMessage(
+			*FString::Printf(TEXT("%hs: Invalid Edge Nodes"), __FUNCTION__ ), ELogVerbosity::Error);
 		return false;
 	}
 
@@ -477,7 +504,8 @@ bool UMovieGraphConfig::RemoveAllOutboundEdges(UMovieGraphNode* InNode)
 {
 	if (!InNode)
 	{
-		UE_LOG(LogMovieRenderPipeline, Error, TEXT("RemoveAllOutboundEdges: Invalid Edge Nodes"));
+		FFrame::KismetExecutionMessage(
+			*FString::Printf(TEXT("%hs: Invalid Edge Nodes"), __FUNCTION__), ELogVerbosity::Error);
 		return false;
 	}
 
@@ -500,7 +528,10 @@ bool UMovieGraphConfig::RemoveInboundEdges(UMovieGraphNode* InNode, const FName&
 {
 	if (!InNode)
 	{
-		UE_LOG(LogMovieRenderPipeline, Error, TEXT("RemoveInboundEdges: Invalid Edge Nodes"));
+		FFrame::KismetExecutionMessage(
+			*FString::Printf(
+				TEXT("%hs: Invalid Edge Nodes"), __FUNCTION__),
+				ELogVerbosity::Error);
 		return false;
 	}
 
@@ -523,7 +554,10 @@ bool UMovieGraphConfig::RemoveOutboundEdges(UMovieGraphNode* InNode, const FName
 {
 	if (!InNode)
 	{
-		UE_LOG(LogMovieRenderPipeline, Error, TEXT("RemoveOutboundEdges: Invalid Edge Nodes"));
+		FFrame::KismetExecutionMessage(
+			*FString::Printf(
+				TEXT("%hs: Invalid Edge Nodes"), __FUNCTION__),
+				ELogVerbosity::Error);
 		return false;
 	}
 
@@ -546,6 +580,19 @@ void UMovieGraphConfig::AddNode(UMovieGraphNode* InNode)
 {
 	if (!InNode)
 	{
+		FFrame::KismetExecutionMessage(
+			*FString::Printf(
+				TEXT("%hs: No node was specified for add."), __FUNCTION__),
+				ELogVerbosity::Display);
+		return;
+	}
+
+	if (!InNode->CanBeAddedByUser())
+	{
+		FFrame::KismetExecutionMessage(
+			*FString::Printf(
+				TEXT("%hs:Cannot add node of type %s."), __FUNCTION__, *InNode->GetClass()->GetName()),
+				ELogVerbosity::Display);
 		return;
 	}
 
@@ -573,7 +620,10 @@ bool UMovieGraphConfig::RemoveNode(UMovieGraphNode* InNode)
 {
 	if (!InNode)
 	{
-		UE_LOG(LogMovieRenderPipeline, Error, TEXT("RemoveNode: Invalid Node"));
+		FFrame::KismetExecutionMessage(
+			*FString::Printf(
+				TEXT("%hs: Could not remove invalid InNode"), __FUNCTION__),
+				ELogVerbosity::Display);
 		return false;
 	}
 
@@ -609,7 +659,10 @@ RetType* UMovieGraphConfig::AddMember(TArray<TObjectPtr<ArrType>>& InMemberArray
 	
 	if (!NewMember)
 	{
-		UE_LOG(LogMovieRenderPipeline, Warning, TEXT("Unable to create new member object in the graph."));
+		FFrame::KismetExecutionMessage(
+			*FString::Printf(
+				TEXT("%hs: Unable to create new member object in the graph."), __FUNCTION__),
+				ELogVerbosity::Error);
 		return nullptr;
 	}
 
@@ -739,7 +792,10 @@ bool UMovieGraphConfig::DeleteMember(UMovieGraphMember* MemberToDelete)
 
 	if (!MemberToDelete->IsDeletable())
 	{
-		UE_LOG(LogMovieRenderPipeline, Error, TEXT("DeleteMember: The member '%s' cannot be deleted because it is flagged as non-deletable."), *MemberToDelete->GetMemberName());
+		FFrame::KismetExecutionMessage(
+			*FString::Printf(
+				TEXT("%hs: The member '%s' cannot be deleted because it is flagged as non-deletable."), __FUNCTION__, *MemberToDelete->GetMemberName()),
+				ELogVerbosity::Error);
 		return false;
 	}
 
@@ -1174,7 +1230,9 @@ void UMovieGraphConfig::CreateFlattenedGraph_Recursive(UMovieGraphEvaluatedConfi
 				{
 					// ToDo: This won't work long term if you have two different branches visiting the same node
 					// also we need to reset this every time we go start from the root.
-					UE_LOG(LogMovieRenderPipeline, Error, TEXT("Circular graph?"));
+					FFrame::KismetExecutionMessage(
+						*FString::Printf(TEXT("%hs: Circular graph?"), __FUNCTION__),
+						ELogVerbosity::Warning);
 					continue;
 				}
 
diff --git a/Engine/Plugins/MovieScene/MovieRenderPipeline/Source/MovieRenderPipelineCore/Private/Graph/MovieGraphPin.cpp b/Engine/Plugins/MovieScene/MovieRenderPipeline/Source/MovieRenderPipelineCore/Private/Graph/MovieGraphPin.cpp
--- a/Engine/Plugins/MovieScene/MovieRenderPipeline/Source/MovieRenderPipelineCore/Private/Graph/MovieGraphPin.cpp
+++ b/Engine/Plugins/MovieScene/MovieRenderPipeline/Source/MovieRenderPipelineCore/Private/Graph/MovieGraphPin.cpp
@@ -1,9 +1,13 @@
 // Copyright Epic Games, Inc. All Rights Reserved.
 
 #include "Graph/MovieGraphPin.h"
+
+#include "Algo/Find.h"
+#include "EdGraph/EdGraphSchema.h"
 #include "Graph/MovieGraphNode.h"
 #include "Graph/MovieGraphEdge.h"
 #include "MovieRenderPipelineCoreModule.h"
+#include "Graph/MovieGraphConfig.h"
 
 bool UMovieGraphPin::AddEdgeTo(UMovieGraphPin* InOtherPin)
 {
@@ -88,6 +92,59 @@ bool UMovieGraphPin::BreakAllEdges()
 	return bChanged;
 }
 
+FPinConnectionResponse UMovieGraphPin::CanCreateConnection_PinConnectionResponse(const UMovieGraphPin* InOtherPin) const
+{
+	if (!InOtherPin)
+	{
+		return FPinConnectionResponse(CONNECT_RESPONSE_DISALLOW, NSLOCTEXT("MoviePipeline", "InvalidPinError", "InOtherPin is invalid!"));
+	}
+	
+	// No Circular Connections
+	if (Node == InOtherPin->Node)
+	{
+		return FPinConnectionResponse(CONNECT_RESPONSE_DISALLOW, NSLOCTEXT("MoviePipeline", "CircularPinError", "No Circular Connections!"));
+	}
+
+	const bool bBothPinsAreBranch = Properties.bIsBranch && InOtherPin->Properties.bIsBranch;
+
+	const bool bBothPinsAreSameType = bBothPinsAreBranch ||				// Both are branches or
+		(!Properties.bIsBranch && !InOtherPin->Properties.bIsBranch &&	// Neither is branch and
+		Properties.Type == InOtherPin->Properties.Type)	;				// They have the same property type
+
+	// Pins need to be the same type
+	if (!bBothPinsAreSameType)				
+	{
+		return FPinConnectionResponse(CONNECT_RESPONSE_DISALLOW, NSLOCTEXT("MoviePipeline", "PinTypeMismatchError", "Pin types don't match!"));
+	}
+
+	if (!IsPinDirectionCompatibleWith(InOtherPin))
+	{
+		return FPinConnectionResponse(CONNECT_RESPONSE_DISALLOW, NSLOCTEXT("MoviePipeline", "PinDirectionMismatchError", "Directions are not compatible!"));
+	}
+
+	// Determine if the connection would violate branch restrictions enforced by the nodes involved in the connection.
+	FText BranchRestrictionError;
+	if (bBothPinsAreBranch && !IsConnectionToBranchAllowed(InOtherPin, BranchRestrictionError))
+	{
+		return FPinConnectionResponse(CONNECT_RESPONSE_DISALLOW, BranchRestrictionError);
+	}
+
+	// We don't allow multiple things to be connected to an Input Pin
+	const UMovieGraphPin* InputPin = IsInputPin() ? this : InOtherPin;
+	if(InputPin->GetAllConnectedPins().Num() > 0)
+	{
+		const ECanCreateConnectionResponse ReplyBreakInputs = IsInputPin() ? CONNECT_RESPONSE_BREAK_OTHERS_A : CONNECT_RESPONSE_BREAK_OTHERS_B;
+		return FPinConnectionResponse(ReplyBreakInputs, NSLOCTEXT("MoviePipeline", "PinInputReplaceExisting","Replace existing input connections"));
+	}
+	
+	return FPinConnectionResponse(CONNECT_RESPONSE_MAKE, NSLOCTEXT("MoviePipeline", "PinConnect", "Connect nodes"));
+}
+
+bool UMovieGraphPin::CanCreateConnection(const UMovieGraphPin* InOtherPin) const
+{
+	return CanCreateConnection_PinConnectionResponse(InOtherPin).Response != CONNECT_RESPONSE_DISALLOW;
+}
+
 bool UMovieGraphPin::IsConnected() const
 {
 	for (UMovieGraphEdge* Edge : Edges)
@@ -100,6 +157,13 @@ bool UMovieGraphPin::IsConnected() const
 
 	return false;
 }
+
+bool UMovieGraphPin::IsInputPin() const
+{
+	check(Node);
+	return Node->GetInputPin(Properties.Label) == this;
+}
+
 bool UMovieGraphPin::IsOutputPin() const
 {
 	check(Node);
@@ -154,4 +218,91 @@ TArray<UMovieGraphPin*> UMovieGraphPin::GetAllConnectedPins() const
 	}
 
 	return ConnectedPins;
-}
\ No newline at end of file
+}
+
+bool UMovieGraphPin::IsConnectionToBranchAllowed(const UMovieGraphPin* OtherPin, FText& OutError) const
+{
+	const UMovieGraphPin* InputPin = IsInputPin() ? this : OtherPin;
+	const UMovieGraphPin* OutputPin = !IsInputPin() ? this : OtherPin;
+	
+	if (!InputPin || !OutputPin)
+	{
+		OutError = NSLOCTEXT("MoviePipeline", "PinDirectionMismatchError", "Directions are not compatible!");
+		return false;
+	}
+	
+	const TObjectPtr<UMovieGraphNode> ToNode = InputPin->Node;
+	const TObjectPtr<UMovieGraphNode> FromNode = OutputPin->Node;
+	check(ToNode && FromNode);
+	const UMovieGraphConfig* GraphConfig = ToNode->GetGraph();
+
+	// Get all upstream/downstream nodes that occur on the connection -- these are the nodes that need to be checked for branch restrictions.
+	// FromNode/ToNode themselves also needs to be part of the validation checks.
+	TArray<UMovieGraphNode*> NodesToCheck = {FromNode, ToNode};
+	GraphConfig->VisitUpstreamNodes(FromNode, UMovieGraphConfig::FVisitNodesCallback::CreateLambda(
+		[&NodesToCheck](UMovieGraphNode* VisitedNode, const UMovieGraphPin* VisitedPin)
+		{
+			NodesToCheck.Add(VisitedNode);
+		}));
+
+	GraphConfig->VisitDownstreamNodes(ToNode, UMovieGraphConfig::FVisitNodesCallback::CreateLambda(
+		[&NodesToCheck](UMovieGraphNode* VisitedNode, const UMovieGraphPin* VisitedPin)
+		{
+			NodesToCheck.Add(VisitedNode);
+		}));
+
+	// Determine which branch(es) are connected to this node up/downstream.
+	const TArray<FString> DownstreamBranchNames = GraphConfig->GetDownstreamBranchNames(ToNode, InputPin);
+	const TArray<FString> UpstreamBranchNames = GraphConfig->GetUpstreamBranchNames(FromNode, OutputPin);
+	const bool bGlobalsIsDownstream = DownstreamBranchNames.Contains(UMovieGraphNode::GlobalsPinNameString);
+	const bool bGlobalsIsUpstream = UpstreamBranchNames.Contains(UMovieGraphNode::GlobalsPinNameString);
+	const bool bDownstreamBranchExistsAndIsntOnlyGlobals =
+		!DownstreamBranchNames.IsEmpty() && ((DownstreamBranchNames.Num() != 1) || (DownstreamBranchNames[0] != UMovieGraphNode::GlobalsPinNameString));
+	const bool bUpstreamBranchExistsAndIsntOnlyGlobals =
+		!UpstreamBranchNames.IsEmpty() && ((UpstreamBranchNames.Num() != 1) || (UpstreamBranchNames[0] != UMovieGraphNode::GlobalsPinNameString));
+
+	// Globals branches can only be connected to Globals branches
+	if ((bGlobalsIsDownstream && bUpstreamBranchExistsAndIsntOnlyGlobals) || (bGlobalsIsUpstream && bDownstreamBranchExistsAndIsntOnlyGlobals))
+	{
+		OutError = NSLOCTEXT("MoviePipeline", "GlobalsBranchMismatchError", "Globals branches can only be connected to other Globals branches.");
+		return false;
+	}
+
+	// Error out if any of the nodes that are part of the connection cannot be connected to the upstream/downstream branches.
+	for (const UMovieGraphNode* NodeToCheck : NodesToCheck)
+	{
+		if (NodeToCheck->GetBranchRestriction() == EMovieGraphBranchRestriction::Globals)
+		{
+			// Globals-specific nodes have to be connected such that the only upstream/downstream branches are Globals.
+			// If either the upstream/downstream branches are empty (ie, the node isn't connected to Inputs/Outputs yet)
+			// then the connection is OK for now -- the branch restriction will be enforced when nodes are connected to
+			// Inputs/Outputs.
+			if (bDownstreamBranchExistsAndIsntOnlyGlobals || bUpstreamBranchExistsAndIsntOnlyGlobals)
+			{
+				OutError = FText::Format(
+					NSLOCTEXT("MoviePipeline", "GlobalsBranchRestrictionError", "The node '{0}' can only be connected to the Globals branch."),
+						FText::FromString(NodeToCheck->GetName()));
+				return false;
+			}
+		}
+
+		// Check that render-layer-only nodes aren't connected to Globals.
+		if (NodeToCheck->GetBranchRestriction() == EMovieGraphBranchRestriction::RenderLayer)
+		{
+			if (bGlobalsIsDownstream || bGlobalsIsUpstream)
+			{
+				OutError = FText::Format(
+					NSLOCTEXT("MoviePipeline", "RenderLayerBranchRestrictionError", "The node '{0}' can only be connected to a render layer branch."),
+						FText::FromString(NodeToCheck->GetName()));
+				return false;
+			}
+		}
+	}
+
+	return true;
+}
+
+bool UMovieGraphPin::IsPinDirectionCompatibleWith(const UMovieGraphPin* OtherPin) const
+{
+	return IsInputPin() != OtherPin->IsInputPin();
+}
diff --git a/Engine/Plugins/MovieScene/MovieRenderPipeline/Source/MovieRenderPipelineCore/Public/Graph/MovieGraphNode.h b/Engine/Plugins/MovieScene/MovieRenderPipeline/Source/MovieRenderPipelineCore/Public/Graph/MovieGraphNode.h
--- a/Engine/Plugins/MovieScene/MovieRenderPipeline/Source/MovieRenderPipelineCore/Public/Graph/MovieGraphNode.h
+++ b/Engine/Plugins/MovieScene/MovieRenderPipeline/Source/MovieRenderPipelineCore/Public/Graph/MovieGraphNode.h
@@ -202,6 +202,12 @@ public:
 	/** Toggles the promotion of the property with the given name to a pin on the node. */
 	virtual void TogglePromotePropertyToPin(const FName& PropertyName);
 
+	/**
+	 * Determines if this node type can be added to the graph interactively by a user or via the API when constructing a graph.
+	 * @return true if the object can be added via the API, false otherwise
+	 */
+	virtual bool CanBeAddedByUser() const { return true; }
+
 	/**
 	 * Gets all overrideable properties that are defined on the node. This includes UPROPERTY-defined properties, as
 	 * well as dynamic properties. "Overrideable" means that the property has a corresponding property prefixed with
diff --git a/Engine/Plugins/MovieScene/MovieRenderPipeline/Source/MovieRenderPipelineCore/Public/Graph/MovieGraphPin.h b/Engine/Plugins/MovieScene/MovieRenderPipeline/Source/MovieRenderPipelineCore/Public/Graph/MovieGraphPin.h
--- a/Engine/Plugins/MovieScene/MovieRenderPipeline/Source/MovieRenderPipelineCore/Public/Graph/MovieGraphPin.h
+++ b/Engine/Plugins/MovieScene/MovieRenderPipeline/Source/MovieRenderPipelineCore/Public/Graph/MovieGraphPin.h
@@ -9,6 +9,8 @@
 class UMovieGraphNode;
 class UMovieGraphEdge;
 
+struct FPinConnectionResponse;
+
 USTRUCT(BlueprintType)
 struct FMovieGraphPinProperties
 {
@@ -76,7 +78,10 @@ public:
 	bool AddEdgeTo(UMovieGraphPin* InOtherPin);
 	bool BreakEdgeTo(UMovieGraphPin* InOtherPin);
 	bool BreakAllEdges();
+	FPinConnectionResponse CanCreateConnection_PinConnectionResponse(const UMovieGraphPin* InOtherPin) const;
+	bool CanCreateConnection(const UMovieGraphPin* InOtherPin) const;
 	bool IsConnected() const;
+	bool IsInputPin() const;
 	bool IsOutputPin() const;
 	int32 EdgeCount() const;
 	bool AllowsMultipleConnections() const;
@@ -87,6 +92,13 @@ public:
 	/** Gets all connected pins. */
 	TArray<UMovieGraphPin*> GetAllConnectedPins() const;
 
+	/**
+	 * Determines if the connection between this pin and OtherPin follows branch restriction rules. OutError is populated
+	 * with an error if the connection should be rejected and the function will return false.
+	 */
+	bool IsConnectionToBranchAllowed(const UMovieGraphPin* OtherPin, FText& OutError) const;
+	bool IsPinDirectionCompatibleWith(const UMovieGraphPin* OtherPin) const;
+
 public:
 	// The node that this pin belongs to.
 	UPROPERTY(BlueprintReadOnly, Category = "Properties")
diff --git a/Engine/Plugins/MovieScene/MovieRenderPipeline/Source/MovieRenderPipelineCore/Public/Graph/Nodes/MovieGraphInputNode.h b/Engine/Plugins/MovieScene/MovieRenderPipeline/Source/MovieRenderPipelineCore/Public/Graph/Nodes/MovieGraphInputNode.h
--- a/Engine/Plugins/MovieScene/MovieRenderPipeline/Source/MovieRenderPipelineCore/Public/Graph/Nodes/MovieGraphInputNode.h
+++ b/Engine/Plugins/MovieScene/MovieRenderPipeline/Source/MovieRenderPipelineCore/Public/Graph/Nodes/MovieGraphInputNode.h
@@ -19,6 +19,8 @@ public:
 	virtual TArray<FMovieGraphPinProperties> GetOutputPinProperties() const override;
 	virtual TArray<UMovieGraphPin*> EvaluatePinsToFollow(FMovieGraphEvaluationContext& InContext) const override;
 	virtual bool CanBeDisabled() const override;
+	
+	virtual bool CanBeAddedByUser() const override { return false; }
 
 #if WITH_EDITOR
 	virtual FText GetNodeTitle(const bool bGetDescriptive = false) const override;
diff --git a/Engine/Plugins/MovieScene/MovieRenderPipeline/Source/MovieRenderPipelineCore/Public/Graph/Nodes/MovieGraphOutputNode.h b/Engine/Plugins/MovieScene/MovieRenderPipeline/Source/MovieRenderPipelineCore/Public/Graph/Nodes/MovieGraphOutputNode.h
--- a/Engine/Plugins/MovieScene/MovieRenderPipeline/Source/MovieRenderPipelineCore/Public/Graph/Nodes/MovieGraphOutputNode.h
+++ b/Engine/Plugins/MovieScene/MovieRenderPipeline/Source/MovieRenderPipelineCore/Public/Graph/Nodes/MovieGraphOutputNode.h
@@ -17,6 +17,8 @@ public:
 	
 	virtual TArray<FMovieGraphPinProperties> GetInputPinProperties() const override;
 	virtual bool CanBeDisabled() const override;
+	
+	virtual bool CanBeAddedByUser() const override { return false; }
 
 #if WITH_EDITOR
 	virtual FText GetNodeTitle(const bool bGetDescriptive = false) const override;
diff --git a/Engine/Plugins/MovieScene/MovieRenderPipeline/Source/MovieRenderPipelineEditor/Private/Graph/MovieGraphSchema.cpp b/Engine/Plugins/MovieScene/MovieRenderPipeline/Source/MovieRenderPipelineEditor/Private/Graph/MovieGraphSchema.cpp
--- a/Engine/Plugins/MovieScene/MovieRenderPipeline/Source/MovieRenderPipelineEditor/Private/Graph/MovieGraphSchema.cpp
+++ b/Engine/Plugins/MovieScene/MovieRenderPipeline/Source/MovieRenderPipelineEditor/Private/Graph/MovieGraphSchema.cpp
@@ -135,76 +135,10 @@ void UMovieGraphSchema::InitMoviePipelineNodeClasses()
 
 bool UMovieGraphSchema::IsConnectionToBranchAllowed(const UEdGraphPin* InputPin, const UEdGraphPin* OutputPin, FText& OutError) const
 {
-	UMovieGraphNode* ToNode = UE::MovieGraph::Private::GetGraphNodeFromEdPin(InputPin);
-	UMovieGraphNode* FromNode = UE::MovieGraph::Private::GetGraphNodeFromEdPin(OutputPin);
 	const UMovieGraphPin* ToPin = UE::MovieGraph::Private::GetGraphPinFromEdPin(InputPin);
 	const UMovieGraphPin* FromPin = UE::MovieGraph::Private::GetGraphPinFromEdPin(OutputPin);
-	const UMovieGraphConfig* GraphConfig = UE::MovieGraph::Private::GetGraphFromEdPin(InputPin);
-
-	// Get all upstream/downstream nodes that occur on the connection -- these are the nodes that need to be checked for branch restrictions.
-	// FromNode/ToNode themselves also needs to be part of the validation checks.
-	TArray<UMovieGraphNode*> NodesToCheck = {FromNode, ToNode};
-	GraphConfig->VisitUpstreamNodes(FromNode, UMovieGraphConfig::FVisitNodesCallback::CreateLambda(
-		[&NodesToCheck](UMovieGraphNode* VisitedNode, const UMovieGraphPin* VisitedPin)
-		{
-			NodesToCheck.Add(VisitedNode);
-		}));
-
-	GraphConfig->VisitDownstreamNodes(ToNode, UMovieGraphConfig::FVisitNodesCallback::CreateLambda(
-		[&NodesToCheck](UMovieGraphNode* VisitedNode, const UMovieGraphPin* VisitedPin)
-		{
-			NodesToCheck.Add(VisitedNode);
-		}));
-
-	// Determine which branch(es) are connected to this node up/downstream.
-	const TArray<FString> DownstreamBranchNames = GraphConfig->GetDownstreamBranchNames(ToNode, ToPin);
-	const TArray<FString> UpstreamBranchNames = GraphConfig->GetUpstreamBranchNames(FromNode, FromPin);
-	const bool bGlobalsIsDownstream = DownstreamBranchNames.Contains(UMovieGraphNode::GlobalsPinNameString);
-	const bool bGlobalsIsUpstream = UpstreamBranchNames.Contains(UMovieGraphNode::GlobalsPinNameString);
-	const bool bDownstreamBranchExistsAndIsntOnlyGlobals =
-		!DownstreamBranchNames.IsEmpty() && ((DownstreamBranchNames.Num() != 1) || (DownstreamBranchNames[0] != UMovieGraphNode::GlobalsPinNameString));
-	const bool bUpstreamBranchExistsAndIsntOnlyGlobals =
-		!UpstreamBranchNames.IsEmpty() && ((UpstreamBranchNames.Num() != 1) || (UpstreamBranchNames[0] != UMovieGraphNode::GlobalsPinNameString));
-
-	// Globals branches can only be connected to Globals branches
-	if ((bGlobalsIsDownstream && bUpstreamBranchExistsAndIsntOnlyGlobals) || (bGlobalsIsUpstream && bDownstreamBranchExistsAndIsntOnlyGlobals))
-	{
-		OutError = NSLOCTEXT("MoviePipeline", "GlobalsBranchMismatchError", "Globals branches can only be connected to other Globals branches.");
-		return false;
-	}
-
-	// Error out if any of the nodes that are part of the connection cannot be connected to the upstream/downstream branches.
-	for (const UMovieGraphNode* NodeToCheck : NodesToCheck)
-	{
-		if (NodeToCheck->GetBranchRestriction() == EMovieGraphBranchRestriction::Globals)
-		{
-			// Globals-specific nodes have to be connected such that the only upstream/downstream branches are Globals.
-			// If either the upstream/downstream branches are empty (ie, the node isn't connected to Inputs/Outputs yet)
-			// then the connection is OK for now -- the branch restriction will be enforced when nodes are connected to
-			// Inputs/Outputs.
-			if (bDownstreamBranchExistsAndIsntOnlyGlobals || bUpstreamBranchExistsAndIsntOnlyGlobals)
-			{
-				OutError = FText::Format(
-					NSLOCTEXT("MoviePipeline", "GlobalsBranchRestrictionError", "The node '{0}' can only be connected to the Globals branch."),
-						NodeToCheck->GetNodeTitle());
-				return false;
-			}
-		}
-
-		// Check that render-layer-only nodes aren't connected to Globals.
-		if (NodeToCheck->GetBranchRestriction() == EMovieGraphBranchRestriction::RenderLayer)
-		{
-			if (bGlobalsIsDownstream || bGlobalsIsUpstream)
-			{
-				OutError = FText::Format(
-					NSLOCTEXT("MoviePipeline", "RenderLayerBranchRestrictionError", "The node '{0}' can only be connected to a render layer branch."),
-						NodeToCheck->GetNodeTitle());
-				return false;
-			}
-		}
-	}
-
-	return true;
+	
+	return FromPin->IsConnectionToBranchAllowed(ToPin, OutError);
 }
 
 void UMovieGraphSchema::GetGraphContextActions(FGraphContextMenuBuilder& ContextMenuBuilder) const
@@ -273,44 +207,10 @@ void UMovieGraphSchema::GetGraphContextActions(FGraphContextMenuBuilder& Context
 
 const FPinConnectionResponse UMovieGraphSchema::CanCreateConnection(const UEdGraphPin* PinA, const UEdGraphPin* PinB) const
 {
-	// No Circular Connections
-	if (PinA->GetOwningNode() == PinB->GetOwningNode())
-	{
-		return FPinConnectionResponse(CONNECT_RESPONSE_DISALLOW, NSLOCTEXT("MoviePipeline", "CircularPinError", "No Circular Connections!"));
-	}
-
-	// Pins need to be the same type
-	if (PinA->PinType.PinCategory != PinB->PinType.PinCategory)
-	{
-		return FPinConnectionResponse(CONNECT_RESPONSE_DISALLOW, NSLOCTEXT("MoviePipeline", "PinTypeMismatchError", "Pin types don't match!"));
-	}
-
-	// Re-organize PinA/PinB to Input/Output by comparing internal directions to avoid having to check both cases depending on which
-	// direction the conneciton was made.
-	const UEdGraphPin* InputPin = nullptr;
-	const UEdGraphPin* OutputPin = nullptr;
-
-	if (!CategorizePinsByDirection(PinA, PinB, /*out*/ InputPin, /*out*/ OutputPin))
-	{
-		return FPinConnectionResponse(CONNECT_RESPONSE_DISALLOW, NSLOCTEXT("MoviePipeline", "PinDirectionMismatchError", "Directions are not compatible!"));
-	}
-
-	// Determine if the connection would violate branch restrictions enforced by the nodes involved in the connection.
-	FText BranchRestrictionError;
-	if (!IsConnectionToBranchAllowed(InputPin, OutputPin, BranchRestrictionError))
-	{
-		return FPinConnectionResponse(CONNECT_RESPONSE_DISALLOW, BranchRestrictionError);
-	}
-
-	// We don't allow multiple things to be connected to an Input Pin
-	if(InputPin->HasAnyConnections())
-	{
-		const ECanCreateConnectionResponse ReplyBreakInputs = (PinA == InputPin) ? CONNECT_RESPONSE_BREAK_OTHERS_A : CONNECT_RESPONSE_BREAK_OTHERS_B;
-		return FPinConnectionResponse(ReplyBreakInputs, NSLOCTEXT("MoviePipeline", "PinInputReplaceExisting","Replace existing input connections"));
-	}
+	const UMovieGraphPin* FromPin = UE::MovieGraph::Private::GetGraphPinFromEdPin(PinA);
+	const UMovieGraphPin* ToPin = UE::MovieGraph::Private::GetGraphPinFromEdPin(PinB);
 	
-	// Make sure the pins are not on the same node
-	return FPinConnectionResponse(CONNECT_RESPONSE_MAKE, NSLOCTEXT("MoviePipeline", "PinConnect", "Connect nodes"));
+	return FromPin->CanCreateConnection_PinConnectionResponse(ToPin);
 }
 
 bool UMovieGraphSchema::TryCreateConnection(UEdGraphPin* InA, UEdGraphPin* InB) const
diff --git a/Engine/Plugins/MovieScene/MovieRenderPipeline/Source/MovieRenderPipelineEditor/Private/Graph/MovieGraphSchema.h b/Engine/Plugins/MovieScene/MovieRenderPipeline/Source/MovieRenderPipelineEditor/Private/Graph/MovieGraphSchema.h
--- a/Engine/Plugins/MovieScene/MovieRenderPipeline/Source/MovieRenderPipelineEditor/Private/Graph/MovieGraphSchema.h
+++ b/Engine/Plugins/MovieScene/MovieRenderPipeline/Source/MovieRenderPipelineEditor/Private/Graph/MovieGraphSchema.h
@@ -41,7 +41,7 @@ public:
 private:
 	/**
 	 * Determines if the connection between InputPin and OutputPin follows branch restriction rules. OutError is populated
-	 * with an error if the connection should be rejected.
+	 * with an error if the connection should be rejected and the function will return false.
 	 */
 	bool IsConnectionToBranchAllowed(const UEdGraphPin* InputPin, const UEdGraphPin* OutputPin, FText& OutError) const;
 
diff --git a/Engine/Plugins/MovieScene/MovieRenderPipeline/Source/MovieRenderPipelineEditor/Private/Tests/MovieGraphEditorTestUtilities.cpp b/Engine/Plugins/MovieScene/MovieRenderPipeline/Source/MovieRenderPipelineEditor/Private/Tests/MovieGraphEditorTestUtilities.cpp
new file mode 100644
--- /dev/null
+++ b/Engine/Plugins/MovieScene/MovieRenderPipeline/Source/MovieRenderPipelineEditor/Private/Tests/MovieGraphEditorTestUtilities.cpp
@@ -0,0 +1,115 @@
+﻿// Copyright Epic Games, Inc. All Rights Reserved.
+
+#include "MovieGraphEditorTestUtilities.h"
+
+#include "Graph/MovieGraphConfig.h"
+#include "Graph/MovieGraphConfigFactory.h"
+
+#include "Algo/RemoveIf.h"
+#include "AssetRegistry/AssetRegistryModule.h"
+#include "AssetToolsModule.h"
+#include "IAssetTools.h"
+#include "UObject/Package.h"
+
+namespace UE::MovieGraph::Private::Tests
+{
+	UMovieGraphConfig* CreateNewMovieGraphConfig(const FName InName)
+	{
+		return NewObject<UMovieGraphConfig>(
+			GetTransientPackage(), UMovieGraphConfig::StaticClass(), InName, RF_Transient);
+	}
+
+	UMovieGraphConfig* CreateDefaultMovieGraphConfig()
+	{
+		static const FName NAME_AssetTools = "AssetTools";
+		IAssetTools* AssetTools = &FModuleManager::GetModuleChecked<FAssetToolsModule>(NAME_AssetTools).Get();
+
+		UMovieGraphConfigFactory* ConfigFactory =
+			NewObject<UMovieGraphConfigFactory>(UMovieGraphConfigFactory::StaticClass());
+
+		UMovieGraphConfig* GraphConfig =
+			Cast<UMovieGraphConfig>(
+				AssetTools->CreateAsset(
+					"GraphTestConfig", "/Game/", UMovieGraphConfig::StaticClass(), ConfigFactory
+				)
+			);
+
+		return GraphConfig;
+	}
+
+	void OpenGraphConfigInEditor(UMovieGraphConfig* InGraphConfig)
+	{
+		static const FName NAME_AssetTools = "AssetTools";
+		IAssetTools* AssetTools = &FModuleManager::GetModuleChecked<FAssetToolsModule>(NAME_AssetTools).Get();
+
+		AssetTools->OpenEditorForAssets({InGraphConfig});
+	}
+
+	TArray<UClass*> GetAllDerivedClasses(UClass* BaseClass, bool bRecursive)
+	{
+		TArray<UClass*> DerivedClasses = GetNativeClasses(BaseClass, bRecursive);
+		DerivedClasses.Append(GetBlueprintClasses(BaseClass, bRecursive));
+		return DerivedClasses;
+	}
+
+	void RemoveAbstractClasses(TArray<UClass*>& InClassArray)
+	{
+		Algo::StableRemoveIf(InClassArray,
+			[](const UClass* Class)
+			{
+				return !Class ||Class->HasAnyClassFlags(CLASS_Abstract);
+			});
+	}
+
+	TArray<UClass*> GetNativeClasses(UClass* BaseClass, bool bRecursive)
+	{
+		TArray<UClass*> DerivedClasses;
+		GetDerivedClasses(BaseClass, DerivedClasses, bRecursive);
+
+		DerivedClasses.Remove(BaseClass);
+
+		RemoveAbstractClasses(DerivedClasses);
+
+		return DerivedClasses;
+	}
+
+	bool IsBlueprintChildOfBaseClass(const FAssetData& BlueprintClassData, UClass* BaseClass)
+	{		
+		const FString NativeParentClassPath =
+			BlueprintClassData.GetTagValueRef<FString>(FBlueprintTags::NativeParentClassPath);
+		const FSoftClassPath ClassPath(NativeParentClassPath);
+		const UClass* NativeParentClass = ClassPath.ResolveClass();
+
+		return NativeParentClass && (NativeParentClass == BaseClass || NativeParentClass->IsChildOf(BaseClass));
+	}
+
+	TArray<UClass*> GetBlueprintClasses(
+		UClass* BaseClass, bool bRecursive)
+	{
+		const FAssetRegistryModule& AssetRegistryModule =
+			FModuleManager::LoadModuleChecked< FAssetRegistryModule >(FName("AssetRegistry"));
+		const IAssetRegistry& AssetRegistry = AssetRegistryModule.Get();
+
+		TArray<UClass*> DerivedClasses;
+		TArray< FAssetData > Assets;
+		AssetRegistry.GetAssetsByClass(UBlueprint::StaticClass()->GetClassPathName(), Assets);	
+		for (const FAssetData& Asset : Assets)
+		{
+			if (IsBlueprintChildOfBaseClass(Asset, BaseClass))
+			{
+				if (const UBlueprint* BlueprintAsset = Cast<UBlueprint>(Asset.GetAsset()))
+				{
+					UClass* LoadedClass = BlueprintAsset->GeneratedClass;
+					if (ensure(LoadedClass && BlueprintAsset->ParentClass))
+					{
+						DerivedClasses.AddUnique(LoadedClass);
+					}
+				}
+			}
+		}
+
+		RemoveAbstractClasses(DerivedClasses);
+
+		return DerivedClasses;
+	}
+}
diff --git a/Engine/Plugins/MovieScene/MovieRenderPipeline/Source/MovieRenderPipelineEditor/Private/Tests/MovieGraphEditorTestUtilities.h b/Engine/Plugins/MovieScene/MovieRenderPipeline/Source/MovieRenderPipelineEditor/Private/Tests/MovieGraphEditorTestUtilities.h
new file mode 100644
--- /dev/null
+++ b/Engine/Plugins/MovieScene/MovieRenderPipeline/Source/MovieRenderPipelineEditor/Private/Tests/MovieGraphEditorTestUtilities.h
@@ -0,0 +1,17 @@
+﻿// Copyright Epic Games, Inc. All Rights Reserved.
+
+#pragma once
+
+#include "CoreMinimal.h"
+
+class UMovieGraphConfig;
+
+namespace UE::MovieGraph::Private::Tests
+{
+	UMovieGraphConfig* CreateNewMovieGraphConfig(const FName InName = "GraphTestConfig");
+	UMovieGraphConfig* CreateDefaultMovieGraphConfig();
+	void OpenGraphConfigInEditor(UMovieGraphConfig* InGraphConfig);
+	TArray<UClass*> GetAllDerivedClasses(UClass* BaseClass, bool bRecursive);
+	TArray<UClass*> GetNativeClasses(UClass* BaseClass, bool bRecursive);
+	TArray<UClass*> GetBlueprintClasses(UClass* BaseClass, bool bRecursive);
+}
diff --git a/Engine/Plugins/MovieScene/MovieRenderPipeline/Source/MovieRenderPipelineEditor/Private/Tests/MovieGraphTests.cpp b/Engine/Plugins/MovieScene/MovieRenderPipeline/Source/MovieRenderPipelineEditor/Private/Tests/MovieGraphTests.cpp
new file mode 100644
--- /dev/null
+++ b/Engine/Plugins/MovieScene/MovieRenderPipeline/Source/MovieRenderPipelineEditor/Private/Tests/MovieGraphTests.cpp
@@ -0,0 +1,290 @@
+﻿// Copyright Epic Games, Inc. All Rights Reserved.
+
+#include "MovieGraphEditorTestUtilities.h"
+
+#include "Algo/RemoveIf.h"
+#include "Graph/MovieGraphConfig.h"
+#include "Graph/Nodes/MovieGraphBranchNode.h"
+#include "Graph/Nodes/MovieGraphInputNode.h"
+#include "Graph/Nodes/MovieGraphOutputNode.h"
+#include "Graph/Nodes/MovieGraphSubgraphNode.h"
+#include "Graph/Nodes/MovieGraphVariableNode.h"
+
+#include "Misc/AutomationTest.h"
+
+namespace UE::MovieGraph::Private::Tests
+{
+	void FilterOutUndesirableClasses(TArray<UClass*>& InClassArray)
+	{
+		static const TSet<UClass*> UndesirableClasses = {
+			UMovieGraphVariableNode::StaticClass(),	// Will be tested separately in this test
+			UMovieGraphInputNode::StaticClass(),	// Should not be able to be added via API
+			UMovieGraphOutputNode::StaticClass(),	// Should not be able to be added via API
+			UMovieGraphSubgraphNode::StaticClass()	// Will be tested separately in another test
+		};
+
+		Algo::StableRemoveIf(InClassArray,
+			[&](const UClass* Class)
+			{
+				return UndesirableClasses.Contains(Class);
+			});
+	}
+
+	/**
+	* 1) Add a null node, verify that NumNodes == 0
+	* 2) Add an unsupported node, verify that NumNodes == 0
+	* 3) Add a supported node, Verify that NumNodes == NumSupportedNodes
+	* 4) Add a variable and variable node, verify the variable node exists includes the variable member
+	*/
+	IMPLEMENT_SIMPLE_AUTOMATION_TEST(FAddGraphNodeTest, "VirtualProduction.MovieGraph.AddGraphNodeTest", (EAutomationTestFlags::EditorContext | EAutomationTestFlags::ProductFilter));
+	bool FAddGraphNodeTest::RunTest(const FString& Parameters)
+	{
+		// Create Config
+		UMovieGraphConfig* GraphConfig = CreateNewMovieGraphConfig("AddGraphNodeTest");
+		TestTrue(TEXT("MovieGraphConfig successfully created"), GraphConfig != nullptr);
+
+		GraphConfig->ConstructRuntimeNode<UMovieGraphNode>(nullptr);
+		TestTrue(TEXT("Null Node not added"), GraphConfig->GetNodes().IsEmpty());
+
+		// Ensure input and output nodes can't be added via the API
+		const UMovieGraphInputNode* TestInputRuntimeNode = GraphConfig->ConstructRuntimeNode<UMovieGraphInputNode>(UMovieGraphInputNode::StaticClass());
+		if(TestInputRuntimeNode && GraphConfig->GetNodes().Contains(TestInputRuntimeNode))
+		{
+			AddError(*FString::Printf(TEXT("Graph allowed user construction of %s"), *TestInputRuntimeNode->GetClass()->GetName()));
+		}
+		const UMovieGraphOutputNode* TestOutputRuntimeNode = GraphConfig->ConstructRuntimeNode<UMovieGraphOutputNode>(UMovieGraphOutputNode::StaticClass());
+		if(TestOutputRuntimeNode && GraphConfig->GetNodes().Contains(TestOutputRuntimeNode))
+		{
+			AddError(*FString::Printf(TEXT("Graph allowed user construction of %s"), *TestOutputRuntimeNode->GetClass()->GetName()));
+		}
+		
+		// Add one node of each type, excepting variables, inputs and outputs
+		TArray<UClass*> AllNodeClasses = GetAllDerivedClasses(UMovieGraphNode::StaticClass(), true);
+		FilterOutUndesirableClasses(AllNodeClasses);
+		for (UClass* Class : AllNodeClasses)
+		{
+			GraphConfig->ConstructRuntimeNode<UMovieGraphNode>(Class);
+		}
+		TestTrue(TEXT("All supported nodes types added"), GraphConfig->GetNodes().Num() == AllNodeClasses.Num());
+
+		// Test adding a variable node
+		UMovieGraphVariable* Variable = GraphConfig->AddVariable(TEXT("A_Cool_Variable"));
+		TestTrue(TEXT("Variable member added"), Variable != nullptr && GraphConfig->GetVariables().Contains(Variable));
+		UMovieGraphVariableNode* VariableNode = GraphConfig->ConstructRuntimeNode<UMovieGraphVariableNode>(UMovieGraphVariableNode::StaticClass());
+		TestTrue(TEXT("Variable node type added"), VariableNode != nullptr && GraphConfig->GetNodes().Contains(VariableNode));
+		VariableNode->SetVariable(Variable);
+		TestTrue(TEXT("Variable node variable member set"), VariableNode->GetVariable() == Variable);
+
+		// Open graph config in editor if desired
+		//OpenGraphConfigInEditor(GraphConfig);
+
+		return true;
+	}
+
+	/**
+	* 1) Add two nodes to the graph
+	* 2) Remove one node, verify NumNodes == 1
+	* 3) Try to remove a nullptr, verify NumNodes == 1
+	* 4) Try to remove a node that doesn't exist in the graph, verify NumNodes == 1
+	* 5) Remove a variable, variable node, and verify it's not in GetNodes()
+	*/
+	IMPLEMENT_SIMPLE_AUTOMATION_TEST(FRemoveGraphNodeTest, "VirtualProduction.MovieGraph.RemoveGraphNodeTest", (EAutomationTestFlags::EditorContext | EAutomationTestFlags::ProductFilter));
+	bool FRemoveGraphNodeTest::RunTest(const FString& Parameters)
+	{
+		// Create Config
+		UMovieGraphConfig* GraphConfig = CreateNewMovieGraphConfig("RemoveGraphNodeTest");
+		TestTrue(TEXT("MovieGraphConfig successfully created"), GraphConfig != nullptr);
+		
+		// Add two nodes
+		TArray<UClass*> AllNodeClasses = GetAllDerivedClasses(UMovieGraphNode::StaticClass(), true);
+		FilterOutUndesirableClasses(AllNodeClasses);
+		UMovieGraphNode* NodeOne = GraphConfig->ConstructRuntimeNode<UMovieGraphNode>(AllNodeClasses[0]);
+		UMovieGraphNode* NodeTwo = GraphConfig->ConstructRuntimeNode<UMovieGraphNode>(AllNodeClasses[1]);
+		TestTrue(
+			FString::Printf(TEXT("Node types added: %s and %s"), *AllNodeClasses[0]->GetName(), *AllNodeClasses[1]->GetName()),
+			GraphConfig->GetNodes().Num() == 2);
+
+		// Test removing one node
+		GraphConfig->RemoveNode(NodeOne);
+		TestTrue(TEXT("Node removed"), GraphConfig->GetNodes().Num() == 1);
+
+		// Test removing a nullptr
+		GraphConfig->RemoveNode(nullptr);
+		TestTrue(TEXT("Calling RemoveNode with nullptr should do nothing"), GraphConfig->GetNodes().Num() == 1);
+
+		// Test removing a node which doesn't exist in the graph
+		GraphConfig->RemoveNode(NodeOne);
+		TestTrue(TEXT("Attempting to remove node which does not exist in the graph should do nothing"), GraphConfig->GetNodes().Num() == 1);
+		
+		// Test adding and removing a variable member and variable node
+		UMovieGraphVariable* Variable = GraphConfig->AddVariable(TEXT("A_Cool_Variable"));
+		TestTrue(TEXT("Variable member added"), Variable != nullptr && GraphConfig->GetVariables().Contains(Variable));
+		
+		UMovieGraphVariableNode* VariableNode = GraphConfig->ConstructRuntimeNode<UMovieGraphVariableNode>(UMovieGraphVariableNode::StaticClass());
+		TestTrue(TEXT("Variable node type added"), VariableNode != nullptr && GraphConfig->GetNodes().Contains(VariableNode));
+		
+		VariableNode->SetVariable(Variable);
+		TestTrue(TEXT("Variable node variable member set"), VariableNode->GetVariable() == Variable);
+		
+		VariableNode->SetVariable(nullptr);
+		TestTrue(TEXT("Calling SetVariable with nullptr should do nothing"), VariableNode->GetVariable() == Variable);
+		
+		GraphConfig->RemoveNode(VariableNode);
+		TestTrue(TEXT("Variable node removed"), GraphConfig->GetNodes().Num() == 1);
+		
+		GraphConfig->DeleteMember(Variable);
+		TestTrue(TEXT("Variable member removed"), !GraphConfig->GetVariables().Contains(Variable));
+
+		// Open graph config in editor if desired
+		//OpenGraphConfigInEditor(GraphConfig);
+
+		return true;
+	}
+	
+	/**
+	* 1) Add nodes to the graph, verify that NumNodes == AllNodeClasses
+	* 2) Connect nodes together
+	* 3) Use the Pin API to verify that the connection is between the correct nodes
+	*/
+	IMPLEMENT_SIMPLE_AUTOMATION_TEST(FGraphNodeAddEdgeTest, "VirtualProduction.MovieGraph.GraphNodeAddEdgeTest", (EAutomationTestFlags::EditorContext | EAutomationTestFlags::ProductFilter));
+	bool FGraphNodeAddEdgeTest::RunTest(const FString& Parameters)
+	{
+		// Create Config
+		UMovieGraphConfig* GraphConfig = CreateNewMovieGraphConfig("GraphNodeAddEdgeTest");
+		TestTrue(TEXT("MovieGraphConfig successfully created"), GraphConfig != nullptr);
+
+		// Add all nodes excepting variables, inputs and outputs
+		TArray<UClass*> AllNodeClasses = GetAllDerivedClasses(UMovieGraphNode::StaticClass(), true);
+		FilterOutUndesirableClasses(AllNodeClasses);
+		for (UClass* Class : AllNodeClasses)
+		{
+			GraphConfig->ConstructRuntimeNode<UMovieGraphNode>(Class);
+		}
+		TestTrue(TEXT("All supported nodes types added"), GraphConfig->GetNodes().Num() == AllNodeClasses.Num());
+
+		// Connect all edges		
+		// Some nodes will be duplicated in between these arrays
+		
+		TArray<TObjectPtr<UMovieGraphNode>> NodesWithOutputs = {GraphConfig->GetInputNode()};
+		TArray<TObjectPtr<UMovieGraphNode>> NodesWithInputs = {GraphConfig->GetOutputNode()};
+		
+		auto SeparateNodesIntoInputAndOutputTypes = [&]()
+		{
+			for (const TObjectPtr<UMovieGraphNode>& Node : GraphConfig->GetNodes())
+			{
+				if (!Node->GetInputPins().IsEmpty())
+				{
+					NodesWithInputs.Add(Node);
+				}
+
+				if (!Node->GetOutputPins().IsEmpty())
+				{
+					NodesWithOutputs.Add(Node);
+				}
+
+				const bool bHasNoInputsOrOutputs = Node->GetInputPins().Num() == 0 && Node->GetOutputPins().Num() == 0;
+			
+				TestTrue(*FString::Printf(TEXT("Node %s has inputs or outputs"), *Node.GetName()), !bHasNoInputsOrOutputs);
+			}
+		};
+		SeparateNodesIntoInputAndOutputTypes();
+
+		auto ConnectAllEdgesToAllNodes = [&]()
+		{
+			for (const TObjectPtr<UMovieGraphNode>& NodeFrom : NodesWithOutputs)
+			{
+				for (const TObjectPtr<UMovieGraphPin>& OutputPin : NodeFrom->GetOutputPins())
+				{
+					for (const TObjectPtr<UMovieGraphNode>& NodeTo : NodesWithInputs)
+					{
+						if (NodeFrom == NodeTo) // Don't try to self-connect
+						{
+							continue;
+						}
+					
+						for (const TObjectPtr<UMovieGraphPin>& InputPin : NodeTo->GetInputPins())
+						{
+							if (OutputPin->CanCreateConnection(InputPin))
+							{
+								GraphConfig->AddLabeledEdge(
+									NodeFrom, OutputPin->Properties.Label,
+									NodeTo, InputPin->Properties.Label);
+
+								TestTrue(*FString::Printf(TEXT("Edge added between %s and %s"), *NodeFrom->GetName(), *NodeTo->GetName()),
+									OutputPin->GetAllConnectedPins().Contains(InputPin));
+							}
+						}
+					}
+				}
+			}
+		};
+		ConnectAllEdgesToAllNodes();
+		
+		// Open graph config in editor if desired
+		//OpenGraphConfigInEditor(GraphConfig);
+		
+		return true;
+	}
+	
+	/**
+	* 1) Add a node, connect it with default input and output nodes
+	* 2) Remove an edge
+	* 3) Use the Pin API to verify no more connection
+	*/
+	IMPLEMENT_SIMPLE_AUTOMATION_TEST(FGraphNodeRemoveEdgeTest, "VirtualProduction.MovieGraph.GraphNodeRemoveEdgeTest", (EAutomationTestFlags::EditorContext | EAutomationTestFlags::ProductFilter));
+	bool FGraphNodeRemoveEdgeTest::RunTest(const FString& Parameters)
+	{
+		// Create Config
+		UMovieGraphConfig* GraphConfig = CreateNewMovieGraphConfig("GraphNodeRemoveEdgeTest");
+		TestTrue(TEXT("MovieGraphConfig successfully created"), GraphConfig != nullptr);
+
+		// Add a node with an input and output for connection
+		UMovieGraphNode* MiddleNode =
+			GraphConfig->ConstructRuntimeNode<UMovieGraphBranchNode>(UMovieGraphBranchNode::StaticClass());
+		TestTrue(TEXT("MiddleNode successfully added"), GraphConfig->GetNodes().Contains(MiddleNode));
+
+		UMovieGraphNode* InputNode = GraphConfig->GetInputNode();
+		UMovieGraphNode* OutputNode = GraphConfig->GetOutputNode();
+
+		const TObjectPtr<UMovieGraphPin> InputNodeOutputPin = InputNode->GetOutputPins()[0];
+		const TObjectPtr<UMovieGraphPin> MiddleNodeInputPin = MiddleNode->GetInputPins()[0];
+		const TObjectPtr<UMovieGraphPin> MiddleNodeOutputPin = MiddleNode->GetOutputPins()[0];
+		const TObjectPtr<UMovieGraphPin> OutputNodeInputPin = OutputNode->GetInputPins()[0];
+
+		// Add first edge
+		GraphConfig->AddLabeledEdge(
+			InputNode, InputNodeOutputPin->Properties.Label,
+			MiddleNode, MiddleNodeInputPin->Properties.Label);
+
+		TestTrue(*FString::Printf(TEXT("Edge added between %s and %s"), *InputNode->GetName(), *MiddleNode->GetName()),
+			InputNodeOutputPin->GetAllConnectedPins().Contains(MiddleNodeInputPin));
+
+		// Remove edge
+		GraphConfig->RemoveEdge(InputNode, InputNodeOutputPin->Properties.Label,
+			MiddleNode, MiddleNodeInputPin->Properties.Label);
+
+		TestTrue(*FString::Printf(TEXT("Edge removed between %s and %s"), *InputNode->GetName(), *MiddleNode->GetName()),
+			!InputNodeOutputPin->GetAllConnectedPins().Contains(MiddleNodeInputPin));
+
+		// Add second edge
+		GraphConfig->AddLabeledEdge(
+			MiddleNode, MiddleNodeOutputPin->Properties.Label,
+			OutputNode, OutputNodeInputPin->Properties.Label);
+
+		TestTrue(*FString::Printf(TEXT("Edge added between %s and %s"), *MiddleNode->GetName(), *OutputNode->GetName()),
+			MiddleNodeOutputPin->GetAllConnectedPins().Contains(OutputNodeInputPin));
+
+		// Remove edge
+		GraphConfig->RemoveEdge(
+			MiddleNode, MiddleNodeOutputPin->Properties.Label,
+			OutputNode, OutputNodeInputPin->Properties.Label);
+
+		TestTrue(*FString::Printf(TEXT("Edge added between %s and %s"), *MiddleNode->GetName(), *OutputNode->GetName()),
+			!MiddleNodeOutputPin->GetAllConnectedPins().Contains(OutputNodeInputPin));
+
+		// Open graph config in editor if desired
+		//OpenGraphConfigInEditor(GraphConfig);
+		
+		return true;
+	}
+}
\ No newline at end of file
