MoveRenderGraph:

-Fix: Issue 377208: Static analyzer warnings in MovieGraphTests.cpp

#[fyi] Matt.Hoffman
#[fyi] Chad Williams
#[fyi] Jason.Walter

diff --git a/Engine/Plugins/MovieScene/MovieRenderPipeline/Source/MovieRenderPipelineCore/Public/Graph/MovieGraphConfig.h b/Engine/Plugins/MovieScene/MovieRenderPipeline/Source/MovieRenderPipelineCore/Public/Graph/MovieGraphConfig.h
--- a/Engine/Plugins/MovieScene/MovieRenderPipeline/Source/MovieRenderPipelineCore/Public/Graph/MovieGraphConfig.h
+++ b/Engine/Plugins/MovieScene/MovieRenderPipeline/Source/MovieRenderPipelineCore/Public/Graph/MovieGraphConfig.h
@@ -673,6 +673,16 @@ public:
 	template<class T>
 	T* ConstructRuntimeNode(TSubclassOf<UMovieGraphNode> PipelineGraphNodeClass = T::StaticClass())
 	{
+		if (!PipelineGraphNodeClass)
+		{
+			FFrame::KismetExecutionMessage(
+				*FString::Printf(
+					TEXT("%hs: Invalid PipelineGraphNodeClass. Please specify a valid class."), __FUNCTION__),
+					ELogVerbosity::Display);
+		
+			return nullptr;
+		}
+		
 		// Construct a new object with ourselves as the outer, then keep track of it.
 		T* RuntimeNode = NewObject<T>(this, PipelineGraphNodeClass, NAME_None, RF_Transactional);
 		RuntimeNode->UpdateDynamicProperties();
diff --git a/Engine/Plugins/MovieScene/MovieRenderPipeline/Source/MovieRenderPipelineEditor/Private/Tests/MovieGraphTests.cpp b/Engine/Plugins/MovieScene/MovieRenderPipeline/Source/MovieRenderPipelineEditor/Private/Tests/MovieGraphTests.cpp
--- a/Engine/Plugins/MovieScene/MovieRenderPipeline/Source/MovieRenderPipelineEditor/Private/Tests/MovieGraphTests.cpp
+++ b/Engine/Plugins/MovieScene/MovieRenderPipeline/Source/MovieRenderPipelineEditor/Private/Tests/MovieGraphTests.cpp
@@ -42,6 +42,10 @@ namespace UE::MovieGraph::Private::Tests
 		// Create Config
 		UMovieGraphConfig* GraphConfig = CreateNewMovieGraphConfig("AddGraphNodeTest");
 		TestTrue(TEXT("MovieGraphConfig successfully created"), GraphConfig != nullptr);
+		if (!GraphConfig)
+		{
+			return false;
+		}
 
 		GraphConfig->ConstructRuntimeNode<UMovieGraphNode>(nullptr);
 		TestTrue(TEXT("Null Node not added"), GraphConfig->GetNodes().IsEmpty());
@@ -70,8 +74,16 @@ namespace UE::MovieGraph::Private::Tests
 		// Test adding a variable node
 		UMovieGraphVariable* Variable = GraphConfig->AddVariable(TEXT("A_Cool_Variable"));
 		TestTrue(TEXT("Variable member added"), Variable != nullptr && GraphConfig->GetVariables().Contains(Variable));
+		if (!Variable)
+		{
+			return false;
+		}
 		UMovieGraphVariableNode* VariableNode = GraphConfig->ConstructRuntimeNode<UMovieGraphVariableNode>(UMovieGraphVariableNode::StaticClass());
 		TestTrue(TEXT("Variable node type added"), VariableNode != nullptr && GraphConfig->GetNodes().Contains(VariableNode));
+		if (!VariableNode)
+		{
+			return false;
+		}
 		VariableNode->SetVariable(Variable);
 		TestTrue(TEXT("Variable node variable member set"), VariableNode->GetVariable() == Variable);
 
@@ -94,12 +106,21 @@ namespace UE::MovieGraph::Private::Tests
 		// Create Config
 		UMovieGraphConfig* GraphConfig = CreateNewMovieGraphConfig("RemoveGraphNodeTest");
 		TestTrue(TEXT("MovieGraphConfig successfully created"), GraphConfig != nullptr);
+		if (!GraphConfig)
+		{
+			return false;
+		}
 		
 		// Add two nodes
 		TArray<UClass*> AllNodeClasses = GetAllDerivedClasses(UMovieGraphNode::StaticClass(), true);
 		FilterOutUndesirableClasses(AllNodeClasses);
 		UMovieGraphNode* NodeOne = GraphConfig->ConstructRuntimeNode<UMovieGraphNode>(AllNodeClasses[0]);
 		UMovieGraphNode* NodeTwo = GraphConfig->ConstructRuntimeNode<UMovieGraphNode>(AllNodeClasses[1]);
+		if (!NodeOne || !NodeTwo)
+		{
+			AddError(*FString::Printf(TEXT("Unable to create two nodes.")));
+			return false;
+		}
 		TestTrue(
 			FString::Printf(TEXT("Node types added: %s and %s"), *AllNodeClasses[0]->GetName(), *AllNodeClasses[1]->GetName()),
 			GraphConfig->GetNodes().Num() == 2);
@@ -119,9 +140,17 @@ namespace UE::MovieGraph::Private::Tests
 		// Test adding and removing a variable member and variable node
 		UMovieGraphVariable* Variable = GraphConfig->AddVariable(TEXT("A_Cool_Variable"));
 		TestTrue(TEXT("Variable member added"), Variable != nullptr && GraphConfig->GetVariables().Contains(Variable));
+		if (!Variable)
+		{
+			return false;
+		}
 		
 		UMovieGraphVariableNode* VariableNode = GraphConfig->ConstructRuntimeNode<UMovieGraphVariableNode>(UMovieGraphVariableNode::StaticClass());
 		TestTrue(TEXT("Variable node type added"), VariableNode != nullptr && GraphConfig->GetNodes().Contains(VariableNode));
+		if (!VariableNode)
+		{
+			return false;
+		}
 		
 		VariableNode->SetVariable(Variable);
 		TestTrue(TEXT("Variable node variable member set"), VariableNode->GetVariable() == Variable);
@@ -152,6 +181,10 @@ namespace UE::MovieGraph::Private::Tests
 		// Create Config
 		UMovieGraphConfig* GraphConfig = CreateNewMovieGraphConfig("GraphNodeAddEdgeTest");
 		TestTrue(TEXT("MovieGraphConfig successfully created"), GraphConfig != nullptr);
+		if (!GraphConfig)
+		{
+			return false;
+		}
 
 		// Add all nodes excepting variables, inputs and outputs
 		TArray<UClass*> AllNodeClasses = GetAllDerivedClasses(UMovieGraphNode::StaticClass(), true);
@@ -237,20 +270,39 @@ namespace UE::MovieGraph::Private::Tests
 		// Create Config
 		UMovieGraphConfig* GraphConfig = CreateNewMovieGraphConfig("GraphNodeRemoveEdgeTest");
 		TestTrue(TEXT("MovieGraphConfig successfully created"), GraphConfig != nullptr);
+		if (!GraphConfig)
+		{
+			return false;
+		}
 
 		// Add a node with an input and output for connection
 		UMovieGraphNode* MiddleNode =
 			GraphConfig->ConstructRuntimeNode<UMovieGraphBranchNode>(UMovieGraphBranchNode::StaticClass());
 		TestTrue(TEXT("MiddleNode successfully added"), GraphConfig->GetNodes().Contains(MiddleNode));
+		if (!MiddleNode)
+		{
+			return false;
+		}
 
 		UMovieGraphNode* InputNode = GraphConfig->GetInputNode();
 		UMovieGraphNode* OutputNode = GraphConfig->GetOutputNode();
+		if (!InputNode || !OutputNode)
+		{
+			AddError(*FString::Printf(TEXT("Unable to find Input and Output nodes.")));
+			return false;
+		}
 
 		const TObjectPtr<UMovieGraphPin> InputNodeOutputPin = InputNode->GetOutputPins()[0];
 		const TObjectPtr<UMovieGraphPin> MiddleNodeInputPin = MiddleNode->GetInputPins()[0];
 		const TObjectPtr<UMovieGraphPin> MiddleNodeOutputPin = MiddleNode->GetOutputPins()[0];
 		const TObjectPtr<UMovieGraphPin> OutputNodeInputPin = OutputNode->GetInputPins()[0];
 
+		if (!InputNodeOutputPin || !MiddleNodeInputPin || !MiddleNodeOutputPin || !OutputNodeInputPin)
+		{
+			AddError(*FString::Printf(TEXT("Unable to find all Pins for Input, Middle and Output Nodes.")));
+			return false;
+		}
+
 		// Add first edge
 		GraphConfig->AddLabeledEdge(
 			InputNode, InputNodeOutputPin->Properties.Label,
