#UE Fix potential use of uninitialized pointer

diff --git a/Engine/Plugins/MovieScene/MovieRenderPipeline/Source/MovieRenderPipelineRenderPasses/Private/MoviePipelineEXROutput.cpp b/Engine/Plugins/MovieScene/MovieRenderPipeline/Source/MovieRenderPipelineRenderPasses/Private/MoviePipelineEXROutput.cpp
--- a/Engine/Plugins/MovieScene/MovieRenderPipeline/Source/MovieRenderPipelineRenderPasses/Private/MoviePipelineEXROutput.cpp
+++ b/Engine/Plugins/MovieScene/MovieRenderPipeline/Source/MovieRenderPipelineRenderPasses/Private/MoviePipelineEXROutput.cpp
@@ -200,7 +200,7 @@ bool FEXRImageWriteTask::WriteToDisk()
 					continue;
 				}
 
-				void const* RawDataPtr;
+				void const* RawDataPtr = nullptr;
 				int64 RawDataSize;
 				bSuccess = Layer->GetRawData(RawDataPtr, RawDataSize);
 				if (!bSuccess)
