MoviePipeline: Added the majority of the missing metadata fields for EXRs. The exception is camera-related metadata -- some has been added, but most is still missing.

Unrelated to metadata, but also fixed a regression on the render nodes which removed the icon.

#rnx

diff --git a/Engine/Plugins/MovieScene/MovieRenderPipeline/Source/MovieRenderPipelineCore/Private/Graph/MovieGraphBlueprintLibrary.cpp b/Engine/Plugins/MovieScene/MovieRenderPipeline/Source/MovieRenderPipelineCore/Private/Graph/MovieGraphBlueprintLibrary.cpp
--- a/Engine/Plugins/MovieScene/MovieRenderPipeline/Source/MovieRenderPipelineCore/Private/Graph/MovieGraphBlueprintLibrary.cpp
+++ b/Engine/Plugins/MovieScene/MovieRenderPipeline/Source/MovieRenderPipelineCore/Private/Graph/MovieGraphBlueprintLibrary.cpp
@@ -80,10 +80,11 @@ FString UMovieGraphBlueprintLibrary::ResolveFilenameFormatArguments(const FStrin
 		OutMergedFormatArgs.FilenameArguments.Add(TEXT("renderer_sub_name"), InParams.RenderDataIdentifier.SubResourceName);
 		OutMergedFormatArgs.FilenameArguments.Add(TEXT("render_layer"), RenderLayerName);
 
+		// TODO: Some of these are per render layer and need to be stored that way. EXRs will have the metadata for all the layers/cameras/etc. in one file.
 		OutMergedFormatArgs.FileMetadata.Add(TEXT("unreal/cameraName"), InParams.RenderDataIdentifier.CameraName);
 		OutMergedFormatArgs.FileMetadata.Add(TEXT("unreal/rendererName"), InParams.RenderDataIdentifier.RendererName);
-		OutMergedFormatArgs.FileMetadata.Add(TEXT("unreal/rendererSubName"), InParams.RenderDataIdentifier.SubResourceName);
-		OutMergedFormatArgs.FileMetadata.Add(TEXT("unreal/renderLayer"), RenderLayerName);
+		OutMergedFormatArgs.FileMetadata.Add(TEXT("unreal/subResourceName"), InParams.RenderDataIdentifier.SubResourceName);
+		OutMergedFormatArgs.FileMetadata.Add(TEXT("unreal/layerName"), RenderLayerName);
 	}
 
 
@@ -103,12 +104,16 @@ FString UMovieGraphBlueprintLibrary::ResolveFilenameFormatArguments(const FStrin
 
 		if (InParams.Job)
 		{
+			FString LevelName = InParams.Job->Map.GetAssetName();
 			FString SequenceName = InParams.Job->Sequence.GetAssetName();
+			
+			OutMergedFormatArgs.FilenameArguments.Add(TEXT("level_name"), SequenceName);
 			OutMergedFormatArgs.FilenameArguments.Add(TEXT("sequence_name"), SequenceName);
+			
+			OutMergedFormatArgs.FileMetadata.Add(TEXT("unreal/levelName"), LevelName);
 			OutMergedFormatArgs.FileMetadata.Add(TEXT("unreal/sequenceName"), SequenceName);
 		}
-
-
+		
 		// Add KVP data from the job (date, time, job name, job author, job comment)
 		UE::MoviePipeline::GetSharedFormatArguments(OutMergedFormatArgs.FilenameArguments, OutMergedFormatArgs.FileMetadata, InParams.InitializationTime, InParams.Version, InParams.Job);
 	}
@@ -133,7 +138,7 @@ FString UMovieGraphBlueprintLibrary::ResolveFilenameFormatArguments(const FStrin
 
 		if (SettingInstance)
 		{
-			SettingInstance->GetFormatResolveArgs(OutMergedFormatArgs);
+			SettingInstance->GetFormatResolveArgs(OutMergedFormatArgs, InParams.RenderDataIdentifier);
 		}
 	}
 
diff --git a/Engine/Plugins/MovieScene/MovieRenderPipeline/Source/MovieRenderPipelineCore/Private/Graph/MovieGraphPipeline.cpp b/Engine/Plugins/MovieScene/MovieRenderPipeline/Source/MovieRenderPipelineCore/Private/Graph/MovieGraphPipeline.cpp
--- a/Engine/Plugins/MovieScene/MovieRenderPipeline/Source/MovieRenderPipelineCore/Private/Graph/MovieGraphPipeline.cpp
+++ b/Engine/Plugins/MovieScene/MovieRenderPipeline/Source/MovieRenderPipelineCore/Private/Graph/MovieGraphPipeline.cpp
@@ -1106,8 +1106,15 @@ void UMovieGraphPipeline::ProcessOutstandingFinishedFrames()
 		{
 			TimeStats->EndTime = FDateTime::UtcNow();
 
-			int32 FrameNumber = OutputFrame.TraversalContext.Time.RenderedFrameNumber;
+			// Add render time metadata
+			FString StartTimeStr = TimeStats->StartTime.ToString();
+			FString EndTimeStr = TimeStats->EndTime.ToString();
 			FString DurationTimeStr = (TimeStats->EndTime - TimeStats->StartTime).ToString();
+			OutputFrame.FileMetadata.Add(TEXT("unreal/frameRenderStartTimeUTC"), StartTimeStr);
+			OutputFrame.FileMetadata.Add(TEXT("unreal/frameRenderEndTimeUTC"), EndTimeStr);
+			OutputFrame.FileMetadata.Add(TEXT("unreal/frameRenderDuration"), DurationTimeStr);
+
+			// int32 FrameNumber = OutputFrame.TraversalContext.Time.RenderedFrameNumber;
 			// UE_LOG(LogTemp, Log, TEXT("Frame: %d Duration: %s"), FrameNumber, *DurationTimeStr);
 		}
 
diff --git a/Engine/Plugins/MovieScene/MovieRenderPipeline/Source/MovieRenderPipelineCore/Private/Graph/Nodes/MovieGraphCameraNode.cpp b/Engine/Plugins/MovieScene/MovieRenderPipeline/Source/MovieRenderPipelineCore/Private/Graph/Nodes/MovieGraphCameraNode.cpp
--- a/Engine/Plugins/MovieScene/MovieRenderPipeline/Source/MovieRenderPipelineCore/Private/Graph/Nodes/MovieGraphCameraNode.cpp
+++ b/Engine/Plugins/MovieScene/MovieRenderPipeline/Source/MovieRenderPipelineCore/Private/Graph/Nodes/MovieGraphCameraNode.cpp
@@ -4,6 +4,15 @@
 
 #include "Styling/AppStyle.h"
 
+void UMovieGraphCameraSettingNode::GetFormatResolveArgs(FMovieGraphResolveArgs& OutMergedFormatArgs, const FMovieGraphRenderDataIdentifier& InRenderDataIdentifier) const
+{
+	OutMergedFormatArgs.FilenameArguments.Add(TEXT("shutter_timing"), StaticEnum<EMoviePipelineShutterTiming>()->GetNameStringByValue(static_cast<int64>(ShutterTiming)));
+	OutMergedFormatArgs.FilenameArguments.Add(TEXT("overscan_percentage"), FString::SanitizeFloat(OverscanPercentage));
+	
+	OutMergedFormatArgs.FileMetadata.Add(TEXT("unreal/camera/shutterTiming"), StaticEnum<EMoviePipelineShutterTiming>()->GetNameStringByValue(static_cast<int64>(ShutterTiming)));
+	OutMergedFormatArgs.FileMetadata.Add(TEXT("unreal/camera/overscanPercentage"), FString::SanitizeFloat(OverscanPercentage));
+}
+
 #if WITH_EDITOR
 FText UMovieGraphCameraSettingNode::GetNodeTitle(const bool bGetDescriptive) const
 {
diff --git a/Engine/Plugins/MovieScene/MovieRenderPipeline/Source/MovieRenderPipelineCore/Private/Graph/Nodes/MovieGraphDeferredPassNode.cpp b/Engine/Plugins/MovieScene/MovieRenderPipeline/Source/MovieRenderPipelineCore/Private/Graph/Nodes/MovieGraphDeferredPassNode.cpp
--- a/Engine/Plugins/MovieScene/MovieRenderPipeline/Source/MovieRenderPipelineCore/Private/Graph/Nodes/MovieGraphDeferredPassNode.cpp
+++ b/Engine/Plugins/MovieScene/MovieRenderPipeline/Source/MovieRenderPipelineCore/Private/Graph/Nodes/MovieGraphDeferredPassNode.cpp
@@ -15,6 +15,12 @@ UMovieGraphDeferredRenderPassNode::UMovieGraphDeferredRenderPassNode()
 {
 }
 
+void UMovieGraphDeferredRenderPassNode::GetFormatResolveArgs(FMovieGraphResolveArgs& OutMergedFormatArgs, const FMovieGraphRenderDataIdentifier& InRenderDataIdentifier) const
+{
+	OutMergedFormatArgs.FilenameArguments.Add(TEXT("ss_count"), FString::FromInt(SpatialSampleCount));
+	OutMergedFormatArgs.FileMetadata.Add(TEXT("unreal/sampling/spatialSampleCount"), FString::FromInt(SpatialSampleCount));
+}
+
 #if WITH_EDITOR
 FText UMovieGraphDeferredRenderPassNode::GetNodeTitle(const bool bGetDescriptive) const
 {
diff --git a/Engine/Plugins/MovieScene/MovieRenderPipeline/Source/MovieRenderPipelineCore/Private/Graph/Nodes/MovieGraphOutputSettingNode.cpp b/Engine/Plugins/MovieScene/MovieRenderPipeline/Source/MovieRenderPipelineCore/Private/Graph/Nodes/MovieGraphOutputSettingNode.cpp
--- a/Engine/Plugins/MovieScene/MovieRenderPipeline/Source/MovieRenderPipelineCore/Private/Graph/Nodes/MovieGraphOutputSettingNode.cpp
+++ b/Engine/Plugins/MovieScene/MovieRenderPipeline/Source/MovieRenderPipelineCore/Private/Graph/Nodes/MovieGraphOutputSettingNode.cpp
@@ -16,7 +16,7 @@ UMovieGraphOutputSettingNode::UMovieGraphOutputSettingNode()
 	OutputDirectory.Path = TEXT("{project_dir}/Saved/MovieRenders/");
 }
 
-void UMovieGraphOutputSettingNode::GetFormatResolveArgs(FMovieGraphResolveArgs& OutMergedFormatArgs) const
+void UMovieGraphOutputSettingNode::GetFormatResolveArgs(FMovieGraphResolveArgs& OutMergedFormatArgs, const FMovieGraphRenderDataIdentifier& InRenderDataIdentifier) const
 {
 	const FString ResolvedProjectDir = FPaths::ConvertRelativePathToFull(FPaths::ProjectDir());
 	OutMergedFormatArgs.FilenameArguments.Add(TEXT("project_dir"), ResolvedProjectDir);
diff --git a/Engine/Plugins/MovieScene/MovieRenderPipeline/Source/MovieRenderPipelineCore/Private/Graph/Nodes/MovieGraphSamplingMethodNode.cpp b/Engine/Plugins/MovieScene/MovieRenderPipeline/Source/MovieRenderPipelineCore/Private/Graph/Nodes/MovieGraphSamplingMethodNode.cpp
--- a/Engine/Plugins/MovieScene/MovieRenderPipeline/Source/MovieRenderPipelineCore/Private/Graph/Nodes/MovieGraphSamplingMethodNode.cpp
+++ b/Engine/Plugins/MovieScene/MovieRenderPipeline/Source/MovieRenderPipelineCore/Private/Graph/Nodes/MovieGraphSamplingMethodNode.cpp
@@ -18,6 +18,29 @@ EMovieGraphBranchRestriction UMovieGraphSamplingMethodNode::GetBranchRestriction
 	return EMovieGraphBranchRestriction::Globals;
 }
 
+void UMovieGraphSamplingMethodNode::GetFormatResolveArgs(FMovieGraphResolveArgs& OutMergedFormatArgs, const FMovieGraphRenderDataIdentifier& InRenderDataIdentifier) const
+{
+	OutMergedFormatArgs.FilenameArguments.Add(TEXT("ts_count"), FString::FromInt(TemporalSampleCount));
+    OutMergedFormatArgs.FileMetadata.Add(TEXT("unreal/sampling/temporalSampleCount"), FString::FromInt(TemporalSampleCount));
+
+	if (const UClass* SamplingClass = SamplingMethodClass.ResolveClass())
+	{
+#if WITH_EDITOR
+		const FString SamplingModeDisplayName = SamplingClass->GetDisplayNameText().ToString();
+#else
+		const FString SamplingModeDisplayName = SamplingClass->GetName();
+#endif
+		
+		OutMergedFormatArgs.FilenameArguments.Add(TEXT("sampling_mode"), SamplingModeDisplayName);
+		OutMergedFormatArgs.FileMetadata.Add(TEXT("unreal/sampling/mode"), SamplingModeDisplayName);
+	}
+	else
+	{
+		OutMergedFormatArgs.FilenameArguments.Add(TEXT("sampling_mode"), TEXT("Unknown"));
+		OutMergedFormatArgs.FileMetadata.Add(TEXT("unreal/sampling/mode"), TEXT("Unknown"));
+	}
+}
+
 #if WITH_EDITOR
 FText UMovieGraphSamplingMethodNode::GetNodeTitle(const bool bGetDescriptive) const
 {
diff --git a/Engine/Plugins/MovieScene/MovieRenderPipeline/Source/MovieRenderPipelineCore/Public/Graph/MovieGraphDataTypes.h b/Engine/Plugins/MovieScene/MovieRenderPipeline/Source/MovieRenderPipelineCore/Public/Graph/MovieGraphDataTypes.h
--- a/Engine/Plugins/MovieScene/MovieRenderPipeline/Source/MovieRenderPipelineCore/Public/Graph/MovieGraphDataTypes.h
+++ b/Engine/Plugins/MovieScene/MovieRenderPipeline/Source/MovieRenderPipelineCore/Public/Graph/MovieGraphDataTypes.h
@@ -352,6 +352,9 @@ namespace UE::MovieGraph
 
 		/** Stores the actual pixel data for each render pass. */
 		TMap<FMovieGraphRenderDataIdentifier, TUniquePtr<FImagePixelData>> ImageOutputData;
+
+		/** Additional metadata to be added to the output (if supported by the output container). */
+		TMap<FString, FString> FileMetadata;
 	};
 
 	/**
diff --git a/Engine/Plugins/MovieScene/MovieRenderPipeline/Source/MovieRenderPipelineCore/Public/Graph/MovieGraphFilenameResolveParams.h b/Engine/Plugins/MovieScene/MovieRenderPipeline/Source/MovieRenderPipelineCore/Public/Graph/MovieGraphFilenameResolveParams.h
--- a/Engine/Plugins/MovieScene/MovieRenderPipeline/Source/MovieRenderPipelineCore/Public/Graph/MovieGraphFilenameResolveParams.h
+++ b/Engine/Plugins/MovieScene/MovieRenderPipeline/Source/MovieRenderPipelineCore/Public/Graph/MovieGraphFilenameResolveParams.h
@@ -60,6 +60,7 @@ struct FMovieGraphFilenameResolveParams
 		, Shot(nullptr)
 		, EvaluatedConfig(nullptr)
 		, FrameNumberOffset(0)
+		, DefaultFrameRate(FFrameRate(24, 1))
 	{
 	}
 	
@@ -125,7 +126,7 @@ struct FMovieGraphFilenameResolveParams
 	UPROPERTY(BlueprintReadWrite, EditAnywhere, Category = "Movie Graph")
 	FDateTime InitializationTime;
 	
-	/** When converitng frame numbers to strings, how many digits should we pad them up to? ie: 5 => 0005 with a count of 4. */
+	/** When converting frame numbers to strings, how many digits should we pad them up to? ie: 5 => 0005 with a count of 4. */
 	UPROPERTY(BlueprintReadWrite, EditAnywhere, Category = "Movie Graph")
 	int32 ZeroPadFrameNumberCount;
 
diff --git a/Engine/Plugins/MovieScene/MovieRenderPipeline/Source/MovieRenderPipelineCore/Public/Graph/MovieGraphNode.h b/Engine/Plugins/MovieScene/MovieRenderPipeline/Source/MovieRenderPipelineCore/Public/Graph/MovieGraphNode.h
--- a/Engine/Plugins/MovieScene/MovieRenderPipeline/Source/MovieRenderPipelineCore/Public/Graph/MovieGraphNode.h
+++ b/Engine/Plugins/MovieScene/MovieRenderPipeline/Source/MovieRenderPipelineCore/Public/Graph/MovieGraphNode.h
@@ -376,7 +376,7 @@ public:
 	* Because this is called either on the CDO or on a flattened instance, there is no need to worry about
 	* resolving the settings of the graph, the node only needs to read its own values.
 	*/
-	virtual void GetFormatResolveArgs(FMovieGraphResolveArgs& OutMergedFormatArgs) const {}
+	virtual void GetFormatResolveArgs(FMovieGraphResolveArgs& OutMergedFormatArgs, const FMovieGraphRenderDataIdentifier& InRenderDataIdentifier) const {}
 };
 
 /**
diff --git a/Engine/Plugins/MovieScene/MovieRenderPipeline/Source/MovieRenderPipelineCore/Public/Graph/Nodes/MovieGraphCameraNode.h b/Engine/Plugins/MovieScene/MovieRenderPipeline/Source/MovieRenderPipelineCore/Public/Graph/Nodes/MovieGraphCameraNode.h
--- a/Engine/Plugins/MovieScene/MovieRenderPipeline/Source/MovieRenderPipelineCore/Public/Graph/Nodes/MovieGraphCameraNode.h
+++ b/Engine/Plugins/MovieScene/MovieRenderPipeline/Source/MovieRenderPipelineCore/Public/Graph/Nodes/MovieGraphCameraNode.h
@@ -16,6 +16,7 @@ public:
 	UMovieGraphCameraSettingNode() = default;
 
 	virtual EMovieGraphBranchRestriction GetBranchRestriction() const override { return EMovieGraphBranchRestriction::Globals; }
+	virtual void GetFormatResolveArgs(FMovieGraphResolveArgs& OutMergedFormatArgs, const FMovieGraphRenderDataIdentifier& InRenderDataIdentifier) const override;
 
 #if WITH_EDITOR
 	virtual FText GetNodeTitle(const bool bGetDescriptive = false) const override;
diff --git a/Engine/Plugins/MovieScene/MovieRenderPipeline/Source/MovieRenderPipelineCore/Public/Graph/Nodes/MovieGraphDeferredPassNode.h b/Engine/Plugins/MovieScene/MovieRenderPipeline/Source/MovieRenderPipelineCore/Public/Graph/Nodes/MovieGraphDeferredPassNode.h
--- a/Engine/Plugins/MovieScene/MovieRenderPipeline/Source/MovieRenderPipelineCore/Public/Graph/Nodes/MovieGraphDeferredPassNode.h
+++ b/Engine/Plugins/MovieScene/MovieRenderPipeline/Source/MovieRenderPipelineCore/Public/Graph/Nodes/MovieGraphDeferredPassNode.h
@@ -15,6 +15,8 @@ class MOVIERENDERPIPELINECORE_API UMovieGraphDeferredRenderPassNode : public UMo
 public:
 	UMovieGraphDeferredRenderPassNode();
 
+	virtual void GetFormatResolveArgs(FMovieGraphResolveArgs& OutMergedFormatArgs, const FMovieGraphRenderDataIdentifier& InRenderDataIdentifier) const override;
+
 #if WITH_EDITOR
 	virtual FText GetNodeTitle(const bool bGetDescriptive = false) const override;
 #endif
diff --git a/Engine/Plugins/MovieScene/MovieRenderPipeline/Source/MovieRenderPipelineCore/Public/Graph/Nodes/MovieGraphOutputSettingNode.h b/Engine/Plugins/MovieScene/MovieRenderPipeline/Source/MovieRenderPipelineCore/Public/Graph/Nodes/MovieGraphOutputSettingNode.h
--- a/Engine/Plugins/MovieScene/MovieRenderPipeline/Source/MovieRenderPipelineCore/Public/Graph/Nodes/MovieGraphOutputSettingNode.h
+++ b/Engine/Plugins/MovieScene/MovieRenderPipeline/Source/MovieRenderPipelineCore/Public/Graph/Nodes/MovieGraphOutputSettingNode.h
@@ -15,7 +15,7 @@ public:
 	UMovieGraphOutputSettingNode();
 
 	// UMovieGraphSettingNode Interface
-	virtual void GetFormatResolveArgs(FMovieGraphResolveArgs& OutMergedFormatArgs) const override;
+	virtual void GetFormatResolveArgs(FMovieGraphResolveArgs& OutMergedFormatArgs, const FMovieGraphRenderDataIdentifier& InRenderDataIdentifier) const override;
 
 #if WITH_EDITOR
 	virtual FText GetNodeTitle(const bool bGetDescriptive) const override;
diff --git a/Engine/Plugins/MovieScene/MovieRenderPipeline/Source/MovieRenderPipelineCore/Public/Graph/Nodes/MovieGraphRenderPassNode.h b/Engine/Plugins/MovieScene/MovieRenderPipeline/Source/MovieRenderPipelineCore/Public/Graph/Nodes/MovieGraphRenderPassNode.h
--- a/Engine/Plugins/MovieScene/MovieRenderPipeline/Source/MovieRenderPipelineCore/Public/Graph/Nodes/MovieGraphRenderPassNode.h
+++ b/Engine/Plugins/MovieScene/MovieRenderPipeline/Source/MovieRenderPipelineCore/Public/Graph/Nodes/MovieGraphRenderPassNode.h
@@ -1,8 +1,11 @@
 // Copyright Epic Games, Inc. All Rights Reserved.
 
 #pragma once
+
 #include "Graph/MovieGraphNode.h"
 #include "Graph/MovieGraphDataTypes.h"
+#include "Styling/AppStyle.h"
+
 #include "MovieGraphRenderPassNode.generated.h"
 
 // Forward Declare
@@ -52,6 +55,18 @@ public:
 	{
 		return NSLOCTEXT("MovieGraphNodes", "RenderPassGraphNode_Category", "Rendering");
 	}
+
+	virtual FLinearColor GetNodeTitleColor() const override
+	{
+		return FLinearColor(0.572f, 0.274f, 1.f);
+	}
+	
+	virtual FSlateIcon GetIconAndTint(FLinearColor& OutColor) const override
+	{
+		static const FSlateIcon DeferredRendererIcon = FSlateIcon(FAppStyle::GetAppStyleSetName(), "SequenceRecorder.TabIcon");
+		OutColor = FLinearColor::White;
+		return DeferredRendererIcon;
+	}
 #endif
 
 protected:
diff --git a/Engine/Plugins/MovieScene/MovieRenderPipeline/Source/MovieRenderPipelineCore/Public/Graph/Nodes/MovieGraphSamplingMethodNode.h b/Engine/Plugins/MovieScene/MovieRenderPipeline/Source/MovieRenderPipelineCore/Public/Graph/Nodes/MovieGraphSamplingMethodNode.h
--- a/Engine/Plugins/MovieScene/MovieRenderPipeline/Source/MovieRenderPipelineCore/Public/Graph/Nodes/MovieGraphSamplingMethodNode.h
+++ b/Engine/Plugins/MovieScene/MovieRenderPipeline/Source/MovieRenderPipelineCore/Public/Graph/Nodes/MovieGraphSamplingMethodNode.h
@@ -16,6 +16,7 @@ public:
 	UMovieGraphSamplingMethodNode();
 
 	virtual EMovieGraphBranchRestriction GetBranchRestriction() const override;
+	virtual void GetFormatResolveArgs(FMovieGraphResolveArgs& OutMergedFormatArgs, const FMovieGraphRenderDataIdentifier& InRenderDataIdentifier) const override;
 
 #if WITH_EDITOR
 	virtual FText GetNodeTitle(const bool bGetDescriptive = false) const override;
diff --git a/Engine/Plugins/MovieScene/MovieRenderPipeline/Source/MovieRenderPipelineRenderPasses/Private/MovieGraphImageSequenceOutputNode.cpp b/Engine/Plugins/MovieScene/MovieRenderPipeline/Source/MovieRenderPipelineRenderPasses/Private/MovieGraphImageSequenceOutputNode.cpp
--- a/Engine/Plugins/MovieScene/MovieRenderPipeline/Source/MovieRenderPipelineRenderPasses/Private/MovieGraphImageSequenceOutputNode.cpp
+++ b/Engine/Plugins/MovieScene/MovieRenderPipeline/Source/MovieRenderPipelineRenderPasses/Private/MovieGraphImageSequenceOutputNode.cpp
@@ -326,6 +326,9 @@ void UMovieGraphImageSequenceOutputNode_EXR::OnReceiveImageDataImpl(UMovieGraphP
 		// MultiLayerImageTask->CompressionLevel is intentionally skipped because it doesn't seem to make any practical difference
 		// so we don't expose it to the user because that will just cause confusion where the setting doesn't seem to do anything.
 
+		// Add in hardware usage metadata
+		UE::MoviePipeline::GetHardwareUsageMetadata(FilenameToResolveArgs[Filename].FileMetadata, FPaths::GetPath(Filename));
+
 		// FileMetadata has been generated by ResolveFilenameFormatArgs, but we need to convert from FString, FString
 		// (needed for BP/Python purposes) to a FStringFormatArg as we need to preserve numeric metadata types later in
 		// the image writing process (for compression level).
@@ -336,6 +339,12 @@ void UMovieGraphImageSequenceOutputNode_EXR::OnReceiveImageDataImpl(UMovieGraphP
 		}
 		MultiLayerImageTask->FileMetadata = NewFileMetadataMap;
 
+		// Add in any metadata from the output merger frame
+		for (const TPair<FString, FString>& Metadata : InRawFrameData->FileMetadata)
+		{
+			MultiLayerImageTask->FileMetadata.Add(Metadata.Key, Metadata.Value);
+		}
+
 		// Add color space metadata to the output: xy chromaticity coordinates and/or the color space source/dest names.
 		// TODO: Support is also needed for regular exrs via the image wrapper module.
 		// TODO: No OCIO node yet
@@ -372,7 +381,8 @@ void UMovieGraphImageSequenceOutputNode_EXR::OnReceiveImageDataImpl(UMovieGraphP
 			// as most programs that handle EXRs expect the main image data to be in an unnamed layer.
 			if (LayerIndex == 0)
 			{
-				// Add task information that is common to all layers
+				// Add task information that is common to all layers. This metadata may be redundant with unreal/* metadata,
+				// but these are "standard" fields in EXR metadata.
 				MultiLayerImageTask->FileMetadata.Add("owner", UE::MoviePipeline::GetJobAuthor(Payload->TraversalContext.Job));
 				MultiLayerImageTask->FileMetadata.Add("comments", Payload->TraversalContext.Job->Comment);
 
