MovieGraph: Add support for Camera Shutter Timing field.

For a given frame (ie: 0), the different settings will represent a different range of Sequencer time:
Open should make the data from [0, 1) be represented by the motion blur.
Center should make the data from [-0.5, 0.5) be represented by motion blur (Default, matches what you see evaluated in Sequencer)
Closed should make the data from [-1, 0) be represented by motion blur

diff --git a/Engine/Plugins/MovieScene/MovieRenderPipeline/Source/MovieRenderPipelineCore/Private/Graph/MovieGraphCoreTimeStep.cpp b/Engine/Plugins/MovieScene/MovieRenderPipeline/Source/MovieRenderPipelineCore/Private/Graph/MovieGraphCoreTimeStep.cpp
--- a/Engine/Plugins/MovieScene/MovieRenderPipeline/Source/MovieRenderPipelineCore/Private/Graph/MovieGraphCoreTimeStep.cpp
+++ b/Engine/Plugins/MovieScene/MovieRenderPipeline/Source/MovieRenderPipelineCore/Private/Graph/MovieGraphCoreTimeStep.cpp
@@ -5,6 +5,7 @@
 #include "Graph/MovieGraphPipeline.h"
 #include "Graph/MovieGraphBlueprintLibrary.h"
 #include "Graph/Nodes/MovieGraphOutputSettingNode.h"
+#include "Graph/Nodes/MovieGraphCameraNode.h"
 #include "MoviePipelineQueue.h"
 #include "MovieRenderPipelineCoreModule.h"
 #include "LevelSequence.h"
@@ -593,21 +594,32 @@ void UMovieGraphCoreTimeStep::UpdateFrameMetrics()
 	// Shutter timing is a bias applied to the final evaluation time to let us change what we consider a frame
 	// ie: Do we consider a frame the start of the timespan we captured? Or is the frame the end of the timespan?
 	// We default to Centered so that the center of your evaluated time is what you see in Level Sequences.
-	// ToDo: Shutter Timing Offset
-	//switch (CameraSettings->ShutterTiming)
-	//{
-	//	// Subtract the entire time the shutter is open.
-	//case EMoviePipelineShutterTiming::FrameClose:
-	//	Output.ShutterOffsetTicks = -Output.TicksWhileShutterOpen;
-	//	break;
-	//	// Only subtract half the time the shutter is open.
-	//case EMoviePipelineShutterTiming::FrameCenter:
-	FrameData.ShutterOffsetFrameTime = -FrameData.FrameTimeWhileShutterOpen / 2.0;
-	//	break;
-	//	// No offset needed
-	//case EMoviePipelineShutterTiming::FrameOpen:
-	//	break;
-	//}
+	EMoviePipelineShutterTiming ShutterTiming = EMoviePipelineShutterTiming::FrameCenter;
+	
+	// We don't always have a config set up when this function is called.
+	if (CurrentTimeStepData.EvaluatedConfig)
+	{
+		if (UMovieGraphCameraSettingNode* CameraSetting = CurrentTimeStepData.EvaluatedConfig->GetSettingForBranch< UMovieGraphCameraSettingNode>(UMovieGraphNode::GlobalsPinName))
+		{
+			ShutterTiming = CameraSetting->ShutterTiming;
+		}
+	}
+	
+	switch (ShutterTiming)
+	{
+	// Subtract the entire time the shutter is open.
+	case EMoviePipelineShutterTiming::FrameClose:
+		FrameData.ShutterOffsetFrameTime = -FrameData.FrameTimeWhileShutterOpen;
+		break;
+	// Only subtract half the time the shutter is open.
+	case EMoviePipelineShutterTiming::FrameCenter:
+		FrameData.ShutterOffsetFrameTime = -FrameData.FrameTimeWhileShutterOpen / 2.0;
+		break;
+	// No offset needed
+	case EMoviePipelineShutterTiming::FrameOpen:
+		FrameData.ShutterOffsetFrameTime = FFrameTime(0);
+		break;
+	}
 
 	// Then, calculate our motion blur offset. Motion Blur in the engine is always
 	// centered around the object so we offset our time sampling by half of the
diff --git a/Engine/Plugins/MovieScene/MovieRenderPipeline/Source/MovieRenderPipelineCore/Public/Graph/Nodes/MovieGraphCameraNode.h b/Engine/Plugins/MovieScene/MovieRenderPipeline/Source/MovieRenderPipelineCore/Public/Graph/Nodes/MovieGraphCameraNode.h
--- a/Engine/Plugins/MovieScene/MovieRenderPipeline/Source/MovieRenderPipelineCore/Public/Graph/Nodes/MovieGraphCameraNode.h
+++ b/Engine/Plugins/MovieScene/MovieRenderPipeline/Source/MovieRenderPipelineCore/Public/Graph/Nodes/MovieGraphCameraNode.h
@@ -13,7 +13,10 @@ class MOVIERENDERPIPELINECORE_API UMovieGraphCameraSettingNode : public UMovieGr
 	GENERATED_BODY()
 
 public:
-	UMovieGraphCameraSettingNode() = default;
+	UMovieGraphCameraSettingNode()
+		: ShutterTiming(EMoviePipelineShutterTiming::FrameCenter)
+		, OverscanPercentage(0.f)
+	{}
 
 	virtual EMovieGraphBranchRestriction GetBranchRestriction() const override { return EMovieGraphBranchRestriction::Globals; }
 	virtual void GetFormatResolveArgs(FMovieGraphResolveArgs& OutMergedFormatArgs, const FMovieGraphRenderDataIdentifier& InRenderDataIdentifier) const override;
