MovieGraph: Add Camera Overscan support.

When set to 10.0, a 1080p image should become a 2112x1188 image, with the Field of View increased, so that a center-out crop is the same pixels as the original 1080p image.

When writing out a jpg, write out the jpeg at the 2112x1188 size, with a visibly wider FoV.

When writing out exrs, write out the exr at 1920x1080 with the extra (96,54) pixels (per side) hidden in the Data Window feature of exrs. The visible 1920x1080 pixels should match a zero% overscan image.

diff --git a/Engine/Plugins/MovieScene/MovieRenderPipeline/Source/MovieRenderPipelineCore/Private/Graph/MovieGraphBlueprintLibrary.cpp b/Engine/Plugins/MovieScene/MovieRenderPipeline/Source/MovieRenderPipelineCore/Private/Graph/MovieGraphBlueprintLibrary.cpp
--- a/Engine/Plugins/MovieScene/MovieRenderPipeline/Source/MovieRenderPipelineCore/Private/Graph/MovieGraphBlueprintLibrary.cpp
+++ b/Engine/Plugins/MovieScene/MovieRenderPipeline/Source/MovieRenderPipelineCore/Private/Graph/MovieGraphBlueprintLibrary.cpp
@@ -317,7 +317,7 @@ FIntPoint UMovieGraphBlueprintLibrary::GetEffectiveOutputResolution(UMovieGraphE
 	if (UMovieGraphCameraSettingNode* CameraSetting = InEvaluatedGraph->GetSettingForBranch<UMovieGraphCameraSettingNode>(UMovieGraphNode::GlobalsPinName, bIncludeCDOs))
 	{
 		// The old system used [0-1] range for floats, the new system will use [0-100], so we rescale down before calling through.
-		RescaledOverscan = CameraSetting->OverscanPercentage / 100.f;
+		RescaledOverscan = FMath::Clamp(CameraSetting->OverscanPercentage / 100.f, 0.f, 1.f);
 	}
 
 	return UMoviePipelineBlueprintLibrary::Utility_GetEffectiveOutputResolution(RescaledOverscan, OutputSetting->OutputResolution);
diff --git a/Engine/Plugins/MovieScene/MovieRenderPipeline/Source/MovieRenderPipelineCore/Private/Graph/Renderers/MovieGraphDeferredPass.cpp b/Engine/Plugins/MovieScene/MovieRenderPipeline/Source/MovieRenderPipelineCore/Private/Graph/Renderers/MovieGraphDeferredPass.cpp
--- a/Engine/Plugins/MovieScene/MovieRenderPipeline/Source/MovieRenderPipelineCore/Private/Graph/Renderers/MovieGraphDeferredPass.cpp
+++ b/Engine/Plugins/MovieScene/MovieRenderPipeline/Source/MovieRenderPipelineCore/Private/Graph/Renderers/MovieGraphDeferredPass.cpp
@@ -98,14 +98,14 @@ void FMovieGraphDeferredPass::Render(const FMovieGraphTraversalContext& InFrameT
 	int32 NumSpatialSamples = FMath::Max(1, ParentNode->SpatialSampleCount);
 	const bool bDisableToneCurve = ParentNode->bDisableToneCurve;
 	const EAntiAliasingMethod AntiAliasingMethod = ParentNode->AntiAliasingMethod;
-	float OverscanPercentage = 0.f;
+	float OverscanFraction = 0.f;
 	const float TileOverlapPadRatio = 0.0f; // No tiling support right now
 
 	// Camera nodes are optional
 	const UMovieGraphCameraSettingNode* CameraNode = InTimeData.EvaluatedConfig->GetSettingForBranch<UMovieGraphCameraSettingNode>(LayerData.BranchName, bIncludeCDOs);
 	if (CameraNode)
 	{
-		OverscanPercentage = FMath::Clamp(0.f, 1.f, CameraNode->OverscanPercentage / 100.f);
+		OverscanFraction = FMath::Clamp(CameraNode->OverscanPercentage / 100.f, 0.f, 1.f);
 	}
 	
 	FIntPoint AccumulatorResolution = UMovieGraphBlueprintLibrary::GetEffectiveOutputResolution(InTimeData.EvaluatedConfig, LayerData.BranchName);
@@ -165,7 +165,7 @@ void FMovieGraphDeferredPass::Render(const FMovieGraphTraversalContext& InFrameT
 		CameraInfo.TilingParams.OverlapPad = FVector2f(0.f, 0.f); // No tile support
 		CameraInfo.TilingParams.TileCount = FIntPoint(1, 1); // No tile support
 		CameraInfo.TilingParams.TileIndexes = FIntPoint(0, 0); // No tile support
-		CameraInfo.OverscanFraction = OverscanPercentage;
+		CameraInfo.OverscanFraction = OverscanFraction;
 		CameraInfo.ProjectionMatrixJitterAmount = FVector2D((SpatialShiftAmount.X) * 2.0f / (float)BackbufferResolution.X, SpatialShiftAmount.Y * -2.0f / (float)BackbufferResolution.Y);
 
 		// For this particular tile, what is the offset into the output image
@@ -232,7 +232,7 @@ void FMovieGraphDeferredPass::Render(const FMovieGraphTraversalContext& InFrameT
 			SampleState.OverlappedPad = OverlappedPad;
 			SampleState.OverlappedOffset = OverlappedOffset;
 			SampleState.OverlappedSubpixelShift = OverlappedSubpixelShift;
-
+			SampleState.OverscanFraction = OverscanFraction;
 		}
 
 		// If this was just to contribute to the history buffer, no need to go any further.
diff --git a/Engine/Plugins/MovieScene/MovieRenderPipeline/Source/MovieRenderPipelineCore/Private/Graph/Renderers/MovieGraphImagePassBase.cpp b/Engine/Plugins/MovieScene/MovieRenderPipeline/Source/MovieRenderPipelineCore/Private/Graph/Renderers/MovieGraphImagePassBase.cpp
--- a/Engine/Plugins/MovieScene/MovieRenderPipeline/Source/MovieRenderPipelineCore/Private/Graph/Renderers/MovieGraphImagePassBase.cpp
+++ b/Engine/Plugins/MovieScene/MovieRenderPipeline/Source/MovieRenderPipelineCore/Private/Graph/Renderers/MovieGraphImagePassBase.cpp
@@ -324,9 +324,6 @@ FMatrix FMovieGraphImagePassBase::CalculateProjectionMatrix(const UE::MovieGraph
 	// Calculate a Projection Matrix. This code unfortunately ends up similar to, but not quite the same as FMinimalViewInfo::CalculateProjectionMatrixGivenView
 	FMatrix BaseProjMatrix;
 	
-	// Potentially modified when we modify projection matrix for tiling.
-	float DofSensorScale = 1.0f;
-
 	// TileSize should respect the actual backbuffer size being used by the render.
 	float ViewRectWidth = InCameraInfo.TilingParams.TileSize.X;
 	float ViewRectHeight = InCameraInfo.TilingParams.TileSize.Y;
diff --git a/Engine/Plugins/MovieScene/MovieRenderPipeline/Source/MovieRenderPipelineCore/Public/Graph/MovieGraphDataTypes.h b/Engine/Plugins/MovieScene/MovieRenderPipeline/Source/MovieRenderPipelineCore/Public/Graph/MovieGraphDataTypes.h
--- a/Engine/Plugins/MovieScene/MovieRenderPipeline/Source/MovieRenderPipelineCore/Public/Graph/MovieGraphDataTypes.h
+++ b/Engine/Plugins/MovieScene/MovieRenderPipeline/Source/MovieRenderPipelineCore/Public/Graph/MovieGraphDataTypes.h
@@ -280,6 +280,7 @@ namespace UE::MovieGraph
 			, bRequiresAccumulator(false)
 			, bFetchFromAccumulator(false)
 			, bCompositeOnOtherRenders(false)
+			, OverscanFraction(0.f)
 			, CompositingSortOrder(0)
 		{}
 
@@ -316,6 +317,9 @@ namespace UE::MovieGraph
 		/** When using spatial jitters, how much do we need to shift the output data during accumulation to counter-act the jitter. */
 		FVector2D OverlappedSubpixelShift;
 
+		/** When using camera overscan, what fraction (0-1) did this render use? Needed so that exrs can take a center-out crop of the data. */
+		float OverscanFraction;
+
 		/**
 		* If multiple passes are composited on top of a render, the sort order determines the order in which they're composited.
 		* Passes with a low sort order will composite on top of passes with a higher sort order.
diff --git a/Engine/Plugins/MovieScene/MovieRenderPipeline/Source/MovieRenderPipelineRenderPasses/Private/MovieGraphImageSequenceOutputNode.cpp b/Engine/Plugins/MovieScene/MovieRenderPipeline/Source/MovieRenderPipelineRenderPasses/Private/MovieGraphImageSequenceOutputNode.cpp
--- a/Engine/Plugins/MovieScene/MovieRenderPipeline/Source/MovieRenderPipelineRenderPasses/Private/MovieGraphImageSequenceOutputNode.cpp
+++ b/Engine/Plugins/MovieScene/MovieRenderPipeline/Source/MovieRenderPipelineRenderPasses/Private/MovieGraphImageSequenceOutputNode.cpp
@@ -390,8 +390,7 @@ void UMovieGraphImageSequenceOutputNode_EXR::OnReceiveImageDataImpl(UMovieGraphP
 				MultiLayerImageTask->Width = Resolution.X;
 				MultiLayerImageTask->Height = Resolution.Y;
 				
-				// TODO: Overscan not available in the graph yet
-				// MultiLayerImageTask->OverscanPercentage = Payload->SampleState.OverscanPercentage;
+				MultiLayerImageTask->OverscanPercentage = Payload->OverscanFraction;
 			}
 			else
 			{
