MoviePipeline: Added a new implementation for collections in the graph which aligns with the upcoming UX. There need to be additional query types added, but this will provide a good base for inital testing.

Also fixed a small regression in FMovieGraphDeferredPass from the refactor which broke collections.

#rnx

diff --git a/Engine/Plugins/MovieScene/MovieRenderPipeline/Source/MovieRenderPipelineCore/Private/Graph/MovieGraphPipeline.cpp b/Engine/Plugins/MovieScene/MovieRenderPipeline/Source/MovieRenderPipelineCore/Private/Graph/MovieGraphPipeline.cpp
--- a/Engine/Plugins/MovieScene/MovieRenderPipeline/Source/MovieRenderPipelineCore/Private/Graph/MovieGraphPipeline.cpp
+++ b/Engine/Plugins/MovieScene/MovieRenderPipeline/Source/MovieRenderPipelineCore/Private/Graph/MovieGraphPipeline.cpp
@@ -233,13 +233,14 @@ void UMovieGraphPipeline::UpdateLayerContentsInRenderLayerSubsystem(const UMovie
 			// Find the collection that this modifier is modifying
 			for (const UMovieGraphCollectionNode* CollectionNode : CollectionNodes)
 			{
-				if (CollectionNode && (CollectionNode->CollectionName == ModifierNode->ModifiedCollectionName))
+				if (!CollectionNode || !CollectionNode->Collection)
 				{
-					UMoviePipelineCollection* NewCollection = NewObject<UMoviePipelineCollection>();
-					NewCollection->AddQuery(CollectionNode->QueryClass);
-					NewCollection->SetCollectionName(CollectionNode->CollectionName);
-					
-					ModifierNode->ModifierClass->SetCollections({NewCollection});
+					continue;
+				}
+				
+				if (CollectionNode->Collection->GetCollectionName() == ModifierNode->ModifiedCollectionName)
+				{
+					ModifierNode->ModifierClass->SetCollections({CollectionNode->Collection});
 					break;
 				}
 			}
diff --git a/Engine/Plugins/MovieScene/MovieRenderPipeline/Source/MovieRenderPipelineCore/Private/Graph/MoviePipelineRenderLayerSubsystem.cpp b/Engine/Plugins/MovieScene/MovieRenderPipeline/Source/MovieRenderPipelineCore/Private/Graph/MoviePipelineRenderLayerSubsystem.cpp
--- a/Engine/Plugins/MovieScene/MovieRenderPipeline/Source/MovieRenderPipelineCore/Private/Graph/MoviePipelineRenderLayerSubsystem.cpp
+++ b/Engine/Plugins/MovieScene/MovieRenderPipeline/Source/MovieRenderPipelineCore/Private/Graph/MoviePipelineRenderLayerSubsystem.cpp
@@ -9,6 +9,7 @@
 #include "EngineUtils.h"
 #include "Materials/MaterialInterface.h"
 #include "MovieRenderPipelineCoreModule.h"
+#include "Styling/AppStyle.h"
 #include "UObject/Package.h"
 
 void UMoviePipelineMaterialModifier::ApplyModifier(const UWorld* World)
@@ -21,14 +22,14 @@ void UMoviePipelineMaterialModifier::ApplyModifier(const UWorld* World)
 
 	ModifiedComponents.Empty();
 	
-	for (const UMoviePipelineCollection* Collection : Collections)
+	for (const UMovieGraphCollection* Collection : Collections)
 	{
 		if (!Collection)
 		{
 			continue;
 		}
 
-		for (const AActor* Actor : Collection->GetMatchingActors(World, bUseInvertedActors))
+		for (const AActor* Actor : Collection->Evaluate(World))
 		{
 			const bool bIncludeFromChildActors = true;
 			TArray<UPrimitiveComponent*> PrimitiveComponents;
@@ -100,20 +101,18 @@ void UMoviePipelineVisibilityModifier::ApplyModifier(const UWorld* World)
 	NewVisibilityState.bAffectIndirectLightingWhileHidden = bAffectIndirectLightingWhileHidden;
 	NewVisibilityState.bHoldout = bHoldout;
 	
-	for (const UMoviePipelineCollection* Collection : Collections)
+	for (const UMovieGraphCollection* Collection : Collections)
 	{
 		if (!Collection)
 		{
 			continue;
 		}
 
-		const TArray<AActor*> MatchingActors = Collection->GetMatchingActors(World, bUseInvertedActors);
+		const TSet<AActor*> MatchingActors = Collection->Evaluate(World);
 		ModifiedActors.Reserve(MatchingActors.Num());
 
-		for (int32 Index = 0; Index < MatchingActors.Num(); ++Index)
+		for (const AActor* Actor : MatchingActors)
 		{
-			const AActor* Actor = MatchingActors[Index];
-
 			// Save out visibility state before the modifier is applied
 			FActorVisibilityState OriginalVisibilityState;
 			OriginalVisibilityState.Actor = Actor;
@@ -301,10 +300,10 @@ TArray<AActor*> UMoviePipelineCollection::GetMatchingActors(const UWorld* World,
 	return MatchingActors;
 }
 
-void UMoviePipelineCollectionModifier::AddCollection(UMoviePipelineCollection* Collection)
+void UMoviePipelineCollectionModifier::AddCollection(UMovieGraphCollection* Collection)
 {
 	// Don't allow adding a duplicate collection
-	for (const UMoviePipelineCollection* ExistingCollection : Collections)
+	for (const UMovieGraphCollection* ExistingCollection : Collections)
 	{
 		if (Collection && ExistingCollection && Collection->GetCollectionName().Equals(ExistingCollection->GetCollectionName()))
 		{
@@ -323,7 +322,401 @@ void UMoviePipelineCollection::AddQuery(UMoviePipelineCollectionQuery* Query)
 	}
 }
 
-UMoviePipelineCollection* UMoviePipelineRenderLayer::GetCollectionByName(const FString& Name) const
+UMovieGraphConditionGroupQueryBase::UMovieGraphConditionGroupQueryBase()
+	: OpType(EMovieGraphConditionGroupQueryOpType::Union)
+	, bIsEnabled(true)
+{
+}
+
+void UMovieGraphConditionGroupQueryBase::SetOperationType(const EMovieGraphConditionGroupQueryOpType OperationType)
+{
+	// Always allow setting the operation type to Union. If not setting to Union, only allow setting the operation type if this is not the first
+	// query in the condition group. The first query is always a Union.
+	if (OperationType == EMovieGraphConditionGroupQueryOpType::Union)
+	{
+		OpType = EMovieGraphConditionGroupQueryOpType::Union;
+		return;
+	}
+
+	const UMovieGraphConditionGroup* ParentConditionGroup = GetTypedOuter<UMovieGraphConditionGroup>();
+	if (ensureMsgf(ParentConditionGroup, TEXT("Cannot set the operation type on a condition group query that doesn't have a condition group outer")))
+	{
+		if (ParentConditionGroup->GetQueries().Find(this) != 0)
+		{
+			OpType = OperationType;
+		}
+	}
+}
+
+EMovieGraphConditionGroupQueryOpType UMovieGraphConditionGroupQueryBase::GetOperationType() const
+{
+	return OpType;
+}
+
+void UMovieGraphConditionGroupQueryBase::Evaluate(const TArray<AActor*>& InActorsToQuery, TSet<AActor*>& OutMatchingActors) const
+{
+	// No implementation
+}
+
+bool UMovieGraphConditionGroupQueryBase::ShouldHidePropertyNames() const
+{
+	// Show property names by default; subclassed queries can opt-out if they want a cleaner UI 
+	return false;
+}
+
+const FSlateIcon& UMovieGraphConditionGroupQueryBase::GetIcon() const
+{
+	static const FSlateIcon EmptyIcon = FSlateIcon();
+	return EmptyIcon;
+}
+
+bool UMovieGraphConditionGroupQueryBase::IsEditorOnlyQuery() const
+{
+	return false;
+}
+
+void UMovieGraphConditionGroupQueryBase::SetEnabled(const bool bEnabled)
+{
+	bIsEnabled = bEnabled;
+}
+
+bool UMovieGraphConditionGroupQueryBase::IsEnabled() const
+{
+	return bIsEnabled;
+}
+
+void UMovieGraphConditionGroupQuery_ActorTag::Evaluate(const TArray<AActor*>& InActorsToQuery, TSet<AActor*>& OutMatchingActors) const
+{
+	TRACE_CPUPROFILER_EVENT_SCOPE(UMovieGraphConditionGroupQuery_ActorTag::Evaluate);
+	
+	for (AActor* Actor : InActorsToQuery)
+	{
+		if (Actor->Tags.Contains(TagToMatch))
+		{
+			OutMatchingActors.Add(Actor);
+		}
+	}
+}
+
+const FSlateIcon& UMovieGraphConditionGroupQuery_ActorTag::GetIcon() const
+{
+	// TODO: This icon is wrong
+	static const FSlateIcon ActorTagIcon = FSlateIcon(FAppStyle::GetAppStyleSetName(), "TODO");
+	return ActorTagIcon;
+}
+
+void UMovieGraphConditionGroupQuery_ActorName::Evaluate(const TArray<AActor*>& InActorsToQuery, TSet<AActor*>& OutMatchingActors) const
+{
+	TRACE_CPUPROFILER_EVENT_SCOPE(UMovieGraphConditionGroupQuery_ActorName::Evaluate);
+	
+	const bool bIsMatchAll = (WildcardSearch == TEXT("*"));
+
+	// Actor names can be specified on multiple lines
+	TArray<FString> AllActorNames;
+	WildcardSearch.ParseIntoArrayLines(AllActorNames);
+
+	// Quick early-out if "*" is used as the wildcard. Faster than doing the wildcard matching.
+	if (bIsMatchAll)
+	{
+		OutMatchingActors.Append(InActorsToQuery);
+		return;
+	}
+
+	for (AActor* Actor : InActorsToQuery)
+	{
+#if WITH_EDITOR
+		for (const FString& ActorName : AllActorNames)
+		{
+			if (Actor->GetActorLabel().MatchesWildcard(ActorName))
+			{
+				OutMatchingActors.Add(Actor);
+			}
+		}
+#endif
+	}
+}
+
+const FSlateIcon& UMovieGraphConditionGroupQuery_ActorName::GetIcon() const
+{
+	static const FSlateIcon ActorTagIcon = FSlateIcon(FAppStyle::GetAppStyleSetName(), "ClassIcon.TextRenderActor");
+	return ActorTagIcon;
+}
+
+bool UMovieGraphConditionGroupQuery_ActorName::IsEditorOnly() const
+{
+	// GetActorLabel() is editor-only
+	return true;
+}
+
+void UMovieGraphConditionGroupQuery_ActorType::Evaluate(const TArray<AActor*>& InActorsToQuery, TSet<AActor*>& OutMatchingActors) const
+{
+	TRACE_CPUPROFILER_EVENT_SCOPE(UMovieGraphConditionGroupQuery_ActorType::Evaluate);
+	
+	for (AActor* Actor : InActorsToQuery)
+	{
+		if (Actor->GetClass() == ActorType)
+		{
+			OutMatchingActors.Add(Actor);
+		}
+	}
+}
+
+const FSlateIcon& UMovieGraphConditionGroupQuery_ActorType::GetIcon() const
+{
+	static const FSlateIcon ActorTagIcon = FSlateIcon(FAppStyle::GetAppStyleSetName(), "ClassIcon.ActorComponent");
+	return ActorTagIcon;
+}
+
+void UMovieGraphConditionGroupQuery_ComponentTag::Evaluate(const TArray<AActor*>& InActorsToQuery, TSet<AActor*>& OutMatchingActors) const
+{
+	TRACE_CPUPROFILER_EVENT_SCOPE(UMovieGraphConditionGroupQuery_ComponentTag::Evaluate);
+	
+	TInlineComponentArray<UActorComponent*> ActorComponents;
+
+	for (AActor* Actor : InActorsToQuery)
+	{
+		Actor->GetComponents<UActorComponent*>(ActorComponents);
+		
+		for (const UActorComponent* Component : ActorComponents)
+		{
+			if (Component && Component->ComponentTags.Contains(TagToMatch))
+			{
+				OutMatchingActors.Add(Actor);
+			}
+		}
+
+		ActorComponents.Reset();
+	}
+}
+
+const FSlateIcon& UMovieGraphConditionGroupQuery_ComponentTag::GetIcon() const
+{
+	// TODO: This icon is wrong
+	static const FSlateIcon ActorTagIcon = FSlateIcon(FAppStyle::GetAppStyleSetName(), "ClassIcon.DataAsset");
+	return ActorTagIcon;
+}
+
+void UMovieGraphConditionGroupQuery_ComponentType::Evaluate(const TArray<AActor*>& InActorsToQuery, TSet<AActor*>& OutMatchingActors) const
+{
+	TRACE_CPUPROFILER_EVENT_SCOPE(UMovieGraphConditionGroupQuery_ComponentType::Evaluate);
+	
+	TInlineComponentArray<UActorComponent*> ActorComponents;
+
+	for (AActor* Actor : InActorsToQuery)
+	{
+		Actor->GetComponents<UActorComponent*>(ActorComponents);
+		
+		for (const UActorComponent* Component : ActorComponents)
+		{
+			if (Component->GetClass() == ComponentType)
+			{
+				OutMatchingActors.Add(Actor);
+			}
+		}
+
+		ActorComponents.Reset();
+	}
+}
+
+const FSlateIcon& UMovieGraphConditionGroupQuery_ComponentType::GetIcon() const
+{
+	static const FSlateIcon ActorTagIcon = FSlateIcon(FAppStyle::GetAppStyleSetName(), "ClassIcon.ActorComponent");
+	return ActorTagIcon;
+}
+
+UMovieGraphConditionGroup::UMovieGraphConditionGroup()
+	: OpType(EMovieGraphConditionGroupOpType::Union)
+	, bIsEnabled(true)
+{
+}
+
+void UMovieGraphConditionGroup::SetOperationType(const EMovieGraphConditionGroupOpType OperationType)
+{
+	// Always allow setting the operation type to Union. If not setting to Union, only allow setting the operation type if this is not the first
+	// condition group in the collection. The first condition group is always a Union.
+	if (OperationType == EMovieGraphConditionGroupOpType::Union)
+	{
+		OpType = EMovieGraphConditionGroupOpType::Union;
+		return;
+	}
+
+	const UMovieGraphCollection* ParentCollection = GetTypedOuter<UMovieGraphCollection>();
+	if (ensureMsgf(ParentCollection, TEXT("Cannot set the operation type on a condition group that doesn't have a collection outer")))
+	{
+		if (ParentCollection->GetConditionGroups().Find(this) != 0)
+		{
+			OpType = OperationType;
+		}
+	}
+}
+
+EMovieGraphConditionGroupOpType UMovieGraphConditionGroup::GetOperationType() const
+{
+	return OpType;
+}
+
+TSet<AActor*> UMovieGraphConditionGroup::Evaluate(const UWorld* InWorld) const
+{
+	TRACE_CPUPROFILER_EVENT_SCOPE(UMovieGraphConditionGroup::Evaluate);
+
+	// Reset the TSet for evaluation results; it is persisted across frames to prevent constantly re-allocating it
+	EvaluationResult.Reset();
+
+	// Generate a list of actors that can be fed to the queries once, rather than having all queries perform this
+	TArray<AActor*> AllActors;
+	for (TActorIterator<AActor> ActorItr(InWorld); ActorItr; ++ActorItr)
+	{
+		if (AActor* Actor = *ActorItr)
+		{
+			AllActors.Add(Actor);
+		}
+	}
+
+	for (int32 QueryIndex = 0; QueryIndex < Queries.Num(); ++QueryIndex)
+	{
+		const TObjectPtr<UMovieGraphConditionGroupQueryBase>& Query = Queries[QueryIndex];
+		if (!Query || !Query->IsEnabled())
+		{
+			continue;
+		}
+		
+		if (QueryIndex == 0)
+		{
+			// The first query should always be a Union
+			ensure(Query->GetOperationType() == EMovieGraphConditionGroupQueryOpType::Union);
+		}
+
+		// Similar to EvaluationResult, QueryResult is persisted+reset to prevent constantly re-allocating it
+		QueryResult.Reset();
+
+		Query->Evaluate(AllActors, QueryResult);
+		
+		switch (Query->GetOperationType())
+		{
+		case EMovieGraphConditionGroupQueryOpType::Union:
+			// Append() is faster than Union() because we don't need to allocate a new set
+			EvaluationResult.Append(QueryResult);
+			break;
+
+		case EMovieGraphConditionGroupQueryOpType::Intersect:
+			EvaluationResult = EvaluationResult.Intersect(QueryResult);
+			break;
+
+		case EMovieGraphConditionGroupQueryOpType::Subtract:
+			EvaluationResult = EvaluationResult.Difference(QueryResult);
+			break;
+		}
+	}
+	
+	return EvaluationResult;
+}
+
+UMovieGraphConditionGroupQueryBase* UMovieGraphConditionGroup::AddQuery(const TSubclassOf<UMovieGraphConditionGroupQueryBase>& InQueryType)
+{
+	UMovieGraphConditionGroupQueryBase* NewQueryObj = NewObject<UMovieGraphConditionGroupQueryBase>(this, InQueryType.Get());
+	Queries.Add(NewQueryObj);
+	return NewQueryObj;
+}
+
+const TArray<UMovieGraphConditionGroupQueryBase*>& UMovieGraphConditionGroup::GetQueries() const
+{
+	return Queries;
+}
+
+bool UMovieGraphConditionGroup::RemoveQuery(UMovieGraphConditionGroupQueryBase* InQuery)
+{
+	return Queries.RemoveSingle(InQuery) == 1;
+}
+
+void UMovieGraphConditionGroup::SetQueryOrder(TArray<UMovieGraphConditionGroupQueryBase*>& InQueryOrder)
+{
+	// TODO
+}
+
+void UMovieGraphConditionGroup::SetEnabled(const bool bEnabled)
+{
+	bIsEnabled = bEnabled;
+}
+
+bool UMovieGraphConditionGroup::IsEnabled() const
+{
+	return bIsEnabled;
+}
+
+TSet<AActor*> UMovieGraphCollection::Evaluate(const UWorld* InWorld) const
+{
+	TRACE_CPUPROFILER_EVENT_SCOPE(UMovieGraphCollection::Evaluate);
+	
+	TSet<AActor*> ResultSet;
+
+	for (int32 ConditionGroupIndex = 0; ConditionGroupIndex < ConditionGroups.Num(); ++ConditionGroupIndex)
+	{
+		const TObjectPtr<UMovieGraphConditionGroup>& ConditionGroup = ConditionGroups[ConditionGroupIndex];
+		if (!ConditionGroup || !ConditionGroup->IsEnabled())
+		{
+			continue;
+		}
+		
+		if (ConditionGroupIndex == 0)
+		{
+			// The first condition group should always be a Union
+			ensure(ConditionGroup->GetOperationType() == EMovieGraphConditionGroupOpType::Union);
+		}
+
+		const TSet<AActor*> QueryResult = ConditionGroup->Evaluate(InWorld);
+		
+		switch (ConditionGroup->GetOperationType())
+		{
+		case EMovieGraphConditionGroupOpType::Union:
+			ResultSet = ResultSet.Union(QueryResult);
+			break;
+
+		case EMovieGraphConditionGroupOpType::Intersect:
+			ResultSet = ResultSet.Intersect(QueryResult);
+			break;
+
+		case EMovieGraphConditionGroupOpType::Subtract:
+			ResultSet = ResultSet.Difference(QueryResult);
+			break;
+		}
+	}
+	
+	return ResultSet;
+}
+
+UMovieGraphConditionGroup* UMovieGraphCollection::AddConditionGroup()
+{
+	UMovieGraphConditionGroup* NewConditionGroup = NewObject<UMovieGraphConditionGroup>(this);
+	ConditionGroups.Add(NewConditionGroup);
+	return NewConditionGroup;
+}
+
+const TArray<UMovieGraphConditionGroup*>& UMovieGraphCollection::GetConditionGroups() const
+{
+	return ConditionGroups;
+}
+
+bool UMovieGraphCollection::RemoveConditionGroup(UMovieGraphConditionGroup* InConditionGroup)
+{
+	return ConditionGroups.RemoveSingle(InConditionGroup) == 1;
+}
+
+void UMovieGraphCollection::SetCollectionName(const FString& InName)
+{
+	CollectionName = InName;
+}
+
+const FString& UMovieGraphCollection::GetCollectionName() const
+{
+	return CollectionName;
+}
+
+void UMovieGraphCollection::SetConditionGroupOrder(TArray<UMovieGraphConditionGroup*>& InConditionGroupOrder)
+{
+	// TODO
+}
+
+UMovieGraphCollection* UMoviePipelineRenderLayer::GetCollectionByName(const FString& Name) const
 {
 	for (const UMoviePipelineCollectionModifier* Modifier : Modifiers)
 	{
@@ -332,7 +725,7 @@ UMoviePipelineCollection* UMoviePipelineRenderLayer::GetCollectionByName(const F
 			continue;
 		}
 		
-		for (UMoviePipelineCollection* Collection : Modifier->GetCollections())
+		for (UMovieGraphCollection* Collection : Modifier->GetCollections())
 		{
 			if (Collection && Collection->GetCollectionName().Equals(Name))
 			{
@@ -402,7 +795,7 @@ UMoviePipelineRenderLayerSubsystem* UMoviePipelineRenderLayerSubsystem::GetFromW
 
 void UMoviePipelineRenderLayerSubsystem::Initialize(FSubsystemCollectionBase& Collection)
 {
-	VisualizationEmptyCollection = NewObject<UMoviePipelineCollection>(GetTransientPackage(), NAME_None, RF_Transient);
+	VisualizationEmptyCollection = NewObject<UMovieGraphCollection>(GetTransientPackage(), NAME_None, RF_Transient);
 
 	// By default, the visualizer should hide everything in the world
 	VisualizationModifier_HideWorld = NewObject<UMoviePipelineVisibilityModifier>(GetTransientPackage(), NAME_None, RF_Transient);
@@ -499,7 +892,7 @@ void UMoviePipelineRenderLayerSubsystem::ClearActiveRenderLayer()
 	ClearAllPreviews();
 }
 
-void UMoviePipelineRenderLayerSubsystem::PreviewCollection(UMoviePipelineCollection* Collection)
+void UMoviePipelineRenderLayerSubsystem::PreviewCollection(UMovieGraphCollection* Collection)
 {
 	if (!Collection)
 	{
@@ -562,7 +955,7 @@ void UMoviePipelineRenderLayerSubsystem::PreviewModifier(UMoviePipelineCollectio
 
 	// Add the modifier's collections to the viz as well, so the actors that the modifier affects are actually visible
 	// TODO: This may need special handing for visibility modifiers
-	for (UMoviePipelineCollection* Collection : ActiveModifier->GetCollections())
+	for (UMovieGraphCollection* Collection : ActiveModifier->GetCollections())
 	{
 		VisualizationModifier_VisibleCollections->AddCollection(Collection);
 	}
diff --git a/Engine/Plugins/MovieScene/MovieRenderPipeline/Source/MovieRenderPipelineCore/Private/Graph/Nodes/MovieGraphCollectionNode.cpp b/Engine/Plugins/MovieScene/MovieRenderPipeline/Source/MovieRenderPipelineCore/Private/Graph/Nodes/MovieGraphCollectionNode.cpp
--- a/Engine/Plugins/MovieScene/MovieRenderPipeline/Source/MovieRenderPipelineCore/Private/Graph/Nodes/MovieGraphCollectionNode.cpp
+++ b/Engine/Plugins/MovieScene/MovieRenderPipeline/Source/MovieRenderPipelineCore/Private/Graph/Nodes/MovieGraphCollectionNode.cpp
@@ -3,19 +3,25 @@
 #include "Graph/Nodes/MovieGraphCollectionNode.h"
 
 #include "Graph/MovieGraphConfig.h"
+#include "Graph/MoviePipelineRenderLayerSubsystem.h"
 #include "Styling/AppStyle.h"
 
 #define LOCTEXT_NAMESPACE "MovieGraph"
 
+UMovieGraphCollectionNode::UMovieGraphCollectionNode()
+{
+	Collection = CreateDefaultSubobject<UMovieGraphCollection>(TEXT("Collection"));
+}
+
 #if WITH_EDITOR
 FText UMovieGraphCollectionNode::GetNodeTitle(const bool bGetDescriptive) const
 {
 	static const FText CollectionNodeName = LOCTEXT("NodeName_Collection", "Collection");
 	static const FText CollectionNodeDescription = LOCTEXT("NodeDescription_Collection", "Collection\n{0}");
 
-	if (bGetDescriptive && !CollectionName.IsEmpty())
+	if (bGetDescriptive && Collection && !Collection->GetCollectionName().IsEmpty())
 	{
-		return FText::Format(CollectionNodeDescription, FText::FromString(CollectionName));
+		return FText::Format(CollectionNodeDescription, FText::FromString(Collection->GetCollectionName()));
 	}
 
 	return CollectionNodeName;
@@ -44,14 +50,19 @@ void UMovieGraphCollectionNode::PostEditChangeProperty(FPropertyChangedEvent& Pr
 {
 	Super::PostEditChangeProperty(PropertyChangedEvent);
 
+	// TODO: Ideally this only fires when the collection name changes
 	// Broadcast a node-changed delegate so that the node title's UI gets updated.
-	if (PropertyChangedEvent.GetPropertyName() == GET_MEMBER_NAME_CHECKED(UMovieGraphCollectionNode, CollectionName) ||
-		PropertyChangedEvent.GetPropertyName() == GET_MEMBER_NAME_CHECKED(UMovieGraphCollectionNode, bOverride_CollectionName))
+	if (PropertyChangedEvent.GetPropertyName() == GET_MEMBER_NAME_CHECKED(UMovieGraphCollectionNode, Collection) ||
+		PropertyChangedEvent.GetPropertyName() == GET_MEMBER_NAME_CHECKED(UMovieGraphCollectionNode, bOverride_Collection))
 	{
 		OnNodeChangedDelegate.Broadcast(this);
 	}
 }
-
 #endif // WITH_EDITOR
 
+FString UMovieGraphCollectionNode::GetNodeInstanceName() const
+{
+	return Collection ? Collection->GetCollectionName() : FString();
+}
+
 #undef LOCTEXT_NAMESPACE // "MovieGraph"
\ No newline at end of file
diff --git a/Engine/Plugins/MovieScene/MovieRenderPipeline/Source/MovieRenderPipelineCore/Private/Graph/Renderers/MovieGraphDeferredPass.cpp b/Engine/Plugins/MovieScene/MovieRenderPipeline/Source/MovieRenderPipelineCore/Private/Graph/Renderers/MovieGraphDeferredPass.cpp
--- a/Engine/Plugins/MovieScene/MovieRenderPipeline/Source/MovieRenderPipelineCore/Private/Graph/Renderers/MovieGraphDeferredPass.cpp
+++ b/Engine/Plugins/MovieScene/MovieRenderPipeline/Source/MovieRenderPipelineCore/Private/Graph/Renderers/MovieGraphDeferredPass.cpp
@@ -74,6 +74,11 @@ void FMovieGraphDeferredPass::AddReferencedObjects(FReferenceCollector& Collecto
 	}
 }
 
+FName FMovieGraphDeferredPass::GetBranchName() const
+{
+	return LayerData.BranchName;
+}
+
 void FMovieGraphDeferredPass::Render(const FMovieGraphTraversalContext& InFrameTraversalContext, const FMovieGraphTimeStepData& InTimeData)
 {
 	// ToDo: InFrameTraversalContext includes a copy of TimeData, but may be the one cached at the first temporal sample,
diff --git a/Engine/Plugins/MovieScene/MovieRenderPipeline/Source/MovieRenderPipelineCore/Public/Graph/MoviePipelineRenderLayerSubsystem.h b/Engine/Plugins/MovieScene/MovieRenderPipeline/Source/MovieRenderPipelineCore/Public/Graph/MoviePipelineRenderLayerSubsystem.h
--- a/Engine/Plugins/MovieScene/MovieRenderPipeline/Source/MovieRenderPipelineCore/Public/Graph/MoviePipelineRenderLayerSubsystem.h
+++ b/Engine/Plugins/MovieScene/MovieRenderPipeline/Source/MovieRenderPipelineCore/Public/Graph/MoviePipelineRenderLayerSubsystem.h
@@ -4,6 +4,7 @@
 
 #include "CoreMinimal.h"
 #include "Subsystems/WorldSubsystem.h"
+#include "Textures/SlateIcon.h"
 
 #include "MoviePipelineRenderLayerSubsystem.generated.h"
 
@@ -40,6 +41,311 @@ private:
 	TArray<TObjectPtr<UMoviePipelineCollectionQuery>> Queries;
 };
 
+/** Operation types available on condition groups. */
+UENUM(BlueprintType)
+enum class EMovieGraphConditionGroupOpType : uint8
+{
+	/** Adds the contents of the condition group to the results from the previous condition group (if any). */
+	Union,
+
+	/** Removes the contents of the condition group from the result of the previous condition group (if any). Any items in this condition group that aren't also found in the previous condition group will be ignored. */
+	Subtract,
+
+	/** Replaces the results of the previous condition group(s) with only the elements that exist in both that group, and this group. Intersecting with an empty condition group will result in an empty condition group. */
+	Intersect
+};
+
+/** Operation types available on condition group queries. */
+UENUM(BlueprintType)
+enum class EMovieGraphConditionGroupQueryOpType : uint8
+{
+	/** Adds the results of the query to the results from the previous query (if any). */
+	Union,
+
+	/** Removes the results of the query from the results of the previous query (if any). Any items in this query result that aren't also found in the previous query result will be ignored. */
+	Subtract,
+	
+	/** Replaces the results of the previous queries with only the items that exist in both those queries, and this query result. Intersecting with a query which returns nothing will create an empty query result. */
+	Intersect
+};
+
+/** Base class that all condition group queries must inherit from. */
+UCLASS(Abstract, EditInlineNew)
+class UMovieGraphConditionGroupQueryBase : public UObject
+{
+	GENERATED_BODY()
+
+public:
+	UMovieGraphConditionGroupQueryBase();
+
+	/**
+	 * Sets how the condition group query interacts with the condition group. This call is ignored for the first query
+	 * in the condition group (the first is always Union).
+	 */
+	UFUNCTION(BlueprintCallable, Category = "Settings")
+	void SetOperationType(const EMovieGraphConditionGroupQueryOpType OperationType);
+
+	/** Gets the condition group query operation type. */
+	UFUNCTION(BlueprintCallable, Category = "Settings")
+	EMovieGraphConditionGroupQueryOpType GetOperationType() const;
+
+	/** Determines which of the provided actors match the query. Matches are added to OutMatchingActors.  */
+	UFUNCTION(BlueprintCallable, Category = "Settings")
+	virtual void Evaluate(const TArray<AActor*>& InActorsToQuery, TSet<AActor*>& OutMatchingActors) const;
+
+	/**
+	 * Determines if the public properties on the query class will have their names hidden in the details panel. Returns
+	 * false by default. Most query subclasses will only have one property and do not need to clutter the UI with the
+	 * property name (eg, the "Actor Name" query only shows one text box with entries for the actor names, no need to
+	 * show the property name).
+	 */
+	UFUNCTION(BlueprintCallable, Category = "Settings")
+	virtual bool ShouldHidePropertyNames() const;
+
+	/** Gets the icon that represents this query class in the UI. */
+	virtual const FSlateIcon& GetIcon() const;
+
+	/** Determines if this query is only respected when run within the editor. Used for providing a UI hint. */
+	UFUNCTION(BlueprintCallable, Category = "Settings")
+	virtual bool IsEditorOnlyQuery() const;
+
+	/** Sets whether this query is enabled. */
+	UFUNCTION(BlueprintCallable, Category = "Settings")
+	void SetEnabled(const bool bEnabled);
+
+	/** Determines if this query is enabled. */
+	UFUNCTION(BlueprintCallable, Category = "Settings")
+	bool IsEnabled() const;
+
+private:
+	/** The operation type that the query is using. */
+	UPROPERTY(EditAnywhere, Category="General")
+	EMovieGraphConditionGroupQueryOpType OpType;
+
+	/** Whether this query is currently enabled within the condition group. */
+	UPROPERTY(EditAnywhere, Category="General")
+	bool bIsEnabled;
+};
+
+UCLASS(BlueprintType, EditInlineNew, meta = (DisplayName = "Actor Tag"))
+class UMovieGraphConditionGroupQuery_ActorTag final : public UMovieGraphConditionGroupQueryBase
+{
+	GENERATED_BODY()
+
+public:
+	virtual void Evaluate(const TArray<AActor*>& InActorsToQuery, TSet<AActor*>& OutMatchingActors) const override;
+	virtual const FSlateIcon& GetIcon() const override;
+
+public:
+	/** The tag that the actor needs to have in order to be a match. */
+	UPROPERTY(EditAnywhere, Category="General")
+	FName TagToMatch;
+};
+
+UCLASS(BlueprintType, EditInlineNew, meta = (DisplayName = "Actor Name"))
+class UMovieGraphConditionGroupQuery_ActorName final : public UMovieGraphConditionGroupQueryBase
+{
+	GENERATED_BODY()
+
+public:
+	virtual void Evaluate(const TArray<AActor*>& InActorsToQuery, TSet<AActor*>& OutMatchingActors) const override;
+	virtual const FSlateIcon& GetIcon() const override;
+	virtual bool IsEditorOnly() const override;
+
+public:
+	/**
+	 * The name that the actor needs to have in order to be a match. Not case sensitive. Wildcards ("?" and "*") are supported but not required.
+	 * The "*" wildcard matches zero or more characters, and "?" matches exactly one character (and that character must be present).
+	 * 
+	 * Wildcard examples:
+	 * Foo* would match Foo, FooBar, and FooBaz, but not BarFoo.
+	 * *Foo* would match the above in addition to BarFoo.
+	 * Foo?Bar would match Foo.Bar and Foo_Bar, but not FooBar.
+	 * Foo? would match Food, but not FooBar or BarFoo.
+	 * Foo??? would match FooBar and FooBaz, but not Foo or Food.
+	 * ?oo? would match Food, but not Foo.
+	 * ?Foo* would match AFooBar, but not FooBar 
+	 */
+	UPROPERTY(EditAnywhere, Category="General", meta=(MultiLine=true))
+	FString WildcardSearch;
+};
+
+UCLASS(BlueprintType, EditInlineNew, meta = (DisplayName = "Actor Type"))
+class UMovieGraphConditionGroupQuery_ActorType final : public UMovieGraphConditionGroupQueryBase
+{
+	GENERATED_BODY()
+
+public:
+	virtual void Evaluate(const TArray<AActor*>& InActorsToQuery, TSet<AActor*>& OutMatchingActors) const override;
+	virtual const FSlateIcon& GetIcon() const override;
+
+public:
+	/** The type (class) that the actor needs to have in order to be a match. */
+	UPROPERTY(EditAnywhere, Category="General")
+	TSubclassOf<AActor> ActorType;
+};
+
+UCLASS(BlueprintType, EditInlineNew, meta = (DisplayName = "Component Tag"))
+class UMovieGraphConditionGroupQuery_ComponentTag final : public UMovieGraphConditionGroupQueryBase
+{
+	GENERATED_BODY()
+
+public:
+	virtual void Evaluate(const TArray<AActor*>& InActorsToQuery, TSet<AActor*>& OutMatchingActors) const override;
+	virtual const FSlateIcon& GetIcon() const override;
+
+public:
+	/** The tag that a component on the actor needs to have in order to be a match. */
+	UPROPERTY(EditAnywhere, Category="General")
+	FName TagToMatch;
+};
+
+UCLASS(BlueprintType, EditInlineNew, meta = (DisplayName = "Component Type"))
+class UMovieGraphConditionGroupQuery_ComponentType final : public UMovieGraphConditionGroupQueryBase
+{
+	GENERATED_BODY()
+
+public:
+	virtual void Evaluate(const TArray<AActor*>& InActorsToQuery, TSet<AActor*>& OutMatchingActors) const override;
+	virtual const FSlateIcon& GetIcon() const override;
+
+public:
+	/** The type that a component on the actor needs to have in order to be a match. */
+	UPROPERTY(EditAnywhere, Category="General")
+	TSubclassOf<UActorComponent> ComponentType;
+};
+
+/** A group of queries which can be added to a collection. */
+UCLASS(BlueprintType, EditInlineNew)
+class UMovieGraphConditionGroup : public UObject
+{
+	GENERATED_BODY()
+
+public:
+	UMovieGraphConditionGroup();
+
+	/**
+	 * Sets how the condition group interacts with the collection. This call is ignored for the first condition group
+	 * in the collection (the first is always Union).
+	 */
+	UFUNCTION(BlueprintCallable, Category = "Settings")
+	void SetOperationType(const EMovieGraphConditionGroupOpType OperationType);
+
+	/** Gets the condition group operation type. */
+	UFUNCTION(BlueprintCallable, Category = "Settings")
+	EMovieGraphConditionGroupOpType GetOperationType() const;
+
+	/** Determines the actors that match the condition group by running the queries contained in it. */
+	UFUNCTION(BlueprintCallable, Category = "Settings")
+	TSet<AActor*> Evaluate(const UWorld* InWorld) const;
+
+	/**
+	 * Adds a new condition group query to the condition group and returns a ptr to it. The condition group owns the
+	 * created query.
+	 */
+	UFUNCTION(BlueprintCallable, Category = "Settings")
+	UMovieGraphConditionGroupQueryBase* AddQuery(const TSubclassOf<UMovieGraphConditionGroupQueryBase>& InQueryType);
+
+	/** Gets all queries currently contained in the condition group. */
+	UFUNCTION(BlueprintCallable, Category = "Settings")
+	const TArray<UMovieGraphConditionGroupQueryBase*>& GetQueries() const;
+
+	/** Removes the specified query from the condition group if it exists. Returns true on success, else false. */
+	UFUNCTION(BlueprintCallable, Category = "Settings")
+	bool RemoveQuery(UMovieGraphConditionGroupQueryBase* InQuery);
+
+	/** TODO: API around reordering queries, pending research on how Slate provides this information. */
+	UFUNCTION(BlueprintCallable, Category = "Settings")
+	void SetQueryOrder(TArray<UMovieGraphConditionGroupQueryBase*>& InQueryOrder);
+
+	// TODO: May not exist in the final API if the design does not call for enabling/disabling condition groups.
+	/** Sets whether this condition group is enabled. */
+	UFUNCTION(BlueprintCallable, Category = "Settings")
+	void SetEnabled(const bool bEnabled);
+
+	// TODO: May not exist in the final API if the design does not call for enabling/disabling condition groups.
+	/** Determines if this condition group is enabled */
+	UFUNCTION(BlueprintCallable, Category = "Settings")
+	bool IsEnabled() const;
+
+private:
+	/** The operation type that the condition group is using. */
+	UPROPERTY(EditAnywhere, Category="General")
+	EMovieGraphConditionGroupOpType OpType;
+
+	/** Whether this condition group is currently enabled within the collection. */
+	UPROPERTY(EditAnywhere, Category="General")
+	bool bIsEnabled;
+
+	/** The queries that are contained within the condition group. */
+	UPROPERTY(EditAnywhere, Category="General", Instanced, meta=(EditInline))
+	TArray<TObjectPtr<UMovieGraphConditionGroupQueryBase>> Queries;
+
+	/** Persisted actor set which can be re-used for query evaluations across frames to prevent constantly re-allocating it. */
+	UPROPERTY(Transient)
+	mutable TSet<AActor*> QueryResult;
+
+	/** Persisted actor set which can be re-used for condition group evaluations across frames to prevent constantly re-allocating it. */
+	UPROPERTY(Transient)
+	mutable TSet<AActor*> EvaluationResult;
+};
+
+/** A group of actors generated by actor queries. */
+UCLASS(BlueprintType, EditInlineNew)
+class MOVIERENDERPIPELINECORE_API UMovieGraphCollection : public UObject
+{
+	GENERATED_BODY()
+
+public:
+	UMovieGraphCollection() = default;
+
+	/** Sets the name of the collection as seen in the UI. */
+	UFUNCTION(BlueprintCallable, Category = "Settings")
+	void SetCollectionName(const FString& InName);
+
+	/** Gets the name of the collection as seen in the UI. */
+	UFUNCTION(BlueprintCallable, Category = "Settings")
+	const FString& GetCollectionName() const;
+
+	/**
+	 * Gets matching actors by having condition groups evaluate themselves, and performing set operations on the
+	 * condition group results (eg, union'ing condition group A and B).
+	 */
+	UFUNCTION(BlueprintCallable, Category = "Settings")
+	TSet<AActor*> Evaluate(const UWorld* InWorld) const;
+
+	/**
+	 * Adds a new condition group to the collection and returns a ptr to it. The collection owns the created
+	 * condition group.
+	 */
+	UFUNCTION(BlueprintCallable, Category = "Settings")
+	UMovieGraphConditionGroup* AddConditionGroup();
+
+	/** Gets all condition groups currently contained in the collection. */
+	UFUNCTION(BlueprintCallable, Category = "Settings")
+	const TArray<UMovieGraphConditionGroup*>& GetConditionGroups() const;
+
+	/**
+	 * Removes the specified condition group from the collection if it exists. Returns true on success, else false.
+	 * Removes all child queries that belong to this group at the same time.
+	 */
+	UFUNCTION(BlueprintCallable, Category = "Settings")
+	bool RemoveConditionGroup(UMovieGraphConditionGroup* InConditionGroup);
+
+	/** TODO: API around reordering condition groups, pending research on how Slate provides this information. */
+	UFUNCTION(BlueprintCallable, Category = "Settings")
+	void SetConditionGroupOrder(TArray<UMovieGraphConditionGroup*>& InConditionGroupOrder);
+
+private:
+	/** The display name of the collection, shown in the UI. Does not need to be unique across collections. */
+	UPROPERTY(EditAnywhere, Category="General")
+	FString CollectionName;
+	
+	/** The condition groups that are contained within the collection. */
+	UPROPERTY(EditAnywhere, Category="General", Instanced, meta=(EditInline))
+	TArray<TObjectPtr<UMovieGraphConditionGroup>> ConditionGroups;
+};
+
 /**
  * Base class for providing actor modification functionality via collections.
  */
@@ -51,14 +357,14 @@ class UMoviePipelineCollectionModifier : public UObject
 public:
 	/** Adds a collection to the existing set of collections in this modifier. */
 	UFUNCTION(BlueprintCallable, Category = "Settings")
-	void AddCollection(UMoviePipelineCollection* Collection);
+	void AddCollection(UMovieGraphCollection* Collection);
 
 	/** Overwrites the existing collections with the provided array of collections. */
 	UFUNCTION(BlueprintCallable, Category = "Settings")
-	void SetCollections(const TArray<UMoviePipelineCollection*> InCollections) { Collections = InCollections; }
+	void SetCollections(const TArray<UMovieGraphCollection*> InCollections) { Collections = InCollections; }
 	
 	UFUNCTION(BlueprintCallable, Category = "Settings")
-	TArray<UMoviePipelineCollection*> GetCollections() const { return Collections; }
+	TArray<UMovieGraphCollection*> GetCollections() const { return Collections; }
 
 	virtual void ApplyModifier(const UWorld* World) PURE_VIRTUAL(UMoviePipelineCollectionModifier::ApplyModifier, );
 	virtual void UndoModifier() PURE_VIRTUAL(UMoviePipelineCollectionModifier::UndoModifier, );
@@ -71,7 +377,7 @@ public:
 protected:
 	/** The collections which this modifier will operate on. */
 	UPROPERTY()
-	TArray<TObjectPtr<UMoviePipelineCollection>> Collections;
+	TArray<TObjectPtr<UMovieGraphCollection>> Collections;
 };
 
 /**
@@ -262,7 +568,7 @@ public:
 	void SetRenderLayerName(const FName& NewName) { RenderLayerName = NewName; }
 
 	UFUNCTION(BlueprintCallable, Category = "Settings")
-	UMoviePipelineCollection* GetCollectionByName(const FString& Name) const;
+	UMovieGraphCollection* GetCollectionByName(const FString& Name) const;
 
 	UFUNCTION(BlueprintCallable, Category = "Settings")
 	void AddModifier(UMoviePipelineCollectionModifier* Modifier);
@@ -344,7 +650,7 @@ public:
 	void ClearActiveRenderLayer();
 
 	UFUNCTION(BlueprintCallable, Category = "Settings")
-	void PreviewCollection(UMoviePipelineCollection* Collection);
+	void PreviewCollection(UMovieGraphCollection* Collection);
 
 	UFUNCTION(BlueprintCallable, Category = "Settings")
 	void ClearCollectionPreview();
@@ -373,7 +679,7 @@ private:
 
 	/** The collection that is currently being viewed/previewed. */
 	UPROPERTY(Transient)
-	TObjectPtr<UMoviePipelineCollection> ActiveCollection = nullptr;
+	TObjectPtr<UMovieGraphCollection> ActiveCollection = nullptr;
 
 	/** The modifier that is currently being viewed/previewed. */
 	UPROPERTY(Transient)
@@ -385,7 +691,7 @@ private:
 
 	/** Empty collection used for visualization purposes (in conjunction w/ the viz render layer). */
 	UPROPERTY(Transient)
-	TObjectPtr<UMoviePipelineCollection> VisualizationEmptyCollection = nullptr;
+	TObjectPtr<UMovieGraphCollection> VisualizationEmptyCollection = nullptr;
 
 	/** A modifier used for visualization purposes (to hide the entire world). */
 	UPROPERTY(Transient)
diff --git a/Engine/Plugins/MovieScene/MovieRenderPipeline/Source/MovieRenderPipelineCore/Public/Graph/Nodes/MovieGraphCollectionNode.h b/Engine/Plugins/MovieScene/MovieRenderPipeline/Source/MovieRenderPipelineCore/Public/Graph/Nodes/MovieGraphCollectionNode.h
--- a/Engine/Plugins/MovieScene/MovieRenderPipeline/Source/MovieRenderPipelineCore/Public/Graph/Nodes/MovieGraphCollectionNode.h
+++ b/Engine/Plugins/MovieScene/MovieRenderPipeline/Source/MovieRenderPipelineCore/Public/Graph/Nodes/MovieGraphCollectionNode.h
@@ -7,7 +7,7 @@
 #include "MovieGraphCollectionNode.generated.h"
 
 // Forward Declare
-class UMoviePipelineCollectionQuery;
+class UMovieGraphCollection;
 
 /** 
 * A collection node specifies an interface for doing dynamic scene queries for actors in the world. Collections work in tandem with
@@ -19,6 +19,8 @@ class MOVIERENDERPIPELINECORE_API UMovieGraphCollectionNode : public UMovieGraph
 	GENERATED_BODY()
 
 public:
+	UMovieGraphCollectionNode();
+
 #if WITH_EDITOR
 	virtual FText GetNodeTitle(const bool bGetDescriptive = false) const override;
 	virtual FText GetMenuCategory() const override;
@@ -30,20 +32,12 @@ public:
 	//~ End UObject interface
 #endif
 
-	virtual FString GetNodeInstanceName() const override { return CollectionName; }
+	virtual FString GetNodeInstanceName() const override;
 
 public:
 	UPROPERTY(EditAnywhere, BlueprintReadWrite, Category = Overrides, meta = (InlineEditConditionToggle))
-	uint8 bOverride_CollectionName : 1;
-	
-	UPROPERTY(EditAnywhere, BlueprintReadWrite, Category = Overrides, meta = (InlineEditConditionToggle))
-	uint8 bOverride_QueryClass : 1;
-
-	/** The name of this collection, which is used to reference this collection in the graph. */
-	UPROPERTY(EditAnywhere, BlueprintReadWrite, Category = "General", meta=(EditCondition="bOverride_CollectionName"))
-	FString CollectionName;
-	
-	/** The type of query this node should run. */
-	UPROPERTY(BlueprintReadWrite, EditAnywhere, Instanced, Category = "General", meta=(EditCondition="bOverride_QueryClass", EditInline))
-	TObjectPtr<UMoviePipelineCollectionQuery> QueryClass;
+	uint8 bOverride_Collection : 1;
+
+	UPROPERTY(BlueprintReadWrite, EditAnywhere, Instanced, Category = "General", meta=(EditCondition="bOverride_Collection", EditInline))
+	TObjectPtr<UMovieGraphCollection> Collection;
 };
\ No newline at end of file
diff --git a/Engine/Plugins/MovieScene/MovieRenderPipeline/Source/MovieRenderPipelineCore/Public/Graph/Renderers/MovieGraphDeferredPass.h b/Engine/Plugins/MovieScene/MovieRenderPipeline/Source/MovieRenderPipelineCore/Public/Graph/Renderers/MovieGraphDeferredPass.h
--- a/Engine/Plugins/MovieScene/MovieRenderPipeline/Source/MovieRenderPipelineCore/Public/Graph/Renderers/MovieGraphDeferredPass.h
+++ b/Engine/Plugins/MovieScene/MovieRenderPipeline/Source/MovieRenderPipelineCore/Public/Graph/Renderers/MovieGraphDeferredPass.h
@@ -12,6 +12,7 @@ namespace UE::MovieGraph::Rendering
 		virtual void Render(const FMovieGraphTraversalContext& InFrameTraversalContext, const FMovieGraphTimeStepData& InTimeData) override;
 		virtual void GatherOutputPasses(TArray<FMovieGraphRenderDataIdentifier>& OutExpectedPasses) const override;
 		virtual void AddReferencedObjects(FReferenceCollector& Collector) override;
+		virtual FName GetBranchName() const override;
 	protected:
 		virtual void PostRendererSubmission(const UE::MovieGraph::FMovieGraphSampleState& InSampleState, const UE::MovieGraph::DefaultRenderer::FRenderTargetInitParams& InRenderTargetInitParams, FCanvas& InCanvas, const UE::MovieGraph::DefaultRenderer::FCameraInfo& InCameraInfo) override;
 
