MovieGraph: Fixed an issue where multiple shot renders would crash due to changing Time Step instances mid-way through using the instance. This also resolves an issue where "relative" numbers would reset between shots (including the ones meant to be relative to the start of the whole render).

diff --git a/Engine/Plugins/MovieScene/MovieRenderPipeline/Source/MovieRenderPipelineCore/Private/Graph/MovieGraphBlueprintLibrary.cpp b/Engine/Plugins/MovieScene/MovieRenderPipeline/Source/MovieRenderPipelineCore/Private/Graph/MovieGraphBlueprintLibrary.cpp
--- a/Engine/Plugins/MovieScene/MovieRenderPipeline/Source/MovieRenderPipelineCore/Private/Graph/MovieGraphBlueprintLibrary.cpp
+++ b/Engine/Plugins/MovieScene/MovieRenderPipeline/Source/MovieRenderPipelineCore/Private/Graph/MovieGraphBlueprintLibrary.cpp
@@ -338,7 +338,7 @@ void UMovieGraphBlueprintLibrary::GetOverallOutputFrames(const UMovieGraphPipeli
 	OutCurrentIndex = 0;
 	OutTotalCount = 0;
 	
-	if (InMovieGraphPipeline)
+	if (InMovieGraphPipeline && InMovieGraphPipeline->GetTimeStepInstance())
 	{
 		OutCurrentIndex = InMovieGraphPipeline->GetTimeStepInstance()->GetCalculatedTimeData().OutputFrameNumber;
 
diff --git a/Engine/Plugins/MovieScene/MovieRenderPipeline/Source/MovieRenderPipelineCore/Private/Graph/MovieGraphCoreTimeStep.cpp b/Engine/Plugins/MovieScene/MovieRenderPipeline/Source/MovieRenderPipelineCore/Private/Graph/MovieGraphCoreTimeStep.cpp
--- a/Engine/Plugins/MovieScene/MovieRenderPipeline/Source/MovieRenderPipelineCore/Private/Graph/MovieGraphCoreTimeStep.cpp
+++ b/Engine/Plugins/MovieScene/MovieRenderPipeline/Source/MovieRenderPipelineCore/Private/Graph/MovieGraphCoreTimeStep.cpp
@@ -19,12 +19,8 @@
 
 UMovieGraphCoreTimeStep::UMovieGraphCoreTimeStep()
 {
-	CustomTimeStep = CreateDefaultSubobject<UMovieGraphEngineTimeStep>("MovieGraphEngineTimeStep");
-
 	// This set up our internal state to pick up on the first temporal sub-sample in the pattern
 	ResetForEndOfOutputFrame();
-	CurrentTimeStepData.OutputFrameNumber = 0;
-	CurrentTimeStepData.RenderedFrameNumber = 0;
 }
 
 void UMovieGraphCoreTimeStep::TickProducingFrames()
@@ -33,18 +29,6 @@ void UMovieGraphCoreTimeStep::TickProducingFrames()
 	const TArray<TObjectPtr<UMoviePipelineExecutorShot>>& ActiveShotList = GetOwningGraph()->GetActiveShotList();
 	const TObjectPtr<UMoviePipelineExecutorShot>& CurrentCameraCut = ActiveShotList[CurrentShotIndex];
 
-	// When start up we want to override the engine's Custom Timestep with our own.
-	// This gives us the ability to completely control the engine tick/delta time before the frame
-	// is started so that we don't have to always be thinking of delta times one frame ahead. We need
-	// to do this only once we're ready to set the timestep though, as Initialize can be called as
-	// a result of a OnBeginFrame, meaning that Initialize is called on the frame before TickProducingFrames
-	// so there would be one frame where it used the custom timestep (after initialize) before TPF was called.
-	if (GEngine->GetCustomTimeStep() != CustomTimeStep)
-	{
-		// ToDo: This will restore the wrong timestep at the end of a render if we have different TimeStep instances.
-		PrevCustomTimeStep = GEngine->GetCustomTimeStep();
-		GEngine->SetCustomTimeStep(CustomTimeStep);
-	}
 
 	// We cache the frame metrics for the duration of a single output frame, instead of recalculating them every tick.
 	// This is required for stochastic timesteps (which will jump around the actual evaluated time randomly within one
@@ -330,7 +314,7 @@ void UMovieGraphCoreTimeStep::TickProducingFrames()
 
 				// The delta time isn't very relevant here but you must specify at _a_ delta time each frame to ensure
 				// we never have a frame that uses a stale delta time.
-				CustomTimeStep->SetCachedFrameTiming(UMovieGraphEngineTimeStep::FTimeStepCache(CurrentFrameMetrics.FrameRate.AsInterval()));
+				GetOwningGraph()->GetCustomEngineTimeStep()->SetCachedFrameTiming(UMovieGraphEngineTimeStep::FTimeStepCache(CurrentFrameMetrics.FrameRate.AsInterval()));
 
 				CurrentCameraCut->ShotInfo.State = EMovieRenderShotState::Finished;
 				GetOwningGraph()->TeardownShot(CurrentCameraCut);
@@ -381,8 +365,8 @@ void UMovieGraphCoreTimeStep::TickProducingFrames()
 	CurrentTimeStepData.bIsFirstTemporalSampleForFrame = bIsFirstTemporalSampleOverride.Get(IsFirstTemporalSample());
 	CurrentTimeStepData.bIsLastTemporalSampleForFrame = bIsLastTemporalSampleOverride.Get(IsLastTemporalSample());
 	CurrentTimeStepData.bRequiresAccumulator = CurrentFrameData.TemporalSampleCount > 1 && !CurrentTimeStepData.bDiscardOutput;
-	CurrentTimeStepData.OutputFrameNumber = CurrentFrameData.OutputFrameNumber;
-	CurrentTimeStepData.RenderedFrameNumber = CurrentFrameData.RenderedFrameNumber;
+	CurrentTimeStepData.OutputFrameNumber = GetOwningGraph()->GetCustomEngineTimeStep()->SharedTimeStepData.OutputFrameNumber;
+	CurrentTimeStepData.RenderedFrameNumber = GetOwningGraph()->GetCustomEngineTimeStep()->SharedTimeStepData.RenderedFrameNumber;
 	CurrentTimeStepData.TemporalSampleCount = CurrentFrameData.TemporalSampleCount;
 	CurrentTimeStepData.TemporalSampleIndex = CurrentFrameData.TemporalSampleIndex;
 	CurrentTimeStepData.EvaluatedConfig = TObjectPtr<UMovieGraphEvaluatedConfig>(CurrentFrameData.EvaluatedConfig.Get());
@@ -418,7 +402,7 @@ void UMovieGraphCoreTimeStep::TickProducingFrames()
 	// Set our time step for the next frame. We use the undilated delta time for the Custom Timestep as the engine will
 	// apply the time dilation to the world tick for us, so we don't want to double up time dilation.
 	double UndilatedDeltaTime = CurrentFrameMetrics.TickResolution.AsSeconds(FrameDeltaTime);
-	CustomTimeStep->SetCachedFrameTiming(UMovieGraphEngineTimeStep::FTimeStepCache(UndilatedDeltaTime));
+	GetOwningGraph()->GetCustomEngineTimeStep()->SetCachedFrameTiming(UMovieGraphEngineTimeStep::FTimeStepCache(UndilatedDeltaTime));
 	GetOwningGraph()->GetDataSourceInstance()->SyncDataSourceTime(FinalEvalTime);
 	if (bShouldJump)
 	{
@@ -438,14 +422,14 @@ void UMovieGraphCoreTimeStep::TickProducingFrames()
 
 		if ( bIsLastTemporalSampleOverride.Get(IsLastTemporalSample()))
 		{
-			CurrentFrameData.RenderedFrameNumber++;
+			GetOwningGraph()->GetCustomEngineTimeStep()->SharedTimeStepData.RenderedFrameNumber++;
 			CurrentFrameData.LastOutputFrameRange = CurrentFrameData.CurrentOutputFrameRange;
 				
 			// Increment the output frame number only on the last temporal sample, and only if
 			// we're actually rendering frames to disk.
 			if (CurrentCameraCut->ShotInfo.State == EMovieRenderShotState::Rendering)
 			{
-				CurrentFrameData.OutputFrameNumber++;
+				GetOwningGraph()->GetCustomEngineTimeStep()->SharedTimeStepData.OutputFrameNumber++;
 			}
 				
 			// If we've rendered the last temporal sub-sample, we've started a new output frame
@@ -471,7 +455,7 @@ void UMovieGraphCoreTimeStep::TickProducingFrames()
 		UE_LOG(LogMovieRenderPipeline, Error, TEXT("Shot ran past evaluation range, this shouldn't be possible."));
 	}
 
-	if (!ensureMsgf(!FMath::IsNearlyZero(CustomTimeStep->TimeCache.UndilatedDeltaTime), TEXT("An incorrect or uninitialized time step was used!")))
+	if (!ensureMsgf(!FMath::IsNearlyZero(GetOwningGraph()->GetCustomEngineTimeStep()->TimeCache.UndilatedDeltaTime), TEXT("An incorrect or uninitialized time step was used!")))
 	{
 		UE_LOG(LogMovieRenderPipeline, Error, TEXT("An incorrect or uninitialized time step was used!"));
 	}
@@ -549,16 +533,6 @@ bool UMovieGraphCoreTimeStep::IsExpansionForTSRequired(const TObjectPtr<UMovieGr
 	return false;
 }
 
-
-void UMovieGraphCoreTimeStep::Shutdown()
-{
-	// Shut down our custom timestep which reqstores some world settings we modified.
-
-	// ToDo: This takes 200ms because of an arbitrary sleep inside of SetCustomTimeStep when shutting down
-	// but that requires some bigger changes to custom timesteps to fix.
-	GEngine->SetCustomTimeStep(PrevCustomTimeStep);
-}
-
 void UMovieGraphCoreTimeStep::UpdateFrameMetrics()
 {
 	FOutputFrameMetrics FrameData;
diff --git a/Engine/Plugins/MovieScene/MovieRenderPipeline/Source/MovieRenderPipelineCore/Private/Graph/MovieGraphPipeline.cpp b/Engine/Plugins/MovieScene/MovieRenderPipeline/Source/MovieRenderPipelineCore/Private/Graph/MovieGraphPipeline.cpp
--- a/Engine/Plugins/MovieScene/MovieRenderPipeline/Source/MovieRenderPipelineCore/Private/Graph/MovieGraphPipeline.cpp
+++ b/Engine/Plugins/MovieScene/MovieRenderPipeline/Source/MovieRenderPipelineCore/Private/Graph/MovieGraphPipeline.cpp
@@ -24,8 +24,6 @@
 #include "ImageWriteQueue.h"
 #include "Modules/ModuleManager.h"
 
-// Temp
-#include "LevelSequence.h"
 
 FString UMovieGraphPipeline::DefaultPreviewWidgetAsset = TEXT("/MovieRenderPipeline/Blueprints/UI_MovieGraphPipelineScreenOverlay.UI_MovieGraphPipelineScreenOverlay_C");
 
@@ -35,6 +33,7 @@ UMovieGraphPipeline::UMovieGraphPipeline()
 	, PipelineState(EMovieRenderPipelineState::Uninitialized)
 {
 	OutputMerger = MakeShared<UE::MovieGraph::FMovieGraphOutputMerger>(this);
+	CustomEngineTimeStep = CreateDefaultSubobject<UMovieGraphEngineTimeStep>("MovieGraphEngineTimeStep");
 
 	Debug_ImageWriteQueue = &FModuleManager::Get().LoadModuleChecked<IImageWriteQueueModule>("ImageWriteQueue").GetWriteQueue();
 }
@@ -330,10 +329,6 @@ void UMovieGraphPipeline::BuildShotListFromDataSource()
 		Shot->ShotInfo.VersionNumber = ResolveVersionForShot(Shot, EvaluatedConfig);
 	}
 
-	// Initialize the time step instance to be the first one. This is usually set up by the shot, but some logic may attempt
-	// to access the time step instance before the shot has a chance to initialize.
-	GraphTimeStepInstance = GraphTimeStepInstances[0];
-
 	// The active shot-list is a subset of the whole shot-list; The ShotInfo contains information about every range it detected to render
 	// but if the user has turned the shot off in the UI then we don't want to render it.
 	ActiveShotList.Empty();
@@ -399,17 +394,42 @@ void UMovieGraphPipeline::TickProducingFrames()
 	//	return;
 	//}
 
-		// When start up we want to override the engine's Custom Timestep with our own.
+	// When start up we want to override the engine's Custom Timestep with our own.
 	// This gives us the ability to completely control the engine tick/delta time before the frame
 	// is started so that we don't have to always be thinking of delta times one frame ahead. We need
 	// to do this only once we're ready to set the timestep though, as Initialize can be called as
 	// a result of a OnBeginFrame, meaning that Initialize is called on the frame before TickProducingFrames
 	// so there would be one frame where it used the custom timestep (after initialize) before TPF was called.
-	//if (GEngine->GetCustomTimeStep() != CustomTimeStep)
-	//{
-	//	CachedPrevCustomTimeStep = GEngine->GetCustomTimeStep();
-	//	GEngine->SetCustomTimeStep(CustomTimeStep);
-	//}
+	if (GEngine->GetCustomTimeStep() != CustomEngineTimeStep)
+	{
+		PrevCustomEngineTimeStep = GEngine->GetCustomTimeStep();
+		GEngine->SetCustomTimeStep(CustomEngineTimeStep);
+	}
+
+	// If we don't have a graph time step instance assigned at all and TickProducingFrames is being called,
+	// then we set the first one as the pending one (because we haven't initialized any shots yet), that way it gets
+	// Initialize called on it, and we avoid anyone relying on a not-initialized GraphTimeStepInstance if 
+	// they tried accessing it before the first shot was started.
+	if (!GraphTimeStepInstance)
+	{
+		PendingTimeStepInstance = GraphTimeStepInstances[0];
+	}
+
+	// We can switch between time-step instances between shots,
+	// but the SetupShot/TeardownShot are called by the current instance,
+	// so we defer the actual pointer change until the next tick.
+	if (PendingTimeStepInstance)
+	{
+		if(GraphTimeStepInstance)
+		{
+			GraphTimeStepInstance->Shutdown();
+		}
+		GraphTimeStepInstance = PendingTimeStepInstance;
+		GraphTimeStepInstance->Initialize();
+		
+		PendingTimeStepInstance = nullptr;
+	}
+
 
 
 	GetTimeStepInstance()->TickProducingFrames();
@@ -553,8 +573,6 @@ void UMovieGraphPipeline::SetupShot(const TObjectPtr<UMoviePipelineExecutorShot>
 	//	}
 	//}
 
-	GraphTimeStepInstance = GraphTimeStepInstances[CurrentShotIndex];
-
 	const FMovieGraphTimeStepData& TimeStepData = GetTimeStepInstance()->GetCalculatedTimeData();
 	const UMovieGraphEvaluatedConfig* EvaluatedConfig = TimeStepData.EvaluatedConfig;
 
@@ -618,6 +636,13 @@ void UMovieGraphPipeline::TeardownShot(const TObjectPtr<UMoviePipelineExecutorSh
 	}
 
 	CurrentShotIndex++;
+
+	// At the end of each shot, set our Pending instance to be the next one, so that all initialization logic is handled within
+	// a single TimeStepInstance instead of being split between the last one and the new one (which would happen if we use SetupShot to do this)
+	if (CurrentShotIndex < GraphTimeStepInstances.Num())
+	{
+		PendingTimeStepInstance = GraphTimeStepInstances[CurrentShotIndex];
+	}
 }
 
 void UMovieGraphPipeline::SetSoloShot(const TObjectPtr<UMoviePipelineExecutorShot>& InShot)
@@ -718,7 +743,26 @@ TArray<TPair<FName, T*>> UMovieGraphPipeline::GetSettingForActiveRenderLayers(co
 	TArray<TPair<FName, T*>> FoundSettings;
 	
 	UMoviePipelineRenderLayerSubsystem* LayerSubsystem = GetWorld()->GetSubsystem<UMoviePipelineRenderLayerSubsystem>();
-	const FMovieGraphTimeStepData& TimeStepData = GetTimeStepInstance()->GetCalculatedTimeData();
+
+
+	// ToDo: This is only called by the Export step for the CommandLineEncoderNode right now, which only
+	// works per-render layer. However, by the time the export step actually starts, there isn't actually
+	// an active shot anymore, and it's not clear what graph we should be actually evaluating now. Perhaps 
+	// if we find a UMovieGraphPostRenderNode on a Render Layer, we should pause and run it at the end of 
+	// that shot? This would cause issues with post-renderer nodes that wanted to do something with the whole
+	// sequence (such as turn all shots into one movie file). Maybe in the post-export step we should only evaluate
+	// the Globals pin on the Primary Config, because we don't know which layers actually got rendered (due to shot
+	// overrides, etc.)
+	//
+	// For now, we're going to fall back on the incorrect legacy behavior (which was relying on stale data from the last
+	// shot), but this is going to have to be revisited.
+	UMovieGraphTimeStepBase* TimeStepInstance = GetTimeStepInstance();
+	if (!TimeStepInstance)
+	{
+		TimeStepInstance = GraphTimeStepInstances.Last();
+	}
+
+	const FMovieGraphTimeStepData& TimeStepData = TimeStepInstance->GetCalculatedTimeData();
 	const TObjectPtr<UMovieGraphEvaluatedConfig> EvaluatedConfig = TimeStepData.EvaluatedConfig;
 
 	for (const UMoviePipelineRenderLayer* RenderLayer : LayerSubsystem->GetRenderLayers())
@@ -735,8 +779,6 @@ TArray<TPair<FName, T*>> UMovieGraphPipeline::GetSettingForActiveRenderLayers(co
 
 UMovieGraphTimeStepBase* UMovieGraphPipeline::GetTimeStepInstance() const
 {
-	check(GraphTimeStepInstance);
-
 	return GraphTimeStepInstance;
 }
 
@@ -938,18 +980,23 @@ void UMovieGraphPipeline::TransitionToState(const EMovieRenderPipelineState InNe
 			// to continue to call it to allow ticking the Finalization stage.
 			FCoreDelegates::OnEndFrame.RemoveAll(this);
 
-			// Restore any custom Time Step that may have been set before. We do this here
-			// because the TimeStepInstance is only expected to be having to calculate times
-			// during ProducingFrames.
-			GetTimeStepInstance()->Shutdown();
-
 			// Ensure all frames have been processed by the GPU and sent to the Output Merger
 			FlushRenderingCommands();
 
 			// And then make sure all frames are sent to the Output Containers before we finalize.
 			ProcessOutstandingFinishedFrames();
 
-			// PreviewTexture = nullptr;
+			// Restore any custom Time Step that may have been set before. We do this here
+			// because the TimeStepInstance is only expected to be having to calculate times
+			// during ProducingFrames.
+			GetTimeStepInstance()->Shutdown();
+
+			// We assign it to nullptr here so that post-ProducingFrame code doesn't accidentally 
+			// use a stale time step instance.
+			GraphTimeStepInstance = nullptr;
+
+			// Shut down our custom timestep which reqstores some world settings we modified.
+			GEngine->SetCustomTimeStep(PrevCustomEngineTimeStep);
 
 			// This is called once notifying output containers that all frames that will be submitted have been submitted.
 			PipelineState = EMovieRenderPipelineState::Finalize;
diff --git a/Engine/Plugins/MovieScene/MovieRenderPipeline/Source/MovieRenderPipelineCore/Public/Graph/MovieGraphCoreTimeStep.h b/Engine/Plugins/MovieScene/MovieRenderPipeline/Source/MovieRenderPipelineCore/Public/Graph/MovieGraphCoreTimeStep.h
--- a/Engine/Plugins/MovieScene/MovieRenderPipeline/Source/MovieRenderPipelineCore/Public/Graph/MovieGraphCoreTimeStep.h
+++ b/Engine/Plugins/MovieScene/MovieRenderPipeline/Source/MovieRenderPipelineCore/Public/Graph/MovieGraphCoreTimeStep.h
@@ -5,7 +5,6 @@
 #include "Graph/MovieGraphConfig.h"
 #include "Misc/FrameRate.h"
 #include "Misc/FrameTime.h"
-#include "Engine/EngineCustomTimeStep.h"
 #include "UObject/StrongObjectPtr.h"
 
 #include "MovieGraphCoreTimeStep.generated.h"
@@ -29,7 +28,6 @@ public:
 
 	// UMovieGraphTimeStepBase Interface
 	virtual void TickProducingFrames() override;
-	virtual void Shutdown() override;
 	virtual FMovieGraphTimeStepData GetCalculatedTimeData() const override { return CurrentTimeStepData; }
 	virtual bool IsExpansionForTSRequired(const TObjectPtr<UMovieGraphEvaluatedConfig>& InConfig) const override;
 	// ~UMovieGraphTimeStepBase Interface
@@ -106,8 +104,6 @@ protected:
 		FCurrentFrameData()
 			: TemporalSampleIndex(0)
 			, TemporalSampleCount(0)
-			, OutputFrameNumber(0)
-			, RenderedFrameNumber(0)
 		{
 		}
 
@@ -117,15 +113,6 @@ protected:
 		/** How many temporal sub-samples are there for the current output frame? */
 		int32 TemporalSampleCount;
 		
-		/** Which output frame are we working on, relative to zero.*/
-		int32 OutputFrameNumber;
-
-		/** 
-		* Index of which frames we've submitted for rendering. Doesn't line up with OutputFrameNumber when using warm-up frames.
-		* Used internally by the rendering engine to keep track of which frames need to be read back.
-		*/
-		int32 RenderedFrameNumber;
-
 		/** A range of time (in Tick Resolution) for the last output frame being worked on. Updated before first TS of next frame. */
 		TRange<FFrameTime> LastOutputFrameRange;
 
@@ -162,53 +149,4 @@ protected:
 	* across engine ticks to work towards producing a single frame.
 	*/
 	FCurrentFrameData CurrentFrameData;
-
-	/**
-	* A custom timestep owned by this object that is used to inform the engine what the delta time for each
-	* frame should be. 
-	*/
-	UPROPERTY(Transient)
-	TObjectPtr<UMovieGraphEngineTimeStep> CustomTimeStep;
-
-	/** The previous custom timestep the engine was using, if any. */
-	UPROPERTY(Transient)
-	TObjectPtr<UEngineCustomTimeStep> PrevCustomTimeStep;
-};
-
-UCLASS()
-class MOVIERENDERPIPELINECORE_API UMovieGraphEngineTimeStep : public UEngineCustomTimeStep
-{
-	GENERATED_BODY()
-public:
-	UMovieGraphEngineTimeStep();
-
-	struct FTimeStepCache
-	{
-		FTimeStepCache() 
-			: UndilatedDeltaTime(0.0)
-		{}
-
-		FTimeStepCache(double InUndilatedDeltaTime)
-			: UndilatedDeltaTime(InUndilatedDeltaTime)
-		{}
-
-		double UndilatedDeltaTime;
-	};
-
-public:
-	void SetCachedFrameTiming(const FTimeStepCache& InTimeCache);
-	
-	// UEngineCustomTimeStep Interface
-	virtual bool Initialize(UEngine* InEngine) override;
-	virtual void Shutdown(UEngine* InEngine) override;
-	virtual bool UpdateTimeStep(UEngine* InEngine) override;
-	virtual ECustomTimeStepSynchronizationState GetSynchronizationState() const override { return ECustomTimeStepSynchronizationState::Synchronized; }
-	// ~UEngineCustomTimeStep Interface
-
-	/** We don't do any thinking on our own, instead we just spit out the numbers stored in our time cache. */
-	FTimeStepCache TimeCache;
-
-	// Not cached in TimeCache as TimeCache is reset every frame.
-	float PrevMinUndilatedFrameTime;
-	float PrevMaxUndilatedFrameTime;
 };
diff --git a/Engine/Plugins/MovieScene/MovieRenderPipeline/Source/MovieRenderPipelineCore/Public/Graph/MovieGraphDataTypes.h b/Engine/Plugins/MovieScene/MovieRenderPipeline/Source/MovieRenderPipelineCore/Public/Graph/MovieGraphDataTypes.h
--- a/Engine/Plugins/MovieScene/MovieRenderPipeline/Source/MovieRenderPipelineCore/Public/Graph/MovieGraphDataTypes.h
+++ b/Engine/Plugins/MovieScene/MovieRenderPipeline/Source/MovieRenderPipelineCore/Public/Graph/MovieGraphDataTypes.h
@@ -10,6 +10,7 @@
 #include "Containers/Queue.h"
 #include "Misc/FrameRate.h"
 #include "MovieRenderPipelineDataTypes.h"
+#include "Engine/EngineCustomTimeStep.h"
 
 #include "MovieGraphDataTypes.generated.h"
 
@@ -103,7 +104,9 @@ public:
 	/** Called each frame while the Movie Graph Pipeline is in a producing frames state. */
 	virtual void TickProducingFrames() {}
 
-	/** Called when the Movie Graph Pipeline is shutting down, use this to restore any changes. */
+	/** Called when this time step instance becomes active (ie: at the start of a shot). */
+	virtual void Initialize() {}
+	/** Called when this time step instance is no longer active (ie: at the end of a shot). */
 	virtual void Shutdown() {}
 
 	/** 
@@ -122,6 +125,66 @@ public:
 	UMovieGraphPipeline* GetOwningGraph() const;
 };
 
+
+UCLASS()
+class MOVIERENDERPIPELINECORE_API UMovieGraphEngineTimeStep : public UEngineCustomTimeStep
+{
+	GENERATED_BODY()
+public:
+	UMovieGraphEngineTimeStep();
+
+	struct FTimeStepCache
+	{
+		FTimeStepCache()
+			: UndilatedDeltaTime(0.0)
+		{}
+
+		FTimeStepCache(double InUndilatedDeltaTime)
+			: UndilatedDeltaTime(InUndilatedDeltaTime)
+		{}
+
+		double UndilatedDeltaTime;
+	};
+
+	struct FSharedData
+	{
+		FSharedData()
+			: OutputFrameNumber(0)
+			, RenderedFrameNumber(0)
+		{}
+		
+		
+		/** Which output frame are we working on, relative to zero.*/
+		int32 OutputFrameNumber;
+
+		/** 
+		* Index of which frames we've submitted for rendering. Doesn't line up with OutputFrameNumber when using warm-up frames.
+		* Used internally by the rendering engine to keep track of which frames need to be read back.
+		*/
+		int32 RenderedFrameNumber;
+	};
+
+public:
+	void SetCachedFrameTiming(const FTimeStepCache& InTimeCache);
+
+	// UEngineCustomTimeStep Interface
+	virtual bool Initialize(UEngine* InEngine) override;
+	virtual void Shutdown(UEngine* InEngine) override;
+	virtual bool UpdateTimeStep(UEngine* InEngine) override;
+	virtual ECustomTimeStepSynchronizationState GetSynchronizationState() const override { return ECustomTimeStepSynchronizationState::Synchronized; }
+	// ~UEngineCustomTimeStep Interface
+
+	/** We don't do any thinking on our own, instead we just spit out the numbers stored in our time cache. */
+	FTimeStepCache TimeCache;
+
+	/** Data that should be shared between all shot time step instances. */
+	FSharedData SharedTimeStepData;
+
+	// Not cached in TimeCache as TimeCache is reset every frame.
+	float PrevMinUndilatedFrameTime;
+	float PrevMaxUndilatedFrameTime;
+};
+
 UCLASS(BlueprintType, Abstract)
 class MOVIERENDERPIPELINECORE_API UMovieGraphRendererBase : public UObject
 {
diff --git a/Engine/Plugins/MovieScene/MovieRenderPipeline/Source/MovieRenderPipelineCore/Public/Graph/MovieGraphPipeline.h b/Engine/Plugins/MovieScene/MovieRenderPipeline/Source/MovieRenderPipelineCore/Public/Graph/MovieGraphPipeline.h
--- a/Engine/Plugins/MovieScene/MovieRenderPipeline/Source/MovieRenderPipelineCore/Public/Graph/MovieGraphPipeline.h
+++ b/Engine/Plugins/MovieScene/MovieRenderPipeline/Source/MovieRenderPipelineCore/Public/Graph/MovieGraphPipeline.h
@@ -67,6 +67,9 @@ public:
 	
 	/** Which index of the Active Shot List are we currently on */
 	int32 GetCurrentShotIndex() const { return CurrentShotIndex; }
+
+	/** Get the pointer to the Custom Engine Timestep we use to control engine ticking. Shared by all per-shot Time Steps.*/
+	TObjectPtr<UMovieGraphEngineTimeStep> GetCustomEngineTimeStep() const { return CustomEngineTimeStep; }
 	
 	/** Called by the TimeStepInstance when it's time to set up for another shot. Don't call this unless you know what you're doing. */
 	void SetupShot(const TObjectPtr<UMoviePipelineExecutorShot>& InShot);
@@ -165,6 +168,14 @@ protected:
 	UPROPERTY(Transient)
 	TArray<TObjectPtr<UMovieGraphTimeStepBase>> GraphTimeStepInstances;
 
+	/** 
+	* If set, on the next TickProducingFrames, the GraphTimeStepInstance pointer will be swapped with this one.
+	* Uses a deferred mechanism since SetupShot/TeardownShot are called by the current instance, so we don't want
+	* to swap until the current instance isn't actively being used.
+	*/
+	UPROPERTY(Transient)
+	TObjectPtr<UMovieGraphTimeStepBase> PendingTimeStepInstance;
+
 	/**
 	 * Sometimes the shot index can be incremented before the time step instance should be changed, so the current time
 	 * step instance is tracked. This should generally be used to access the current time step instance, rather than
@@ -240,6 +251,16 @@ protected:
 
 	/** Responsible for managing cvars throughout the lifetime of the pipeline. */
 	TSharedPtr<UE::MovieGraph::Private::FMovieGraphCVarManager> CVarManager;
+
+	/**
+	* A custom timestep created by the Movie Graph Pipeline to allow the time step managers to set what the delta times for the frame should be.
+	*/
+	UPROPERTY(Transient)
+	TObjectPtr<UMovieGraphEngineTimeStep> CustomEngineTimeStep;
+
+	/** The previous custom timestep the engine was using, if any. */
+	UPROPERTY(Transient)
+	TObjectPtr<UEngineCustomTimeStep> PrevCustomEngineTimeStep;
 	
 public:
 	static FString DefaultPreviewWidgetAsset;
