Asset: Move some assets from Animation to Cinematics

#rnx

diff --git a/Engine/Plugins/MovieScene/MovieRenderPipeline/Source/MovieRenderPipelineEditor/Private/AssetDefinition_PipelinePrimaryConfig.h b/Engine/Plugins/MovieScene/MovieRenderPipeline/Source/MovieRenderPipelineEditor/Private/AssetDefinition_PipelinePrimaryConfig.h
--- a/Engine/Plugins/MovieScene/MovieRenderPipeline/Source/MovieRenderPipelineEditor/Private/AssetDefinition_PipelinePrimaryConfig.h
+++ b/Engine/Plugins/MovieScene/MovieRenderPipeline/Source/MovieRenderPipelineEditor/Private/AssetDefinition_PipelinePrimaryConfig.h
@@ -19,7 +19,7 @@ public:
 	virtual TSoftClassPtr<UObject> GetAssetClass() const override { return UMoviePipelinePrimaryConfig::StaticClass(); }
 	virtual TConstArrayView<FAssetCategoryPath> GetAssetCategories() const override
 	{
-		static const auto Categories = { EAssetCategoryPaths::Animation };
+		static const auto Categories = { EAssetCategoryPaths::Cinematics };
 		return Categories;
 	}
 	virtual EAssetCommandResult OpenAssets(const FAssetOpenArgs& OpenArgs) const override;
diff --git a/Engine/Plugins/VirtualProduction/Takes/Source/TakeRecorder/Private/AssetDefinition_TakePreset.h b/Engine/Plugins/VirtualProduction/Takes/Source/TakeRecorder/Private/AssetDefinition_TakePreset.h
--- a/Engine/Plugins/VirtualProduction/Takes/Source/TakeRecorder/Private/AssetDefinition_TakePreset.h
+++ b/Engine/Plugins/VirtualProduction/Takes/Source/TakeRecorder/Private/AssetDefinition_TakePreset.h
@@ -19,7 +19,7 @@ public:
 	virtual TSoftClassPtr<UObject> GetAssetClass() const override { return UTakePreset::StaticClass(); }
 	virtual TConstArrayView<FAssetCategoryPath> GetAssetCategories() const override
 	{
-		static const auto Categories = { EAssetCategoryPaths::Animation };
+		static const auto Categories = { EAssetCategoryPaths::Cinematics };
 		return Categories;
 	}
 	virtual FAssetOpenSupport GetAssetOpenSupport(const FAssetOpenSupportArgs& OpenSupportArgs) const override;
