MoveRenderGraph:

-Unit Test: Testing Node Branch connection (Global vs RenderLayer)
-Fix: Should prevent connection between Global-only nodes and RenderLayer-only nodes
-Restore Error verbosity to appropriate KismetExecutionMessages
-General Test refactor for readability
-Suppress warning and error logs in tests
-OpenGraphConfigInEditorIfDesired now only oens the editor for the asset if static bool bIsDesired == true
-Remove MovieGraphPlaceholderNodes (per Matt Hoffman)

#[fyi] Jason.Walter

diff --git a/Engine/Plugins/MovieScene/MovieRenderPipeline/Source/MovieRenderPipelineCore/Private/Graph/MovieGraphConfig.cpp b/Engine/Plugins/MovieScene/MovieRenderPipeline/Source/MovieRenderPipelineCore/Private/Graph/MovieGraphConfig.cpp
--- a/Engine/Plugins/MovieScene/MovieRenderPipeline/Source/MovieRenderPipelineCore/Private/Graph/MovieGraphConfig.cpp
+++ b/Engine/Plugins/MovieScene/MovieRenderPipeline/Source/MovieRenderPipelineCore/Private/Graph/MovieGraphConfig.cpp
@@ -583,7 +583,7 @@ void UMovieGraphConfig::AddNode(UMovieGraphNode* InNode)
 		FFrame::KismetExecutionMessage(
 			*FString::Printf(
 				TEXT("%hs: No node was specified for add."), __FUNCTION__),
-				ELogVerbosity::Display);
+				ELogVerbosity::Error);
 		return;
 	}
 
@@ -592,7 +592,7 @@ void UMovieGraphConfig::AddNode(UMovieGraphNode* InNode)
 		FFrame::KismetExecutionMessage(
 			*FString::Printf(
 				TEXT("%hs:Cannot add node of type %s."), __FUNCTION__, *InNode->GetClass()->GetName()),
-				ELogVerbosity::Display);
+				ELogVerbosity::Error);
 		return;
 	}
 
@@ -623,7 +623,7 @@ bool UMovieGraphConfig::RemoveNode(UMovieGraphNode* InNode)
 		FFrame::KismetExecutionMessage(
 			*FString::Printf(
 				TEXT("%hs: Could not remove invalid InNode"), __FUNCTION__),
-				ELogVerbosity::Display);
+				ELogVerbosity::Error);
 		return false;
 	}
 
diff --git a/Engine/Plugins/MovieScene/MovieRenderPipeline/Source/MovieRenderPipelineCore/Private/Graph/MovieGraphPin.cpp b/Engine/Plugins/MovieScene/MovieRenderPipeline/Source/MovieRenderPipelineCore/Private/Graph/MovieGraphPin.cpp
--- a/Engine/Plugins/MovieScene/MovieRenderPipeline/Source/MovieRenderPipelineCore/Private/Graph/MovieGraphPin.cpp
+++ b/Engine/Plugins/MovieScene/MovieRenderPipeline/Source/MovieRenderPipelineCore/Private/Graph/MovieGraphPin.cpp
@@ -227,7 +227,7 @@ bool UMovieGraphPin::IsConnectionToBranchAllowed(const UMovieGraphPin* OtherPin,
 	
 	if (!InputPin || !OutputPin)
 	{
-		OutError = NSLOCTEXT("MoviePipeline", "PinDirectionMismatchError", "Directions are not compatible!");
+		OutError = NSLOCTEXT("MovieGraph", "PinDirectionMismatchError", "Directions are not compatible!");
 		return false;
 	}
 	
@@ -236,6 +236,17 @@ bool UMovieGraphPin::IsConnectionToBranchAllowed(const UMovieGraphPin* OtherPin,
 	check(ToNode && FromNode);
 	const UMovieGraphConfig* GraphConfig = ToNode->GetGraph();
 
+	// Test High-Level Node Restrictions
+	const EMovieGraphBranchRestriction FromNodeRestriction = FromNode->GetBranchRestriction();
+	const EMovieGraphBranchRestriction ToNodeRestriction = ToNode->GetBranchRestriction();
+	if (FromNodeRestriction != ToNodeRestriction &&						// If BranchRestrictions are not the same
+		FromNodeRestriction != EMovieGraphBranchRestriction::Any &&		// And neither Node is an 'Any' Node
+		ToNodeRestriction != EMovieGraphBranchRestriction::Any)			// Then do not allow connection
+	{
+		OutError = NSLOCTEXT("MovieGraph", "HighLevelPerNodeBranchRestrictionError", "Cannot connect a Globals-only Node to a RenderLayer-only Node!");
+		return false;
+	}
+
 	// Get all upstream/downstream nodes that occur on the connection -- these are the nodes that need to be checked for branch restrictions.
 	// FromNode/ToNode themselves also needs to be part of the validation checks.
 	TArray<UMovieGraphNode*> NodesToCheck = {FromNode, ToNode};
@@ -264,7 +275,7 @@ bool UMovieGraphPin::IsConnectionToBranchAllowed(const UMovieGraphPin* OtherPin,
 	// Globals branches can only be connected to Globals branches
 	if ((bGlobalsIsDownstream && bUpstreamBranchExistsAndIsntOnlyGlobals) || (bGlobalsIsUpstream && bDownstreamBranchExistsAndIsntOnlyGlobals))
 	{
-		OutError = NSLOCTEXT("MoviePipeline", "GlobalsBranchMismatchError", "Globals branches can only be connected to other Globals branches.");
+		OutError = NSLOCTEXT("MovieGraph", "GlobalsBranchMismatchError", "Globals branches can only be connected to other Globals branches.");
 		return false;
 	}
 
@@ -280,7 +291,7 @@ bool UMovieGraphPin::IsConnectionToBranchAllowed(const UMovieGraphPin* OtherPin,
 			if (bDownstreamBranchExistsAndIsntOnlyGlobals || bUpstreamBranchExistsAndIsntOnlyGlobals)
 			{
 				OutError = FText::Format(
-					NSLOCTEXT("MoviePipeline", "GlobalsBranchRestrictionError", "The node '{0}' can only be connected to the Globals branch."),
+					NSLOCTEXT("MovieGraph", "GlobalsBranchRestrictionError", "The node '{0}' can only be connected to the Globals branch."),
 						FText::FromString(NodeToCheck->GetName()));
 				return false;
 			}
@@ -292,7 +303,7 @@ bool UMovieGraphPin::IsConnectionToBranchAllowed(const UMovieGraphPin* OtherPin,
 			if (bGlobalsIsDownstream || bGlobalsIsUpstream)
 			{
 				OutError = FText::Format(
-					NSLOCTEXT("MoviePipeline", "RenderLayerBranchRestrictionError", "The node '{0}' can only be connected to a render layer branch."),
+					NSLOCTEXT("MovieGraph", "RenderLayerBranchRestrictionError", "The node '{0}' can only be connected to a render layer branch."),
 						FText::FromString(NodeToCheck->GetName()));
 				return false;
 			}
diff --git a/Engine/Plugins/MovieScene/MovieRenderPipeline/Source/MovieRenderPipelineCore/Public/Graph/MovieGraphConfig.h b/Engine/Plugins/MovieScene/MovieRenderPipeline/Source/MovieRenderPipelineCore/Public/Graph/MovieGraphConfig.h
--- a/Engine/Plugins/MovieScene/MovieRenderPipeline/Source/MovieRenderPipelineCore/Public/Graph/MovieGraphConfig.h
+++ b/Engine/Plugins/MovieScene/MovieRenderPipeline/Source/MovieRenderPipelineCore/Public/Graph/MovieGraphConfig.h
@@ -678,7 +678,7 @@ public:
 			FFrame::KismetExecutionMessage(
 				*FString::Printf(
 					TEXT("%hs: Invalid PipelineGraphNodeClass. Please specify a valid class."), __FUNCTION__),
-					ELogVerbosity::Display);
+					ELogVerbosity::Error);
 		
 			return nullptr;
 		}
diff --git a/Engine/Plugins/MovieScene/MovieRenderPipeline/Source/MovieRenderPipelineCore/Public/Graph/MovieGraphPlaceholderNodes.cpp b/Engine/Plugins/MovieScene/MovieRenderPipeline/Source/MovieRenderPipelineCore/Public/Graph/MovieGraphPlaceholderNodes.cpp
deleted file mode 100644
diff --git a/Engine/Plugins/MovieScene/MovieRenderPipeline/Source/MovieRenderPipelineCore/Public/Graph/MovieGraphPlaceholderNodes.h b/Engine/Plugins/MovieScene/MovieRenderPipeline/Source/MovieRenderPipelineCore/Public/Graph/MovieGraphPlaceholderNodes.h
deleted file mode 100644
diff --git a/Engine/Plugins/MovieScene/MovieRenderPipeline/Source/MovieRenderPipelineEditor/Private/Tests/MovieGraphEditorTestUtilities.cpp b/Engine/Plugins/MovieScene/MovieRenderPipeline/Source/MovieRenderPipelineEditor/Private/Tests/MovieGraphEditorTestUtilities.cpp
--- a/Engine/Plugins/MovieScene/MovieRenderPipeline/Source/MovieRenderPipelineEditor/Private/Tests/MovieGraphEditorTestUtilities.cpp
+++ b/Engine/Plugins/MovieScene/MovieRenderPipeline/Source/MovieRenderPipelineEditor/Private/Tests/MovieGraphEditorTestUtilities.cpp
@@ -4,11 +4,13 @@
 
 #include "Graph/MovieGraphConfig.h"
 #include "Graph/MovieGraphConfigFactory.h"
+#include "Graph/MovieGraphNode.h"
 
 #include "Algo/RemoveIf.h"
 #include "AssetRegistry/AssetRegistryModule.h"
 #include "AssetToolsModule.h"
 #include "IAssetTools.h"
+#include "Misc/AutomationTest.h"
 #include "UObject/Package.h"
 
 namespace UE::MovieGraph::Private::Tests
@@ -83,8 +85,7 @@ namespace UE::MovieGraph::Private::Tests
 		return NativeParentClass && (NativeParentClass == BaseClass || NativeParentClass->IsChildOf(BaseClass));
 	}
 
-	TArray<UClass*> GetBlueprintClasses(
-		UClass* BaseClass, bool bRecursive)
+	TArray<UClass*> GetBlueprintClasses(UClass* BaseClass, bool bRecursive)
 	{
 		const FAssetRegistryModule& AssetRegistryModule =
 			FModuleManager::LoadModuleChecked< FAssetRegistryModule >(FName("AssetRegistry"));
@@ -112,4 +113,32 @@ namespace UE::MovieGraph::Private::Tests
 
 		return DerivedClasses;
 	}
+	
+	void SuppressLogWarnings(FAutomationTestBase* InTestBase)
+	{
+		check(InTestBase);
+		InTestBase->bSuppressLogWarnings = true;
+	}
+
+	void SuppressLogErrors(FAutomationTestBase* InTestBase)
+	{
+		check(InTestBase);
+		InTestBase->bSuppressLogErrors = true;
+	}
+
+	void SetupTest(
+		FAutomationTestBase* InTestBase, const bool bSuppressLogWarnings, const bool bSuppressLogErrors)
+	{
+		check(InTestBase);
+		
+		if (bSuppressLogWarnings)
+		{
+			SuppressLogWarnings(InTestBase);
+		}
+
+		if (bSuppressLogErrors)
+		{
+			SuppressLogErrors(InTestBase);
+		}
+	}
 }
diff --git a/Engine/Plugins/MovieScene/MovieRenderPipeline/Source/MovieRenderPipelineEditor/Private/Tests/MovieGraphEditorTestUtilities.h b/Engine/Plugins/MovieScene/MovieRenderPipeline/Source/MovieRenderPipelineEditor/Private/Tests/MovieGraphEditorTestUtilities.h
--- a/Engine/Plugins/MovieScene/MovieRenderPipeline/Source/MovieRenderPipelineEditor/Private/Tests/MovieGraphEditorTestUtilities.h
+++ b/Engine/Plugins/MovieScene/MovieRenderPipeline/Source/MovieRenderPipelineEditor/Private/Tests/MovieGraphEditorTestUtilities.h
@@ -4,7 +4,9 @@
 
 #include "CoreMinimal.h"
 
+class FAutomationTestBase;
 class UMovieGraphConfig;
+class UMovieGraphNode;
 
 namespace UE::MovieGraph::Private::Tests
 {
@@ -14,4 +16,8 @@ namespace UE::MovieGraph::Private::Tests
 	TArray<UClass*> GetAllDerivedClasses(UClass* BaseClass, bool bRecursive);
 	TArray<UClass*> GetNativeClasses(UClass* BaseClass, bool bRecursive);
 	TArray<UClass*> GetBlueprintClasses(UClass* BaseClass, bool bRecursive);
+	void SuppressLogWarnings(FAutomationTestBase* InTestBase);
+	void SuppressLogErrors(FAutomationTestBase* InTestBase);
+	void SetupTest(
+		FAutomationTestBase* InTestBase, const bool bSuppressLogWarnings = true, const bool bSuppressLogErrors = true);
 }
diff --git a/Engine/Plugins/MovieScene/MovieRenderPipeline/Source/MovieRenderPipelineEditor/Private/Tests/MovieGraphTests.cpp b/Engine/Plugins/MovieScene/MovieRenderPipeline/Source/MovieRenderPipelineEditor/Private/Tests/MovieGraphTests.cpp
--- a/Engine/Plugins/MovieScene/MovieRenderPipeline/Source/MovieRenderPipelineEditor/Private/Tests/MovieGraphTests.cpp
+++ b/Engine/Plugins/MovieScene/MovieRenderPipeline/Source/MovieRenderPipelineEditor/Private/Tests/MovieGraphTests.cpp
@@ -2,14 +2,18 @@
 
 #include "MovieGraphEditorTestUtilities.h"
 
-#include "Algo/RemoveIf.h"
 #include "Graph/MovieGraphConfig.h"
+#include "Graph/Nodes/MovieGraphApplyCVarPresetNode.h"
 #include "Graph/Nodes/MovieGraphBranchNode.h"
+#include "Graph/Nodes/MovieGraphCommandLineEncoderNode.h"
 #include "Graph/Nodes/MovieGraphInputNode.h"
 #include "Graph/Nodes/MovieGraphOutputNode.h"
+#include "Graph/Nodes/MovieGraphSetCVarValueNode.h"
 #include "Graph/Nodes/MovieGraphSubgraphNode.h"
 #include "Graph/Nodes/MovieGraphVariableNode.h"
 
+#include "Algo/RemoveIf.h"
+#include "Graph/Nodes/MovieGraphSamplingMethodNode.h"
 #include "Misc/AutomationTest.h"
 
 namespace UE::MovieGraph::Private::Tests
@@ -30,6 +34,168 @@ namespace UE::MovieGraph::Private::Tests
 			});
 	}
 
+	void OpenGraphConfigInEditorIfDesired(UMovieGraphConfig* InGraphConfig)
+	{
+		/**
+		* A static switch which governs whether the graph config is opened
+		* useful for test development
+		* just remember to set it to false when done
+		*/
+		static bool bIsDesired = false;
+
+		if (bIsDesired)
+		{
+			OpenGraphConfigInEditor(InGraphConfig);
+		}
+	}
+
+	void RemoveAllEdgesFromAllNodes(UMovieGraphConfig* GraphConfig)
+	{
+		TArray<UMovieGraphNode*> Nodes = GraphConfig->GetNodes();
+		Nodes.Add(GraphConfig->GetInputNode());
+		Nodes.Add(GraphConfig->GetOutputNode());
+		
+		for (UMovieGraphNode* Node : Nodes)
+		{
+			GraphConfig->RemoveAllInboundEdges(Node);
+			GraphConfig->RemoveAllOutboundEdges(Node);
+		}
+	}
+
+	void AddVerifyAndRemoveVariableNode(
+		FAutomationTestBase* InTestBase,
+		UMovieGraphConfig* GraphConfig)
+	{
+		// Test adding and removing a variable member and variable node
+		UMovieGraphVariable* Variable = GraphConfig->AddVariable(TEXT("A_Cool_Variable"));
+		InTestBase->TestTrue(TEXT("Variable member added"), Variable != nullptr && GraphConfig->GetVariables().Contains(Variable));
+		if (!Variable)
+		{
+			return;
+		}
+		
+		UMovieGraphVariableNode* VariableNode = GraphConfig->ConstructRuntimeNode<UMovieGraphVariableNode>(UMovieGraphVariableNode::StaticClass());
+		InTestBase->TestTrue(TEXT("Variable node type added"), VariableNode != nullptr && GraphConfig->GetNodes().Contains(VariableNode));
+		if (!VariableNode)
+		{
+			return;
+		}
+		
+		VariableNode->SetVariable(Variable);
+		InTestBase->TestTrue(TEXT("Variable node variable member set"), VariableNode->GetVariable() == Variable);
+		
+		VariableNode->SetVariable(nullptr);
+		InTestBase->TestTrue(TEXT("Calling SetVariable with nullptr should do nothing"), VariableNode->GetVariable() == Variable);
+		
+		GraphConfig->RemoveNode(VariableNode);
+		InTestBase->TestTrue(TEXT("Variable node removed"), GraphConfig->GetNodes().Num() == 1);
+		
+		GraphConfig->DeleteMember(Variable);
+		InTestBase->TestTrue(TEXT("Variable member removed"), !GraphConfig->GetVariables().Contains(Variable));
+	}
+
+	void SeparateNodesIntoInputAndOutputTypes(
+		FAutomationTestBase* InTestBase,
+		const UMovieGraphConfig* GraphConfig,
+		TArray<TObjectPtr<UMovieGraphNode>>& OutNodesWithOutputs, 
+		TArray<TObjectPtr<UMovieGraphNode>>& OutNodesWithInputs)
+	{
+		for (const TObjectPtr<UMovieGraphNode>& Node : GraphConfig->GetNodes())
+		{
+			if (!Node->GetInputPins().IsEmpty())
+			{
+				OutNodesWithInputs.Add(Node);
+			}
+
+			if (!Node->GetOutputPins().IsEmpty())
+			{
+				OutNodesWithOutputs.Add(Node);
+			}
+
+			const bool bHasNoInputsOrOutputs = Node->GetInputPins().Num() == 0 && Node->GetOutputPins().Num() == 0;
+			
+			InTestBase->TestTrue(*FString::Printf(TEXT("Node %s has inputs or outputs"), *Node.GetName()), !bHasNoInputsOrOutputs);
+		}
+	};
+
+	void ConnectAndTestEdges_ExpectedSuccess(
+		FAutomationTestBase* InTestBase,
+		UMovieGraphConfig* GraphConfig,
+		UMovieGraphNode* NodeFrom, UMovieGraphNode* NodeTo,
+		UMovieGraphPin* OutputPin, UMovieGraphPin* InputPin)
+	{
+		GraphConfig->AddLabeledEdge(
+			NodeFrom, OutputPin->Properties.Label,
+			NodeTo, InputPin->Properties.Label);
+
+		InTestBase->TestTrue(
+			*FString::Printf(TEXT("Edge added between %s and %s"), *NodeFrom->GetName(), *NodeTo->GetName()),
+			OutputPin->GetAllConnectedPins().Contains(InputPin));
+	}
+	
+	void ConnectAndTestEdges_ExpectedFailure(
+		FAutomationTestBase* InTestBase,
+		UMovieGraphConfig* GraphConfig,
+		UMovieGraphNode* NodeFrom, UMovieGraphNode* NodeTo,
+		UMovieGraphPin* OutputPin, UMovieGraphPin* InputPin)
+	{
+		GraphConfig->AddLabeledEdge(
+			NodeFrom, OutputPin->Properties.Label,
+			NodeTo, InputPin->Properties.Label);
+
+		if (OutputPin->GetAllConnectedPins().Contains(InputPin))
+		{
+			InTestBase->AddError(
+				*FString::Printf(TEXT("Edge added between %s and %s, but edge connection was expected to fail."),
+					*NodeFrom->GetName(), *NodeTo->GetName()));
+		}
+	}
+
+	void RemoveEdgeAndVerify(
+		FAutomationTestBase* InTestBase,
+		UMovieGraphConfig* GraphConfig,
+		UMovieGraphNode* NodeFrom, UMovieGraphNode* NodeTo,
+		UMovieGraphPin* OutputPin, UMovieGraphPin* InputPin)
+	{
+		GraphConfig->RemoveEdge(NodeFrom, OutputPin->Properties.Label,
+			NodeTo, InputPin->Properties.Label);
+
+		InTestBase->TestTrue(
+			*FString::Printf(
+				TEXT("Edge removed between %s and %s"), *NodeFrom->GetName(), *NodeTo->GetName()),
+			!OutputPin->GetAllConnectedPins().Contains(InputPin));
+	}
+
+	void ConnectAllEdgesToAllNodes(
+		FAutomationTestBase* InTestBase, 
+		UMovieGraphConfig* GraphConfig,
+		const TArray<TObjectPtr<UMovieGraphNode>>& NodesWithOutputs, 
+		const TArray<TObjectPtr<UMovieGraphNode>>& NodesWithInputs)
+	{
+		for (const TObjectPtr<UMovieGraphNode>& NodeFrom : NodesWithOutputs)
+		{
+			for (const TObjectPtr<UMovieGraphPin>& OutputPin : NodeFrom->GetOutputPins())
+			{
+				for (const TObjectPtr<UMovieGraphNode>& NodeTo : NodesWithInputs)
+				{
+					if (NodeFrom == NodeTo) // Don't try to self-connect
+					{
+						continue;
+					}
+					
+					for (const TObjectPtr<UMovieGraphPin>& InputPin : NodeTo->GetInputPins())
+					{
+						if (OutputPin->CanCreateConnection(InputPin))
+						{
+							ConnectAndTestEdges_ExpectedSuccess(
+								InTestBase, GraphConfig, NodeFrom, NodeTo, OutputPin, InputPin);
+						}
+					}
+				}
+			}
+		}
+	};
+
 	/**
 	* 1) Add a null node, verify that NumNodes == 0
 	* 2) Add an unsupported node, verify that NumNodes == 0
@@ -39,6 +205,8 @@ namespace UE::MovieGraph::Private::Tests
 	IMPLEMENT_SIMPLE_AUTOMATION_TEST(FAddGraphNodeTest, "VirtualProduction.MovieGraph.AddGraphNodeTest", (EAutomationTestFlags::EditorContext | EAutomationTestFlags::ProductFilter));
 	bool FAddGraphNodeTest::RunTest(const FString& Parameters)
 	{
+		SetupTest(this);
+		
 		// Create Config
 		UMovieGraphConfig* GraphConfig = CreateNewMovieGraphConfig("AddGraphNodeTest");
 		TestTrue(TEXT("MovieGraphConfig successfully created"), GraphConfig != nullptr);
@@ -51,15 +219,21 @@ namespace UE::MovieGraph::Private::Tests
 		TestTrue(TEXT("Null Node not added"), GraphConfig->GetNodes().IsEmpty());
 
 		// Ensure input and output nodes can't be added via the API
-		const UMovieGraphInputNode* TestInputRuntimeNode = GraphConfig->ConstructRuntimeNode<UMovieGraphInputNode>(UMovieGraphInputNode::StaticClass());
+		const UMovieGraphInputNode* TestInputRuntimeNode = GraphConfig->ConstructRuntimeNode<UMovieGraphInputNode>(
+			UMovieGraphInputNode::StaticClass());
 		if(TestInputRuntimeNode && GraphConfig->GetNodes().Contains(TestInputRuntimeNode))
 		{
-			AddError(*FString::Printf(TEXT("Graph allowed user construction of %s"), *TestInputRuntimeNode->GetClass()->GetName()));
+			AddError(
+				*FString::Printf(
+					TEXT("Graph allowed user construction of %s"), *TestInputRuntimeNode->GetClass()->GetName()));
 		}
-		const UMovieGraphOutputNode* TestOutputRuntimeNode = GraphConfig->ConstructRuntimeNode<UMovieGraphOutputNode>(UMovieGraphOutputNode::StaticClass());
+		const UMovieGraphOutputNode* TestOutputRuntimeNode = GraphConfig->ConstructRuntimeNode<UMovieGraphOutputNode>(
+			UMovieGraphOutputNode::StaticClass());
 		if(TestOutputRuntimeNode && GraphConfig->GetNodes().Contains(TestOutputRuntimeNode))
 		{
-			AddError(*FString::Printf(TEXT("Graph allowed user construction of %s"), *TestOutputRuntimeNode->GetClass()->GetName()));
+			AddError(
+				*FString::Printf(
+					TEXT("Graph allowed user construction of %s"), *TestOutputRuntimeNode->GetClass()->GetName()));
 		}
 		
 		// Add one node of each type, excepting variables, inputs and outputs
@@ -69,17 +243,23 @@ namespace UE::MovieGraph::Private::Tests
 		{
 			GraphConfig->ConstructRuntimeNode<UMovieGraphNode>(Class);
 		}
-		TestTrue(TEXT("All supported nodes types added"), GraphConfig->GetNodes().Num() == AllNodeClasses.Num());
+		TestTrue(
+			TEXT("All supported nodes types added"), GraphConfig->GetNodes().Num() == AllNodeClasses.Num());
 
 		// Test adding a variable node
 		UMovieGraphVariable* Variable = GraphConfig->AddVariable(TEXT("A_Cool_Variable"));
-		TestTrue(TEXT("Variable member added"), Variable != nullptr && GraphConfig->GetVariables().Contains(Variable));
+		TestTrue(
+			TEXT("Variable member added"),
+			Variable != nullptr && GraphConfig->GetVariables().Contains(Variable));
 		if (!Variable)
 		{
 			return false;
 		}
-		UMovieGraphVariableNode* VariableNode = GraphConfig->ConstructRuntimeNode<UMovieGraphVariableNode>(UMovieGraphVariableNode::StaticClass());
-		TestTrue(TEXT("Variable node type added"), VariableNode != nullptr && GraphConfig->GetNodes().Contains(VariableNode));
+		UMovieGraphVariableNode* VariableNode = GraphConfig->ConstructRuntimeNode<UMovieGraphVariableNode>(
+			UMovieGraphVariableNode::StaticClass());
+		TestTrue(
+			TEXT("Variable node type added"),
+			VariableNode != nullptr && GraphConfig->GetNodes().Contains(VariableNode));
 		if (!VariableNode)
 		{
 			return false;
@@ -87,8 +267,7 @@ namespace UE::MovieGraph::Private::Tests
 		VariableNode->SetVariable(Variable);
 		TestTrue(TEXT("Variable node variable member set"), VariableNode->GetVariable() == Variable);
 
-		// Open graph config in editor if desired
-		//OpenGraphConfigInEditor(GraphConfig);
+		OpenGraphConfigInEditorIfDesired(GraphConfig);
 
 		return true;
 	}
@@ -103,6 +282,8 @@ namespace UE::MovieGraph::Private::Tests
 	IMPLEMENT_SIMPLE_AUTOMATION_TEST(FRemoveGraphNodeTest, "VirtualProduction.MovieGraph.RemoveGraphNodeTest", (EAutomationTestFlags::EditorContext | EAutomationTestFlags::ProductFilter));
 	bool FRemoveGraphNodeTest::RunTest(const FString& Parameters)
 	{
+		SetupTest(this);
+		
 		// Create Config
 		UMovieGraphConfig* GraphConfig = CreateNewMovieGraphConfig("RemoveGraphNodeTest");
 		TestTrue(TEXT("MovieGraphConfig successfully created"), GraphConfig != nullptr);
@@ -118,11 +299,12 @@ namespace UE::MovieGraph::Private::Tests
 		UMovieGraphNode* NodeTwo = GraphConfig->ConstructRuntimeNode<UMovieGraphNode>(AllNodeClasses[1]);
 		if (!NodeOne || !NodeTwo)
 		{
-			AddError(*FString::Printf(TEXT("Unable to create two nodes.")));
+			AddError(TEXT("Unable to create two nodes."));
 			return false;
 		}
 		TestTrue(
-			FString::Printf(TEXT("Node types added: %s and %s"), *AllNodeClasses[0]->GetName(), *AllNodeClasses[1]->GetName()),
+			FString::Printf(
+				TEXT("Node types added: %s and %s"), *AllNodeClasses[0]->GetName(), *AllNodeClasses[1]->GetName()),
 			GraphConfig->GetNodes().Num() == 2);
 
 		// Test removing one node
@@ -131,41 +313,19 @@ namespace UE::MovieGraph::Private::Tests
 
 		// Test removing a nullptr
 		GraphConfig->RemoveNode(nullptr);
-		TestTrue(TEXT("Calling RemoveNode with nullptr should do nothing"), GraphConfig->GetNodes().Num() == 1);
+		TestTrue(
+			TEXT("Calling RemoveNode with nullptr should do nothing"), GraphConfig->GetNodes().Num() == 1);
 
 		// Test removing a node which doesn't exist in the graph
 		GraphConfig->RemoveNode(NodeOne);
-		TestTrue(TEXT("Attempting to remove node which does not exist in the graph should do nothing"), GraphConfig->GetNodes().Num() == 1);
-		
+		TestTrue(
+			TEXT("Attempting to remove node which does not exist in the graph should do nothing"),
+			GraphConfig->GetNodes().Num() == 1);
+
 		// Test adding and removing a variable member and variable node
-		UMovieGraphVariable* Variable = GraphConfig->AddVariable(TEXT("A_Cool_Variable"));
-		TestTrue(TEXT("Variable member added"), Variable != nullptr && GraphConfig->GetVariables().Contains(Variable));
-		if (!Variable)
-		{
-			return false;
-		}
-		
-		UMovieGraphVariableNode* VariableNode = GraphConfig->ConstructRuntimeNode<UMovieGraphVariableNode>(UMovieGraphVariableNode::StaticClass());
-		TestTrue(TEXT("Variable node type added"), VariableNode != nullptr && GraphConfig->GetNodes().Contains(VariableNode));
-		if (!VariableNode)
-		{
-			return false;
-		}
-		
-		VariableNode->SetVariable(Variable);
-		TestTrue(TEXT("Variable node variable member set"), VariableNode->GetVariable() == Variable);
-		
-		VariableNode->SetVariable(nullptr);
-		TestTrue(TEXT("Calling SetVariable with nullptr should do nothing"), VariableNode->GetVariable() == Variable);
-		
-		GraphConfig->RemoveNode(VariableNode);
-		TestTrue(TEXT("Variable node removed"), GraphConfig->GetNodes().Num() == 1);
-		
-		GraphConfig->DeleteMember(Variable);
-		TestTrue(TEXT("Variable member removed"), !GraphConfig->GetVariables().Contains(Variable));
+		AddVerifyAndRemoveVariableNode(this, GraphConfig);
 
-		// Open graph config in editor if desired
-		//OpenGraphConfigInEditor(GraphConfig);
+		OpenGraphConfigInEditorIfDesired(GraphConfig);
 
 		return true;
 	}
@@ -178,6 +338,8 @@ namespace UE::MovieGraph::Private::Tests
 	IMPLEMENT_SIMPLE_AUTOMATION_TEST(FGraphNodeAddEdgeTest, "VirtualProduction.MovieGraph.GraphNodeAddEdgeTest", (EAutomationTestFlags::EditorContext | EAutomationTestFlags::ProductFilter));
 	bool FGraphNodeAddEdgeTest::RunTest(const FString& Parameters)
 	{
+		SetupTest(this);
+		
 		// Create Config
 		UMovieGraphConfig* GraphConfig = CreateNewMovieGraphConfig("GraphNodeAddEdgeTest");
 		TestTrue(TEXT("MovieGraphConfig successfully created"), GraphConfig != nullptr);
@@ -193,7 +355,8 @@ namespace UE::MovieGraph::Private::Tests
 		{
 			GraphConfig->ConstructRuntimeNode<UMovieGraphNode>(Class);
 		}
-		TestTrue(TEXT("All supported nodes types added"), GraphConfig->GetNodes().Num() == AllNodeClasses.Num());
+		TestTrue(
+			TEXT("All supported nodes types added"), GraphConfig->GetNodes().Num() == AllNodeClasses.Num());
 
 		// Connect all edges		
 		// Some nodes will be duplicated in between these arrays
@@ -201,60 +364,10 @@ namespace UE::MovieGraph::Private::Tests
 		TArray<TObjectPtr<UMovieGraphNode>> NodesWithOutputs = {GraphConfig->GetInputNode()};
 		TArray<TObjectPtr<UMovieGraphNode>> NodesWithInputs = {GraphConfig->GetOutputNode()};
 		
-		auto SeparateNodesIntoInputAndOutputTypes = [&]()
-		{
-			for (const TObjectPtr<UMovieGraphNode>& Node : GraphConfig->GetNodes())
-			{
-				if (!Node->GetInputPins().IsEmpty())
-				{
-					NodesWithInputs.Add(Node);
-				}
-
-				if (!Node->GetOutputPins().IsEmpty())
-				{
-					NodesWithOutputs.Add(Node);
-				}
-
-				const bool bHasNoInputsOrOutputs = Node->GetInputPins().Num() == 0 && Node->GetOutputPins().Num() == 0;
-			
-				TestTrue(*FString::Printf(TEXT("Node %s has inputs or outputs"), *Node.GetName()), !bHasNoInputsOrOutputs);
-			}
-		};
-		SeparateNodesIntoInputAndOutputTypes();
-
-		auto ConnectAllEdgesToAllNodes = [&]()
-		{
-			for (const TObjectPtr<UMovieGraphNode>& NodeFrom : NodesWithOutputs)
-			{
-				for (const TObjectPtr<UMovieGraphPin>& OutputPin : NodeFrom->GetOutputPins())
-				{
-					for (const TObjectPtr<UMovieGraphNode>& NodeTo : NodesWithInputs)
-					{
-						if (NodeFrom == NodeTo) // Don't try to self-connect
-						{
-							continue;
-						}
-					
-						for (const TObjectPtr<UMovieGraphPin>& InputPin : NodeTo->GetInputPins())
-						{
-							if (OutputPin->CanCreateConnection(InputPin))
-							{
-								GraphConfig->AddLabeledEdge(
-									NodeFrom, OutputPin->Properties.Label,
-									NodeTo, InputPin->Properties.Label);
-
-								TestTrue(*FString::Printf(TEXT("Edge added between %s and %s"), *NodeFrom->GetName(), *NodeTo->GetName()),
-									OutputPin->GetAllConnectedPins().Contains(InputPin));
-							}
-						}
-					}
-				}
-			}
-		};
-		ConnectAllEdgesToAllNodes();
+		SeparateNodesIntoInputAndOutputTypes(this, GraphConfig, NodesWithOutputs, NodesWithInputs);
+		ConnectAllEdgesToAllNodes(this, GraphConfig, NodesWithOutputs, NodesWithInputs);
 		
-		// Open graph config in editor if desired
-		//OpenGraphConfigInEditor(GraphConfig);
+		OpenGraphConfigInEditorIfDesired(GraphConfig);
 		
 		return true;
 	}
@@ -267,6 +380,8 @@ namespace UE::MovieGraph::Private::Tests
 	IMPLEMENT_SIMPLE_AUTOMATION_TEST(FGraphNodeRemoveEdgeTest, "VirtualProduction.MovieGraph.GraphNodeRemoveEdgeTest", (EAutomationTestFlags::EditorContext | EAutomationTestFlags::ProductFilter));
 	bool FGraphNodeRemoveEdgeTest::RunTest(const FString& Parameters)
 	{
+		SetupTest(this);
+		
 		// Create Config
 		UMovieGraphConfig* GraphConfig = CreateNewMovieGraphConfig("GraphNodeRemoveEdgeTest");
 		TestTrue(TEXT("MovieGraphConfig successfully created"), GraphConfig != nullptr);
@@ -288,7 +403,7 @@ namespace UE::MovieGraph::Private::Tests
 		UMovieGraphNode* OutputNode = GraphConfig->GetOutputNode();
 		if (!InputNode || !OutputNode)
 		{
-			AddError(*FString::Printf(TEXT("Unable to find Input and Output nodes.")));
+			AddError(TEXT("Unable to find Input and Output nodes."));
 			return false;
 		}
 
@@ -299,43 +414,409 @@ namespace UE::MovieGraph::Private::Tests
 
 		if (!InputNodeOutputPin || !MiddleNodeInputPin || !MiddleNodeOutputPin || !OutputNodeInputPin)
 		{
-			AddError(*FString::Printf(TEXT("Unable to find all Pins for Input, Middle and Output Nodes.")));
+			AddError(TEXT("Unable to find all Pins for Input, Middle and Output Nodes."));
 			return false;
 		}
 
 		// Add first edge
-		GraphConfig->AddLabeledEdge(
-			InputNode, InputNodeOutputPin->Properties.Label,
-			MiddleNode, MiddleNodeInputPin->Properties.Label);
+		ConnectAndTestEdges_ExpectedSuccess(
+			this, GraphConfig, InputNode, MiddleNode, InputNodeOutputPin, MiddleNodeInputPin);
+
+		// Remove edge
+		RemoveEdgeAndVerify(
+			this, GraphConfig, InputNode, MiddleNode,
+			InputNodeOutputPin, MiddleNodeInputPin);
 
-		TestTrue(*FString::Printf(TEXT("Edge added between %s and %s"), *InputNode->GetName(), *MiddleNode->GetName()),
-			InputNodeOutputPin->GetAllConnectedPins().Contains(MiddleNodeInputPin));
+		// Add second edge
+		ConnectAndTestEdges_ExpectedSuccess(
+			this, GraphConfig, MiddleNode, OutputNode, MiddleNodeOutputPin, OutputNodeInputPin);
 
 		// Remove edge
-		GraphConfig->RemoveEdge(InputNode, InputNodeOutputPin->Properties.Label,
-			MiddleNode, MiddleNodeInputPin->Properties.Label);
+		RemoveEdgeAndVerify(
+			this, GraphConfig, MiddleNode, OutputNode,
+			MiddleNodeOutputPin, OutputNodeInputPin);
 
-		TestTrue(*FString::Printf(TEXT("Edge removed between %s and %s"), *InputNode->GetName(), *MiddleNode->GetName()),
-			!InputNodeOutputPin->GetAllConnectedPins().Contains(MiddleNodeInputPin));
+		OpenGraphConfigInEditorIfDesired(GraphConfig);
+		
+		return true;
+	}
 
-		// Add second edge
-		GraphConfig->AddLabeledEdge(
-			MiddleNode, MiddleNodeOutputPin->Properties.Label,
-			OutputNode, OutputNodeInputPin->Properties.Label);
+	/**
+	* Tests Node to Node Branch Connection without connection to input or output:
+	* Any-> Any, Any->Global, Any->Render Layer,  (all expected to work),
+	* Globals-> Globals (should work), Globals-> Any (should work), Globals-> RenderLayer (should fail),
+	* RenderLayer-> RenderLayer (should work), RenderLayer-> Any (should work), RenderLayer-> Globals (should fail)
+	*/
+	IMPLEMENT_SIMPLE_AUTOMATION_TEST(FGraphNodeBranchRestrictionIndividualNodeTest, "VirtualProduction.MovieGraph.GraphNodeBranchRestrictionIndividualNodeTest", (EAutomationTestFlags::EditorContext | EAutomationTestFlags::ProductFilter));
+	bool FGraphNodeBranchRestrictionIndividualNodeTest::RunTest(const FString& Parameters)
+	{
+		SetupTest(this);
+		
+		// Globals: UMovieGraphSetCVarValueNode, UMovieGraphApplyCVarPresetNode, UMovieGraphSamplingMethodNode,
+		// UMovieGraphAudioOutputNode, UMovieGraphWarmUpSettingNode
+		// RenderLayer: UMovieGraphCommandLineEncoderNode
+		// Any: All other node types as of October 6, 2023
+		
+		// Create Config
+		UMovieGraphConfig* GraphConfig = CreateNewMovieGraphConfig("GraphNodeRemoveEdgeTest");
+		TestTrue(TEXT("MovieGraphConfig successfully created"), GraphConfig != nullptr);
+		if (!GraphConfig)
+		{
+			return false;
+		}
+
+		// Add middle nodes of BranchRestriction type 'Any'
+		UMovieGraphNode* AnyNodeA =
+			GraphConfig->ConstructRuntimeNode<UMovieGraphBranchNode>(
+				UMovieGraphBranchNode::StaticClass());
+		TestTrue(
+			TEXT("AnyNodeA successfully added"), GraphConfig->GetNodes().Contains(AnyNodeA));
+		if (!AnyNodeA)
+		{
+			return false;
+		}
+		UMovieGraphNode* AnyNodeB =
+			GraphConfig->ConstructRuntimeNode<UMovieGraphSamplingMethodNode>(
+				UMovieGraphSamplingMethodNode::StaticClass());
+		TestTrue(
+			TEXT("AnyNodeB successfully added"), GraphConfig->GetNodes().Contains(AnyNodeB));
+		if (!AnyNodeB)
+		{
+			return false;
+		}
 
-		TestTrue(*FString::Printf(TEXT("Edge added between %s and %s"), *MiddleNode->GetName(), *OutputNode->GetName()),
-			MiddleNodeOutputPin->GetAllConnectedPins().Contains(OutputNodeInputPin));
+		// Add middle nodes of BranchRestriction type 'Globals'
+		UMovieGraphNode* GlobalsNodeA =
+			GraphConfig->ConstructRuntimeNode<UMovieGraphSetCVarValueNode>(
+				UMovieGraphSetCVarValueNode::StaticClass());
+		TestTrue(
+			TEXT("GlobalsNodeA successfully added"), GraphConfig->GetNodes().Contains(GlobalsNodeA));
+		if (!GlobalsNodeA)
+		{
+			return false;
+		}
+		UMovieGraphNode* GlobalsNodeB =
+			GraphConfig->ConstructRuntimeNode<UMovieGraphApplyCVarPresetNode>(
+				UMovieGraphApplyCVarPresetNode::StaticClass());
+		TestTrue(
+			TEXT("GlobalsNodeB successfully added"), GraphConfig->GetNodes().Contains(GlobalsNodeB));
+		if (!GlobalsNodeB)
+		{
+			return false;
+		}
 
-		// Remove edge
-		GraphConfig->RemoveEdge(
-			MiddleNode, MiddleNodeOutputPin->Properties.Label,
-			OutputNode, OutputNodeInputPin->Properties.Label);
+		// Add middle nodes of BranchRestriction type 'RenderLayer'
+		// Ideally we'd want two different types but only one exists at present
+		UMovieGraphNode* RenderLayerNodeA =
+			GraphConfig->ConstructRuntimeNode<UMovieGraphCommandLineEncoderNode>(
+				UMovieGraphCommandLineEncoderNode::StaticClass());
+		TestTrue(
+			TEXT("RenderLayerNodeA successfully added"), GraphConfig->GetNodes().Contains(RenderLayerNodeA));
+		if (!RenderLayerNodeA)
+		{
+			return false;
+		}
+		UMovieGraphNode* RenderLayerNodeB =
+			GraphConfig->ConstructRuntimeNode<UMovieGraphCommandLineEncoderNode>(
+				UMovieGraphCommandLineEncoderNode::StaticClass());
+		TestTrue(
+			TEXT("RenderLayerNodeB successfully added"), GraphConfig->GetNodes().Contains(RenderLayerNodeB));
+		if (!RenderLayerNodeB)
+		{
+			return false;
+		}
+
+		// Identify Pins
+		// Need two nodes of each Type
+		const TObjectPtr<UMovieGraphPin> AnyNodeInputPinA = AnyNodeA->GetInputPins()[0];
+		const TObjectPtr<UMovieGraphPin> AnyNodeOutputPinA = AnyNodeA->GetOutputPins()[0];
+		
+		const TObjectPtr<UMovieGraphPin> AnyNodeInputPinB = AnyNodeB->GetInputPins()[0];
+		const TObjectPtr<UMovieGraphPin> AnyNodeOutputPinB = AnyNodeB->GetOutputPins()[0];
+
+		const TObjectPtr<UMovieGraphPin> GlobalsNodeInputPinA = GlobalsNodeA->GetInputPins()[0];
+		const TObjectPtr<UMovieGraphPin> GlobalsNodeOutputPinA = GlobalsNodeA->GetOutputPins()[0];
+		
+		const TObjectPtr<UMovieGraphPin> GlobalsNodeInputPinB = GlobalsNodeB->GetInputPins()[0];
+		const TObjectPtr<UMovieGraphPin> GlobalsNodeOutputPinB = GlobalsNodeB->GetOutputPins()[0];
+
+		const TObjectPtr<UMovieGraphPin> RenderLayerNodeInputPinA = RenderLayerNodeA->GetInputPins()[0];
+		const TObjectPtr<UMovieGraphPin> RenderLayerNodeOutputPinA = RenderLayerNodeA->GetOutputPins()[0];
+		
+		const TObjectPtr<UMovieGraphPin> RenderLayerNodeInputPinB = RenderLayerNodeB->GetInputPins()[0];
+		const TObjectPtr<UMovieGraphPin> RenderLayerNodeOutputPinB = RenderLayerNodeB->GetOutputPins()[0];
+
+		if (!AnyNodeInputPinA || !AnyNodeOutputPinA ||
+			!AnyNodeInputPinB || !AnyNodeOutputPinB ||
+			!GlobalsNodeInputPinA || !GlobalsNodeOutputPinA ||
+			!GlobalsNodeInputPinB || !GlobalsNodeOutputPinB ||
+			!RenderLayerNodeInputPinA || !RenderLayerNodeOutputPinA ||
+			!RenderLayerNodeInputPinB || !RenderLayerNodeOutputPinB)
+		{
+			AddError(TEXT("Unable to find all Pins for each Node."));
+			return false;
+		}
+
+		// Any-> Any, Any-> Global, Any-> RenderLayer,  (all expected to work),
+		// This connects AnyNodeA to AnyNodeB, GlobalsNodeA, and RenderLayerNodeA simultaneously
+
+		// Any-> Any
+		ConnectAndTestEdges_ExpectedSuccess(
+			this, GraphConfig, AnyNodeA, AnyNodeB,
+			AnyNodeOutputPinA, AnyNodeInputPinB);
 
-		TestTrue(*FString::Printf(TEXT("Edge added between %s and %s"), *MiddleNode->GetName(), *OutputNode->GetName()),
-			!MiddleNodeOutputPin->GetAllConnectedPins().Contains(OutputNodeInputPin));
+		// Any-> Global
+		ConnectAndTestEdges_ExpectedSuccess(
+			this, GraphConfig, AnyNodeA, GlobalsNodeA,
+			AnyNodeOutputPinA, GlobalsNodeInputPinA);
+
+		// Any-> RenderLayer
+		ConnectAndTestEdges_ExpectedSuccess(
+			this, GraphConfig, AnyNodeA, RenderLayerNodeA,
+			AnyNodeOutputPinA, RenderLayerNodeInputPinA);
+
+		// Globals-> Globals (should work), Globals-> Any (should work), Globals-> RenderLayer (should fail)
+		// This connects GlobalsNodeA to GlobalsNodeB and AnyNodeB simultaneously.
+		// It attempts a connection to RenderLayerNodeA but fails.
+		// AnyNodeB's connection to AnyNodeA will be severed.
+		
+		// Globals-> Globals
+		ConnectAndTestEdges_ExpectedSuccess(
+			this, GraphConfig, GlobalsNodeA, GlobalsNodeB,
+			GlobalsNodeOutputPinA, GlobalsNodeInputPinB);
+
+		// Globals-> Any
+		ConnectAndTestEdges_ExpectedSuccess(
+			this, GraphConfig, GlobalsNodeA, AnyNodeB,
+			GlobalsNodeOutputPinA, AnyNodeInputPinB);
+
+		// Globals-> RenderLayer
+		ConnectAndTestEdges_ExpectedFailure(
+			this, GraphConfig, GlobalsNodeA, RenderLayerNodeA,
+			GlobalsNodeOutputPinA, RenderLayerNodeInputPinA);
+		
+		// RenderLayer-> RenderLayer (should work), RenderLayer-> Any (should work), RenderLayer-> Globals (should fail)
+		// This connects RenderLayerNodeA to RenderLayerNodeB and AnyNodeB simultaneously.
+		// It attempts a connection to GlobalsNodeB but fails.
+		// AnyNodeB's connection to GlobalsNodeA will be severed.
+		
+		// RenderLayer-> RenderLayer
+		ConnectAndTestEdges_ExpectedSuccess(
+			this, GraphConfig, RenderLayerNodeA, RenderLayerNodeB,
+			RenderLayerNodeOutputPinA, RenderLayerNodeInputPinB);
+
+		// RenderLayer-> Any
+		ConnectAndTestEdges_ExpectedSuccess(
+			this, GraphConfig, RenderLayerNodeA, AnyNodeB,
+			RenderLayerNodeOutputPinA, AnyNodeInputPinB);
+
+		// RenderLayer-> Globals
+		ConnectAndTestEdges_ExpectedFailure(
+			this, GraphConfig, RenderLayerNodeA, GlobalsNodeB,
+			RenderLayerNodeOutputPinA, GlobalsNodeInputPinB);
+
+		OpenGraphConfigInEditorIfDesired(GraphConfig);
+		
+		return true;
+	}
+
+	/**
+	* Tests Node to Chain Branch Connection:
+	* Input Globals-> Globals Node-> Any Node-> Output Globals (should work),
+	* Input Globals-> Globals Node-> Any Node-> Output RenderLayer (should fail),
+	* Input Globals-> Globals Node-> Any Node-> RenderLayer Node (should fail),
+	* Input Globals-> RenderLayer Node (should fail),
+	* Unconnected Globals Node-> Output Globals (should work),
+	* Unconnected Globals Node-> Output RenderLayer (should fail),
+	* Input RenderLayer-> RenderLayer Node-> Any Node-> Output RenderLayer (should work),
+	* Input RenderLayer-> RenderLayer Node-> Any Node-> Output Globals (should fail),
+	* Input RenderLayer-> RenderLayer Node-> Any Node-> Globals Node (should fail),
+	* Input RenderLayer-> Globals Node (should fail),
+	* Unconnected RenderLayer Node-> Output Globals (should fail),
+	* Unconnected RenderLayer Node-> Output RenderLayer (should work)
+	*/
+	IMPLEMENT_SIMPLE_AUTOMATION_TEST(FGraphNodeBranchRestrictionUpstreamAndDownstreamTest, "VirtualProduction.MovieGraph.GraphNodeBranchRestrictionUpstreamAndDownstreamTest", (EAutomationTestFlags::EditorContext | EAutomationTestFlags::ProductFilter));
+	bool FGraphNodeBranchRestrictionUpstreamAndDownstreamTest::RunTest(const FString& Parameters)
+	{
+		SetupTest(this);
+		
+		// Globals: UMovieGraphSetCVarValueNode, UMovieGraphApplyCVarPresetNode, UMovieGraphSamplingMethodNode,
+		// UMovieGraphAudioOutputNode, UMovieGraphWarmUpSettingNode
+		// RenderLayer: UMovieGraphCommandLineEncoderNode
+		// Any: All other node types as of October 5, 2023
+		
+		// Create Config
+		UMovieGraphConfig* GraphConfig = CreateNewMovieGraphConfig("GraphNodeRemoveEdgeTest");
+		TestTrue(TEXT("MovieGraphConfig successfully created"), GraphConfig != nullptr);
+		if (!GraphConfig)
+		{
+			return false;
+		}
+
+		// Add inputs and outputs for a render layer
+		const UMovieGraphInput* NewInput = GraphConfig->AddInput();
+		const UMovieGraphOutput* NewOutput = GraphConfig->AddOutput();
+		if (!NewInput || !NewOutput)
+		{
+			AddError(TEXT("Unable to create new Input and Output."));
+			return false;
+		}
+
+		UMovieGraphNode* InputNode = GraphConfig->GetInputNode();
+		UMovieGraphNode* OutputNode = GraphConfig->GetOutputNode();
+		if (!InputNode || !OutputNode)
+		{
+			AddError(TEXT("Unable to find Input and Output nodes."));
+			return false;
+		}
+
+		// Add middle nodes with an input and output for connection
+		UMovieGraphNode* AnyNode =
+			GraphConfig->ConstructRuntimeNode<UMovieGraphBranchNode>(UMovieGraphBranchNode::StaticClass());
+		TestTrue(TEXT("AnyNode successfully added"), GraphConfig->GetNodes().Contains(AnyNode));
+		if (!AnyNode)
+		{
+			return false;
+		}
+
+		UMovieGraphNode* GlobalsNode =
+			GraphConfig->ConstructRuntimeNode<UMovieGraphSetCVarValueNode>(UMovieGraphSetCVarValueNode::StaticClass());
+		TestTrue(TEXT("GlobalsNode successfully added"), GraphConfig->GetNodes().Contains(GlobalsNode));
+		if (!GlobalsNode)
+		{
+			return false;
+		}
+
+		UMovieGraphNode* RenderLayerNode =
+			GraphConfig->ConstructRuntimeNode<UMovieGraphCommandLineEncoderNode>(
+				UMovieGraphCommandLineEncoderNode::StaticClass());
+		TestTrue(
+			TEXT("RenderLayerNode successfully added"), GraphConfig->GetNodes().Contains(RenderLayerNode));
+		if (!RenderLayerNode)
+		{
+			return false;
+		}
+
+		// Identify Pins
+		const TObjectPtr<UMovieGraphPin> InputNodeGlobalsOutputPin = InputNode->GetOutputPins()[0];
+		const TObjectPtr<UMovieGraphPin> InputNodeRenderLayerOutputPin = InputNode->GetOutputPins()[1];
+		
+		const TObjectPtr<UMovieGraphPin> AnyNodeInputPin = AnyNode->GetInputPins()[0];
+		const TObjectPtr<UMovieGraphPin> AnyNodeOutputPin = AnyNode->GetOutputPins()[0];
+
+		const TObjectPtr<UMovieGraphPin> GlobalsNodeInputPin = GlobalsNode->GetInputPins()[0];
+		const TObjectPtr<UMovieGraphPin> GlobalsNodeOutputPin = GlobalsNode->GetOutputPins()[0];
+
+		const TObjectPtr<UMovieGraphPin> RenderLayerNodeInputPin = RenderLayerNode->GetInputPins()[0];
+		const TObjectPtr<UMovieGraphPin> RenderLayerNodeOutputPin = RenderLayerNode->GetOutputPins()[0];
+		
+		const TObjectPtr<UMovieGraphPin> OutputNodeGlobalsInputPin = OutputNode->GetInputPins()[0];
+		const TObjectPtr<UMovieGraphPin> OutputNodeRenderLayerInputPin = OutputNode->GetInputPins()[1];
+
+		if (!InputNodeGlobalsOutputPin || !InputNodeRenderLayerOutputPin ||!AnyNodeInputPin ||
+			!AnyNodeOutputPin || !GlobalsNodeInputPin || !GlobalsNodeOutputPin ||
+			!RenderLayerNodeInputPin || !RenderLayerNodeOutputPin ||
+			!OutputNodeGlobalsInputPin || !OutputNodeRenderLayerInputPin)
+		{
+			AddError(TEXT("Unable to find all Pins for Input, Output and other Nodes."));
+			return false;
+		}
+
+		/**
+		* Input Globals-> Globals Node-> Any Node-> Output Globals (should work),
+		* 
+		* Input Globals-> Globals Node-> Any Node-> Output RenderLayer (should fail),
+		* Input Globals-> Globals Node-> Any Node-> RenderLayer Node (should fail),
+		* 
+		* Input Globals-> RenderLayer Node (should fail)
+		*/
+		ConnectAndTestEdges_ExpectedSuccess(
+			this, GraphConfig, InputNode, GlobalsNode,
+			InputNodeGlobalsOutputPin, GlobalsNodeInputPin);
+		ConnectAndTestEdges_ExpectedSuccess(
+			this, GraphConfig, GlobalsNode, AnyNode,
+			GlobalsNodeOutputPin, AnyNodeInputPin);
+		ConnectAndTestEdges_ExpectedSuccess(
+			this, GraphConfig, AnyNode, OutputNode,
+			AnyNodeOutputPin, OutputNodeGlobalsInputPin);
+		
+		ConnectAndTestEdges_ExpectedFailure(
+			this, GraphConfig, AnyNode, OutputNode,
+			AnyNodeOutputPin, OutputNodeRenderLayerInputPin);
+		ConnectAndTestEdges_ExpectedFailure(
+			this, GraphConfig, AnyNode, RenderLayerNode,
+			AnyNodeOutputPin, RenderLayerNodeInputPin);
+
+		RemoveAllEdgesFromAllNodes(GraphConfig);
+		
+		ConnectAndTestEdges_ExpectedFailure(
+			this, GraphConfig, InputNode, RenderLayerNode,
+			InputNodeGlobalsOutputPin, RenderLayerNodeInputPin);
+
+		RemoveAllEdgesFromAllNodes(GraphConfig);
+
+		/**
+		* Unconnected Globals Node-> Output Globals (should work),
+		* Unconnected Globals Node-> Output RenderLayer (should fail)
+		*/
+		ConnectAndTestEdges_ExpectedSuccess(
+			this, GraphConfig, GlobalsNode, OutputNode,
+			GlobalsNodeOutputPin, OutputNodeGlobalsInputPin);
+		
+		RemoveAllEdgesFromAllNodes(GraphConfig);
+		
+		ConnectAndTestEdges_ExpectedFailure(
+			this, GraphConfig, GlobalsNode, OutputNode,
+			GlobalsNodeOutputPin, OutputNodeRenderLayerInputPin);
+
+		/**
+		* Input RenderLayer-> RenderLayer Node-> Any Node-> Output RenderLayer (should work),
+		* 
+		* Input RenderLayer-> RenderLayer Node-> Any Node-> Output Globals (should fail),
+		* Input RenderLayer-> RenderLayer Node-> Any Node-> Globals Node (should fail),
+		* 
+		* Input RenderLayer-> Globals Node (should fail)
+		*/
+		ConnectAndTestEdges_ExpectedSuccess(
+			this, GraphConfig, InputNode, RenderLayerNode,
+			InputNodeRenderLayerOutputPin, RenderLayerNodeInputPin);
+		ConnectAndTestEdges_ExpectedSuccess(
+			this, GraphConfig, RenderLayerNode, AnyNode,
+			RenderLayerNodeOutputPin, AnyNodeInputPin);
+		ConnectAndTestEdges_ExpectedSuccess(
+			this, GraphConfig, AnyNode, OutputNode,
+			AnyNodeOutputPin, OutputNodeRenderLayerInputPin);
+		
+		ConnectAndTestEdges_ExpectedFailure(
+			this, GraphConfig, AnyNode, OutputNode,
+			AnyNodeOutputPin, OutputNodeGlobalsInputPin);
+		ConnectAndTestEdges_ExpectedFailure(
+			this, GraphConfig, AnyNode, GlobalsNode,
+			AnyNodeOutputPin, GlobalsNodeInputPin);
+
+		RemoveAllEdgesFromAllNodes(GraphConfig);
+		
+		ConnectAndTestEdges_ExpectedFailure(
+			this, GraphConfig, InputNode, GlobalsNode,
+			InputNodeRenderLayerOutputPin, GlobalsNodeInputPin);
+
+		RemoveAllEdgesFromAllNodes(GraphConfig);
+
+		/**
+		* Unconnected RenderLayer Node-> Output Globals (should fail),
+		* Unconnected RenderLayer Node-> Output RenderLayer (should work)
+		*/
+		ConnectAndTestEdges_ExpectedFailure(
+			this, GraphConfig, RenderLayerNode, OutputNode,
+			RenderLayerNodeOutputPin, OutputNodeGlobalsInputPin);
+		
+		RemoveAllEdgesFromAllNodes(GraphConfig);
+		
+		ConnectAndTestEdges_ExpectedSuccess(
+			this, GraphConfig, RenderLayerNode, OutputNode,
+			RenderLayerNodeOutputPin, OutputNodeRenderLayerInputPin);
 
-		// Open graph config in editor if desired
-		//OpenGraphConfigInEditor(GraphConfig);
+		OpenGraphConfigInEditorIfDesired(GraphConfig);
 		
 		return true;
 	}
