MRG: Moved OCIO configuration to file output node.

#rnx

diff --git a/Engine/Plugins/MovieScene/MovieRenderPipeline/Source/MovieRenderPipelineCore/Public/Graph/Nodes/MovieGraphOutputSettingNode.h b/Engine/Plugins/MovieScene/MovieRenderPipeline/Source/MovieRenderPipelineCore/Public/Graph/Nodes/MovieGraphOutputSettingNode.h
--- a/Engine/Plugins/MovieScene/MovieRenderPipeline/Source/MovieRenderPipelineCore/Public/Graph/Nodes/MovieGraphOutputSettingNode.h
+++ b/Engine/Plugins/MovieScene/MovieRenderPipeline/Source/MovieRenderPipelineCore/Public/Graph/Nodes/MovieGraphOutputSettingNode.h
@@ -3,7 +3,6 @@
 #pragma once
 #include "Graph/MovieGraphNode.h"
 #include "Misc/FrameRate.h"
-#include "OpenColorIOColorSpace.h"
 #include "MovieGraphOutputSettingNode.generated.h"
 
 
@@ -56,11 +55,6 @@ public:
 	UPROPERTY(EditAnywhere, BlueprintReadWrite, Category = Overrides, meta = (InlineEditConditionToggle))
 	uint8 bOverride_VersionNumber : 1;
 
-#if WITH_EDITORONLY_DATA
-	UPROPERTY(EditAnywhere, BlueprintReadWrite, Category = Overrides, meta = (InlineEditConditionToggle))
-	uint8 bOverride_OCIOConfiguration : 1;
-#endif
-
 	/** What directory should all of our output files be relative to. */
 	UPROPERTY(EditAnywhere, BlueprintReadWrite, Category = "File Output", meta=(EditCondition="bOverride_OutputDirectory"))
 	FDirectoryPath OutputDirectory;
@@ -116,16 +110,4 @@ public:
 	/** The value to use for the version token if versions are not automatically incremented (Auto Version is off). */
 	UPROPERTY(EditAnywhere, BlueprintReadWrite, Category = "Versioning", meta = (EditCondition = "bOverride_VersionNumber", UIMin = 1, UIMax = 50, ClampMin = 1))
 	int32 VersionNumber;
-
-#if WITH_EDITORONLY_DATA
-	/**
-	* OCIO configuration/transform settings.
-	* 
-	* Note: There are differences from the previous implementation in MRQ given that we are now doing CPU-side processing.
-	* 1) This feature only works in editor-mode, since the OpenColorIO library is currently unavailable in game builds.
-	* 2) Users are now responsible for setting the renderer output space to Final Color (HDR) in Linear Working Color Space (SCS_FinalColorHDR).
-	*/
-	UPROPERTY(EditAnywhere, BlueprintReadWrite, Category = "File Output", meta = (EditCondition = "bOverride_OCIOConfiguration"))
-	FOpenColorIODisplayConfiguration OCIOConfiguration;
-#endif
 };
diff --git a/Engine/Plugins/MovieScene/MovieRenderPipeline/Source/MovieRenderPipelineRenderPasses/Private/MovieGraphImageSequenceOutputNode.cpp b/Engine/Plugins/MovieScene/MovieRenderPipeline/Source/MovieRenderPipelineRenderPasses/Private/MovieGraphImageSequenceOutputNode.cpp
--- a/Engine/Plugins/MovieScene/MovieRenderPipeline/Source/MovieRenderPipelineRenderPasses/Private/MovieGraphImageSequenceOutputNode.cpp
+++ b/Engine/Plugins/MovieScene/MovieRenderPipeline/Source/MovieRenderPipelineRenderPasses/Private/MovieGraphImageSequenceOutputNode.cpp
@@ -241,9 +241,9 @@ void UMovieGraphImageSequenceOutputNode::OnReceiveImageDataImpl(UMovieGraphPipel
 		TileImageTask->PixelData = RenderData.Value->CopyImageData();
 
 #if WITH_EDITOR
-		if (OutputSettingNode->OCIOConfiguration.bIsEnabled)
+		if (OCIOConfiguration.bIsEnabled)
 		{
-			FPixelPreProcessor OCIOPixelPreProcessor = UE::MovieGraph::Private::CreateOpenColorIOPixelPreProcessor(OutputSettingNode->OCIOConfiguration.ColorConfiguration);
+			FPixelPreProcessor OCIOPixelPreProcessor = UE::MovieGraph::Private::CreateOpenColorIOPixelPreProcessor(OCIOConfiguration.ColorConfiguration);
 			if (OCIOPixelPreProcessor)
 			{
 				TileImageTask->PixelPreProcessors.Emplace(MoveTemp(OCIOPixelPreProcessor));
diff --git a/Engine/Plugins/MovieScene/MovieRenderPipeline/Source/MovieRenderPipelineRenderPasses/Public/MovieGraphImageSequenceOutputNode.h b/Engine/Plugins/MovieScene/MovieRenderPipeline/Source/MovieRenderPipelineRenderPasses/Public/MovieGraphImageSequenceOutputNode.h
--- a/Engine/Plugins/MovieScene/MovieRenderPipeline/Source/MovieRenderPipelineRenderPasses/Public/MovieGraphImageSequenceOutputNode.h
+++ b/Engine/Plugins/MovieScene/MovieRenderPipeline/Source/MovieRenderPipelineRenderPasses/Public/MovieGraphImageSequenceOutputNode.h
@@ -6,6 +6,7 @@
 
 #include "IImageWrapper.h"
 #include "MoviePipelineEXROutput.h"
+#include "OpenColorIOColorSpace.h"
 #include "Styling/AppStyle.h"
 #include "Async/Future.h"
 
@@ -32,6 +33,21 @@ public:
 	virtual bool IsFinishedWritingToDiskImpl() const override;
 	// ~UMovieGraphFileOutputNode Interface
 
+#if WITH_EDITORONLY_DATA
+	UPROPERTY(EditAnywhere, BlueprintReadWrite, Category = Overrides, meta = (InlineEditConditionToggle))
+	uint8 bOverride_OCIOConfiguration : 1;
+
+	/**
+	* OCIO configuration/transform settings.
+	*
+	* Note: There are differences from the previous implementation in MRQ given that we are now doing CPU-side processing.
+	* 1) This feature only works in editor-mode, since the OpenColorIO library is currently unavailable in game builds.
+	* 2) Users are now responsible for setting the renderer output space to Final Color (HDR) in Linear Working Color Space (SCS_FinalColorHDR).
+	*/
+	UPROPERTY(EditAnywhere, BlueprintReadWrite, Category = "Color", DisplayName="OpenColorIO Configuration", meta = (EditCondition = "bOverride_OCIOConfiguration"))
+	FOpenColorIODisplayConfiguration OCIOConfiguration;
+#endif
+
 protected:
 	/** The output format (as known used by the ImageWriteQueue) to output into. */
 	EImageFormat OutputFormat;
