Movie Render Queue: Follow up to 23938675

Apply same workaround in SMoviePIpelineQueueEditor for finding redirected config

diff --git a/Engine/Plugins/MovieScene/MovieRenderPipeline/Source/MovieRenderPipelineEditor/Private/Widgets/SMoviePipelineQueueEditor.cpp b/Engine/Plugins/MovieScene/MovieRenderPipeline/Source/MovieRenderPipelineEditor/Private/Widgets/SMoviePipelineQueueEditor.cpp
--- a/Engine/Plugins/MovieScene/MovieRenderPipeline/Source/MovieRenderPipelineEditor/Private/Widgets/SMoviePipelineQueueEditor.cpp
+++ b/Engine/Plugins/MovieScene/MovieRenderPipeline/Source/MovieRenderPipelineEditor/Private/Widgets/SMoviePipelineQueueEditor.cpp
@@ -458,6 +458,16 @@ public:
 
 			AssetPickerConfig.AssetShowWarningText = NoAssetsFoundWarning;
 			AssetPickerConfig.Filter.ClassPaths.Add(InClass->GetClassPathName());
+
+			// The redirected masterconfig assets aren't found when the new primaryconfig asset path is specified. This is a temporary workaround 
+			// for the AssetRegistry not applying CoreRedirects to the ClassPath. 
+			// This can be removed once that bug (UE-168245) is fixed and the KnownCase in CheckUnacceptableWords can also be removed.
+			if (InClass->GetClassPathName() == FTopLevelAssetPath(TEXT("/Script/MovieRenderPipelineCore.MoviePipelinePrimaryConfig")))
+			{
+				AssetPickerConfig.Filter.ClassPaths.Add(FTopLevelAssetPath(TEXT("/Script/MovieRenderPipelineCore.MoviePipelineMasterConfig")));
+			}
+
+			AssetPickerConfig.Filter.bRecursiveClasses = true;
 			AssetPickerConfig.OnAssetSelected = InOnAssetSelected;
 		}
 
