Movie Graph:

-Fix: //UE5/Main - QAVirtualProduction Editor Win64 Test=UE.EditorAutomation(RunTest=VirtualProduction) - Expected 'Edge added between MovieGraphCommandLineEncoderNode_1 and MovieGraphSamplingMethodNode_1' to be true

diff --git a/Engine/Plugins/MovieScene/MovieRenderPipeline/Source/MovieRenderPipelineEditor/Private/Tests/MovieGraphTests.cpp b/Engine/Plugins/MovieScene/MovieRenderPipeline/Source/MovieRenderPipelineEditor/Private/Tests/MovieGraphTests.cpp
--- a/Engine/Plugins/MovieScene/MovieRenderPipeline/Source/MovieRenderPipelineEditor/Private/Tests/MovieGraphTests.cpp
+++ b/Engine/Plugins/MovieScene/MovieRenderPipeline/Source/MovieRenderPipelineEditor/Private/Tests/MovieGraphTests.cpp
@@ -3,17 +3,19 @@
 #include "MovieGraphEditorTestUtilities.h"
 
 #include "Graph/MovieGraphConfig.h"
+#include "Graph/MovieGraphNode.h"
 #include "Graph/Nodes/MovieGraphApplyCVarPresetNode.h"
 #include "Graph/Nodes/MovieGraphBranchNode.h"
 #include "Graph/Nodes/MovieGraphCommandLineEncoderNode.h"
+#include "Graph/Nodes/MovieGraphDeferredPassNode.h"
 #include "Graph/Nodes/MovieGraphInputNode.h"
 #include "Graph/Nodes/MovieGraphOutputNode.h"
+#include "Graph/Nodes/MovieGraphSamplingMethodNode.h"
 #include "Graph/Nodes/MovieGraphSetCVarValueNode.h"
 #include "Graph/Nodes/MovieGraphSubgraphNode.h"
 #include "Graph/Nodes/MovieGraphVariableNode.h"
 
 #include "Algo/RemoveIf.h"
-#include "Graph/Nodes/MovieGraphSamplingMethodNode.h"
 #include "Misc/AutomationTest.h"
 
 namespace UE::MovieGraph::Private::Tests
@@ -476,8 +478,8 @@ namespace UE::MovieGraph::Private::Tests
 			return false;
 		}
 		UMovieGraphNode* AnyNodeB =
-			GraphConfig->ConstructRuntimeNode<UMovieGraphSamplingMethodNode>(
-				UMovieGraphSamplingMethodNode::StaticClass());
+			GraphConfig->ConstructRuntimeNode<UMovieGraphDeferredRenderPassNode>(
+				UMovieGraphDeferredRenderPassNode::StaticClass());
 		TestTrue(
 			TEXT("AnyNodeB successfully added"), GraphConfig->GetNodes().Contains(AnyNodeB));
 		if (!AnyNodeB)
