Movie Render Graph

-Add Named Resolutions struct and customization

diff --git a/Engine/Plugins/MovieScene/MovieRenderPipeline/Source/MovieRenderPipelineCore/Private/Graph/MovieGraphBlueprintLibrary.cpp b/Engine/Plugins/MovieScene/MovieRenderPipeline/Source/MovieRenderPipelineCore/Private/Graph/MovieGraphBlueprintLibrary.cpp
--- a/Engine/Plugins/MovieScene/MovieRenderPipeline/Source/MovieRenderPipelineCore/Private/Graph/MovieGraphBlueprintLibrary.cpp
+++ b/Engine/Plugins/MovieScene/MovieRenderPipeline/Source/MovieRenderPipelineCore/Private/Graph/MovieGraphBlueprintLibrary.cpp
@@ -320,7 +320,12 @@ FIntPoint UMovieGraphBlueprintLibrary::GetEffectiveOutputResolution(UMovieGraphE
 		RescaledOverscan = FMath::Clamp(CameraSetting->OverscanPercentage / 100.f, 0.f, 1.f);
 	}
 
-	return UMoviePipelineBlueprintLibrary::Utility_GetEffectiveOutputResolution(RescaledOverscan, OutputSetting->OutputResolution);
+	
+
+	return UMoviePipelineBlueprintLibrary::Utility_GetEffectiveOutputResolution(
+		RescaledOverscan,
+		OutputSetting->GetSyncedOutputResolution()
+	);
 }
 
 FText UMovieGraphBlueprintLibrary::GetJobName(const UMovieGraphPipeline* InMovieGraphPipeline)
diff --git a/Engine/Plugins/MovieScene/MovieRenderPipeline/Source/MovieRenderPipelineCore/Private/Graph/MovieGraphNamedResolution.cpp b/Engine/Plugins/MovieScene/MovieRenderPipeline/Source/MovieRenderPipelineCore/Private/Graph/MovieGraphNamedResolution.cpp
new file mode 100644
--- /dev/null
+++ b/Engine/Plugins/MovieScene/MovieRenderPipeline/Source/MovieRenderPipelineCore/Private/Graph/MovieGraphNamedResolution.cpp
@@ -0,0 +1,22 @@
+﻿// Copyright Epic Games, Inc. All Rights Reserved.
+
+#include "Graph/MovieGraphNamedResolution.h"
+
+#include "Graph/MovieGraphProjectSettings.h"
+
+FMovieGraphNamedResolution::FMovieGraphNamedResolution(const FName& InResolutionProfileName)
+{
+	// Find a matching custom entry from Project Settings
+	if (const UMovieGraphProjectSettings* MovieGraphProjectSettings =
+		GetDefault<UMovieGraphProjectSettings>())
+	{
+		if (const FMovieGraphNamedResolution* Match = MovieGraphProjectSettings->FindNamedResolutionForOption(InResolutionProfileName))
+		{
+			*this = *Match;
+			return;
+		}
+	}
+
+	// Otherwise make a new one
+	ProfileName = InResolutionProfileName;
+}
diff --git a/Engine/Plugins/MovieScene/MovieRenderPipeline/Source/MovieRenderPipelineCore/Private/Graph/MovieGraphProjectSettings.cpp b/Engine/Plugins/MovieScene/MovieRenderPipeline/Source/MovieRenderPipelineCore/Private/Graph/MovieGraphProjectSettings.cpp
new file mode 100644
--- /dev/null
+++ b/Engine/Plugins/MovieScene/MovieRenderPipeline/Source/MovieRenderPipelineCore/Private/Graph/MovieGraphProjectSettings.cpp
@@ -0,0 +1,107 @@
+﻿// Copyright Epic Games, Inc. All Rights Reserved.
+
+#include "Graph/MovieGraphProjectSettings.h"
+
+#include "Algo/Find.h"
+#include "Algo/FindLast.h"
+
+UMovieGraphProjectSettings::UMovieGraphProjectSettings()
+{
+	DefaultNamedResolutions = {
+		{ "720p (HD)", FIntPoint(1280, 720), "High definition" },
+		{ "1080p (FHD)", FIntPoint(1920, 1080), "Full high definition" },
+		{ "1440p (QHD)", FIntPoint(2560, 1440), "Quad high definition" },
+		{ "2160p (4K UHD)", FIntPoint(3840, 2160), "Ultra high definition (4k)" }
+	};
+}
+
+const FMovieGraphNamedResolution* UMovieGraphProjectSettings::FindNamedResolutionForOption(const FName& InOption) const
+{
+	return Algo::FindByPredicate(
+			DefaultNamedResolutions,
+			[&InOption](const FMovieGraphNamedResolution& Other)
+			{
+				return Other.ProfileName == InOption;
+			});
+}
+
+#if WITH_EDITOR
+void UMovieGraphProjectSettings::PostEditChangeProperty(FPropertyChangedEvent& PropertyChangedEvent)
+{
+	Super::PostEditChangeProperty(PropertyChangedEvent);
+
+	EnsureUniqueNamesInDefaultNamedResolutions();
+}
+#endif
+
+void RenameSimilarlyNamedResolutionsLoop(
+	TPair<FName, int32>& NameToCountPair, TArray<FMovieGraphNamedResolution>& DefaultNamedResolutions)
+{
+    // Loop until the count of a given profile name is reduced to 1
+	while (NameToCountPair.Value > 1)
+	{
+		// Find the last occurrence of the profile name in the default named resolutions
+		FMovieGraphNamedResolution* Match =
+			Algo::FindLastByPredicate(
+				DefaultNamedResolutions,
+				[&NameToCountPair](const FMovieGraphNamedResolution& Other)
+				{
+					return Other.ProfileName == NameToCountPair.Key;
+				});
+
+		if (ensureAlwaysMsgf(Match, TEXT("%hs: Failed to find FMovieGraphNamedResolution %s!"), __FUNCTION__, *NameToCountPair.Key.ToString()))
+		{
+			// Generate a new name by appending a number to the existing profile name
+			int32 NumberToAppend = NameToCountPair.Value - 1;
+			FName NewName = *FString::Printf(TEXT("%s_%d"), *Match->ProfileName.ToString(), NumberToAppend);
+
+			// Check if the new name already exists in the default named resolutions
+			while (DefaultNamedResolutions.ContainsByPredicate([&NewName](const FMovieGraphNamedResolution& Other)
+			{
+				return Other.ProfileName.IsEqual(NewName);
+			}))
+			{
+				// If it exists, increment the number and try again
+				NumberToAppend++;
+				NewName = *FString::Printf(TEXT("%s_%d"), *Match->ProfileName.ToString(), NumberToAppend);
+			}
+
+			// Reduce the count of the profile name and assign the new name to the matching resolution
+			NameToCountPair.Value--;
+			Match->ProfileName = NewName;
+		}
+		else
+		{
+			// If no matching resolution is found, break to avoid an infinite loop
+			break;
+		}
+	}
+}
+
+void UMovieGraphProjectSettings::EnsureUniqueNamesInDefaultNamedResolutions()
+{
+    // Create a map to store the count of each profile name
+    TMap<FName, int32> NameToCount;
+
+    // Iterate through each default named resolution
+    for (const FMovieGraphNamedResolution& ResolutionName : DefaultNamedResolutions)
+    {
+        // Check if the profile name exists in the map
+        if (int32* Match = NameToCount.Find(ResolutionName.ProfileName))
+        {
+            // If it exists, increment the count
+            *Match += 1;
+        }
+        else
+        {
+            // If it doesn't exist, add it to the map with a count of 1
+            NameToCount.Add(ResolutionName.ProfileName, 1);
+        }
+    }
+
+    // Iterate through the map of profile names and their counts, renaming redundant items as we go
+    for (TPair<FName, int32>& NameToCountPair : NameToCount)
+    {
+		RenameSimilarlyNamedResolutionsLoop(NameToCountPair, DefaultNamedResolutions);
+    }
+}
diff --git a/Engine/Plugins/MovieScene/MovieRenderPipeline/Source/MovieRenderPipelineCore/Private/Graph/Nodes/MovieGraphOutputSettingNode.cpp b/Engine/Plugins/MovieScene/MovieRenderPipeline/Source/MovieRenderPipelineCore/Private/Graph/Nodes/MovieGraphOutputSettingNode.cpp
--- a/Engine/Plugins/MovieScene/MovieRenderPipeline/Source/MovieRenderPipelineCore/Private/Graph/Nodes/MovieGraphOutputSettingNode.cpp
+++ b/Engine/Plugins/MovieScene/MovieRenderPipeline/Source/MovieRenderPipelineCore/Private/Graph/Nodes/MovieGraphOutputSettingNode.cpp
@@ -1,10 +1,12 @@
 // Copyright Epic Games, Inc. All Rights Reserved.
 
 #include "Graph/Nodes/MovieGraphOutputSettingNode.h"
+
+#include "Graph/MovieGraphProjectSettings.h"
 #include "Styling/AppStyle.h"
 
 UMovieGraphOutputSettingNode::UMovieGraphOutputSettingNode()
-	: OutputResolution(FIntPoint(1920, 1080))
+	: OutputResolution(FMovieGraphNamedResolution(FMovieGraphNamedResolution::DefaultResolutionName))
 	, OutputFrameRate(FFrameRate(24, 1))
 	, bOverwriteExistingOutput(true)
 	, ZeroPadFrameNumbers(4)
@@ -22,12 +24,14 @@ void UMovieGraphOutputSettingNode::GetFormatResolveArgs(FMovieGraphResolveArgs&
 	OutMergedFormatArgs.FilenameArguments.Add(TEXT("project_dir"), ResolvedProjectDir);
 	OutMergedFormatArgs.FileMetadata.Add(TEXT("unreal/project_dir"), ResolvedProjectDir);
 
+	FIntPoint OutputResolutionAsIntPoint = GetSyncedOutputResolution();
+	
 	// Resolution Arguments
 	{
-		FString Resolution = FString::Printf(TEXT("%d_%d"), OutputResolution.X, OutputResolution.Y);
+		FString Resolution = FString::Printf(TEXT("%d_%d"), OutputResolutionAsIntPoint.X, OutputResolutionAsIntPoint.Y);
 		OutMergedFormatArgs.FilenameArguments.Add(TEXT("output_resolution"), Resolution);
-		OutMergedFormatArgs.FilenameArguments.Add(TEXT("output_width"), FString::FromInt(OutputResolution.X));
-		OutMergedFormatArgs.FilenameArguments.Add(TEXT("output_height"), FString::FromInt(OutputResolution.Y));
+		OutMergedFormatArgs.FilenameArguments.Add(TEXT("output_width"), FString::FromInt(OutputResolutionAsIntPoint.X));
+		OutMergedFormatArgs.FilenameArguments.Add(TEXT("output_height"), FString::FromInt(OutputResolutionAsIntPoint.Y));
 	}
 
 	// We don't resolve the version here because that's handled on a per-file/shot basis
@@ -58,4 +62,21 @@ FSlateIcon UMovieGraphOutputSettingNode::GetIconAndTint(FLinearColor& OutColor)
 	OutColor = FLinearColor::White;
 	return SettingsIcon;
 }
-#endif // WITH_EDITOR
\ No newline at end of file
+#endif // WITH_EDITOR
+
+FIntPoint UMovieGraphOutputSettingNode::GetSyncedOutputResolution() const
+{
+	// Try to find a matching entry from Project Settings to stay in sync
+	const UMovieGraphProjectSettings* MovieGraphProjectSettings = GetDefault<UMovieGraphProjectSettings>();
+	if (ensureAlwaysMsgf(MovieGraphProjectSettings, TEXT("%hs: Failed to find UMovieGraphProjectSettings!"), __FUNCTION__))
+	{
+		const FMovieGraphNamedResolution* FoundResolution = MovieGraphProjectSettings->FindNamedResolutionForOption(OutputResolution.ProfileName);
+		if (FoundResolution)
+		{
+			return FoundResolution->Resolution;
+		}
+	}
+
+	// Otherwise return what we have saved locally
+	return OutputResolution.Resolution;
+}
diff --git a/Engine/Plugins/MovieScene/MovieRenderPipeline/Source/MovieRenderPipelineCore/Public/Graph/MovieGraphNamedResolution.h b/Engine/Plugins/MovieScene/MovieRenderPipeline/Source/MovieRenderPipelineCore/Public/Graph/MovieGraphNamedResolution.h
new file mode 100644
--- /dev/null
+++ b/Engine/Plugins/MovieScene/MovieRenderPipeline/Source/MovieRenderPipelineCore/Public/Graph/MovieGraphNamedResolution.h
@@ -0,0 +1,50 @@
+﻿// Copyright Epic Games, Inc. All Rights Reserved.
+#pragma once
+
+#include "Engine/EngineTypes.h"
+
+#include "MovieGraphNamedResolution.generated.h"
+
+/**
+ * Holds information about a screen resolution to be used for rendering.
+ */
+USTRUCT(BlueprintType)
+struct MOVIERENDERPIPELINECORE_API FMovieGraphNamedResolution
+{
+	GENERATED_BODY()
+
+public:
+
+	FMovieGraphNamedResolution() = default;
+
+	FMovieGraphNamedResolution(
+		const FName& InResolutionProfileName, const FIntPoint InResolution, const FString& InDescription)
+		: ProfileName(InResolutionProfileName)
+		, Resolution(InResolution)
+		, Description(InDescription)
+	{}
+
+	FMovieGraphNamedResolution(const FName& InResolutionProfileName);
+	
+	/** The name of the resolution this links to */
+	UPROPERTY(config, EditAnywhere, BlueprintReadWrite, Category = "Movie Graph|Resolution")
+	FName ProfileName = NAME_None;
+
+	/** The screen resolution (in pixels). */
+	UPROPERTY(config, EditAnywhere, BlueprintReadWrite, Category = "Movie Graph|Resolution")
+	FIntPoint Resolution = FIntPoint(1920, 1080);
+	
+	/** The description text for this screen resolution. */
+	UPROPERTY(config, EditAnywhere, BlueprintReadWrite, Category = "Movie Graph|Resolution")
+	FString Description = FString();
+
+	bool IsValid() const
+	{
+		return !ProfileName.IsNone() && Resolution.X > 0 && Resolution.Y > 0;
+	}
+
+	/**
+	 * The default resolution name to use when one is not defined.
+	 */
+	inline static FName DefaultResolutionName = TEXT("1080p (FHD)");
+};
\ No newline at end of file
diff --git a/Engine/Plugins/MovieScene/MovieRenderPipeline/Source/MovieRenderPipelineCore/Public/Graph/MovieGraphProjectSettings.h b/Engine/Plugins/MovieScene/MovieRenderPipeline/Source/MovieRenderPipelineCore/Public/Graph/MovieGraphProjectSettings.h
new file mode 100644
--- /dev/null
+++ b/Engine/Plugins/MovieScene/MovieRenderPipeline/Source/MovieRenderPipelineCore/Public/Graph/MovieGraphProjectSettings.h
@@ -0,0 +1,50 @@
+﻿// Copyright Epic Games, Inc. All Rights Reserved.
+#pragma once
+
+#include "Graph/MovieGraphNamedResolution.h"
+
+#include "Engine/DeveloperSettings.h"
+#include "Engine/EngineTypes.h"
+#include "UObject/Object.h"
+#include "UObject/UnrealType.h"
+
+#include "MovieGraphProjectSettings.generated.h"
+
+/**
+ * Settings that apply to the Movie Graph.
+ */
+UCLASS(config = MovieRenderPipeline, DefaultConfig, meta=(DisplayName = "Movie Graph Settings"))
+class MOVIERENDERPIPELINECORE_API UMovieGraphProjectSettings : public UDeveloperSettings
+{
+public:
+	GENERATED_BODY()
+
+	UMovieGraphProjectSettings();
+
+	/** Gets the settings container name for the settings, either Project or Editor */
+	virtual FName GetContainerName() const override { return FName("Project"); }
+	/** Gets the category for the settings, some high level grouping like, Editor, Engine, Game...etc. */
+	virtual FName GetCategoryName() const override { return FName("Plugins"); }
+	
+	/**
+	 * Tries to find a match for the given option in DefaultNamedResolutions.
+	 * @param InOption The name of the option.
+	 * @return A raw pointer to an FMovieGraphNamedResolution if a match is found, nullptr otherwise.
+	 */
+	const FMovieGraphNamedResolution* FindNamedResolutionForOption(const FName& InOption) const;
+	
+	/**
+	 * A list of default resolutions to render with, defined in Config/BaseMovieRenderPipeline.ini or Project Settings.
+	 * These resolutions will appear in the MovieGraph's various output settings nodes.
+	 */
+	UPROPERTY(Config, EditAnywhere, Category = "Movie Graph")
+	TArray<FMovieGraphNamedResolution> DefaultNamedResolutions;
+
+private:
+#if WITH_EDITOR
+	virtual void PostEditChangeProperty(struct FPropertyChangedEvent& PropertyChangedEvent) override;
+#endif
+
+	void EnsureUniqueNamesInDefaultNamedResolutions();
+	
+};
\ No newline at end of file
diff --git a/Engine/Plugins/MovieScene/MovieRenderPipeline/Source/MovieRenderPipelineCore/Public/Graph/Nodes/MovieGraphOutputSettingNode.h b/Engine/Plugins/MovieScene/MovieRenderPipeline/Source/MovieRenderPipelineCore/Public/Graph/Nodes/MovieGraphOutputSettingNode.h
--- a/Engine/Plugins/MovieScene/MovieRenderPipeline/Source/MovieRenderPipelineCore/Public/Graph/Nodes/MovieGraphOutputSettingNode.h
+++ b/Engine/Plugins/MovieScene/MovieRenderPipeline/Source/MovieRenderPipelineCore/Public/Graph/Nodes/MovieGraphOutputSettingNode.h
@@ -1,10 +1,13 @@
 // Copyright Epic Games, Inc. All Rights Reserved.
 
 #pragma once
+
+#include "Graph/MovieGraphNamedResolution.h"
 #include "Graph/MovieGraphNode.h"
+
 #include "Misc/FrameRate.h"
-#include "MovieGraphOutputSettingNode.generated.h"
 
+#include "MovieGraphOutputSettingNode.generated.h"
 
 UCLASS()
 class MOVIERENDERPIPELINECORE_API UMovieGraphOutputSettingNode : public UMovieGraphSettingNode
@@ -24,6 +27,8 @@ public:
 #endif
 	// ~UMovieGraphSettingNode Interface
 
+	FIntPoint GetSyncedOutputResolution() const;
+
 public:
 	UPROPERTY(EditAnywhere, BlueprintReadWrite, Category = Overrides, meta = (InlineEditConditionToggle))
 	uint8 bOverride_OutputDirectory : 1;
@@ -65,7 +70,7 @@ public:
 
 	/** What resolution should our output files be exported at? */
 	UPROPERTY(EditAnywhere, BlueprintReadWrite, Category = "File Output", meta = (EditCondition = "bOverride_OutputResolution"))
-	FIntPoint OutputResolution;
+	FMovieGraphNamedResolution OutputResolution;
 	
 	/** What frame rate should the output files be exported at? This overrides the Display Rate of the target sequence. If not overwritten, uses the default Sequence Display Rate. */
 	UPROPERTY(EditAnywhere, BlueprintReadWrite, Category = "File Output", meta=(EditCondition="bOverride_OutputFrameRate"))
diff --git a/Engine/Plugins/MovieScene/MovieRenderPipeline/Source/MovieRenderPipelineEditor/Private/Customizations/Graph/MovieGraphNamedResolutionCustomization.h b/Engine/Plugins/MovieScene/MovieRenderPipeline/Source/MovieRenderPipelineEditor/Private/Customizations/Graph/MovieGraphNamedResolutionCustomization.h
new file mode 100644
--- /dev/null
+++ b/Engine/Plugins/MovieScene/MovieRenderPipeline/Source/MovieRenderPipelineEditor/Private/Customizations/Graph/MovieGraphNamedResolutionCustomization.h
@@ -0,0 +1,375 @@
+﻿// Copyright Epic Games, Inc. All Rights Reserved.
+#pragma once
+
+#include "Graph/MovieGraphNamedResolution.h"
+#include "Graph/MovieGraphProjectSettings.h"
+
+#include "MovieRenderPipelineCoreModule.h"
+
+#include "DetailLayoutBuilder.h"
+#include "DetailWidgetRow.h"
+#include "IPropertyTypeCustomization.h"
+#include "Widgets/Input/SComboBox.h"
+#include "Widgets/SBoxPanel.h"
+
+#define LOCTEXT_NAMESPACE "FMovieGraphNamedResolutionCustomization"
+
+/* Customizes how named resolutions are displayed in the details pane. */
+class MOVIERENDERPIPELINEEDITOR_API FMovieGraphNamedResolutionCustomization : public IPropertyTypeCustomization
+{
+public:
+	static TSharedRef<IPropertyTypeCustomization> MakeInstance()
+	{
+		return MakeShared<FMovieGraphNamedResolutionCustomization>();
+	}
+
+	virtual ~FMovieGraphNamedResolutionCustomization() override
+	{
+		UMovieGraphProjectSettings* MovieGraphProjectSettings = GetMutableDefault<UMovieGraphProjectSettings>();
+		if (ensureAlwaysMsgf(MovieGraphProjectSettings, TEXT("%hs: Failed to find UMovieGraphProjectSettings!"), __FUNCTION__))
+		{
+			MovieGraphProjectSettings->OnSettingChanged().Remove(OnProjectSettingsModifiedHandle);
+		}
+	}
+
+protected:
+
+	bool IsCurrentOptionValid() const
+	{
+		return ComboBoxOptions.Contains(CurrentOption);
+	}
+	
+	/**
+	 * Checks to see if CurrentOption is None, and if so calls RepopulateOptions() which sets CurrentOption to the first option.
+	 * Otherwise just returns CurrentOption.
+	 * @return The current option selected in the combobox
+	 */
+	const FName& GetCurrentOptionAssigningIfNeeded()
+	{
+		if (!IsCurrentOptionValid())
+		{
+			RepopulateOptions();
+		}
+		return CurrentOption;
+	}
+
+	void SetCurrentOptionAndCacheTooltipText(const FName& InOption)
+	{
+		CurrentOption = InOption;
+		CachedSelectedOptionTooltipText = GetOptionTooltipText(CurrentOption);
+	}
+
+	/**
+	 * Adds all ProfileNames from DefaultNamedResolutions to the combobox,
+	 * and sets CurrentOption to the first option if it's None or isn't a known option.
+	 */
+	void RepopulateOptions()
+	{
+		ComboBoxOptions.Reset(ComboBoxOptions.Num());
+
+		const UMovieGraphProjectSettings* MovieGraphProjectSettings = GetDefault<UMovieGraphProjectSettings>();
+		if (ensureAlwaysMsgf(MovieGraphProjectSettings, TEXT("%hs: Failed to find UMovieGraphProjectSettings!"), __FUNCTION__))
+		{
+			for (const FMovieGraphNamedResolution& Resolution : MovieGraphProjectSettings->DefaultNamedResolutions)
+			{
+				ComboBoxOptions.AddUnique(Resolution.ProfileName);
+			}
+		}
+
+		ComboBoxOptions.AddUnique(CustomEntryName);
+
+		if (!IsCurrentOptionValid())
+		{
+			SetCurrentOptionAndCacheTooltipText(CustomEntryName);
+		}
+	}
+
+	/**
+	 * Returns the tooltip text for the given option by looking up its Description
+	 * in DefaultNamedResolutions from Project Settings.
+	 * If a match isn't found, the tooltip text will be the name of the option.
+	 * @param InOption The name of the option.
+	 * @return The tooltip text for the option.
+	 */
+	FText GetOptionTooltipText(FName InOption) const
+	{
+		if (const FMovieGraphNamedResolution* Match = FindNamedResolutionForOption(InOption))
+		{
+			return FText::FromString(Match->Description);
+		}
+		return FText::FromName(InOption);
+	}
+
+	/**
+	 * Tries to find a match for the given option in DefaultNamedResolutions from Project Settings.
+	 * @param InOption The name of the option.
+	 * @return A raw pointer to an FMovieGraphNamedResolution if a match is found, nullptr otherwise.
+	 */
+	const FMovieGraphNamedResolution* FindNamedResolutionForOption(const FName& InOption) const
+	{
+		// Return the custom entry
+		if (InOption.IsEqual(CustomEntryName))
+		{
+			return &CustomEntry;
+		}
+
+		// Find a matching custom entry from Project Settings
+		const UMovieGraphProjectSettings* MovieGraphProjectSettings = GetDefault<UMovieGraphProjectSettings>();
+		if (ensureAlwaysMsgf(MovieGraphProjectSettings, TEXT("%hs: Failed to find UMovieGraphProjectSettings!"), __FUNCTION__))
+		{
+			return MovieGraphProjectSettings->FindNamedResolutionForOption(InOption);
+		}
+
+		UE_LOG(LogMovieRenderPipeline, Error, TEXT("%hs: No match found for FMovieGraphNamedResolution %s!"), __FUNCTION__, *InOption.ToString());
+		return nullptr;
+	}
+
+	/**
+	 * Generates a widget for the given option with the option name as STextBlock with tooltip text .
+	 * @param InOption The name of the option.
+	 * @return A shared reference to the generated widget.
+	 */
+	TSharedRef<SWidget> MakeWidgetForOption(FName InOption) const
+	{
+		return
+			SNew(STextBlock)
+			.Text(FText::FromName(InOption))
+			.ToolTipText(this, &FMovieGraphNamedResolutionCustomization::GetOptionTooltipText, InOption)
+			.Font(IDetailLayoutBuilder::GetDetailFont())
+		;
+	}
+
+	TSharedRef<SWidget> CreateOptionComboBox()
+	{
+		return
+			SNew(SComboBox<FName>)
+			.OnComboBoxOpening(this, &FMovieGraphNamedResolutionCustomization::RepopulateOptions)
+			.OptionsSource(&ComboBoxOptions)
+			.OnSelectionChanged(this, &FMovieGraphNamedResolutionCustomization::OnComboBoxSelectionChanged)
+			.OnGenerateWidget(this, &FMovieGraphNamedResolutionCustomization::MakeWidgetForOption)
+			.InitiallySelectedItem(GetCurrentOptionAssigningIfNeeded())
+			[
+				SNew(STextBlock)
+				.Text_Lambda([this]()
+				{
+					return FText::FromName(GetCurrentOptionAssigningIfNeeded());
+				})
+				.ToolTipText_Lambda([this]()
+				{
+					return CachedSelectedOptionTooltipText;
+				})
+				.Font(IDetailLayoutBuilder::GetDetailFont())
+			];
+	}
+
+	TSharedRef<SWidget> CreateCustomResolutionWidget(const TSharedRef<IPropertyHandle> InStructPropertyHandle)
+	{
+		const TSharedRef<IPropertyHandle> ResolutionPropHandle = InStructPropertyHandle->GetChildHandle(
+			GET_MEMBER_NAME_CHECKED(FMovieGraphNamedResolution, Resolution)).ToSharedRef();
+		
+		return
+			SNew(SHorizontalBox)
+			.Visibility_Lambda([this]()
+			{
+				return GetCurrentOptionAssigningIfNeeded().IsEqual(CustomEntryName) ?
+					EVisibility::Visible : EVisibility::Collapsed;
+			})
+
+			+SHorizontalBox::Slot()
+			.AutoWidth()
+			.VAlign(VAlign_Center)
+			.Padding(0, 0, 8, 0)
+			[
+				SNew(STextBlock)
+				.Text(FText::FromString(TEXT("Width")))
+				.Font(IDetailLayoutBuilder::GetDetailFont())
+			]
+
+			+SHorizontalBox::Slot()
+			.HAlign(HAlign_Fill)
+			.Padding(0, 0, 8, 0)
+			[
+				ResolutionPropHandle->GetChildHandle(0)->CreatePropertyValueWidget()
+			]
+
+			+SHorizontalBox::Slot()
+			.AutoWidth()
+			.VAlign(VAlign_Center)
+			.Padding(0, 0, 8, 0)
+			[
+				SNew(STextBlock)
+				.Text(FText::FromString(TEXT("Height")))
+				.Font(IDetailLayoutBuilder::GetDetailFont())
+			]
+
+			+SHorizontalBox::Slot()
+			.HAlign(HAlign_Fill)
+			[
+				ResolutionPropHandle->GetChildHandle(1)->CreatePropertyValueWidget()
+			];
+	}
+
+	void BindToOnProjectSettingsModified()
+	{
+		UMovieGraphProjectSettings* MovieGraphProjectSettings = GetMutableDefault<UMovieGraphProjectSettings>();
+		if (ensureAlwaysMsgf(MovieGraphProjectSettings, TEXT("%hs: Unable to get UMovieGraphProjectSettings."), __FUNCTION__))
+		{
+			OnProjectSettingsModifiedHandle =
+				MovieGraphProjectSettings->OnSettingChanged().AddSP(
+					this, &FMovieGraphNamedResolutionCustomization::OnProjectSettingsChanged);
+		}
+	}
+	
+	//~ Begin IPropertyTypeCustomization interface
+	virtual void CustomizeHeader(
+		TSharedRef<IPropertyHandle> InStructPropertyHandle,
+		FDetailWidgetRow& HeaderRow, IPropertyTypeCustomizationUtils& CustomizationUtils) override
+	{
+		StructPropertyHandle = InStructPropertyHandle;
+		
+		const TSharedRef<IPropertyHandle> NameProp = InStructPropertyHandle->GetChildHandle(
+			GET_MEMBER_NAME_CHECKED(FMovieGraphNamedResolution, ProfileName)).ToSharedRef();
+		
+		// Try to 'load' selected option from the struct handle
+		FName LoadedOption = NAME_None;
+		if (NameProp->GetValue(LoadedOption) == FPropertyAccess::Result::Success && !LoadedOption.IsNone())
+		{
+			SetCurrentOptionAndCacheTooltipText(LoadedOption);
+		}
+		
+		// Set up option list for first time. If LoadedOption doesn't exist, use the first option in the list.
+		RepopulateOptions();
+
+		// Repopulate options when the project settings change
+		BindToOnProjectSettingsModified();
+
+		HeaderRow
+		.NameContent()
+		[
+			// Show "Resolution" as the row name
+			InStructPropertyHandle->CreatePropertyNameWidget()
+		]
+		.ValueContent()
+		.HAlign(HAlign_Fill)
+		[
+			SNew(SVerticalBox)
+			
+			+ SVerticalBox::Slot()
+			.Padding(0, 2, 8, 2)
+			.AutoHeight()
+			.HAlign(HAlign_Fill)
+			[
+				CreateOptionComboBox()
+			]
+
+			+ SVerticalBox::Slot()
+			.Padding(0, 2, 8, 2)
+			.AutoHeight()
+			.HAlign(HAlign_Fill)
+			[
+				CreateCustomResolutionWidget(InStructPropertyHandle)
+			]
+		];
+	}
+
+	// Skip customizing children
+	virtual void CustomizeChildren(
+		TSharedRef<IPropertyHandle> InStructPropertyHandle,
+		IDetailChildrenBuilder& StructBuilder, IPropertyTypeCustomizationUtils& StructCustomizationUtils) override { }
+	//~ End IPropertyTypeCustomization interface
+	
+	void OnProjectSettingsChanged(UObject*, FPropertyChangedEvent& PropertyChangedEvent)
+	{
+		if (PropertyChangedEvent.MemberProperty->GetFName().IsEqual(GET_MEMBER_NAME_CHECKED(UMovieGraphProjectSettings, DefaultNamedResolutions)))
+		{
+			RepopulateOptions();
+		}
+	}
+
+	/**
+	 * Updates the struct which owns this customization with the data from the matched named resolution.
+	 */
+	void UpdateOwningStruct() const
+	{
+		const TSharedPtr<IPropertyHandle> PinnedStructPropertyHandle = StructPropertyHandle.Pin();
+		if (!PinnedStructPropertyHandle)
+		{
+			UE_LOG(LogMovieRenderPipeline, Error, TEXT("%hs: Unable to pin StructPropertyHandle."), __FUNCTION__);
+			return;
+		}
+
+		const FMovieGraphNamedResolution* NamedResolution = FindNamedResolutionForOption(CurrentOption);
+		if (!NamedResolution)
+		{
+			return;
+		}
+		
+		void* StructData = nullptr;
+		if (PinnedStructPropertyHandle->GetValueData(StructData) == FPropertyAccess::Success && StructData != nullptr)
+		{
+			// Copy the data from our matched NamedResolution to the owning struct
+			FMemory::Memcpy(StructData, NamedResolution, sizeof(FMovieGraphNamedResolution));
+
+			// Mark owning packages dirty
+			TArray<UObject*> OutObjects;
+			PinnedStructPropertyHandle->GetOuterObjects(OutObjects);
+
+			for (const UObject* Object : OutObjects)
+			{
+				if (!Object)
+				{
+					continue;
+				}
+				ensureAlwaysMsgf(Object->MarkPackageDirty(), TEXT("%hs: Failed to mark package of %s dirty!"), __FUNCTION__, *Object->GetName());
+			}
+		}
+		else
+		{
+			UE_LOG(LogMovieRenderPipeline, Error, TEXT("%hs: Unable to get struct data from StructPropertyHandle."), __FUNCTION__);
+		}
+	}
+
+	void OnComboBoxSelectionChanged(const FName NewValue, ESelectInfo::Type)
+	{
+		SetCurrentOptionAndCacheTooltipText(NewValue);
+
+		UpdateOwningStruct();
+	}
+
+	/**
+	 * Weak ptr to the handle of the struct which owns this customization
+	 */
+	TWeakPtr<IPropertyHandle> StructPropertyHandle;
+
+	/**
+	 * The option currently selected in the combobox
+	 */
+	FName CurrentOption = NAME_None;
+	/**
+	 * The tooltip text for the currently selected option
+	 */
+	FText CachedSelectedOptionTooltipText;
+
+	/**
+	 * The options for the combo box, populated from UMovieGraphProjectSettings::DefaultNamedResolutions.
+	 * A 'custom' entry is added to the end of the list, representing an arbitrary user-defined resolution.
+	 */
+	TArray<FName> ComboBoxOptions;
+
+	/**
+	 * A binding that updates the ComboBoxOptions when UMovieGraphProjectSettings::DefaultNamedResolutions changes
+	 */
+	FDelegateHandle OnProjectSettingsModifiedHandle;
+
+	/**
+	 * Predefined name for the 'custom' resolution option in the combobox
+	 */
+	inline static FName CustomEntryName = TEXT("Custom");
+	/**
+	 * Default FMovieGraphNamedResolution for the 'custom' option
+	 */
+	inline static FMovieGraphNamedResolution CustomEntry =
+		FMovieGraphNamedResolution(CustomEntryName, FIntPoint(1920, 1080), CustomEntryName.ToString());	
+};
+
+#undef LOCTEXT_NAMESPACE
diff --git a/Engine/Plugins/MovieScene/MovieRenderPipeline/Source/MovieRenderPipelineEditor/Private/Graph/MovieGraphAssetToolkit.cpp b/Engine/Plugins/MovieScene/MovieRenderPipeline/Source/MovieRenderPipelineEditor/Private/Graph/MovieGraphAssetToolkit.cpp
--- a/Engine/Plugins/MovieScene/MovieRenderPipeline/Source/MovieRenderPipelineEditor/Private/Graph/MovieGraphAssetToolkit.cpp
+++ b/Engine/Plugins/MovieScene/MovieRenderPipeline/Source/MovieRenderPipelineEditor/Private/Graph/MovieGraphAssetToolkit.cpp
@@ -3,18 +3,21 @@
 #include "MovieGraphAssetToolkit.h"
 
 #include "Customizations/Graph/MovieGraphMemberCustomization.h"
+#include "Customizations/Graph/MovieGraphNamedResolutionCustomization.h"
 #include "Customizations/Graph/MovieGraphNodeCustomization.h"
 #include "Customizations/Graph/MovieGraphSelectNodeCustomization.h"
 #include "Customizations/Graph/MovieGraphShowFlagsCustomization.h"
 #include "Graph/Renderers/MovieGraphShowFlags.h"
-#include "Framework/Commands/GenericCommands.h"
+
 #include "Graph/MovieGraphConfig.h"
 #include "Graph/Nodes/MovieGraphSelectNode.h"
 #include "MovieEdGraphNode.h"
 #include "MovieGraphSchema.h"
-#include "PropertyEditorModule.h"
 #include "SMovieGraphActiveRenderSettingsTabContent.h"
 #include "SMovieGraphMembersTabContent.h"
+
+#include "Framework/Commands/GenericCommands.h"
+#include "PropertyEditorModule.h"
 #include "ToolMenus.h"
 #include "Widgets/Docking/SDockTab.h"
 #include "Widgets/Graph/SMovieGraphConfigPanel.h"
@@ -242,6 +245,10 @@ TSharedRef<SDockTab> FMovieGraphAssetToolkit::SpawnTab_RenderGraphDetails(const
 	SelectedGraphObjectsDetailsWidget->RegisterInstancedCustomPropertyTypeLayout(
 		UMovieGraphShowFlags::StaticClass()->GetFName(),
 		FOnGetPropertyTypeCustomizationInstance::CreateStatic(&FMovieGraphShowFlagsCustomization::MakeInstance));
+
+	SelectedGraphObjectsDetailsWidget->RegisterInstancedCustomPropertyTypeLayout(
+		FMovieGraphNamedResolution::StaticStruct()->GetFName(),
+		FOnGetPropertyTypeCustomizationInstance::CreateStatic(&FMovieGraphNamedResolutionCustomization::MakeInstance));
 	
 	return SNew(SDockTab)
 		.TabColorScale(GetTabColorScale())
