Movie Graph:

-Fix: Add support for double-clicking Subgraph Nodes

diff --git a/Engine/Plugins/MovieScene/MovieRenderPipeline/Source/MovieRenderPipelineEditor/Private/Graph/MovieEdGraph.cpp b/Engine/Plugins/MovieScene/MovieRenderPipeline/Source/MovieRenderPipelineEditor/Private/Graph/MovieEdGraph.cpp
--- a/Engine/Plugins/MovieScene/MovieRenderPipeline/Source/MovieRenderPipelineEditor/Private/Graph/MovieEdGraph.cpp
+++ b/Engine/Plugins/MovieScene/MovieRenderPipeline/Source/MovieRenderPipelineEditor/Private/Graph/MovieEdGraph.cpp
@@ -4,17 +4,20 @@
 #include "Graph/MovieGraphConfig.h"
 #include "Graph/MovieGraphEdge.h"
 #include "Graph/MovieGraphPin.h"
+#include "Graph/Nodes/MovieGraphSubgraphNode.h"
 #include "Graph/Nodes/MovieGraphVariableNode.h"
 #include "Graph/MovieGraphNode.h"
 #include "MovieEdGraphOutputNode.h"
 #include "MovieEdGraphInputNode.h"
 #include "MovieEdGraphVariableNode.h"
+#include "MoviePipelineEdGraphSubgraphNode.h"
 #include "MovieRenderPipelineCoreModule.h"
 #include "EdGraph/EdGraphPin.h"
 
 template UMoviePipelineEdGraphNodeBase* UMoviePipelineEdGraph::CreateNodeFromRuntimeNode<UMoviePipelineEdGraphNodeInput>(UMovieGraphNode* InRuntimeNode);
 template UMoviePipelineEdGraphNodeBase* UMoviePipelineEdGraph::CreateNodeFromRuntimeNode<UMoviePipelineEdGraphNodeOutput>(UMovieGraphNode* InRuntimeNode);
 template UMoviePipelineEdGraphNodeBase* UMoviePipelineEdGraph::CreateNodeFromRuntimeNode<UMoviePipelineEdGraphVariableNode>(UMovieGraphNode* InRuntimeNode);
+template UMoviePipelineEdGraphNodeBase* UMoviePipelineEdGraph::CreateNodeFromRuntimeNode<UMoviePipelineEdGraphSubgraphNode>(UMovieGraphNode* InRuntimeNode); 
 template UMoviePipelineEdGraphNodeBase* UMoviePipelineEdGraph::CreateNodeFromRuntimeNode<UMoviePipelineEdGraphNode>(UMovieGraphNode* InRuntimeNode);
 
 UMovieGraphConfig* UMoviePipelineEdGraph::GetPipelineGraph() const
@@ -48,6 +51,10 @@ void UMoviePipelineEdGraph::InitFromRuntimeGraph(UMovieGraphConfig* InGraph)
 		{
 			NodeLookup.Add(RuntimeNode, CreateNodeFromRuntimeNode<UMoviePipelineEdGraphVariableNode>(RuntimeNode));
 		}
+		else if (RuntimeNode.IsA<UMovieGraphSubgraphNode>())
+		{
+			NodeLookup.Add(RuntimeNode, CreateNodeFromRuntimeNode<UMoviePipelineEdGraphSubgraphNode>(RuntimeNode));
+		}
 		else
 		{
 			NodeLookup.Add(RuntimeNode, CreateNodeFromRuntimeNode<UMoviePipelineEdGraphNode>(RuntimeNode));
diff --git a/Engine/Plugins/MovieScene/MovieRenderPipeline/Source/MovieRenderPipelineEditor/Private/Graph/MoviePipelineEdGraphSubgraphNode.cpp b/Engine/Plugins/MovieScene/MovieRenderPipeline/Source/MovieRenderPipelineEditor/Private/Graph/MoviePipelineEdGraphSubgraphNode.cpp
new file mode 100644
--- /dev/null
+++ b/Engine/Plugins/MovieScene/MovieRenderPipeline/Source/MovieRenderPipelineEditor/Private/Graph/MoviePipelineEdGraphSubgraphNode.cpp
@@ -0,0 +1,18 @@
+﻿// Copyright Epic Games, Inc. All Rights Reserved.
+
+#include "MoviePipelineEdGraphSubgraphNode.h"
+
+#include "Graph/Nodes/MovieGraphSubgraphNode.h"
+
+UObject* UMoviePipelineEdGraphSubgraphNode::GetJumpTargetForDoubleClick() const
+{
+	if (const UMovieGraphSubgraphNode* SubgraphNode = Cast<UMovieGraphSubgraphNode>(RuntimeNode))
+	{
+		if (UMovieGraphConfig* SubgraphAsset =  SubgraphNode->GetSubgraphAsset())
+		{
+			return SubgraphAsset;
+		}
+	}
+	
+	return Super::GetJumpTargetForDoubleClick();
+}
diff --git a/Engine/Plugins/MovieScene/MovieRenderPipeline/Source/MovieRenderPipelineEditor/Private/Graph/MoviePipelineEdGraphSubgraphNode.h b/Engine/Plugins/MovieScene/MovieRenderPipeline/Source/MovieRenderPipelineEditor/Private/Graph/MoviePipelineEdGraphSubgraphNode.h
new file mode 100644
--- /dev/null
+++ b/Engine/Plugins/MovieScene/MovieRenderPipeline/Source/MovieRenderPipelineEditor/Private/Graph/MoviePipelineEdGraphSubgraphNode.h
@@ -0,0 +1,19 @@
+﻿// Copyright Epic Games, Inc. All Rights Reserved.
+
+#pragma once
+
+#include "Graph/MovieEdGraphNode.h"
+
+#include "MoviePipelineEdGraphSubgraphNode.generated.h"
+
+
+UCLASS()
+class UMoviePipelineEdGraphSubgraphNode : public UMoviePipelineEdGraphNode
+{
+	GENERATED_BODY()
+
+public:
+	// ~Begin UEdGraphNode interface
+	virtual UObject* GetJumpTargetForDoubleClick() const override;
+	// ~End UEdGraphNode interface
+};
diff --git a/Engine/Plugins/MovieScene/MovieRenderPipeline/Source/MovieRenderPipelineEditor/Private/Widgets/Graph/SMovieGraphConfigPanel.cpp b/Engine/Plugins/MovieScene/MovieRenderPipeline/Source/MovieRenderPipelineEditor/Private/Widgets/Graph/SMovieGraphConfigPanel.cpp
--- a/Engine/Plugins/MovieScene/MovieRenderPipeline/Source/MovieRenderPipelineEditor/Private/Widgets/Graph/SMovieGraphConfigPanel.cpp
+++ b/Engine/Plugins/MovieScene/MovieRenderPipeline/Source/MovieRenderPipelineEditor/Private/Widgets/Graph/SMovieGraphConfigPanel.cpp
@@ -52,6 +52,7 @@
 #include "Misc/FileHelper.h"
 #include "EdGraphUtilities.h"
 #include "HAL/PlatformApplicationMisc.h"
+#include "Subsystems/AssetEditorSubsystem.h"
 
 #include "Graph/MovieGraphConfig.h"
 #include "Graph/Nodes/MovieGraphVariableNode.h"
@@ -665,6 +666,13 @@ void SMoviePipelineGraphPanel::OnSelectedNodesChanged(const TSet<UObject*>& NewS
 
 void SMoviePipelineGraphPanel::OnNodeDoubleClicked(class UEdGraphNode* Node)
 {
+	if (Node != nullptr)
+	{
+		if (UObject* Object = Node->GetJumpTargetForDoubleClick())
+		{
+			GEditor->GetEditorSubsystem<UAssetEditorSubsystem>()->OpenEditorForAsset(Object);
+		}
+	}
 }
 
 void SMoviePipelineGraphPanel::OnNodeTitleCommitted(const FText& NewText, ETextCommit::Type CommitInfo, UEdGraphNode* NodeBeingChanged)
