MovieGraph: Fixed rendering with landscapes crashing due to missing view extensions.

diff --git a/Engine/Plugins/MovieScene/MovieRenderPipeline/Source/MovieRenderPipelineCore/Private/Graph/Renderers/MovieGraphImagePassBase.cpp b/Engine/Plugins/MovieScene/MovieRenderPipeline/Source/MovieRenderPipelineCore/Private/Graph/Renderers/MovieGraphImagePassBase.cpp
--- a/Engine/Plugins/MovieScene/MovieRenderPipeline/Source/MovieRenderPipelineCore/Private/Graph/Renderers/MovieGraphImagePassBase.cpp
+++ b/Engine/Plugins/MovieScene/MovieRenderPipeline/Source/MovieRenderPipelineCore/Private/Graph/Renderers/MovieGraphImagePassBase.cpp
@@ -11,6 +11,9 @@
 #include "MovieRenderOverlappedImage.h"
 #include "Engine/RendererSettings.h"
 #include "UnrealClient.h"
+#include "SceneViewExtensionContext.h"
+#include "SceneViewExtension.h"
+#include "Engine/Engine.h"
 
 namespace UE::MovieGraph::Rendering
 {
@@ -142,6 +145,8 @@ TSharedRef<FSceneViewFamilyContext> FMovieGraphImagePassBase::CreateSceneViewFam
 		.SetTime(FGameTime::CreateUndilated(InInitData.TimeData.WorldSeconds, InInitData.TimeData.FrameDeltaTime))
 		.SetRealtimeUpdate(true));
 
+	// Need to add the engine-wide view extensions, as rendering code may depend on them (ie: landscapes)
+	OutViewFamily->ViewExtensions.Append(GEngine->ViewExtensions->GatherActiveExtensions(FSceneViewExtensionContext(InInitData.World->Scene)));
 
 	return OutViewFamily;
 }
