MovieGraph: Fixed a deadlock if you canceled a movie render (by hitting escape) but only during the last shot of the render.

diff --git a/Engine/Plugins/MovieScene/MovieRenderPipeline/Source/MovieRenderPipelineCore/Private/Graph/MovieGraphPipeline.cpp b/Engine/Plugins/MovieScene/MovieRenderPipeline/Source/MovieRenderPipelineCore/Private/Graph/MovieGraphPipeline.cpp
--- a/Engine/Plugins/MovieScene/MovieRenderPipeline/Source/MovieRenderPipelineCore/Private/Graph/MovieGraphPipeline.cpp
+++ b/Engine/Plugins/MovieScene/MovieRenderPipeline/Source/MovieRenderPipelineCore/Private/Graph/MovieGraphPipeline.cpp
@@ -895,7 +895,7 @@ void UMovieGraphPipeline::ShutdownImpl(bool bIsError)
 
 	// It's possible for a previous call to RequestShutdown to have set an error before this call and
 	// we don't want to just blow away the error flag even if this function was then called normally.
-	// (ie: Don't accidnetally unset error state)
+	// (ie: Don't accidentally unset error state)
 	if (bIsError)
 	{
 		bShutdownSetErrorFlag = true;
@@ -970,7 +970,7 @@ void UMovieGraphPipeline::TransitionToState(const EMovieRenderPipelineState InNe
 
 			// If we had naturally finished the last shot before doing this transition it will have
 			// already been torn down, so this only catches mid-shot transitions to ensure teardown.
-			if (CurrentShotIndex < (ActiveShotList.Num() - 1))
+			if (CurrentShotIndex < ActiveShotList.Num())
 			{
 				// Ensures all in-flight work for that shot is handled.
 				TeardownShot(ActiveShotList[CurrentShotIndex]);
