MRG: Fix for difference between png/jpg. We prevent png output node from writing linear16bit .

#rnx

diff --git a/Engine/Plugins/MovieScene/MovieRenderPipeline/Source/MovieRenderPipelineRenderPasses/Private/MovieGraphImageSequenceOutputNode.cpp b/Engine/Plugins/MovieScene/MovieRenderPipeline/Source/MovieRenderPipelineRenderPasses/Private/MovieGraphImageSequenceOutputNode.cpp
--- a/Engine/Plugins/MovieScene/MovieRenderPipeline/Source/MovieRenderPipelineRenderPasses/Private/MovieGraphImageSequenceOutputNode.cpp
+++ b/Engine/Plugins/MovieScene/MovieRenderPipeline/Source/MovieRenderPipelineRenderPasses/Private/MovieGraphImageSequenceOutputNode.cpp
@@ -11,6 +11,7 @@
 #include "Graph/MovieGraphBlueprintLibrary.h"
 #include "Modules/ModuleManager.h"
 #include "MoviePipelineUtils.h"
+#include "MoviePipelineImageSequenceOutput.h" // for FAsyncImageQuantization
 #include "MovieRenderPipelineCoreModule.h"
 #include "ImageWriteQueue.h"
 #include "Misc/Paths.h"
@@ -240,6 +241,7 @@ void UMovieGraphImageSequenceOutputNode::OnReceiveImageDataImpl(UMovieGraphPipel
 		TileImageTask->Filename = FileName;
 		TileImageTask->PixelData = RenderData.Value->CopyImageData();
 
+		bool bQuantizationEncodeSRGB = true;
 #if WITH_EDITOR
 		if (OCIOConfiguration.bIsEnabled)
 		{
@@ -247,12 +249,21 @@ void UMovieGraphImageSequenceOutputNode::OnReceiveImageDataImpl(UMovieGraphPipel
 			if (OCIOPixelPreProcessor)
 			{
 				TileImageTask->PixelPreProcessors.Emplace(MoveTemp(OCIOPixelPreProcessor));
+				
+				// We assume that any encoding on the output transform should be done by OCIO
+				bQuantizationEncodeSRGB = false;
 			}
 		}
 #endif
 
 		EImagePixelType PixelType = TileImageTask->PixelData->GetType();
 
+		if (bQuantizeTo8Bit && TileImageTask->PixelData->GetBitDepth() > 8u)
+		{
+			TileImageTask->PixelPreProcessors.Emplace(UE::MoviePipeline::FAsyncImageQuantization(TileImageTask.Get(), bQuantizationEncodeSRGB));
+			PixelType = EImagePixelType::Color;
+		}
+
 		// Perform compositing if any composited passes were found earlier
 		for (TPair<FMovieGraphRenderDataIdentifier, TUniquePtr<FImagePixelData>>& CompositedPass : CompositedPasses)
 		{
diff --git a/Engine/Plugins/MovieScene/MovieRenderPipeline/Source/MovieRenderPipelineRenderPasses/Private/MoviePipelineImageSequenceOutput.cpp b/Engine/Plugins/MovieScene/MovieRenderPipeline/Source/MovieRenderPipelineRenderPasses/Private/MoviePipelineImageSequenceOutput.cpp
--- a/Engine/Plugins/MovieScene/MovieRenderPipeline/Source/MovieRenderPipelineRenderPasses/Private/MoviePipelineImageSequenceOutput.cpp
+++ b/Engine/Plugins/MovieScene/MovieRenderPipeline/Source/MovieRenderPipelineRenderPasses/Private/MoviePipelineImageSequenceOutput.cpp
@@ -26,23 +26,27 @@
 #include UE_INLINE_GENERATED_CPP_BY_NAME(MoviePipelineImageSequenceOutput)
 
 DECLARE_CYCLE_STAT(TEXT("ImgSeqOutput_RecieveImageData"), STAT_ImgSeqRecieveImageData, STATGROUP_MoviePipeline);
-struct FAsyncImageQuantization
-{
-	FAsyncImageQuantization(FImageWriteTask* InWriteTask, const bool bInApplysRGB)
-		: ParentWriteTask(InWriteTask)
-		, bApplysRGB(bInApplysRGB)
-	{}
 
-	void operator()(FImagePixelData* PixelData)
+namespace UE
+{
+	namespace MoviePipeline
 	{
-		// Convert the incoming data to 8-bit, potentially with sRGB applied.
-		TUniquePtr<FImagePixelData> QuantizedPixelData = UE::MoviePipeline::QuantizeImagePixelDataToBitDepth(PixelData, 8, nullptr,  bApplysRGB);
-		ParentWriteTask->PixelData = MoveTemp(QuantizedPixelData);
-	}
+		FAsyncImageQuantization::FAsyncImageQuantization(FImageWriteTask* InWriteTask, const bool bInConvertToSRGB)
+			: ParentWriteTask(InWriteTask)
+			, bConvertToSRGB(bInConvertToSRGB)
+		{
+		}
 
-	FImageWriteTask* ParentWriteTask;
-	bool bApplysRGB;
-};
+		void FAsyncImageQuantization::operator()(FImagePixelData* PixelData)
+		{
+			// Note: Ideally we would use FImageCore routines here, but there is no easy way to construct pixel data from an FImage currently.
+
+			// Convert the incoming data to 8-bit, potentially with sRGB applied.
+			TUniquePtr<FImagePixelData> QuantizedPixelData = QuantizeImagePixelDataToBitDepth(PixelData, 8, nullptr, bConvertToSRGB);
+			ParentWriteTask->PixelData = MoveTemp(QuantizedPixelData);
+		}
+	}
+}
 
 UMoviePipelineImageSequenceOutputBase::UMoviePipelineImageSequenceOutputBase()
 {
@@ -208,7 +212,7 @@ void UMoviePipelineImageSequenceOutputBase::OnReceiveImageDataImpl(FMoviePipelin
 			// All three of these formats only support 8 bit data, so we need to take the incoming buffer type,
 			// copy it into a new 8-bit array and apply a little noise to the data to help hide gradient banding.
 			const bool bApplysRGB = !(ColorSetting && ColorSetting->OCIOConfiguration.bIsEnabled);
-			TileImageTask->PixelPreProcessors.Add(FAsyncImageQuantization(TileImageTask.Get(), bApplysRGB));
+			TileImageTask->PixelPreProcessors.Add(UE::MoviePipeline::FAsyncImageQuantization(TileImageTask.Get(), bApplysRGB));
 
 			// The pixel type will get changed by this pre-processor so future calculations below need to know the correct type they'll be editing.
 			QuantizedPixelType = EImagePixelType::Color; 
diff --git a/Engine/Plugins/MovieScene/MovieRenderPipeline/Source/MovieRenderPipelineRenderPasses/Public/MovieGraphImageSequenceOutputNode.h b/Engine/Plugins/MovieScene/MovieRenderPipeline/Source/MovieRenderPipelineRenderPasses/Public/MovieGraphImageSequenceOutputNode.h
--- a/Engine/Plugins/MovieScene/MovieRenderPipeline/Source/MovieRenderPipelineRenderPasses/Public/MovieGraphImageSequenceOutputNode.h
+++ b/Engine/Plugins/MovieScene/MovieRenderPipeline/Source/MovieRenderPipelineRenderPasses/Public/MovieGraphImageSequenceOutputNode.h
@@ -52,6 +52,9 @@ protected:
 	/** The output format (as known used by the ImageWriteQueue) to output into. */
 	EImageFormat OutputFormat;
 
+	/** Whether we enforce 8-bit depth on the output. */
+	bool bQuantizeTo8Bit;
+
 	/** A pointer to the image write queue used for asynchronously writing images */
 	IImageWriteQueue* ImageWriteQueue;
 
@@ -71,6 +74,7 @@ public:
 	UMovieGraphImageSequenceOutputNode_EXR()
 	{
 		OutputFormat = EImageFormat::EXR;
+		bQuantizeTo8Bit = false;
 		Compression = EEXRCompressionFormat::PIZ;
 		bMultilayer = true;
 	}
@@ -159,6 +163,7 @@ public:
 	UMovieGraphImageSequenceOutputNode_BMP()
 	{
 		OutputFormat = EImageFormat::BMP;
+		bQuantizeTo8Bit = true;
 	}
 
 #if WITH_EDITOR
@@ -207,6 +212,7 @@ public:
 	UMovieGraphImageSequenceOutputNode_JPG()
 	{
 		OutputFormat = EImageFormat::JPEG;
+		bQuantizeTo8Bit = true;
 	}
 
 #if WITH_EDITOR
@@ -254,6 +260,9 @@ public:
 	UMovieGraphImageSequenceOutputNode_PNG()
 	{
 		OutputFormat = EImageFormat::PNG;
+
+		// Note: we could offer linear 16-bit pngs simply by letting users turn this to false.
+		bQuantizeTo8Bit = true;
 	}
 
 #if WITH_EDITOR
diff --git a/Engine/Plugins/MovieScene/MovieRenderPipeline/Source/MovieRenderPipelineRenderPasses/Public/MoviePipelineImageSequenceOutput.h b/Engine/Plugins/MovieScene/MovieRenderPipeline/Source/MovieRenderPipelineRenderPasses/Public/MoviePipelineImageSequenceOutput.h
--- a/Engine/Plugins/MovieScene/MovieRenderPipeline/Source/MovieRenderPipelineRenderPasses/Public/MoviePipelineImageSequenceOutput.h
+++ b/Engine/Plugins/MovieScene/MovieRenderPipeline/Source/MovieRenderPipelineRenderPasses/Public/MoviePipelineImageSequenceOutput.h
@@ -6,7 +6,30 @@
 #include "MoviePipelineImageSequenceOutput.generated.h"
 
 // Forward Declare
+struct FImagePixelData;
 class IImageWriteQueue;
+class FImageWriteTask;
+
+namespace UE
+{
+	namespace MoviePipeline
+	{
+		struct FAsyncImageQuantization
+		{
+			/** Constructor.*/
+			FAsyncImageQuantization(FImageWriteTask* InWriteTask, const bool bInConvertToSRGB);
+
+			/** Process/quantize pixel data */
+			void operator()(FImagePixelData* PixelData);
+
+			/** Parent image write task. */
+			FImageWriteTask* ParentWriteTask;
+
+			/** True if the quantization should also apply an sRGB encoding. */
+			bool bConvertToSRGB;
+		};
+	}
+}
 
 UCLASS(Blueprintable, Abstract)
 class MOVIERENDERPIPELINERENDERPASSES_API UMoviePipelineImageSequenceOutputBase : public UMoviePipelineOutputBase
