MoviePipeline: Fixed an issue where MoviePipeline.h shadowed GetPipelineState which causes the MoviePipelinePIEExecutor to call the base implementation (which always returns uninitialized).

diff --git a/Engine/Plugins/MovieScene/MovieRenderPipeline/Source/MovieRenderPipelineCore/Public/MoviePipeline.h b/Engine/Plugins/MovieScene/MovieRenderPipeline/Source/MovieRenderPipelineCore/Public/MoviePipeline.h
--- a/Engine/Plugins/MovieScene/MovieRenderPipeline/Source/MovieRenderPipelineCore/Public/MoviePipeline.h
+++ b/Engine/Plugins/MovieScene/MovieRenderPipeline/Source/MovieRenderPipelineCore/Public/MoviePipeline.h
@@ -58,6 +58,7 @@ public:
 	virtual void RequestShutdownImpl(bool bIsError) override;
 	virtual void ShutdownImpl(bool bError ) override;
 	virtual bool IsShutdownRequestedImpl() const override { return bShutdownRequested; }
+	virtual EMovieRenderPipelineState GetPipelineStateImpl() const override { return PipelineState; }
 	// ~UMoviePipelineBase Interface
 
 	/**
@@ -140,7 +141,6 @@ public:
 	UFUNCTION(BlueprintPure, Category = "Movie Render Pipeline")
 	UMoviePipelineExecutorJob* GetCurrentJob() const { return CurrentJob; }
 
-	EMovieRenderPipelineState GetPipelineState() const { return PipelineState; }
 	FMoviePipelineOutputData GetOutputDataParams();
 
 	void GetSidecarCameraData(UMoviePipelineExecutorShot* InShot, int32 InCameraIndex, FMinimalViewInfo& OutViewInfo, class UCameraComponent** OutCameraComponent) const;
