MoviePipeline: Added a details customization for the new collections UX.

Also updated the existing condition group queries to provide extra info to the UI (GetWidgets() and GetAddMenuContents()), and added an additional Actor filter.

#rnx

diff --git a/Engine/Plugins/MovieScene/MovieRenderPipeline/Source/MovieRenderPipelineCore/MovieRenderPipelineCore.Build.cs b/Engine/Plugins/MovieScene/MovieRenderPipeline/Source/MovieRenderPipelineCore/MovieRenderPipelineCore.Build.cs
--- a/Engine/Plugins/MovieScene/MovieRenderPipeline/Source/MovieRenderPipelineCore/MovieRenderPipelineCore.Build.cs
+++ b/Engine/Plugins/MovieScene/MovieRenderPipeline/Source/MovieRenderPipelineCore/MovieRenderPipelineCore.Build.cs
@@ -42,6 +42,13 @@ public class MovieRenderPipelineCore : ModuleRules
 
 		if (Target.bBuildEditor == true)
         {
+	        PrivateDependencyModuleNames.AddRange(
+		        new string[]
+		        {
+			        "ClassViewer",
+			        "SceneOutliner",
+		        });
+	        
 			PublicDependencyModuleNames.Add("MovieSceneTools");
         }
 	}
diff --git a/Engine/Plugins/MovieScene/MovieRenderPipeline/Source/MovieRenderPipelineCore/Private/Graph/MoviePipelineRenderLayerSubsystem.cpp b/Engine/Plugins/MovieScene/MovieRenderPipeline/Source/MovieRenderPipelineCore/Private/Graph/MoviePipelineRenderLayerSubsystem.cpp
--- a/Engine/Plugins/MovieScene/MovieRenderPipeline/Source/MovieRenderPipelineCore/Private/Graph/MoviePipelineRenderLayerSubsystem.cpp
+++ b/Engine/Plugins/MovieScene/MovieRenderPipeline/Source/MovieRenderPipelineCore/Private/Graph/MoviePipelineRenderLayerSubsystem.cpp
@@ -7,10 +7,181 @@
 #include "Components/SkyLightComponent.h"
 #include "Components/SphereReflectionCaptureComponent.h"
 #include "EngineUtils.h"
+#include "Framework/Application/SlateApplication.h"
 #include "Materials/MaterialInterface.h"
+#include "Modules/ModuleManager.h"
 #include "MovieRenderPipelineCoreModule.h"
 #include "Styling/AppStyle.h"
+#include "Styling/SlateIconFinder.h"
 #include "UObject/Package.h"
+#include "Widgets/Input/SMultiLineEditableTextBox.h"
+
+#if WITH_EDITOR
+#include "ActorTreeItem.h"
+#include "ClassViewerFilter.h"
+#include "ClassViewerModule.h"
+#include "ISceneOutliner.h"
+#include "SceneOutlinerModule.h"
+#include "SceneOutlinerPublicTypes.h"
+#include "SClassViewer.h"
+#include "Widgets/SCompoundWidget.h"
+#endif
+
+#define LOCTEXT_NAMESPACE "MovieGraph"
+
+namespace UE::MovieGraph::Private
+{
+#if WITH_EDITOR
+	/**
+	 * A widget which lists items in rows w/ an alternating row color, where each item has an icon and text. A summary row appears at the end of the
+	 * list indicating how many items are in the list.
+	 */
+	template<typename ListType>
+	class SQueryContentsList final : public SCompoundWidget
+	{
+	public:
+		DECLARE_DELEGATE_RetVal_OneParam(const FSlateBrush*, FGetRowIcon, ListType);
+		DECLARE_DELEGATE_RetVal_OneParam(FText, FGetRowText, ListType);
+		
+		SLATE_BEGIN_ARGS(SQueryContentsList<ListType>)
+		{}
+			SLATE_ATTRIBUTE(TArray<ListType>*, DataSource)
+			SLATE_ATTRIBUTE(FText, DataType)
+			SLATE_ATTRIBUTE(FText, DataTypePlural)
+			SLATE_EVENT(FGetRowIcon, OnGetRowIcon);
+			SLATE_EVENT(FGetRowText, OnGetRowText);
+		SLATE_END_ARGS()
+
+		const FSlateBrush* GetRowIcon(const ListType& InListData) const
+		{
+			if (OnGetRowIcon.IsBound())
+			{
+				return OnGetRowIcon.Execute(InListData);
+			}
+			
+			return nullptr;
+		}
+
+		FText GetRowText(const ListType& InListData) const
+		{
+			if (OnGetRowText.IsBound())
+			{
+				return OnGetRowText.Execute(InListData);
+			}
+			
+			return FText();
+		}
+
+		void Construct(const FArguments& InArgs)
+		{
+			DataSource = InArgs._DataSource.Get();
+			DataType = InArgs._DataType.Get();
+			DataTypePlural = InArgs._DataTypePlural.Get();
+			OnGetRowIcon = InArgs._OnGetRowIcon;
+			OnGetRowText = InArgs._OnGetRowText;
+			
+			ChildSlot
+			.HAlign(HAlign_Fill)
+			.VAlign(VAlign_Top)
+			[
+				SNew(SVerticalBox)
+				+ SVerticalBox::Slot()
+				.AutoHeight()
+				[
+					SNew(SListView<ListType>)
+					.ListItemsSource(DataSource)
+					.OnGenerateRow_Lambda([this](const ListType& InListData, const TSharedRef<STableViewBase>& InOwnerTable)
+					{
+						return SNew(STableRow<ListType>, InOwnerTable)
+							.Style(FAppStyle::Get(), "TableView.AlternatingRow")
+							.ShowWires(false)
+							[
+								SNew(SHorizontalBox)
+								+ SHorizontalBox::Slot()
+								.VAlign(VAlign_Center)
+								.Padding(7.f, 5.f, 7.f, 5.f)
+								.AutoWidth()
+								[
+									SNew(SImage)
+									.Image(GetRowIcon(InListData))
+								]
+									
+								+ SHorizontalBox::Slot()
+								.VAlign(VAlign_Center)
+								.HAlign(HAlign_Fill)
+								[
+									SNew(STextBlock)
+									.Font(FAppStyle::GetFontStyle(TEXT("PropertyWindow.NormalFont")))
+									.Text(GetRowText(InListData))
+								]
+							];
+					})
+				]
+
+				+ SVerticalBox::Slot()
+				.AutoHeight()
+				[
+					SNew(SBorder)
+					.BorderImage(FAppStyle::Get().GetBrush("Brushes.Header"))
+					.VAlign(VAlign_Center)
+					.HAlign(HAlign_Fill)
+					.Padding(FMargin(14, 4))
+					[
+						SNew(STextBlock)
+						.Font(FAppStyle::GetFontStyle(TEXT("PropertyWindow.NormalFont")))
+						.Text_Lambda([this]()
+						{
+							return FText::FromString(FString::Printf(TEXT("%i %s"), DataSource->Num(), DataSource->Num() == 1 ? *DataType.ToString() : *DataTypePlural.ToString()));
+						})
+					]
+				]
+			];
+		}
+
+	private:
+		FGetRowIcon OnGetRowIcon;
+		FGetRowText OnGetRowText;
+		FText DataType;
+		FText DataTypePlural;
+		TArray<ListType>* DataSource = nullptr;
+	};
+
+	/**
+	 * A filter that can be used in the class viewer that appears in the Add menu. Filters out specified classes, and optionally filters out classes
+	 * that do not have a specific base class.
+	 */
+	class FClassViewerTypeFilter final : public IClassViewerFilter
+	{
+	public:
+		explicit FClassViewerTypeFilter(TArray<UClass*>* InClassesToDisallow, UClass* InRequiredBaseClass = nullptr)
+			: ClassesToDisallow(InClassesToDisallow)
+			, RequiredBaseClass(InRequiredBaseClass)
+		{
+			
+		}
+
+		virtual bool IsClassAllowed(const FClassViewerInitializationOptions& InInitOptions, const UClass* InClass, TSharedRef<FClassViewerFilterFuncs> InFilterFuncs) override
+		{
+			return
+				InClass &&
+				!ClassesToDisallow->Contains(InClass) &&
+				(RequiredBaseClass ? InClass->IsChildOf(RequiredBaseClass) : true);
+		}
+		
+		virtual bool IsUnloadedClassAllowed(const FClassViewerInitializationOptions& InInitOptions, const TSharedRef<const IUnloadedBlueprintData> InUnloadedClassData, TSharedRef<FClassViewerFilterFuncs> InFilterFuncs) override
+		{
+			return false;
+		}
+
+	private:
+		/** Classes which should be prevented from showing up in the class viewer. */
+		TArray<UClass*>* ClassesToDisallow;
+
+		/** Classes must have this base class to pass the filter. */
+		UClass* RequiredBaseClass = nullptr;
+	};
+#endif
+}
 
 void UMoviePipelineMaterialModifier::ApplyModifier(const UWorld* World)
 {
@@ -323,7 +494,7 @@ void UMoviePipelineCollection::AddQuery(UMoviePipelineCollectionQuery* Query)
 }
 
 UMovieGraphConditionGroupQueryBase::UMovieGraphConditionGroupQueryBase()
-	: OpType(EMovieGraphConditionGroupQueryOpType::Union)
+	: OpType(EMovieGraphConditionGroupQueryOpType::Add)
 	, bIsEnabled(true)
 {
 }
@@ -332,9 +503,9 @@ void UMovieGraphConditionGroupQueryBase::SetOperationType(const EMovieGraphCondi
 {
 	// Always allow setting the operation type to Union. If not setting to Union, only allow setting the operation type if this is not the first
 	// query in the condition group. The first query is always a Union.
-	if (OperationType == EMovieGraphConditionGroupQueryOpType::Union)
+	if (OperationType == EMovieGraphConditionGroupQueryOpType::Add)
 	{
-		OpType = EMovieGraphConditionGroupQueryOpType::Union;
+		OpType = EMovieGraphConditionGroupQueryOpType::Add;
 		return;
 	}
 
@@ -370,6 +541,18 @@ const FSlateIcon& UMovieGraphConditionGroupQueryBase::GetIcon() const
 	return EmptyIcon;
 }
 
+#if WITH_EDITOR
+TArray<TSharedRef<SWidget>> UMovieGraphConditionGroupQueryBase::GetWidgets()
+{
+	return TArray<TSharedRef<SWidget>>();
+}
+
+TSharedRef<SWidget> UMovieGraphConditionGroupQueryBase::GetAddMenuContents(const FMovieGraphConditionGroupQueryContentsChanged& OnAddFinished)
+{
+	return SNullWidget::NullWidget;
+}
+#endif
+
 bool UMovieGraphConditionGroupQueryBase::IsEditorOnlyQuery() const
 {
 	return false;
@@ -385,43 +568,223 @@ bool UMovieGraphConditionGroupQueryBase::IsEnabled() const
 	return bIsEnabled;
 }
 
-void UMovieGraphConditionGroupQuery_ActorTag::Evaluate(const TArray<AActor*>& InActorsToQuery, TSet<AActor*>& OutMatchingActors) const
+bool UMovieGraphConditionGroupQueryBase::IsFirstConditionGroupQuery() const
 {
-	TRACE_CPUPROFILER_EVENT_SCOPE(UMovieGraphConditionGroupQuery_ActorTag::Evaluate);
+	const UMovieGraphConditionGroup* ParentConditionGroup = GetTypedOuter<UMovieGraphConditionGroup>();
+	if (ensureMsgf(ParentConditionGroup, TEXT("Cannot determine if this is the first condition group query when no parent condition group is present")))
+	{
+		// GetQueries() returns an array of non-const pointers, so Find() doesn't like having a const pointer passed to it.
+		// Find() won't mutate the condition group query though, so the const_cast here is OK.
+		return ParentConditionGroup->GetQueries().Find(const_cast<UMovieGraphConditionGroupQueryBase*>(this)) == 0;
+	}
 	
+	return false;
+}
+
+void UMovieGraphConditionGroupQuery_Actor::Evaluate(const TArray<AActor*>& InActorsToQuery, TSet<AActor*>& OutMatchingActors) const
+{
+	TRACE_CPUPROFILER_EVENT_SCOPE(UMovieGraphConditionGroupQuery_Actor::Evaluate);
+
 	for (AActor* Actor : InActorsToQuery)
 	{
-		if (Actor->Tags.Contains(TagToMatch))
+		if (ActorsToMatch.Contains(Actor))
 		{
 			OutMatchingActors.Add(Actor);
 		}
 	}
 }
 
-const FSlateIcon& UMovieGraphConditionGroupQuery_ActorTag::GetIcon() const
+const FSlateIcon& UMovieGraphConditionGroupQuery_Actor::GetIcon() const
+{
+	static const FSlateIcon ActorIcon = FSlateIcon(FAppStyle::GetAppStyleSetName(), "ClassIcon.Actor");
+	return ActorIcon;
+}
+
+#if WITH_EDITOR
+TArray<TSharedRef<SWidget>> UMovieGraphConditionGroupQuery_Actor::GetWidgets()
+{
+	TArray<TSharedRef<SWidget>> Widgets;
+
+	// Create the data source for the list view
+	ListDataSource.Empty();
+	for (TSoftObjectPtr<AActor>& Actor : ActorsToMatch)
+	{
+		ListDataSource.Add(MakeShared<TSoftObjectPtr<AActor>>(Actor));
+	}
+
+	Widgets.Add(
+		SNew(UE::MovieGraph::Private::SQueryContentsList<TSharedPtr<TSoftObjectPtr<AActor>>>)
+			.DataSource(&ListDataSource)
+			.DataType(FText::FromString("Actor"))
+			.DataTypePlural(FText::FromString("Actors"))
+			.OnGetRowText_Static(&GetRowText)
+			.OnGetRowIcon_Static(&GetRowIcon)
+	);
+
+	return Widgets;
+}
+
+TSharedRef<SWidget> UMovieGraphConditionGroupQuery_Actor::GetAddMenuContents(const FMovieGraphConditionGroupQueryContentsChanged& OnAddFinished)
+{	
+	FSceneOutlinerInitializationOptions SceneOutlinerInitOptions;
+	SceneOutlinerInitOptions.bShowHeaderRow = true;
+	SceneOutlinerInitOptions.bShowSearchBox = true;
+	SceneOutlinerInitOptions.bShowCreateNewFolder = false;
+	SceneOutlinerInitOptions.bFocusSearchBoxWhenOpened = true;
+
+	// Show the name/label column and the type column
+	SceneOutlinerInitOptions.ColumnMap.Add(
+		FSceneOutlinerBuiltInColumnTypes::Label(),
+		FSceneOutlinerColumnInfo(ESceneOutlinerColumnVisibility::Visible, 0, FCreateSceneOutlinerColumn(), false, TOptional<float>(), FSceneOutlinerBuiltInColumnTypes::Label_Localized()));
+	SceneOutlinerInitOptions.ColumnMap.Add(
+		FSceneOutlinerBuiltInColumnTypes::ActorInfo(),
+		FSceneOutlinerColumnInfo(ESceneOutlinerColumnVisibility::Visible, 10,  FCreateSceneOutlinerColumn(), false, TOptional<float>(), FSceneOutlinerBuiltInColumnTypes::ActorInfo_Localized()));
+
+	// Don't show actors which have already been picked
+	SceneOutlinerInitOptions.Filters->AddFilterPredicate<FActorTreeItem>(
+		FActorTreeItem::FFilterPredicate::CreateLambda([this](const AActor* InActor)
+		{
+			return !ActorsToMatch.Contains(InActor);
+		}));
+	
+	const FSceneOutlinerModule& SceneOutlinerModule = FModuleManager::LoadModuleChecked<FSceneOutlinerModule>("SceneOutliner");
+
+	ActorPickerWidget = SceneOutlinerModule.CreateActorPicker(
+		SceneOutlinerInitOptions,
+		FOnActorPicked::CreateLambda([this, OnAddFinished](AActor* InActor)
+		{
+			ActorsToMatch.Add(InActor);
+			ListDataSource.Add(MakeShared<TSoftObjectPtr<AActor>>(ActorsToMatch.Last()));
+			OnAddFinished.ExecuteIfBound();
+
+			// Ensure that the filter runs again so duplicate actors cannot be selected
+			if (ActorPickerWidget.IsValid())
+			{
+				ActorPickerWidget->FullRefresh();
+			}
+		}));
+	
+	return
+		SNew(SBox)
+		.WidthOverride(400.f)
+		.HeightOverride(300.f)
+		[
+			ActorPickerWidget.ToSharedRef()
+		];
+}
+
+const FSlateBrush* UMovieGraphConditionGroupQuery_Actor::GetRowIcon(TSharedPtr<TSoftObjectPtr<AActor>> InActor)
+{
+	if (InActor.IsValid())
+	{
+		if (InActor.Get()->IsValid())
+		{
+			// The first Get() returns the TSoftObjectPtr, the second Get() dereferences the TSoftObjectPtr
+			return FSlateIconFinder::FindIconForClass(InActor.Get()->Get()->GetClass()).GetIcon();
+		}
+	}
+
+	return FSlateIconFinder::FindIconForClass(AActor::StaticClass()).GetIcon();
+}
+
+FText UMovieGraphConditionGroupQuery_Actor::GetRowText(TSharedPtr<TSoftObjectPtr<AActor>> InActor)
+{
+	if (InActor.IsValid())
+	{
+		if (InActor.Get()->IsValid())
+		{
+			// The first Get() returns the TSoftObjectPtr, the second Get() dereferences the TSoftObjectPtr
+			return FText::FromString(InActor.Get()->Get()->GetActorLabel());
+		}
+	}
+
+	return LOCTEXT("MovieGraphActorConditionGroupQuery_InvalidActor", "(invalid)");
+}
+#endif	// WITH_EDITOR
+
+void UMovieGraphConditionGroupQuery_ActorTagName::Evaluate(const TArray<AActor*>& InActorsToQuery, TSet<AActor*>& OutMatchingActors) const
+{
+	TRACE_CPUPROFILER_EVENT_SCOPE(UMovieGraphConditionGroupQuery_ActorTag::Evaluate);
+	
+	// Quick early-out if "*" is used as the wildcard. Faster than doing the wildcard matching.
+	if (TagsToMatch == TEXT("*"))
+	{
+		OutMatchingActors.Append(InActorsToQuery);
+		return;
+	}
+
+	// Actor tags can be specified on multiple lines
+	TArray<FString> AllTagNameStrings;
+	TagsToMatch.ParseIntoArrayLines(AllTagNameStrings);
+
+	for (AActor* Actor : InActorsToQuery)
+	{
+		for (const FString& TagToMatch : AllTagNameStrings)
+		{
+			bool bMatchedTag = false;
+			
+			for (const FName& ActorTag : Actor->Tags)
+			{
+				if (ActorTag.ToString().MatchesWildcard(TagToMatch))
+				{
+					OutMatchingActors.Add(Actor);
+					bMatchedTag = true;
+					break;
+				}
+			}
+
+			// Skip comparing the rest of the tags if one tag already matched 
+			if (bMatchedTag)
+			{
+				break;
+			}
+		}
+	}
+}
+
+const FSlateIcon& UMovieGraphConditionGroupQuery_ActorTagName::GetIcon() const
 {
 	// TODO: This icon is wrong
-	static const FSlateIcon ActorTagIcon = FSlateIcon(FAppStyle::GetAppStyleSetName(), "TODO");
+	static const FSlateIcon ActorTagIcon = FSlateIcon(FAppStyle::GetAppStyleSetName(), "LevelEditor.Tabs.Debug");
 	return ActorTagIcon;
 }
 
+#if WITH_EDITOR
+TArray<TSharedRef<SWidget>> UMovieGraphConditionGroupQuery_ActorTagName::GetWidgets()
+{
+	TArray<TSharedRef<SWidget>> Widgets;
+
+	Widgets.Add(
+		SNew(SBox)
+		.HAlign(HAlign_Fill)
+		.Padding(7.f, 2.f)
+		[
+			SNew(SMultiLineEditableTextBox)
+			.Text_Lambda([this]() { return FText::FromString(TagsToMatch); })
+			.OnTextChanged_Lambda([this](const FText& InText) { TagsToMatch = InText.ToString(); })
+			.HintText(LOCTEXT("MovieGraphActorTagNameQueryHintText", "The actor must match one or more tags. Wildcards allowed.\nEnter each tag on a separate line."))
+		]
+	);
+
+	return Widgets;
+}
+#endif
+
 void UMovieGraphConditionGroupQuery_ActorName::Evaluate(const TArray<AActor*>& InActorsToQuery, TSet<AActor*>& OutMatchingActors) const
 {
 	TRACE_CPUPROFILER_EVENT_SCOPE(UMovieGraphConditionGroupQuery_ActorName::Evaluate);
 	
-	const bool bIsMatchAll = (WildcardSearch == TEXT("*"));
-
-	// Actor names can be specified on multiple lines
-	TArray<FString> AllActorNames;
-	WildcardSearch.ParseIntoArrayLines(AllActorNames);
-
 	// Quick early-out if "*" is used as the wildcard. Faster than doing the wildcard matching.
-	if (bIsMatchAll)
+	if (WildcardSearch == TEXT("*"))
 	{
 		OutMatchingActors.Append(InActorsToQuery);
 		return;
 	}
 
+	// Actor names can be specified on multiple lines
+	TArray<FString> AllActorNames;
+	WildcardSearch.ParseIntoArrayLines(AllActorNames);
+
 	for (AActor* Actor : InActorsToQuery)
 	{
 #if WITH_EDITOR
@@ -442,6 +805,27 @@ const FSlateIcon& UMovieGraphConditionGroupQuery_ActorName::GetIcon() const
 	return ActorTagIcon;
 }
 
+#if WITH_EDITOR
+TArray<TSharedRef<SWidget>> UMovieGraphConditionGroupQuery_ActorName::GetWidgets()
+{
+	TArray<TSharedRef<SWidget>> Widgets;
+
+	Widgets.Add(
+		SNew(SBox)
+		.HAlign(HAlign_Fill)
+		.Padding(7.f, 2.f)
+		[
+			SNew(SMultiLineEditableTextBox)
+			.Text_Lambda([this]() { return FText::FromString(WildcardSearch); })
+			.OnTextChanged_Lambda([this](const FText& InText) { WildcardSearch = InText.ToString(); })
+			.HintText(LOCTEXT("MovieGraphActorNameQueryHintText", "Actor names to query. Wildcards allowed.\nEnter each actor name on a separate line."))
+		]
+	);
+
+	return Widgets;
+}
+#endif
+
 bool UMovieGraphConditionGroupQuery_ActorName::IsEditorOnly() const
 {
 	// GetActorLabel() is editor-only
@@ -454,7 +838,7 @@ void UMovieGraphConditionGroupQuery_ActorType::Evaluate(const TArray<AActor*>& I
 	
 	for (AActor* Actor : InActorsToQuery)
 	{
-		if (Actor->GetClass() == ActorType)
+		if (ActorTypes.Contains(Actor->GetClass()))
 		{
 			OutMatchingActors.Add(Actor);
 		}
@@ -467,10 +851,89 @@ const FSlateIcon& UMovieGraphConditionGroupQuery_ActorType::GetIcon() const
 	return ActorTagIcon;
 }
 
-void UMovieGraphConditionGroupQuery_ComponentTag::Evaluate(const TArray<AActor*>& InActorsToQuery, TSet<AActor*>& OutMatchingActors) const
+#if WITH_EDITOR
+TArray<TSharedRef<SWidget>> UMovieGraphConditionGroupQuery_ActorType::GetWidgets()
+{
+	TArray<TSharedRef<SWidget>> Widgets;
+
+	Widgets.Add(
+		SNew(UE::MovieGraph::Private::SQueryContentsList<UClass*>)
+			.DataSource(&ActorTypes)
+			.DataType(FText::FromString("Actor Type"))
+			.DataTypePlural(FText::FromString("Actor Types"))
+			.OnGetRowText_Static(&GetRowText)
+			.OnGetRowIcon_Static(&GetRowIcon)
+	);
+
+	return Widgets;
+}
+
+TSharedRef<SWidget> UMovieGraphConditionGroupQuery_ActorType::GetAddMenuContents(const FMovieGraphConditionGroupQueryContentsChanged& OnAddFinished)
+{
+	FClassViewerModule& ClassViewerModule = FModuleManager::LoadModuleChecked<FClassViewerModule>("ClassViewer");
+
+	FClassViewerInitializationOptions Options;
+	Options.Mode = EClassViewerMode::ClassPicker;
+	Options.NameTypeToDisplay = EClassViewerNameTypeToDisplay::DisplayName;
+	Options.bShowNoneOption = false;
+	Options.bIsActorsOnly = true;
+	Options.bShowUnloadedBlueprints = false;
+
+	// Add a class filter to disallow adding duplicates of actor types that were already picked
+	Options.ClassFilters.Add(MakeShared<UE::MovieGraph::Private::FClassViewerTypeFilter>(&ActorTypes));
+
+	const TSharedRef<SWidget> ClassViewer = ClassViewerModule.CreateClassViewer(
+		Options,
+		FOnClassPicked::CreateLambda([this, OnAddFinished](UClass* InNewClass)
+		{
+			FSlateApplication::Get().DismissAllMenus();
+			
+			ActorTypes.Add(InNewClass);
+			OnAddFinished.ExecuteIfBound();
+
+			// Ensure that the class filters run again so duplicate actor types cannot be selected
+			if (ClassViewerWidget.IsValid())
+			{
+				ClassViewerWidget->Refresh();
+			}
+		}));
+
+	ClassViewerWidget = StaticCastSharedPtr<SClassViewer>(ClassViewer.ToSharedPtr()); 
+	
+	return SNew(SBox)
+		.WidthOverride(300.f)
+		.HeightOverride(300.f)
+		[
+			ClassViewerWidget.ToSharedRef()
+		];
+}
+
+const FSlateBrush* UMovieGraphConditionGroupQuery_ActorType::GetRowIcon(UClass* InActorType)
+{
+	return FSlateIconFinder::FindIconForClass(InActorType).GetIcon();
+}
+
+FText UMovieGraphConditionGroupQuery_ActorType::GetRowText(UClass* InActorType)
+{
+	return InActorType->GetDisplayNameText();
+}
+#endif
+
+void UMovieGraphConditionGroupQuery_ComponentTagName::Evaluate(const TArray<AActor*>& InActorsToQuery, TSet<AActor*>& OutMatchingActors) const
 {
 	TRACE_CPUPROFILER_EVENT_SCOPE(UMovieGraphConditionGroupQuery_ComponentTag::Evaluate);
 	
+	// Quick early-out if "*" is used as the wildcard. Faster than doing the wildcard matching.
+	if (TagsToMatch == TEXT("*"))
+	{
+		OutMatchingActors.Append(InActorsToQuery);
+		return;
+	}
+
+	// Component tags can be specified on multiple lines
+	TArray<FString> AllTagNameStrings;
+	TagsToMatch.ParseIntoArrayLines(AllTagNameStrings);
+	
 	TInlineComponentArray<UActorComponent*> ActorComponents;
 
 	for (AActor* Actor : InActorsToQuery)
@@ -479,9 +942,31 @@ void UMovieGraphConditionGroupQuery_ComponentTag::Evaluate(const TArray<AActor*>
 		
 		for (const UActorComponent* Component : ActorComponents)
 		{
-			if (Component && Component->ComponentTags.Contains(TagToMatch))
+			bool bMatchedTag = false;
+			
+			for (const FString& TagToMatch : AllTagNameStrings)
+			{			
+				for (const FName& ComponentTag : Component->ComponentTags)
+				{
+					if (ComponentTag.ToString().MatchesWildcard(TagToMatch))
+					{
+						OutMatchingActors.Add(Actor);
+						bMatchedTag = true;
+						break;
+					}
+				}
+
+				// Skip comparing the rest of the tags if one tag already matched 
+				if (bMatchedTag)
+				{
+					break;
+				}
+			}
+
+			// Skip comparing the rest of the components if one component already matched
+			if (bMatchedTag)
 			{
-				OutMatchingActors.Add(Actor);
+				break;
 			}
 		}
 
@@ -489,13 +974,34 @@ void UMovieGraphConditionGroupQuery_ComponentTag::Evaluate(const TArray<AActor*>
 	}
 }
 
-const FSlateIcon& UMovieGraphConditionGroupQuery_ComponentTag::GetIcon() const
+const FSlateIcon& UMovieGraphConditionGroupQuery_ComponentTagName::GetIcon() const
 {
 	// TODO: This icon is wrong
 	static const FSlateIcon ActorTagIcon = FSlateIcon(FAppStyle::GetAppStyleSetName(), "ClassIcon.DataAsset");
 	return ActorTagIcon;
 }
 
+#if WITH_EDITOR
+TArray<TSharedRef<SWidget>> UMovieGraphConditionGroupQuery_ComponentTagName::GetWidgets()
+{
+	TArray<TSharedRef<SWidget>> Widgets;
+
+	Widgets.Add(
+		SNew(SBox)
+		.HAlign(HAlign_Fill)
+		.Padding(7.f, 2.f)
+		[
+			SNew(SMultiLineEditableTextBox)
+			.Text_Lambda([this]() { return FText::FromString(TagsToMatch); })
+			.OnTextChanged_Lambda([this](const FText& InText) { TagsToMatch = InText.ToString(); })
+			.HintText(LOCTEXT("MovieGraphComponentTagNameQueryHintText", "A component on the actor must match one or more component tags.\nWildcards allowed. Enter each tag on a separate line."))
+		]
+	);
+
+	return Widgets;
+}
+#endif
+
 void UMovieGraphConditionGroupQuery_ComponentType::Evaluate(const TArray<AActor*>& InActorsToQuery, TSet<AActor*>& OutMatchingActors) const
 {
 	TRACE_CPUPROFILER_EVENT_SCOPE(UMovieGraphConditionGroupQuery_ComponentType::Evaluate);
@@ -508,7 +1014,7 @@ void UMovieGraphConditionGroupQuery_ComponentType::Evaluate(const TArray<AActor*
 		
 		for (const UActorComponent* Component : ActorComponents)
 		{
-			if (Component->GetClass() == ComponentType)
+			if (ComponentTypes.Contains(Component->GetClass()))
 			{
 				OutMatchingActors.Add(Actor);
 			}
@@ -524,9 +1030,77 @@ const FSlateIcon& UMovieGraphConditionGroupQuery_ComponentType::GetIcon() const
 	return ActorTagIcon;
 }
 
+#if WITH_EDITOR
+TArray<TSharedRef<SWidget>> UMovieGraphConditionGroupQuery_ComponentType::GetWidgets()
+{
+	TArray<TSharedRef<SWidget>> Widgets;
+
+	Widgets.Add(
+		SNew(UE::MovieGraph::Private::SQueryContentsList<UClass*>)
+			.DataSource(&ComponentTypes)
+			.DataType(FText::FromString("Component Type"))
+			.DataTypePlural(FText::FromString("Component Types"))
+			.OnGetRowText_Static(&GetRowText)
+			.OnGetRowIcon_Static(&GetRowIcon)
+	);
+
+	return Widgets;
+}
+
+TSharedRef<SWidget> UMovieGraphConditionGroupQuery_ComponentType::GetAddMenuContents(const FMovieGraphConditionGroupQueryContentsChanged& OnAddFinished)
+{
+	FClassViewerModule& ClassViewerModule = FModuleManager::LoadModuleChecked<FClassViewerModule>("ClassViewer");
+
+	FClassViewerInitializationOptions Options;
+	Options.Mode = EClassViewerMode::ClassPicker;
+	Options.NameTypeToDisplay = EClassViewerNameTypeToDisplay::DisplayName;
+	Options.bShowNoneOption = false;
+	Options.bIsActorsOnly = false;
+	Options.bShowUnloadedBlueprints = false;
+
+	// Add a class filter to disallow adding duplicates of component types that were already picked, as well as restrict the types of classes displayed
+	// to only show component classes
+	Options.ClassFilters.Add(MakeShared<UE::MovieGraph::Private::FClassViewerTypeFilter>(&ComponentTypes, UActorComponent::StaticClass()));
+
+	const TSharedRef<SWidget> ClassViewer = ClassViewerModule.CreateClassViewer(
+		Options,
+		FOnClassPicked::CreateLambda([this, OnAddFinished](UClass* InNewClass)
+		{
+			FSlateApplication::Get().DismissAllMenus();
+			
+			ComponentTypes.Add(InNewClass);
+			OnAddFinished.ExecuteIfBound();
+
+			// Ensure that the class filters run again so duplicate actor types cannot be selected
+			if (ClassViewerWidget.IsValid())
+			{
+				ClassViewerWidget->Refresh();
+			}
+		}));
+
+	ClassViewerWidget = StaticCastSharedPtr<SClassViewer>(ClassViewer.ToSharedPtr()); 
+	
+	return SNew(SBox)
+		.WidthOverride(300.f)
+		.HeightOverride(300.f)
+		[
+			ClassViewerWidget.ToSharedRef()
+		];
+}
+
+const FSlateBrush* UMovieGraphConditionGroupQuery_ComponentType::GetRowIcon(UClass* InComponentType)
+{
+	return FSlateIconFinder::FindIconForClass(InComponentType).GetIcon();
+}
+
+FText UMovieGraphConditionGroupQuery_ComponentType::GetRowText(UClass* InComponentType)
+{
+	return InComponentType->GetDisplayNameText();
+}
+#endif	// WITH_EDITOR
+
 UMovieGraphConditionGroup::UMovieGraphConditionGroup()
-	: OpType(EMovieGraphConditionGroupOpType::Union)
-	, bIsEnabled(true)
+	: OpType(EMovieGraphConditionGroupOpType::Add)
 {
 }
 
@@ -534,9 +1108,9 @@ void UMovieGraphConditionGroup::SetOperationType(const EMovieGraphConditionGroup
 {
 	// Always allow setting the operation type to Union. If not setting to Union, only allow setting the operation type if this is not the first
 	// condition group in the collection. The first condition group is always a Union.
-	if (OperationType == EMovieGraphConditionGroupOpType::Union)
+	if (OperationType == EMovieGraphConditionGroupOpType::Add)
 	{
-		OpType = EMovieGraphConditionGroupOpType::Union;
+		OpType = EMovieGraphConditionGroupOpType::Add;
 		return;
 	}
 
@@ -583,7 +1157,7 @@ TSet<AActor*> UMovieGraphConditionGroup::Evaluate(const UWorld* InWorld) const
 		if (QueryIndex == 0)
 		{
 			// The first query should always be a Union
-			ensure(Query->GetOperationType() == EMovieGraphConditionGroupQueryOpType::Union);
+			ensure(Query->GetOperationType() == EMovieGraphConditionGroupQueryOpType::Add);
 		}
 
 		// Similar to EvaluationResult, QueryResult is persisted+reset to prevent constantly re-allocating it
@@ -593,12 +1167,12 @@ TSet<AActor*> UMovieGraphConditionGroup::Evaluate(const UWorld* InWorld) const
 		
 		switch (Query->GetOperationType())
 		{
-		case EMovieGraphConditionGroupQueryOpType::Union:
+		case EMovieGraphConditionGroupQueryOpType::Add:
 			// Append() is faster than Union() because we don't need to allocate a new set
 			EvaluationResult.Append(QueryResult);
 			break;
 
-		case EMovieGraphConditionGroupQueryOpType::Intersect:
+		case EMovieGraphConditionGroupQueryOpType::And:
 			EvaluationResult = EvaluationResult.Intersect(QueryResult);
 			break;
 
@@ -611,10 +1185,20 @@ TSet<AActor*> UMovieGraphConditionGroup::Evaluate(const UWorld* InWorld) const
 	return EvaluationResult;
 }
 
-UMovieGraphConditionGroupQueryBase* UMovieGraphConditionGroup::AddQuery(const TSubclassOf<UMovieGraphConditionGroupQueryBase>& InQueryType)
+UMovieGraphConditionGroupQueryBase* UMovieGraphConditionGroup::AddQuery(const TSubclassOf<UMovieGraphConditionGroupQueryBase>& InQueryType, const int32 InsertIndex)
 {
 	UMovieGraphConditionGroupQueryBase* NewQueryObj = NewObject<UMovieGraphConditionGroupQueryBase>(this, InQueryType.Get());
-	Queries.Add(NewQueryObj);
+
+	if (InsertIndex < 0)
+	{
+		Queries.Add(NewQueryObj);
+	}
+	else
+	{
+		// Clamp the insert index to a valid range in case an invalid one is provided
+		Queries.Insert(NewQueryObj, FMath::Clamp(InsertIndex, 0, Queries.Num()));
+	}
+	
 	return NewQueryObj;
 }
 
@@ -628,19 +1212,50 @@ bool UMovieGraphConditionGroup::RemoveQuery(UMovieGraphConditionGroupQueryBase*
 	return Queries.RemoveSingle(InQuery) == 1;
 }
 
-void UMovieGraphConditionGroup::SetQueryOrder(TArray<UMovieGraphConditionGroupQueryBase*>& InQueryOrder)
+bool UMovieGraphConditionGroup::IsFirstConditionGroup() const
 {
-	// TODO
+	const UMovieGraphCollection* ParentCollection = GetTypedOuter<UMovieGraphCollection>();
+	if (ensureMsgf(ParentCollection, TEXT("Cannot determine if this is the first condition group when no parent collection is present")))
+	{
+		// GetConditionGroups() returns an array of non-const pointers, so Find() doesn't like having a const pointer passed to it.
+		// Find() won't mutate the condition group though, so the const_cast here is OK.
+		return ParentCollection->GetConditionGroups().Find(const_cast<UMovieGraphConditionGroup*>(this)) == 0;
+	}
+	
+	return false;
 }
 
-void UMovieGraphConditionGroup::SetEnabled(const bool bEnabled)
+bool UMovieGraphConditionGroup::MoveQueryToIndex(UMovieGraphConditionGroupQueryBase* InQuery, const int32 NewIndex)
 {
-	bIsEnabled = bEnabled;
-}
+#if WITH_EDITOR
+	Modify();
+#endif
 
-bool UMovieGraphConditionGroup::IsEnabled() const
-{
-	return bIsEnabled;
+	if (!InQuery)
+	{
+		return false;
+	}
+
+	const int32 ExistingIndex = Queries.Find(InQuery);
+	if (ExistingIndex == INDEX_NONE)
+	{
+		return false;
+	}
+
+	// If the new index is greater than the current index, then decrement the destination index so it remains valid after the removal below
+	int32 DestinationIndex = NewIndex;
+	if (DestinationIndex > ExistingIndex)
+	{
+		--DestinationIndex;
+	}
+
+	Queries.Remove(InQuery);
+	Queries.Insert(InQuery, DestinationIndex);
+
+	// Enforce that the first query is set to Union
+	InQuery->SetOperationType(EMovieGraphConditionGroupQueryOpType::Add);
+
+	return true;
 }
 
 TSet<AActor*> UMovieGraphCollection::Evaluate(const UWorld* InWorld) const
@@ -652,7 +1267,7 @@ TSet<AActor*> UMovieGraphCollection::Evaluate(const UWorld* InWorld) const
 	for (int32 ConditionGroupIndex = 0; ConditionGroupIndex < ConditionGroups.Num(); ++ConditionGroupIndex)
 	{
 		const TObjectPtr<UMovieGraphConditionGroup>& ConditionGroup = ConditionGroups[ConditionGroupIndex];
-		if (!ConditionGroup || !ConditionGroup->IsEnabled())
+		if (!ConditionGroup)
 		{
 			continue;
 		}
@@ -660,18 +1275,18 @@ TSet<AActor*> UMovieGraphCollection::Evaluate(const UWorld* InWorld) const
 		if (ConditionGroupIndex == 0)
 		{
 			// The first condition group should always be a Union
-			ensure(ConditionGroup->GetOperationType() == EMovieGraphConditionGroupOpType::Union);
+			ensure(ConditionGroup->GetOperationType() == EMovieGraphConditionGroupOpType::Add);
 		}
 
 		const TSet<AActor*> QueryResult = ConditionGroup->Evaluate(InWorld);
 		
 		switch (ConditionGroup->GetOperationType())
 		{
-		case EMovieGraphConditionGroupOpType::Union:
+		case EMovieGraphConditionGroupOpType::Add:
 			ResultSet = ResultSet.Union(QueryResult);
 			break;
 
-		case EMovieGraphConditionGroupOpType::Intersect:
+		case EMovieGraphConditionGroupOpType::And:
 			ResultSet = ResultSet.Intersect(QueryResult);
 			break;
 
@@ -701,6 +1316,39 @@ bool UMovieGraphCollection::RemoveConditionGroup(UMovieGraphConditionGroup* InCo
 	return ConditionGroups.RemoveSingle(InConditionGroup) == 1;
 }
 
+bool UMovieGraphCollection::MoveConditionGroupToIndex(UMovieGraphConditionGroup* InConditionGroup, const int32 NewIndex)
+{
+#if WITH_EDITOR
+	Modify();
+#endif
+
+	if (!InConditionGroup)
+	{
+		return false;
+	}
+
+	const int32 ExistingIndex = ConditionGroups.Find(InConditionGroup);
+	if (ExistingIndex == INDEX_NONE)
+	{
+		return false;
+	}
+
+	// If the new index is greater than the current index, then decrement the destination index so it remains valid after the removal below
+	int32 DestinationIndex = NewIndex;
+	if (DestinationIndex > ExistingIndex)
+	{
+		--DestinationIndex;
+	}
+
+	ConditionGroups.Remove(InConditionGroup);
+	ConditionGroups.Insert(InConditionGroup, DestinationIndex);
+
+	// Enforce that the first condition group is set to Union
+	InConditionGroup->SetOperationType(EMovieGraphConditionGroupOpType::Add);
+
+	return true;
+}
+
 void UMovieGraphCollection::SetCollectionName(const FString& InName)
 {
 	CollectionName = InName;
@@ -711,11 +1359,6 @@ const FString& UMovieGraphCollection::GetCollectionName() const
 	return CollectionName;
 }
 
-void UMovieGraphCollection::SetConditionGroupOrder(TArray<UMovieGraphConditionGroup*>& InConditionGroupOrder)
-{
-	// TODO
-}
-
 UMovieGraphCollection* UMoviePipelineRenderLayer::GetCollectionByName(const FString& Name) const
 {
 	for (const UMoviePipelineCollectionModifier* Modifier : Modifiers)
@@ -966,4 +1609,6 @@ void UMoviePipelineRenderLayerSubsystem::PreviewModifier(UMoviePipelineCollectio
 void UMoviePipelineRenderLayerSubsystem::ClearModifierPreview()
 {
 	ClearAllPreviews();
-}
\ No newline at end of file
+}
+
+#undef LOCTEXT_NAMESPACE
\ No newline at end of file
diff --git a/Engine/Plugins/MovieScene/MovieRenderPipeline/Source/MovieRenderPipelineCore/Public/Graph/MoviePipelineRenderLayerSubsystem.h b/Engine/Plugins/MovieScene/MovieRenderPipeline/Source/MovieRenderPipelineCore/Public/Graph/MoviePipelineRenderLayerSubsystem.h
--- a/Engine/Plugins/MovieScene/MovieRenderPipeline/Source/MovieRenderPipelineCore/Public/Graph/MoviePipelineRenderLayerSubsystem.h
+++ b/Engine/Plugins/MovieScene/MovieRenderPipeline/Source/MovieRenderPipelineCore/Public/Graph/MoviePipelineRenderLayerSubsystem.h
@@ -8,6 +8,8 @@
 
 #include "MoviePipelineRenderLayerSubsystem.generated.h"
 
+class SWidget;
+
 /**
  * EXPERIMENTAL
  */
@@ -46,13 +48,13 @@ UENUM(BlueprintType)
 enum class EMovieGraphConditionGroupOpType : uint8
 {
 	/** Adds the contents of the condition group to the results from the previous condition group (if any). */
-	Union,
+	Add,
 
 	/** Removes the contents of the condition group from the result of the previous condition group (if any). Any items in this condition group that aren't also found in the previous condition group will be ignored. */
 	Subtract,
 
 	/** Replaces the results of the previous condition group(s) with only the elements that exist in both that group, and this group. Intersecting with an empty condition group will result in an empty condition group. */
-	Intersect
+	And
 };
 
 /** Operation types available on condition group queries. */
@@ -60,24 +62,27 @@ UENUM(BlueprintType)
 enum class EMovieGraphConditionGroupQueryOpType : uint8
 {
 	/** Adds the results of the query to the results from the previous query (if any). */
-	Union,
+	Add,
 
 	/** Removes the results of the query from the results of the previous query (if any). Any items in this query result that aren't also found in the previous query result will be ignored. */
 	Subtract,
 	
 	/** Replaces the results of the previous queries with only the items that exist in both those queries, and this query result. Intersecting with a query which returns nothing will create an empty query result. */
-	Intersect
+	And
 };
 
 /** Base class that all condition group queries must inherit from. */
-UCLASS(Abstract, EditInlineNew)
-class UMovieGraphConditionGroupQueryBase : public UObject
+UCLASS(Abstract)
+class MOVIERENDERPIPELINECORE_API UMovieGraphConditionGroupQueryBase : public UObject
 {
 	GENERATED_BODY()
 
 public:
 	UMovieGraphConditionGroupQueryBase();
 
+	/** Delegate which is called when the contents of a query has changed. */
+	DECLARE_DELEGATE(FMovieGraphConditionGroupQueryContentsChanged)
+
 	/**
 	 * Sets how the condition group query interacts with the condition group. This call is ignored for the first query
 	 * in the condition group (the first is always Union).
@@ -105,6 +110,20 @@ public:
 	/** Gets the icon that represents this query class in the UI. */
 	virtual const FSlateIcon& GetIcon() const;
 
+#if WITH_EDITOR
+	/**
+	 * Gets the widgets that should be displayed for this query. If no custom widgets are specified (returning an empty array), the default
+	 * name/value widgets will be shown for all query properties tagged with EditAnywhere.
+	 */
+	virtual TArray<TSharedRef<SWidget>> GetWidgets();
+
+	/**
+	 * Gets the contents of the "Add" menu in the UI, if any. The Add menu will not display if SNullWidget::NullWidget is returned. When the Add
+	 * menu updates properties within the query, OnAddFinished should be called in order to give the UI a chance to update itself.
+	 */
+	virtual TSharedRef<SWidget> GetAddMenuContents(const FMovieGraphConditionGroupQueryContentsChanged& OnAddFinished);
+#endif
+
 	/** Determines if this query is only respected when run within the editor. Used for providing a UI hint. */
 	UFUNCTION(BlueprintCallable, Category = "Settings")
 	virtual bool IsEditorOnlyQuery() const;
@@ -117,18 +136,59 @@ public:
 	UFUNCTION(BlueprintCallable, Category = "Settings")
 	bool IsEnabled() const;
 
+	/** Determines if this is the first condition group query under the parent condition group. */
+	UFUNCTION(BlueprintCallable, Category = "Settings")
+	bool IsFirstConditionGroupQuery() const;
+
 private:
 	/** The operation type that the query is using. */
-	UPROPERTY(EditAnywhere, Category="General")
+	UPROPERTY()
 	EMovieGraphConditionGroupQueryOpType OpType;
 
 	/** Whether this query is currently enabled within the condition group. */
-	UPROPERTY(EditAnywhere, Category="General")
+	UPROPERTY()
 	bool bIsEnabled;
 };
 
-UCLASS(BlueprintType, EditInlineNew, meta = (DisplayName = "Actor Tag"))
-class UMovieGraphConditionGroupQuery_ActorTag final : public UMovieGraphConditionGroupQueryBase
+/** Query type which filters actors via an explicit actor list. */
+UCLASS(BlueprintType, meta = (DisplayName = "Actor"))
+class MOVIERENDERPIPELINECORE_API UMovieGraphConditionGroupQuery_Actor final : public UMovieGraphConditionGroupQueryBase
+{
+	GENERATED_BODY()
+
+public:
+	virtual void Evaluate(const TArray<AActor*>& InActorsToQuery, TSet<AActor*>& OutMatchingActors) const override;
+	virtual const FSlateIcon& GetIcon() const override;
+
+#if WITH_EDITOR
+	virtual TArray<TSharedRef<SWidget>> GetWidgets() override;
+	virtual TSharedRef<SWidget> GetAddMenuContents(const FMovieGraphConditionGroupQueryContentsChanged& OnAddFinished) override;
+#endif
+
+private:
+#if WITH_EDITOR
+	static const FSlateBrush* GetRowIcon(TSharedPtr<TSoftObjectPtr<AActor>> InActor);
+	static FText GetRowText(TSharedPtr<TSoftObjectPtr<AActor>> InActor);
+#endif
+
+public:
+	/** The query must match one of the actors in order to be a match. */
+	UPROPERTY(EditAnywhere, Category="General")
+	TArray<TSoftObjectPtr<AActor>> ActorsToMatch;
+
+private:
+#if WITH_EDITOR
+	TSharedPtr<class ISceneOutliner> ActorPickerWidget;
+
+	// Not ideal to store a duplicate of ActorsToMatch, but SListView requires TSharedPtr<...> as the data source, and UPROPERTY does not
+	// support TSharedPtr<...>
+	TArray<TSharedPtr<TSoftObjectPtr<AActor>>> ListDataSource;
+#endif
+};
+
+/** Query type which filters actors via tags on actors. */
+UCLASS(BlueprintType, meta = (DisplayName = "Actor Tag Name"))
+class MOVIERENDERPIPELINECORE_API UMovieGraphConditionGroupQuery_ActorTagName final : public UMovieGraphConditionGroupQueryBase
 {
 	GENERATED_BODY()
 
@@ -136,25 +196,47 @@ public:
 	virtual void Evaluate(const TArray<AActor*>& InActorsToQuery, TSet<AActor*>& OutMatchingActors) const override;
 	virtual const FSlateIcon& GetIcon() const override;
 
+#if WITH_EDITOR
+	virtual TArray<TSharedRef<SWidget>> GetWidgets() override;
+#endif
+
 public:
-	/** The tag that the actor needs to have in order to be a match. */
+	/**
+	 * Tags on the actor must match one or more of the specified tags to be a match. Not case sensitive. One tag per line. Wildcards ("?" and "*") are supported but not required.
+	 * The "*" wildcard matches zero or more characters, and "?" matches exactly one character (and that character must be present).
+	 * 
+	 * Wildcard examples:
+	 * Foo* would match Foo, FooBar, and FooBaz, but not BarFoo.
+	 * *Foo* would match the above in addition to BarFoo.
+	 * Foo?Bar would match Foo.Bar and Foo_Bar, but not FooBar.
+	 * Foo? would match Food, but not FooBar or BarFoo.
+	 * Foo??? would match FooBar and FooBaz, but not Foo or Food.
+	 * ?oo? would match Food, but not Foo.
+	 * ?Foo* would match AFooBar, but not FooBar 
+	 */
 	UPROPERTY(EditAnywhere, Category="General")
-	FName TagToMatch;
+	FString TagsToMatch;
 };
 
-UCLASS(BlueprintType, EditInlineNew, meta = (DisplayName = "Actor Name"))
-class UMovieGraphConditionGroupQuery_ActorName final : public UMovieGraphConditionGroupQueryBase
+/** Query type which filters actors via their name (label). */
+UCLASS(BlueprintType, meta = (DisplayName = "Actor Name"))
+class MOVIERENDERPIPELINECORE_API UMovieGraphConditionGroupQuery_ActorName final : public UMovieGraphConditionGroupQueryBase
 {
 	GENERATED_BODY()
 
 public:
 	virtual void Evaluate(const TArray<AActor*>& InActorsToQuery, TSet<AActor*>& OutMatchingActors) const override;
 	virtual const FSlateIcon& GetIcon() const override;
+
+#if WITH_EDITOR
+	virtual TArray<TSharedRef<SWidget>> GetWidgets() override;
+#endif
+
 	virtual bool IsEditorOnly() const override;
 
 public:
 	/**
-	 * The name that the actor needs to have in order to be a match. Not case sensitive. Wildcards ("?" and "*") are supported but not required.
+	 * The name that the actor needs to have in order to be a match. Not case sensitive. One name per line. Wildcards ("?" and "*") are supported but not required.
 	 * The "*" wildcard matches zero or more characters, and "?" matches exactly one character (and that character must be present).
 	 * 
 	 * Wildcard examples:
@@ -166,12 +248,13 @@ public:
 	 * ?oo? would match Food, but not Foo.
 	 * ?Foo* would match AFooBar, but not FooBar 
 	 */
-	UPROPERTY(EditAnywhere, Category="General", meta=(MultiLine=true))
+	UPROPERTY(EditAnywhere, Category="General")
 	FString WildcardSearch;
 };
 
-UCLASS(BlueprintType, EditInlineNew, meta = (DisplayName = "Actor Type"))
-class UMovieGraphConditionGroupQuery_ActorType final : public UMovieGraphConditionGroupQueryBase
+/** Query type which filters actors by type. */
+UCLASS(BlueprintType, meta = (DisplayName = "Actor Type"))
+class MOVIERENDERPIPELINECORE_API UMovieGraphConditionGroupQuery_ActorType final : public UMovieGraphConditionGroupQueryBase
 {
 	GENERATED_BODY()
 
@@ -179,14 +262,28 @@ public:
 	virtual void Evaluate(const TArray<AActor*>& InActorsToQuery, TSet<AActor*>& OutMatchingActors) const override;
 	virtual const FSlateIcon& GetIcon() const override;
 
+#if WITH_EDITOR
+	virtual TArray<TSharedRef<SWidget>> GetWidgets() override;
+	virtual TSharedRef<SWidget> GetAddMenuContents(const FMovieGraphConditionGroupQueryContentsChanged& OnAddFinished) override;
+#endif
+
 public:
 	/** The type (class) that the actor needs to have in order to be a match. */
-	UPROPERTY(EditAnywhere, Category="General")
-	TSubclassOf<AActor> ActorType;
+	TArray<UClass*> ActorTypes;
+
+private:
+#if WITH_EDITOR
+	static const FSlateBrush* GetRowIcon(UClass* InActorType);
+	static FText GetRowText(UClass* InActorType);
+#endif
+	
+	/** The class viewer widget to show in the Add menu. */
+	TSharedPtr<class SClassViewer> ClassViewerWidget;
 };
 
-UCLASS(BlueprintType, EditInlineNew, meta = (DisplayName = "Component Tag"))
-class UMovieGraphConditionGroupQuery_ComponentTag final : public UMovieGraphConditionGroupQueryBase
+/** Query type which filters actors by tags on their components. */
+UCLASS(BlueprintType, meta = (DisplayName = "Component Tag Name"))
+class MOVIERENDERPIPELINECORE_API UMovieGraphConditionGroupQuery_ComponentTagName final : public UMovieGraphConditionGroupQueryBase
 {
 	GENERATED_BODY()
 
@@ -194,14 +291,31 @@ public:
 	virtual void Evaluate(const TArray<AActor*>& InActorsToQuery, TSet<AActor*>& OutMatchingActors) const override;
 	virtual const FSlateIcon& GetIcon() const override;
 
+#if WITH_EDITOR
+	virtual TArray<TSharedRef<SWidget>> GetWidgets() override;
+#endif
+
 public:
-	/** The tag that a component on the actor needs to have in order to be a match. */
+	/**
+	 * A component on the actor must have one or more of the specified tags to be a match. Not case sensitive. One tag per line. Wildcards ("?" and "*") are supported but not required.
+	 * The "*" wildcard matches zero or more characters, and "?" matches exactly one character (and that character must be present).
+	 * 
+	 * Wildcard examples:
+	 * Foo* would match Foo, FooBar, and FooBaz, but not BarFoo.
+	 * *Foo* would match the above in addition to BarFoo.
+	 * Foo?Bar would match Foo.Bar and Foo_Bar, but not FooBar.
+	 * Foo? would match Food, but not FooBar or BarFoo.
+	 * Foo??? would match FooBar and FooBaz, but not Foo or Food.
+	 * ?oo? would match Food, but not Foo.
+	 * ?Foo* would match AFooBar, but not FooBar 
+	 */
 	UPROPERTY(EditAnywhere, Category="General")
-	FName TagToMatch;
+	FString TagsToMatch;
 };
 
-UCLASS(BlueprintType, EditInlineNew, meta = (DisplayName = "Component Type"))
-class UMovieGraphConditionGroupQuery_ComponentType final : public UMovieGraphConditionGroupQueryBase
+/** Query type which filters actors via the components contained in them. */
+UCLASS(BlueprintType, meta = (DisplayName = "Component Type"))
+class MOVIERENDERPIPELINECORE_API UMovieGraphConditionGroupQuery_ComponentType final : public UMovieGraphConditionGroupQueryBase
 {
 	GENERATED_BODY()
 
@@ -209,15 +323,29 @@ public:
 	virtual void Evaluate(const TArray<AActor*>& InActorsToQuery, TSet<AActor*>& OutMatchingActors) const override;
 	virtual const FSlateIcon& GetIcon() const override;
 
+#if WITH_EDITOR
+	virtual TArray<TSharedRef<SWidget>> GetWidgets() override;
+	virtual TSharedRef<SWidget> GetAddMenuContents(const FMovieGraphConditionGroupQueryContentsChanged& OnAddFinished) override;
+#endif
+
+private:
+#if WITH_EDITOR
+	static const FSlateBrush* GetRowIcon(UClass* InComponentType);
+	static FText GetRowText(UClass* InComponentType);
+#endif
+	
+	/** The class viewer widget to show in the Add menu. */
+	TSharedPtr<class SClassViewer> ClassViewerWidget;
+
 public:
-	/** The type that a component on the actor needs to have in order to be a match. */
+	/** The actor must have one or more of the component type(s) in order to be a match. */
 	UPROPERTY(EditAnywhere, Category="General")
-	TSubclassOf<UActorComponent> ComponentType;
+	TArray<UClass*> ComponentTypes;
 };
 
 /** A group of queries which can be added to a collection. */
-UCLASS(BlueprintType, EditInlineNew)
-class UMovieGraphConditionGroup : public UObject
+UCLASS(BlueprintType)
+class MOVIERENDERPIPELINECORE_API UMovieGraphConditionGroup : public UObject
 {
 	GENERATED_BODY()
 
@@ -241,10 +369,11 @@ public:
 
 	/**
 	 * Adds a new condition group query to the condition group and returns a ptr to it. The condition group owns the
-	 * created query.
+	 * created query. By default the query is added to the end, but an optional index can be provided if the query
+	 * should be placed in a specific location among the existing queries.
 	 */
 	UFUNCTION(BlueprintCallable, Category = "Settings")
-	UMovieGraphConditionGroupQueryBase* AddQuery(const TSubclassOf<UMovieGraphConditionGroupQueryBase>& InQueryType);
+	UMovieGraphConditionGroupQueryBase* AddQuery(const TSubclassOf<UMovieGraphConditionGroupQueryBase>& InQueryType, const int32 InsertIndex = -1);
 
 	/** Gets all queries currently contained in the condition group. */
 	UFUNCTION(BlueprintCallable, Category = "Settings")
@@ -254,31 +383,24 @@ public:
 	UFUNCTION(BlueprintCallable, Category = "Settings")
 	bool RemoveQuery(UMovieGraphConditionGroupQueryBase* InQuery);
 
-	/** TODO: API around reordering queries, pending research on how Slate provides this information. */
+	/** Determines if this is the first condition group under the parent collection. */
 	UFUNCTION(BlueprintCallable, Category = "Settings")
-	void SetQueryOrder(TArray<UMovieGraphConditionGroupQueryBase*>& InQueryOrder);
+	bool IsFirstConditionGroup() const;
 
-	// TODO: May not exist in the final API if the design does not call for enabling/disabling condition groups.
-	/** Sets whether this condition group is enabled. */
-	UFUNCTION(BlueprintCallable, Category = "Settings")
-	void SetEnabled(const bool bEnabled);
-
-	// TODO: May not exist in the final API if the design does not call for enabling/disabling condition groups.
-	/** Determines if this condition group is enabled */
+	/**
+	 * Move the specified query to a new index within the condition group. Returns false if the query was not found or the index
+	 * specified is invalid, else true.
+	 */
 	UFUNCTION(BlueprintCallable, Category = "Settings")
-	bool IsEnabled() const;
+	bool MoveQueryToIndex(UMovieGraphConditionGroupQueryBase* InQuery, const int32 NewIndex);
 
 private:
 	/** The operation type that the condition group is using. */
 	UPROPERTY(EditAnywhere, Category="General")
 	EMovieGraphConditionGroupOpType OpType;
 
-	/** Whether this condition group is currently enabled within the collection. */
-	UPROPERTY(EditAnywhere, Category="General")
-	bool bIsEnabled;
-
 	/** The queries that are contained within the condition group. */
-	UPROPERTY(EditAnywhere, Category="General", Instanced, meta=(EditInline))
+	UPROPERTY(EditAnywhere, Category="General")
 	TArray<TObjectPtr<UMovieGraphConditionGroupQueryBase>> Queries;
 
 	/** Persisted actor set which can be re-used for query evaluations across frames to prevent constantly re-allocating it. */
@@ -291,7 +413,7 @@ private:
 };
 
 /** A group of actors generated by actor queries. */
-UCLASS(BlueprintType, EditInlineNew)
+UCLASS(BlueprintType)
 class MOVIERENDERPIPELINECORE_API UMovieGraphCollection : public UObject
 {
 	GENERATED_BODY()
@@ -332,24 +454,27 @@ public:
 	UFUNCTION(BlueprintCallable, Category = "Settings")
 	bool RemoveConditionGroup(UMovieGraphConditionGroup* InConditionGroup);
 
-	/** TODO: API around reordering condition groups, pending research on how Slate provides this information. */
+	/**
+	 * Move the specified condition group to a new index within the collection. Returns false if the condition group was not found or the index
+	 * specified is invalid, else true.
+	 */
 	UFUNCTION(BlueprintCallable, Category = "Settings")
-	void SetConditionGroupOrder(TArray<UMovieGraphConditionGroup*>& InConditionGroupOrder);
+	bool MoveConditionGroupToIndex(UMovieGraphConditionGroup* InConditionGroup, const int32 NewIndex);
 
 private:
 	/** The display name of the collection, shown in the UI. Does not need to be unique across collections. */
-	UPROPERTY(EditAnywhere, Category="General")
+	UPROPERTY(EditAnywhere, Category="Collection")
 	FString CollectionName;
 	
 	/** The condition groups that are contained within the collection. */
-	UPROPERTY(EditAnywhere, Category="General", Instanced, meta=(EditInline))
+	UPROPERTY(EditAnywhere, Category="Collection")
 	TArray<TObjectPtr<UMovieGraphConditionGroup>> ConditionGroups;
 };
 
 /**
  * Base class for providing actor modification functionality via collections.
  */
-UCLASS(Abstract, EditInlineNew)
+UCLASS(Abstract)
 class UMoviePipelineCollectionModifier : public UObject
 {
 	GENERATED_BODY()
@@ -483,7 +608,7 @@ public:
 /**
  * Base class for providing the ability to determine if an actor matches a query.
  */
-UCLASS(Abstract, EditInlineNew)
+UCLASS(Abstract)
 class UMoviePipelineCollectionQuery : public UObject
 {
 	GENERATED_BODY()
diff --git a/Engine/Plugins/MovieScene/MovieRenderPipeline/Source/MovieRenderPipelineCore/Public/Graph/Nodes/MovieGraphCollectionNode.h b/Engine/Plugins/MovieScene/MovieRenderPipeline/Source/MovieRenderPipelineCore/Public/Graph/Nodes/MovieGraphCollectionNode.h
--- a/Engine/Plugins/MovieScene/MovieRenderPipeline/Source/MovieRenderPipelineCore/Public/Graph/Nodes/MovieGraphCollectionNode.h
+++ b/Engine/Plugins/MovieScene/MovieRenderPipeline/Source/MovieRenderPipelineCore/Public/Graph/Nodes/MovieGraphCollectionNode.h
@@ -36,8 +36,8 @@ public:
 
 public:
 	UPROPERTY(EditAnywhere, BlueprintReadWrite, Category = Overrides, meta = (InlineEditConditionToggle))
-	uint8 bOverride_Collection : 1;
+	uint8 bOverride_Collection : 1 = 1;	// The collection is customized in the details panel, so the override should always be enabled
 
-	UPROPERTY(BlueprintReadWrite, EditAnywhere, Instanced, Category = "General", meta=(EditCondition="bOverride_Collection", EditInline))
+	UPROPERTY(BlueprintReadWrite, EditAnywhere, Category = "General", meta=(EditCondition="bOverride_Collection"))
 	TObjectPtr<UMovieGraphCollection> Collection;
 };
\ No newline at end of file
diff --git a/Engine/Plugins/MovieScene/MovieRenderPipeline/Source/MovieRenderPipelineEditor/Private/Customizations/Graph/MovieGraphCollectionsCustomization.cpp b/Engine/Plugins/MovieScene/MovieRenderPipeline/Source/MovieRenderPipelineEditor/Private/Customizations/Graph/MovieGraphCollectionsCustomization.cpp
new file mode 100644
--- /dev/null
+++ b/Engine/Plugins/MovieScene/MovieRenderPipeline/Source/MovieRenderPipelineEditor/Private/Customizations/Graph/MovieGraphCollectionsCustomization.cpp
@@ -0,0 +1,1460 @@
+﻿// Copyright Epic Games, Inc. All Rights Reserved.
+
+#include "MovieGraphCollectionsCustomization.h"
+
+#include "DetailCategoryBuilder.h"
+#include "DetailLayoutBuilder.h"
+#include "DetailWidgetRow.h"
+#include "Framework/Application/SlateApplication.h"
+#include "Framework/MultiBox/MultiBoxBuilder.h"
+#include "Graph/MoviePipelineRenderLayerSubsystem.h"
+#include "Graph/Nodes/MovieGraphCollectionNode.h"
+#include "MovieRenderPipelineStyle.h"
+#include "PropertyHandle.h"
+#include "ScopedTransaction.h"
+#include "SPositiveActionButton.h"
+#include "UObject/UObjectIterator.h"
+#include "Widgets/Images/SLayeredImage.h"
+#include "Widgets/Input/SButton.h"
+#include "Widgets/Input/SCheckBox.h"
+#include "Widgets/Input/SComboBox.h"
+#include "Widgets/Layout/SBox.h"
+#include "Widgets/SToolTip.h"
+#include "Widgets/Text/STextBlock.h"
+
+#define LOCTEXT_NAMESPACE "MoviePipelineEditorCollectionCustomization"
+
+void SMovieGraphCollectionTreeViewIndent::Construct(const FArguments& InArgs)
+{
+	const int NumIndents = InArgs._NumIndents.Get();
+	const TSharedRef<SHorizontalBox> IndentContainer = SNew(SHorizontalBox);
+	
+	ChildSlot
+	.VAlign(VAlign_Fill)
+	[
+		// A single indent doesn't change the background color; more indents get a darker background
+		SNew(SBorder)
+		.Padding(0)
+		.VAlign(VAlign_Fill)
+		.BorderImage((NumIndents <= 1)
+			? FAppStyle::Get().GetBrush("NoBorder")
+			: FMovieRenderPipelineStyle::Get().GetBrush("MovieRenderGraph.CollectionsTree.IndentedRowBackground"))
+		[
+			IndentContainer
+		]
+	];
+
+	// Add a new indent for each level of indent requested
+	for (int32 Index = 0; Index < NumIndents; ++Index)
+	{
+		IndentContainer->AddSlot()
+		.AutoWidth()
+		.Padding(0)
+		.VAlign(VAlign_Fill)
+		[
+			SNew(SBorder)
+			.BorderImage(FAppStyle::Get().GetBrush("DetailsView.ArrayDropShadow"))
+			[
+				SNew(SBox)
+				.WidthOverride(16.f)
+			]
+		];
+	}
+}
+
+TSharedRef<FMovieGraphCollectionTreeDragOp> FMovieGraphCollectionTreeDragOp::New(const TSharedPtr<FCollectionCustomizationTreeElement>& InElement)
+{
+	TSharedRef<FMovieGraphCollectionTreeDragOp> Operation = MakeShared<FMovieGraphCollectionTreeDragOp>();
+	Operation->Element = InElement;
+	Operation->Construct();
+	
+	return Operation;
+}
+
+void SMovieGraphCollectionTreeView::Construct(const FArguments& InArgs)
+{
+	STreeView<TSharedPtr<FCollectionCustomizationTreeElement>>::FArguments SuperArgs;
+
+	DetailBuilder = InArgs._DetailBuilder;
+	Collection = InArgs._Collection;
+
+	SuperArgs
+	.ItemHeight(28)
+	.TreeItemsSource(&RootElements)
+	.SelectionMode(ESelectionMode::Single)
+	.OnGenerateRow_Lambda([this](const TSharedPtr<FCollectionCustomizationTreeElement>& TreeItem, const TSharedPtr<STableViewBase>& OwnerTable)
+	{
+		return GenerateTreeRow(TreeItem, OwnerTable.ToSharedRef());
+	})
+	.OnGetChildren(this, &SMovieGraphCollectionTreeView::GetChildrenForTree)
+	.OnExpansionChanged(this, &SMovieGraphCollectionTreeView::OnExpansionChanged)
+	.HeaderRow
+	(
+		SNew(SHeaderRow)
+		.Visibility(EVisibility::Collapsed)
+
+		+ SHeaderRow::Column(SMovieGraphCollectionTreeItem::ColumnID_Name)
+		.DefaultLabel(LOCTEXT("NameColumnLabel", "Name"))
+		.SortMode(EColumnSortMode::None)
+		.FillWidth(0.5f)
+
+		+ SHeaderRow::Column(SMovieGraphCollectionTreeItem::ColumnID_Value)
+		.DefaultLabel(LOCTEXT("ValueColumnLabel", "Value"))
+		.SortMode(EColumnSortMode::None)
+		.FillWidth(0.5f)
+	);
+	
+	STreeView<TSharedPtr<FCollectionCustomizationTreeElement>>::Construct(SuperArgs);
+}
+
+void SMovieGraphCollectionTreeView::OnExpansionChanged(const TSharedPtr<FCollectionCustomizationTreeElement> InElement, const bool bInExpanded)
+{
+	if (bInExpanded)
+	{
+		ExpandedElements.Add(InElement->GetHash());
+	}
+	else
+	{
+		ExpandedElements.Remove(InElement->GetHash());
+	}
+}
+
+void SMovieGraphCollectionTreeView::GetChildrenForTree(const TSharedPtr<FCollectionCustomizationTreeElement> InItem, TArray<TSharedPtr<FCollectionCustomizationTreeElement>>& OutChildren)
+{
+	OutChildren.Append(InItem->GetChildren());
+}
+
+void SMovieGraphCollectionTreeView::RestoreExpansionStateRecursive(const TSharedPtr<FCollectionCustomizationTreeElement>& InElement)
+{
+	SetItemExpansion(InElement, ExpandedElements.Contains(InElement->GetHash()));
+	
+	for (const TSharedPtr<FCollectionCustomizationTreeElement>& ChildElement : InElement->GetChildren())
+	{
+		RestoreExpansionStateRecursive(ChildElement);
+	}
+}
+
+TSharedRef<ITableRow> SMovieGraphCollectionTreeView::GenerateTreeRow(TSharedPtr<FCollectionCustomizationTreeElement> InTreeElement, const TSharedRef<STableViewBase>& OwnerTable)
+{
+	if (InTreeElement->Type == FCollectionCustomizationTreeElement::EElementType::QueryProperty)
+	{
+		if (InTreeElement->Query && !InTreeElement->Query->GetWidgets().IsEmpty())
+		{
+			return SNew(SMovieGraphCollectionTreeItem_Stretch, SharedThis(this), InTreeElement);
+		}
+	}
+
+	if ((InTreeElement->Type == FCollectionCustomizationTreeElement::EElementType::EmptyCollection) ||
+		(InTreeElement->Type == FCollectionCustomizationTreeElement::EElementType::EmptyConditionGroup))
+	{
+		return SNew(SMovieGraphCollectionTreeItem_Stretch, SharedThis(this), InTreeElement);
+	}
+	
+	return SNew(SMovieGraphCollectionTreeItem, SharedThis(this), InTreeElement, DetailBuilder);
+}
+
+void SMovieGraphCollectionTreeView::AddRootElement()
+{
+	const TSharedPtr<FCollectionCustomizationTreeElement> RootElement = MakeShared<FCollectionCustomizationTreeElement>(FCollectionCustomizationTreeElement::EElementType::Root);
+	RootElement->Collection = Collection.Get();
+	
+	RootElements.Add(RootElement);
+
+	// Expand the root element by default because it is never visible and the user cannot expand it themselves
+	constexpr bool bIsExpanded = true;
+	SetItemExpansion(RootElement, bIsExpanded);
+	ExpandedElements.Add(RootElement->GetHash());
+}
+
+void SMovieGraphCollectionTreeView::AddRootConditionGroupElements(const TObjectPtr<UMovieGraphCollection>& InCollection) const
+{
+	// Add the collection's condition groups. These will be the root elements of the tree, and each condition group will take
+	// over adding its own children to the tree.
+	const TArray<UMovieGraphConditionGroup*> ConditionGroups = InCollection->GetConditionGroups();
+	for (int32 Index = 0; Index < ConditionGroups.Num(); ++Index)
+	{
+		AddNewConditionGroupElement(ConditionGroups[Index]);
+	}
+}
+
+void SMovieGraphCollectionTreeView::AddNewConditionGroupElement(UMovieGraphConditionGroup* InConditionGroup) const
+{
+	const TSharedPtr<FCollectionCustomizationTreeElement> NewElement = MakeShared<FCollectionCustomizationTreeElement>(FCollectionCustomizationTreeElement::EElementType::ConditionGroup);
+	NewElement->ConditionGroup = InConditionGroup;
+	NewElement->Collection = Collection.Get();
+}
+
+void SMovieGraphCollectionTreeView::AddConditionGroup()
+{
+	if (Collection.IsValid())
+	{
+		UMovieGraphConditionGroup* NewConditionGroup = Collection.Get()->AddConditionGroup();
+		
+		AddNewConditionGroupElement(NewConditionGroup);
+
+		if (!RootElements.IsEmpty())
+		{
+			// Needed so the "empty collection" notice is cleared out (if this is the first condition group being added to the collection).
+			RootElements[0]->ClearCachedChildren();
+		}
+		
+		RefreshView();
+	}
+}
+
+void SMovieGraphCollectionTreeView::RemoveConditionGroup(const TWeakPtr<FCollectionCustomizationTreeElement>& InConditionGroupElement)
+{
+	if (!InConditionGroupElement.IsValid())
+	{
+		return;
+	}
+	
+	const TSharedPtr<FCollectionCustomizationTreeElement> TreeElementPin = InConditionGroupElement.Pin();
+	if (!TreeElementPin->Collection)
+	{
+		return;
+	}
+
+	TreeElementPin->Collection->RemoveConditionGroup(TreeElementPin->ConditionGroup);
+
+	if (TreeElementPin->ParentElement.IsValid())
+	{
+		TreeElementPin->ParentElement.Pin()->ClearCachedChildren();
+	}
+
+	RefreshView();
+}
+
+void SMovieGraphCollectionTreeView::RefreshView()
+{
+	RequestTreeRefresh();
+    
+    // Restore tree expansion state after the refresh
+    for (const TSharedPtr<FCollectionCustomizationTreeElement>& RootElement : RootElements)
+    {
+    	RestoreExpansionStateRecursive(RootElement);
+    }
+}
+
+void SMovieGraphCollectionTreeOpTypeWidget::Construct(const FArguments& InArgs)
+{
+	bIsConditionGroup = InArgs._IsConditionGroup.Get();
+	WeakTreeElement = InArgs._WeakTreeElement.Get();
+	
+	ChildSlot
+	[
+		SNew(SComboBox<FName>)
+		.ButtonStyle(FAppStyle::Get(), "NoBorder")
+		.ForegroundColor(FSlateColor::UseForeground())
+		.IsEnabled(this, &SMovieGraphCollectionTreeOpTypeWidget::IsWidgetEnabled)
+		.OptionsSource(bIsConditionGroup ? GetOpTypes<EMovieGraphConditionGroupOpType>() : GetOpTypes<EMovieGraphConditionGroupQueryOpType>())
+		.OnSelectionChanged(this, &SMovieGraphCollectionTreeOpTypeWidget::SetOpType)
+		.OnGenerateWidget_Lambda([this](const FName& InOpType)
+		{
+			return GetOpTypeContents(InOpType, bIsConditionGroup);
+		})
+		.Content()
+		[
+			GetOpTypeContents(NAME_None, bIsConditionGroup)
+		]
+	];
+}
+
+template <typename T>
+TArray<FName>* SMovieGraphCollectionTreeOpTypeWidget::GetOpTypes()
+{
+	static_assert(TIsEnum<T>::Value, "Provided type must be an enum");
+		
+	static TArray<FName> OpTypes;
+
+	if (OpTypes.IsEmpty())
+	{
+		const UEnum* OpTypeEnum = StaticEnum<T>();
+
+		// -1 to skip the implicit "MAX" added at compile time
+		for (int32 Index = 0; Index < OpTypeEnum->NumEnums() - 1; ++Index)
+		{
+			OpTypes.Add(FName(OpTypeEnum->GetDisplayNameTextByIndex(Index).ToString()));
+		}
+	}
+
+	return &OpTypes;
+}
+
+TSharedRef<SWidget> SMovieGraphCollectionTreeOpTypeWidget::GetOpTypeContents(const FName& InOpName, const bool bIsConditionGroupOp) const
+{
+	TSharedPtr<SWidget> OpTypeTextBlock;
+
+	// This method may only be called once in some scenarios, so provide the option of a text lambda vs. constant text if the
+	// text needs to be updated dynamically when the op type changes.
+	if (InOpName == NAME_None)
+	{
+		SAssignNew(OpTypeTextBlock, STextBlock)
+		.Text_Lambda([this, bIsConditionGroupOp]()
+		{
+			return FText::FromName(bIsConditionGroupOp ? GetCurrentConditionGroupOpType() : GetCurrentConditionGroupQueryOpType());
+		})
+		.Font(FAppStyle::GetFontStyle(TEXT("PropertyWindow.NormalFont")));
+	}
+	else
+	{
+		SAssignNew(OpTypeTextBlock, STextBlock)
+		.Text(FText::FromName(InOpName))
+		.Font(FAppStyle::GetFontStyle(TEXT("PropertyWindow.NormalFont")));
+	}
+		
+	return
+		SNew(SHorizontalBox)
+		+ SHorizontalBox::Slot()
+		.VAlign(VAlign_Center)
+		.Padding(0, 0, 7.f, 0)
+		.AutoWidth()
+		[
+			SNew(SImage)
+			.Image(FAppStyle::GetBrush("Icons.Settings"))
+			.ColorAndOpacity(FSlateColor::UseForeground())
+		]
+
+		+ SHorizontalBox::Slot()
+		.VAlign(VAlign_Center)
+		[
+			OpTypeTextBlock.ToSharedRef()
+		];
+}
+
+bool SMovieGraphCollectionTreeOpTypeWidget::IsWidgetEnabled() const
+{
+	if (WeakTreeElement.IsValid())
+	{
+		if (bIsConditionGroup)
+		{
+			if (const TObjectPtr<UMovieGraphConditionGroup> ConditionGroup = WeakTreeElement.Pin()->ConditionGroup)
+			{
+				// Disable the first condition group
+				return !ConditionGroup->IsFirstConditionGroup();
+			}
+		}
+		else
+		{
+			if (const TObjectPtr<UMovieGraphConditionGroupQueryBase> ConditionGroupQuery = WeakTreeElement.Pin()->Query)
+			{
+				// Disable the first query, or disable it if it was flagged as disabled
+				return !ConditionGroupQuery->IsFirstConditionGroupQuery() && ConditionGroupQuery->IsEnabled();
+			}
+		}
+	}
+	
+	return true;
+}
+
+void SMovieGraphCollectionTreeOpTypeWidget::SetOpType(const FName InNewOpType, ESelectInfo::Type SelectInfo) const
+{
+	if (!WeakTreeElement.IsValid())
+	{
+		return;
+	}
+	
+	if (bIsConditionGroup)
+	{
+		if (const TObjectPtr<UMovieGraphConditionGroup>& ConditionGroup = WeakTreeElement.Pin()->ConditionGroup)
+		{
+			const UEnum* OpTypeEnum = StaticEnum<EMovieGraphConditionGroupOpType>();
+			ConditionGroup->SetOperationType(static_cast<EMovieGraphConditionGroupOpType>(OpTypeEnum->GetValueByName(InNewOpType)));
+		}
+	}
+	else
+	{
+		if (const TObjectPtr<UMovieGraphConditionGroupQueryBase>& Query = WeakTreeElement.Pin()->Query)
+		{
+			const UEnum* OpTypeEnum = StaticEnum<EMovieGraphConditionGroupQueryOpType>();
+			Query->SetOperationType(static_cast<EMovieGraphConditionGroupQueryOpType>(OpTypeEnum->GetValueByName(InNewOpType)));
+		}
+	}
+}
+
+FName SMovieGraphCollectionTreeOpTypeWidget::GetCurrentConditionGroupOpType() const
+{
+	if (WeakTreeElement.IsValid())
+	{
+		if (const TObjectPtr<UMovieGraphConditionGroup>& ConditionGroup = WeakTreeElement.Pin()->ConditionGroup)
+		{
+			const UEnum* OpTypeEnum = StaticEnum<EMovieGraphConditionGroupOpType>();
+			const EMovieGraphConditionGroupOpType OpType = ConditionGroup->GetOperationType();
+			const FText DisplayNameText = OpTypeEnum->GetDisplayNameTextByValue(static_cast<__underlying_type(EMovieGraphConditionGroupOpType)>(OpType));
+			
+			return FName(DisplayNameText.ToString());
+		}
+	}
+		
+	return FName();
+}
+
+FName SMovieGraphCollectionTreeOpTypeWidget::GetCurrentConditionGroupQueryOpType() const
+{
+	if (WeakTreeElement.IsValid())
+	{
+		if (const TObjectPtr<UMovieGraphConditionGroupQueryBase>& Query = WeakTreeElement.Pin()->Query)
+		{
+			const UEnum* OpTypeEnum = StaticEnum<EMovieGraphConditionGroupQueryOpType>();
+			const EMovieGraphConditionGroupQueryOpType OpType = Query->GetOperationType();
+			const FText DisplayNameText = OpTypeEnum->GetDisplayNameTextByValue(static_cast<__underlying_type(EMovieGraphConditionGroupQueryOpType)>(OpType));
+			
+			return FName(DisplayNameText.ToString());
+		}
+	}
+		
+	return FName();
+}
+
+void SMovieGraphCollectionTreeAddQueryContentWidget::Construct(const FArguments& InArgs)
+{
+	TWeakPtr<FCollectionCustomizationTreeElement> WeakTreeElement = InArgs._WeakTreeElement.Get();
+	TWeakPtr<SMovieGraphCollectionTreeView> WeakTreeView = InArgs._WeakTreeView.Get();
+
+	if (!WeakTreeElement.IsValid() || !WeakTreeView.IsValid())
+	{
+		return;
+	}
+
+	// Call the underlying query to get the contents of the menu. The query will call back to the widget if something was added so the tree view has
+	// an opportunity to refresh itself.
+	TSharedRef<SWidget> AddMenuContents = WeakTreeElement.Pin()->Query->GetAddMenuContents(
+		UMovieGraphConditionGroupQueryBase::FMovieGraphConditionGroupQueryContentsChanged::CreateLambda([WeakTreeElement, WeakTreeView]()
+    {
+    	if (WeakTreeView.IsValid())
+    	{
+    		if (WeakTreeElement.IsValid())
+    		{
+    			WeakTreeElement.Pin()->ClearCachedChildren();
+    		}
+    		
+    		WeakTreeView.Pin()->RefreshView();
+    		FSlateApplication::Get().DismissAllMenus();
+    	}
+    }));
+
+    // Generate a (multi-layered) icon for the "Add" menu
+    const TSharedRef<SLayeredImage> AddIcon =
+    	SNew(SLayeredImage)
+    	.ColorAndOpacity(FSlateColor::UseForeground())
+    	.Image(FAppStyle::GetBrush("LevelEditor.OpenAddContent.Background"));
+    AddIcon->AddLayer(FAppStyle::GetBrush("LevelEditor.OpenAddContent.Overlay"));
+
+	ChildSlot
+	[
+		SNew(SComboButton)
+		.ToolTipText(LOCTEXT("AddConditionGroupQueryContent", "Add content to this query."))
+		.ComboButtonStyle(FAppStyle::Get(), "SimpleComboButton")
+		.ContentPadding(0)
+		.HasDownArrow(false)
+		.IsEnabled_Lambda([WeakTreeElement]()
+		{
+			return WeakTreeElement.IsValid() && WeakTreeElement.Pin()->IsQueryEnabled();
+		})
+		.OnGetMenuContent_Lambda([AddMenuContents]()
+		{
+			return AddMenuContents;
+		})
+		.Visibility_Lambda([AddMenuContents]()
+		{
+			// Queries can opt to not show an Add menu by returning NullWidget
+			return (AddMenuContents != SNullWidget::NullWidget) ? EVisibility::Visible : EVisibility::Collapsed;
+		})
+		.ButtonContent()
+		[
+			AddIcon
+		]
+	];
+}
+
+void SMovieGraphCollectionTreeQueryTypeSelectorWidget::Construct(const FArguments& InArgs)
+{
+	WeakTreeElement = InArgs._WeakTreeElement.Get();
+	WeakTreeView = InArgs._WeakTreeView.Get();
+
+	ChildSlot
+	[
+		SNew(SComboBox<UClass*>)
+		.IsEnabled_Lambda([this]()
+		{
+			return WeakTreeElement.IsValid() && WeakTreeElement.Pin()->IsQueryEnabled();
+		})
+		.OptionsSource(GetAvailableQueryTypes())
+		.OnSelectionChanged(this, &SMovieGraphCollectionTreeQueryTypeSelectorWidget::SetQueryType)
+		.OnGenerateWidget_Lambda([this](UClass* InClass)
+		{
+			return GetQueryTypeContents(InClass);
+		})
+		.Content()
+		[
+			GetQueryTypeContents(nullptr)
+		]
+	];
+}
+
+TArray<UClass*>* SMovieGraphCollectionTreeQueryTypeSelectorWidget::GetAvailableQueryTypes()
+{
+	static TArray<UClass*> QueryTypes;
+
+	if (QueryTypes.IsEmpty())
+	{
+		for (TObjectIterator<UClass> ClassIterator; ClassIterator; ++ClassIterator)
+		{
+			UClass* Class = *ClassIterator;
+
+			if (Class->IsChildOf(UMovieGraphConditionGroupQueryBase::StaticClass()) &&
+				!Class->HasAnyClassFlags(CLASS_Abstract | CLASS_Deprecated | CLASS_Hidden))
+			{
+				QueryTypes.Add(Class);
+			}
+		}
+	}
+
+	return &QueryTypes;
+}
+
+TSharedRef<SWidget> SMovieGraphCollectionTreeQueryTypeSelectorWidget::GetQueryTypeContents(UClass* InTypeClass) const
+{
+	TSharedPtr<SWidget> QueryTypeTextBlock;
+
+	// This method may only be called once in some scenarios, so provide the option of a text lambda vs. constant text if the
+	// text needs to be updated dynamically when the query type changes.
+	if (InTypeClass == nullptr)
+	{
+		SAssignNew(QueryTypeTextBlock, STextBlock)
+		.Text_Lambda([this]()
+		{
+			return GetCurrentQueryTypeDisplayName();
+		})
+		.Font(FAppStyle::GetFontStyle(TEXT("PropertyWindow.NormalFont")));
+
+		// For below, assign the query type to the current query type
+		InTypeClass = GetCurrentQueryType();
+	}
+	else
+	{
+		SAssignNew(QueryTypeTextBlock, STextBlock)
+		.Text(InTypeClass->GetDisplayNameText())
+		.Font(FAppStyle::GetFontStyle(TEXT("PropertyWindow.NormalFont")));
+	}
+
+	// Defer to the query to get the icon that should be displayed
+	const FSlateIcon QueryTypeIcon = InTypeClass
+		? GetDefault<UMovieGraphConditionGroupQueryBase>(InTypeClass)->GetIcon()
+		: FSlateIcon();
+	
+	return
+		SNew(SHorizontalBox)
+		+ SHorizontalBox::Slot()
+		.VAlign(VAlign_Center)
+		.Padding(0, 0, 7.f, 0)
+		.AutoWidth()
+		[
+			SNew(SImage)
+			.Image(QueryTypeIcon.GetIcon())
+			.ColorAndOpacity(FSlateColor::UseForeground())
+		]
+
+		+ SHorizontalBox::Slot()
+		.VAlign(VAlign_Center)
+		[
+			QueryTypeTextBlock.ToSharedRef()
+		];
+}
+
+FText SMovieGraphCollectionTreeQueryTypeSelectorWidget::GetCurrentQueryTypeDisplayName() const
+{
+	if (WeakTreeElement.IsValid())
+	{
+		if (const UClass* CurrentQueryType = GetCurrentQueryType())
+		{
+			return CurrentQueryType->GetDisplayNameText();
+		}
+	}
+		
+	return FText();
+}
+
+UClass* SMovieGraphCollectionTreeQueryTypeSelectorWidget::GetCurrentQueryType() const
+{
+	if (WeakTreeElement.IsValid())
+	{
+		return WeakTreeElement.Pin()->Query->GetClass();
+	}
+
+	return nullptr;
+}
+
+void SMovieGraphCollectionTreeQueryTypeSelectorWidget::SetQueryType(UClass* InNewQueryType, ESelectInfo::Type SelectInfo) const
+{
+	if (!WeakTreeElement.IsValid())
+	{
+		return;
+	}
+	
+	const TSharedPtr<FCollectionCustomizationTreeElement> TreeElementPin = WeakTreeElement.Pin();
+
+	// Swap the old query with a new one
+	if (const TObjectPtr<UMovieGraphConditionGroup>& ConditionGroup = TreeElementPin->ConditionGroup)
+	{
+		const auto& ExistingQuery = TreeElementPin->Query;
+		const int32 ExistingQueryIndex = ConditionGroup->GetQueries().Find(ExistingQuery);
+
+		// Add the new query at the index of the current query. If the query couldn't be found for some reason, default to the end of the array (-1)
+		ConditionGroup->AddQuery(InNewQueryType, ExistingQueryIndex != INDEX_NONE ? ExistingQueryIndex : -1);
+	
+		ConditionGroup->RemoveQuery(ExistingQuery);
+	}
+
+	// Changing the query type removes the current query and adds a new one to the parent condition group, which means the parent's children
+	// need to be refreshed
+	if (TreeElementPin->ParentElement.IsValid())
+	{
+		TreeElementPin->ParentElement.Pin()->ClearCachedChildren();			
+	}
+
+	if (WeakTreeView.IsValid())
+	{
+		WeakTreeView.Pin()->RefreshView();
+	}
+}
+
+FCollectionCustomizationTreeElement::FCollectionCustomizationTreeElement(const EElementType InType)
+	: Type(InType)
+{
+		
+}
+
+bool FCollectionCustomizationTreeElement::IsQueryEnabled() const
+{
+	return Query ? Query->IsEnabled() : false;
+}
+
+const TArray<TSharedPtr<FCollectionCustomizationTreeElement>>& FCollectionCustomizationTreeElement::GetChildren() const
+{
+	// Returned the cached children if they are available
+	if (!ChildrenCache.IsEmpty())
+	{
+		return ChildrenCache;
+	}
+
+	/** Creates a new child element of the specified type, with the associated pointers to the objects needed for that type. */
+	auto AddChildElement = [this](
+		const EElementType ElementType, UMovieGraphConditionGroup* InConditionGroup = nullptr,
+		UMovieGraphConditionGroupQueryBase* InQuery = nullptr, FProperty* InQueryProperty = nullptr)
+	{
+		const TSharedPtr<FCollectionCustomizationTreeElement> NewElement = MakeShared<FCollectionCustomizationTreeElement>(ElementType);
+		NewElement->Collection = Collection;
+		NewElement->ConditionGroup = InConditionGroup;
+		NewElement->Query = InQuery;
+		NewElement->QueryProperty = InQueryProperty;
+		NewElement->ParentElement = AsWeak();
+
+		ChildrenCache.Add(NewElement);
+	};
+
+	if (Type == EElementType::Root)
+	{
+		const TArray<UMovieGraphConditionGroup*>& ConditionGroups = Collection->GetConditionGroups();
+
+		if (ConditionGroups.IsEmpty())
+		{
+			AddChildElement(EElementType::EmptyCollection);
+		}
+		else
+		{
+			for (UMovieGraphConditionGroup* Group : ConditionGroups)
+			{
+				AddChildElement(EElementType::ConditionGroup, Group);
+			}
+		}
+	}
+	
+	if (Type == EElementType::ConditionGroup)
+	{
+		const TArray<UMovieGraphConditionGroupQueryBase*>& ConditionGroupQueries = ConditionGroup->GetQueries();
+		
+		if (ConditionGroupQueries.IsEmpty())
+		{
+			AddChildElement(EElementType::EmptyConditionGroup);
+		}
+		else
+		{
+			for (UMovieGraphConditionGroupQueryBase* GroupQuery : ConditionGroup->GetQueries())
+			{
+				AddChildElement(EElementType::Query, ConditionGroup, GroupQuery);
+			}
+		}
+	}
+
+	if (Type == EElementType::Query)
+	{
+		// If the query provides widgets, add a row for each widget. Otherwise, add a row for each EditAnywhere property within the query.
+		const TArray<TSharedRef<SWidget>> QueryWidgets = Query->GetWidgets();
+		if (!QueryWidgets.IsEmpty())
+		{
+			for (const TSharedRef<SWidget>& Widget : QueryWidgets)
+			{
+				FProperty* UnknownProperty = nullptr;
+				AddChildElement(EElementType::QueryProperty, ConditionGroup, Query, UnknownProperty);
+			}
+		}
+		else
+		{
+			for (TFieldIterator<FProperty> PropIt(Query->GetClass()); PropIt; ++PropIt)
+			{
+				FProperty* Prop = *PropIt;
+				if (!Prop->HasAllPropertyFlags(CPF_Edit))
+				{
+					continue;
+				}
+
+				AddChildElement(EElementType::QueryProperty, ConditionGroup, Query, Prop);
+			}
+		}
+	}
+
+	return ChildrenCache;
+}
+
+void FCollectionCustomizationTreeElement::ClearCachedChildren() const
+{
+	ChildrenCache.Empty();
+}
+
+uint32 FCollectionCustomizationTreeElement::GetHash() const
+{
+	if (ElementHash != 0)
+	{
+		return ElementHash;
+	}
+
+	switch (Type)
+	{
+	case EElementType::Root:
+	case EElementType::EmptyCollection:
+	case EElementType::EmptyConditionGroup:
+		ElementHash = GetTypeHash(true);	// There will only ever be one root, or one of the "empty" notices under a given parent
+		break;
+	case EElementType::ConditionGroup:
+		ElementHash = GetTypeHash(ConditionGroup);
+		break;
+	case EElementType::Query:
+		ElementHash = GetTypeHash(Query);
+		break;
+	case EElementType::QueryProperty:
+		ElementHash = GetTypeHash(QueryProperty);
+		break;
+	}
+		
+	if (ParentElement.IsValid())
+	{
+		ElementHash = HashCombine(ElementHash, ParentElement.Pin()->GetHash());
+	}
+
+	return ElementHash;
+}
+
+TArray<TSharedRef<SWidget>> SMovieGraphCollectionTreeItem_Stretch::GetColumnWidgets() const
+{
+	if (!WeakTreeElement.IsValid())
+	{
+		return { SNullWidget::NullWidget };
+	}
+
+	const TSharedPtr<FCollectionCustomizationTreeElement> TreeElement = WeakTreeElement.Pin();
+
+	if (TreeElement->Type == FCollectionCustomizationTreeElement::EElementType::QueryProperty)
+	{
+		// Display the widgets specified by the query
+		return TreeElement->Query->GetWidgets();
+	}
+
+	if (TreeElement->Type == FCollectionCustomizationTreeElement::EElementType::EmptyCollection ||
+		TreeElement->Type == FCollectionCustomizationTreeElement::EElementType::EmptyConditionGroup)
+	{
+		const FText NoticeText = (TreeElement->Type == FCollectionCustomizationTreeElement::EElementType::EmptyCollection)
+			? LOCTEXT("EmptyCollectionNotice", "Add a condition group to this collection.")
+			: LOCTEXT("EmptyConditionGroupNotice", "Add a condition to this condition group.");
+			
+		return {
+			SNew(SBox)
+			.MinDesiredHeight(50)
+			.Padding(10.f)
+			.HAlign(HAlign_Fill)
+			.VAlign(VAlign_Center)
+			.Content()
+			[
+				SNew(STextBlock)
+				.Text(NoticeText)
+			]
+		};
+	}
+	
+	return { SNullWidget::NullWidget };
+}
+
+void SMovieGraphCollectionTreeItem_Stretch::Construct(
+	const FArguments& InArgs, const TSharedPtr<SMovieGraphCollectionTreeView>& InOwnerTable, const TSharedPtr<FCollectionCustomizationTreeElement>& InTreeElement)
+{
+	check(InTreeElement);
+	WeakTreeElement = InTreeElement.ToWeakPtr();
+
+	TSharedPtr<SVerticalBox> WidgetContainer;
+
+	// Add a new indented row which has a designated place for widgets to be added (WidgetContainer)
+	this->ChildSlot
+	[
+		SNew(SHorizontalBox)
+		+ SHorizontalBox::Slot()
+		.AutoWidth()
+		.VAlign(VAlign_Fill)
+		[
+			SNew(SMovieGraphCollectionTreeViewIndent)
+			.NumIndents(2)
+		]
+
+		+ SHorizontalBox::Slot()
+		.VAlign(VAlign_Top)
+		.HAlign(HAlign_Fill)
+		[
+			SNew(SBorder)
+			.Padding(0, 0, 0, 1.f)
+			.HAlign(HAlign_Fill)
+			.BorderImage(FMovieRenderPipelineStyle::Get().GetBrush("MovieRenderGraph.CollectionsTree.BottomBorder"))
+			[
+				SAssignNew(WidgetContainer, SVerticalBox)
+			]
+		]
+	];
+
+	// Add all widgets that are generated by the column
+	for (const TSharedRef<SWidget>& Widget : GetColumnWidgets())
+	{
+		WidgetContainer->AddSlot()
+			.HAlign(HAlign_Fill)
+			[
+				SNew(SBox)
+				.IsEnabled_Lambda([InTreeElement]()
+				{
+					return InTreeElement->IsQueryEnabled();
+				})
+				[
+					Widget
+				]
+			];
+	}
+	
+	ConstructInternal(
+		FArguments()
+			.ShowSelection(true)
+			.Style(FMovieRenderPipelineStyle::Get(), "MovieRenderGraph.CollectionsTree.Row"),
+		InOwnerTable.ToSharedRef());
+}
+
+void SMovieGraphCollectionTreeItem::Construct(const FArguments& InArgs, const TSharedPtr<SMovieGraphCollectionTreeView>& InOwnerTable, const TSharedPtr<FCollectionCustomizationTreeElement>& InTreeElement, const TWeakPtr<IDetailLayoutBuilder>& InDetailBuilder)
+{
+	WeakTreeElement = InTreeElement;
+	WeakTreeView = InOwnerTable.ToWeakPtr();
+	DetailBuilder = InDetailBuilder;
+
+	SMultiColumnTableRow<TSharedPtr<FCollectionCustomizationTreeElement>>::Construct(
+		SMultiColumnTableRow<TSharedPtr<FCollectionCustomizationTreeElement>>::FArguments()
+			.Padding(0)
+			.Style(FMovieRenderPipelineStyle::Get(), "MovieRenderGraph.CollectionsTree.Row")
+			.OnDragDetected(this, &SMovieGraphCollectionTreeItem::HandleOnDragDetected)
+			.OnCanAcceptDrop(this, &SMovieGraphCollectionTreeItem::HandleOnCanAcceptDrop)
+			.OnAcceptDrop(this, &SMovieGraphCollectionTreeItem::HandleOnAcceptDrop)
+		, InOwnerTable.ToSharedRef());
+}
+
+FReply SMovieGraphCollectionTreeItem::HandleOnDragDetected(const FGeometry& InGeometry, const FPointerEvent& MouseEvent) const
+{
+	if (!WeakTreeView.IsValid())
+	{
+		return FReply::Unhandled();
+	}
+	
+	TArray<TSharedPtr<FCollectionCustomizationTreeElement>> SelectedElements = WeakTreeView.Pin()->GetSelectedItems();
+
+	if ((SelectedElements.Num() != 1) || !SelectedElements[0].IsValid())
+	{
+		return FReply::Unhandled();
+	}
+
+	if (MouseEvent.IsMouseButtonDown(EKeys::LeftMouseButton))
+	{
+		const FText DragDropText = SelectedElements[0]->Type == FCollectionCustomizationTreeElement::EElementType::ConditionGroup
+			? LOCTEXT("DragConditionGroup", "Move 1 Condition Group")
+			: LOCTEXT("DragConditionGroupQuery", "Move 1 Condition");
+					
+		const TSharedRef<FMovieGraphCollectionTreeDragOp> DragDropOp = FMovieGraphCollectionTreeDragOp::New(SelectedElements[0]);
+		DragDropOp->CurrentHoverText = DragDropText;
+					
+		return FReply::Handled().BeginDragDrop(DragDropOp);
+	}
+
+	return FReply::Unhandled();
+}
+
+TOptional<EItemDropZone> SMovieGraphCollectionTreeItem::HandleOnCanAcceptDrop(const FDragDropEvent& DragDropEvent, const EItemDropZone DropZone, const TSharedPtr<FCollectionCustomizationTreeElement> TargetItem)
+{
+	TOptional<EItemDropZone> ReturnedDropZone;
+	const TSharedPtr<FMovieGraphCollectionTreeDragOp> DragDropOp = DragDropEvent.GetOperationAs<FMovieGraphCollectionTreeDragOp>();
+
+	// Can't drag an element onto another element -- only changing order is allowed (above/below)
+	if (DropZone == EItemDropZone::OntoItem)
+	{
+		DragDropOp->CurrentIconBrush = FAppStyle::GetBrush(TEXT("Graph.ConnectorFeedback.Error"));
+		return ReturnedDropZone;
+	}
+				
+	if (DragDropOp.IsValid() && DragDropOp->Element.IsValid())
+	{
+		// All tree items can be moved *except* query properties
+		if (TargetItem.Get()->Type != FCollectionCustomizationTreeElement::EElementType::QueryProperty)
+		{
+			// For now, elements cannot be moved to a different parent
+			if (TargetItem.Get()->ParentElement == DragDropOp->Element->ParentElement)
+			{
+				DragDropOp->CurrentIconBrush = FAppStyle::GetBrush(TEXT("Graph.ConnectorFeedback.Ok"));
+				ReturnedDropZone = DropZone;
+			}
+			else
+			{
+				DragDropOp->CurrentIconBrush = FAppStyle::GetBrush(TEXT("Graph.ConnectorFeedback.Error"));
+			}
+		}
+	}
+
+	return ReturnedDropZone;
+}
+
+FReply SMovieGraphCollectionTreeItem::HandleOnAcceptDrop(const FDragDropEvent& DragDropEvent, const EItemDropZone DropZone, const TSharedPtr<FCollectionCustomizationTreeElement> TargetItem) const
+{
+	const TSharedPtr<FMovieGraphCollectionTreeDragOp> DragDropOp = DragDropEvent.GetOperationAs<FMovieGraphCollectionTreeDragOp>();
+	if (!DragDropOp.IsValid())
+	{
+		return FReply::Unhandled();
+	}
+
+	// Figure out the new order of elements
+	if (TargetItem->Type == FCollectionCustomizationTreeElement::EElementType::ConditionGroup)
+	{
+		if (TargetItem->Collection)
+		{
+			const TArray<UMovieGraphConditionGroup*>& ConditionGroups = TargetItem->Collection->GetConditionGroups();
+						
+			// Find the index of the target
+			const int32 TargetIndex = ConditionGroups.Find(TargetItem->ConditionGroup);
+			if (TargetIndex != INDEX_NONE)
+			{
+				FScopedTransaction Transaction(LOCTEXT("ReorderConditionGroup", "Reorder Condition Group"));
+				const int32 NewIndex = DropZone == EItemDropZone::AboveItem ? TargetIndex : TargetIndex + 1;
+				TargetItem->Collection->MoveConditionGroupToIndex(DragDropOp.Get()->Element->ConditionGroup, NewIndex);
+							
+				DragDropOp.Get()->Element->ParentElement.Pin()->ClearCachedChildren();
+				RefreshTreeView();
+			}
+		}
+	}
+	else if (TargetItem->Type == FCollectionCustomizationTreeElement::EElementType::Query)
+	{
+		if (TargetItem->ConditionGroup)
+		{
+			const TArray<UMovieGraphConditionGroupQueryBase*>& Queries = TargetItem->ConditionGroup->GetQueries();
+
+			// Find the index of the target
+			const int32 TargetIndex = Queries.Find(TargetItem->Query);
+			if (TargetIndex != INDEX_NONE)
+			{
+				FScopedTransaction Transaction(LOCTEXT("ReorderConditionGroupQuery", "Reorder Condition Group Query"));
+				const int32 NewIndex = DropZone == EItemDropZone::AboveItem ? TargetIndex : TargetIndex + 1;
+				TargetItem->ConditionGroup->MoveQueryToIndex(DragDropOp.Get()->Element->Query, NewIndex);
+							
+				DragDropOp.Get()->Element->ParentElement.Pin()->ClearCachedChildren();
+				RefreshTreeView();
+			}
+		}
+	}
+
+	return FReply::Handled();
+}
+
+TSharedRef<SWidget> SMovieGraphCollectionTreeItem::GenerateConditionGroupColumnWidget(const FName& ColumnName, const TSharedPtr<FCollectionCustomizationTreeElement>& TreeElement)
+{
+	if (ColumnName == ColumnID_Name)
+	{
+		// Generate the condition group name
+		const int32 GroupIndex = (TreeElement->ConditionGroup && TreeElement->Collection)
+			? TreeElement->Collection->GetConditionGroups().Find(TreeElement->ConditionGroup)
+			: 0;
+		const FString GroupName = FString::Printf(TEXT("Condition Group %i"), (GroupIndex + 1));
+		
+		return
+			SNew(SBorder)
+			.HAlign(HAlign_Fill)
+			.Padding(0)
+			.BorderImage(FMovieRenderPipelineStyle::Get().GetBrush("MovieRenderGraph.CollectionsTree.BottomBorder"))
+			[
+				SNew(SHorizontalBox)
+				+ SHorizontalBox::Slot()
+				.AutoWidth()
+				.VAlign(VAlign_Fill)
+				.Padding(0)
+				[
+					SNew(SMovieGraphCollectionTreeViewIndent)
+					.NumIndents(1)
+				]
+		
+				+ SHorizontalBox::Slot()
+				.AutoWidth()
+				.VAlign(VAlign_Center)
+				[
+					SNew(SExpanderArrow, SharedThis(this))
+					.IndentAmount(0)
+				]
+
+				+ SHorizontalBox::Slot()
+				.AutoWidth()
+				.VAlign(VAlign_Center)
+				.Padding(0, 0, 2.f, 0)
+				[
+					SNew(SImage)
+					.Image(FAppStyle::GetBrush("ClassIcon.GroupActor"))
+					.ColorAndOpacity(FSlateColor::UseForeground())
+				]
+
+				+ SHorizontalBox::Slot()
+				.AutoWidth()
+				.VAlign(VAlign_Center)
+				.Padding(0, 0, 2.f, 0)
+				[
+					SNew(STextBlock)
+					.Text(FText::FromString(GroupName))
+					.Font(FAppStyle::GetFontStyle(TEXT("PropertyWindow.BoldFont")))
+				]
+			];
+	}
+
+	if (ColumnName == ColumnID_Value)
+	{
+		return
+			SNew(SBorder)
+			.Padding(0)
+			.BorderImage(FMovieRenderPipelineStyle::Get().GetBrush("MovieRenderGraph.CollectionsTree.BottomLeftBorder"))
+			[
+				SNew(SHorizontalBox)
+				+ SHorizontalBox::Slot()
+				.VAlign(VAlign_Center)
+				.HAlign(HAlign_Left)
+				.AutoWidth()
+				[
+					SNew(SButton)
+					.HAlign(HAlign_Center)
+					.ToolTip(SNew(SToolTip).Text(LOCTEXT("AddConditionGroupQuery", "Add a condition to this condition group.")))
+					.ButtonStyle(FAppStyle::Get(), "SimpleButton")
+					.OnClicked_Lambda([this]()
+					{
+						if (WeakTreeElement.IsValid())
+						{
+							const TSharedPtr<FCollectionCustomizationTreeElement> TreeElementPin = WeakTreeElement.Pin();
+							
+							if (const TObjectPtr<UMovieGraphConditionGroup> ConditionGroup = TreeElementPin->ConditionGroup)
+							{
+								// Use the Actor Name query as the default
+								ConditionGroup->AddQuery(UMovieGraphConditionGroupQuery_ActorName::StaticClass());
+
+								if (TreeElementPin->ParentElement.IsValid())
+								{
+									TreeElementPin->ParentElement.Pin()->ClearCachedChildren();
+								}
+								
+								RefreshTreeView();
+							}
+						}
+
+						return FReply::Handled();
+					})
+					.Content()
+					[
+						SNew(SImage)
+						.Image(FAppStyle::GetBrush("Icons.PlusCircle"))
+						.ColorAndOpacity(FSlateColor::UseForeground())
+					]
+				]
+
+				+ SHorizontalBox::Slot()
+				.VAlign(VAlign_Center)
+				.HAlign(HAlign_Left)
+				.AutoWidth()
+				[
+					SNew(SComboButton)
+					.ComboButtonStyle(FAppStyle::Get(), "SimpleComboButton")
+					.ContentPadding(0)
+					.ForegroundColor(FSlateColor::UseForeground())
+					.HasDownArrow(true)
+					.OnGetMenuContent_Lambda([this]()
+					{
+						FMenuBuilder ViewOptions(true, nullptr);
+
+						ViewOptions.AddMenuEntry(
+							LOCTEXT("DeleteConditionGroup", "Delete"),
+							LOCTEXT("DeleteConditionGroupTooltip", "Delete this condition group and the queries in it."),
+							FSlateIcon(),
+							FUIAction(
+								FExecuteAction::CreateLambda([this]()
+								{
+									if (WeakTreeView.IsValid())
+									{
+										WeakTreeView.Pin()->RemoveConditionGroup(WeakTreeElement);
+									}
+                                }),
+								FCanExecuteAction()));
+
+						return ViewOptions.MakeWidget();
+					})
+				]
+
+				+ SHorizontalBox::Slot()
+				.VAlign(VAlign_Center)
+				.HAlign(HAlign_Right)
+				.Padding(2.f)
+				.FillWidth(1.f)
+				[
+					SNew(SMovieGraphCollectionTreeOpTypeWidget)
+					.IsConditionGroup(true)
+					.WeakTreeElement(WeakTreeElement)
+				]
+			];
+	}
+	
+	return SNullWidget::NullWidget;
+}
+
+TSharedRef<SWidget> SMovieGraphCollectionTreeItem::GenerateConditionGroupQueryColumnWidget(const FName& ColumnName, const TSharedPtr<FCollectionCustomizationTreeElement>& TreeElement)
+{
+	if (ColumnName == ColumnID_Name)
+	{
+		return
+			SNew(SHorizontalBox)
+			+ SHorizontalBox::Slot()
+			.AutoWidth()
+			.VAlign(VAlign_Fill)
+			[
+				SNew(SMovieGraphCollectionTreeViewIndent)
+				.NumIndents(2)
+			]
+
+			+ SHorizontalBox::Slot()
+			.HAlign(HAlign_Fill)
+			[
+				SNew(SBorder)
+				.Padding(0)
+				.BorderImage(FMovieRenderPipelineStyle::Get().GetBrush("MovieRenderGraph.CollectionsTree.BottomBorder"))
+				[
+					SNew(SHorizontalBox)
+					+ SHorizontalBox::Slot()
+					.AutoWidth()
+					.VAlign(VAlign_Center)
+					[
+						SNew(SExpanderArrow, SharedThis(this))
+						.IndentAmount(0)
+					]
+
+					+ SHorizontalBox::Slot()
+					.AutoWidth()
+					.VAlign(VAlign_Center)
+					.Padding(0, 0, 4.f, 0)
+					[
+						SNew(SCheckBox)
+						.IsChecked_Lambda([TreeElement]()
+						{
+							return TreeElement->IsQueryEnabled() ? ECheckBoxState::Checked : ECheckBoxState::Unchecked;
+						})
+						.OnCheckStateChanged_Lambda([TreeElement](ECheckBoxState NewCheckState)
+						{
+							if (TreeElement.IsValid() && TreeElement->Query)
+							{
+								TreeElement->Query->SetEnabled(NewCheckState == ECheckBoxState::Checked);
+							}
+						})
+					]
+
+					+ SHorizontalBox::Slot()
+					.VAlign(VAlign_Center)
+					.HAlign(HAlign_Fill)
+					.Padding(0, 0, 7.f, 0)
+					[
+						SNew(SMovieGraphCollectionTreeQueryTypeSelectorWidget)
+						.WeakTreeElement(WeakTreeElement)
+						.WeakTreeView(WeakTreeView)
+					]
+				]
+			];
+	}
+
+	if (ColumnName == ColumnID_Value)
+	{			
+		return
+			SNew(SBorder)
+			.Padding(0)
+			.BorderImage(FMovieRenderPipelineStyle::Get().GetBrush("MovieRenderGraph.CollectionsTree.BottomLeftBorder"))
+			[
+				SNew(SHorizontalBox)
+				+ SHorizontalBox::Slot()
+				.VAlign(VAlign_Center)
+				.HAlign(HAlign_Left)
+				.Padding(4.f, 0, 0, 0)
+				.AutoWidth()
+				[
+					SNew(SMovieGraphCollectionTreeAddQueryContentWidget)
+					.WeakTreeElement(WeakTreeElement)
+					.WeakTreeView(WeakTreeView)
+				]
+
+				+ SHorizontalBox::Slot()
+				.VAlign(VAlign_Center)
+				.HAlign(HAlign_Left)
+				.AutoWidth()
+				[
+					SNew(SComboButton)
+					.IsEnabled_Lambda([this]()
+					{
+						return WeakTreeElement.IsValid() ? WeakTreeElement.Pin()->IsQueryEnabled() : false;
+					})
+					.ComboButtonStyle(FAppStyle::Get(), "SimpleComboButton")
+					.ContentPadding(0)
+					.ForegroundColor(FSlateColor::UseForeground())
+					.HasDownArrow(true)
+					.OnGetMenuContent_Lambda([this]()
+					{
+						FMenuBuilder ViewOptions(true, nullptr);
+
+						ViewOptions.AddMenuEntry(
+							LOCTEXT("DeleteConditionGroupQuery", "Delete"),
+							LOCTEXT("DeleteConditionGroupQueryTooltip", "Delete this condition group query."),
+							FSlateIcon(),
+							FUIAction(
+								FExecuteAction::CreateLambda([this]()
+								{
+									if (WeakTreeElement.IsValid())
+									{
+										const TSharedPtr<FCollectionCustomizationTreeElement> TreeElementPin = WeakTreeElement.Pin();
+										
+										TreeElementPin->ConditionGroup->RemoveQuery(TreeElementPin->Query);
+
+										if (TreeElementPin->ParentElement.IsValid())
+										{
+											TreeElementPin->ParentElement.Pin()->ClearCachedChildren();
+										}
+										
+										RefreshTreeView();
+									}
+								}),
+								FCanExecuteAction()));
+
+						return ViewOptions.MakeWidget();
+					})
+				]
+
+				+ SHorizontalBox::Slot()
+				.VAlign(VAlign_Center)
+				.HAlign(HAlign_Right)
+				.Padding(2.f)
+				.FillWidth(1.f)
+				[
+					SNew(SMovieGraphCollectionTreeOpTypeWidget)
+					.IsConditionGroup(false)
+					.WeakTreeElement(WeakTreeElement)
+				]
+			];
+	}
+	
+	return SNullWidget::NullWidget;
+}
+
+TSharedRef<SWidget> SMovieGraphCollectionTreeItem::GenerateQueryPropertyColumnWidget(const FName& ColumnName, const TSharedPtr<FCollectionCustomizationTreeElement>& TreeElement)
+{
+	if (!DetailBuilder.IsValid())
+	{
+		return SNullWidget::NullWidget;
+	}
+		
+	// Get the property handle for the property within the query. This is needed in order to get the name/value widgets for it.
+	const TSharedPtr<IPropertyHandle> QueryPropertyHandle =
+		DetailBuilder.Pin()->AddObjectPropertyData({TreeElement->QueryProperty->GetOwnerUObject()}, TreeElement->QueryProperty->GetFName());
+		
+	if (!QueryPropertyHandle || !QueryPropertyHandle.IsValid())
+	{
+		return SNullWidget::NullWidget;
+	}
+		
+	if (ColumnName == ColumnID_Name)
+	{
+		return
+			SNew(SBorder)
+			.Padding(0)
+			.IsEnabled_Lambda([this]()
+			{
+				return WeakTreeElement.IsValid() ? WeakTreeElement.Pin()->IsQueryEnabled() : false;
+			})
+			.BorderImage(FMovieRenderPipelineStyle::Get().GetBrush("MovieRenderGraph.CollectionsTree.BottomBorder"))
+			[
+				SNew(SHorizontalBox)
+				+ SHorizontalBox::Slot()
+				.AutoWidth()
+				.VAlign(VAlign_Fill)
+				[
+					SNew(SMovieGraphCollectionTreeViewIndent)
+					.NumIndents(2)
+				]
+
+				+ SHorizontalBox::Slot()
+				.AutoWidth()
+				[
+					SNew(SExpanderArrow, SharedThis(this))
+					.IndentAmount(0)
+				]
+
+				+ SHorizontalBox::Slot()
+				[
+					QueryPropertyHandle->CreatePropertyNameWidget()
+				]
+			];
+	}
+
+	if (ColumnName == ColumnID_Value)
+	{
+		return
+			SNew(SBorder)
+			.Padding(5.f, 2.f)
+			.IsEnabled_Lambda([this]()
+			{
+				return WeakTreeElement.IsValid() ? WeakTreeElement.Pin()->IsQueryEnabled() : false;
+			})
+			.BorderImage(FMovieRenderPipelineStyle::Get().GetBrush("MovieRenderGraph.CollectionsTree.BottomLeftBorder"))
+			[
+				QueryPropertyHandle->CreatePropertyValueWidget()
+			];
+	}
+		
+	return SNullWidget::NullWidget;
+}
+
+void SMovieGraphCollectionTreeItem::RefreshTreeView() const
+{
+	if (WeakTreeView.IsValid())
+	{
+		WeakTreeView.Pin()->RefreshView();
+	}
+}
+
+TSharedRef<SWidget> SMovieGraphCollectionTreeItem::GenerateWidgetForColumn(const FName& ColumnName)
+{
+	const TSharedPtr<FCollectionCustomizationTreeElement> TreeElement = WeakTreeElement.Pin();
+	if (!TreeElement.IsValid())
+	{
+		return SNullWidget::NullWidget;
+	}
+
+	if (TreeElement->Type == FCollectionCustomizationTreeElement::EElementType::Root)
+	{
+		// The root element isn't visible, but serves as the single entry point to the tree's hierarchy
+		return SNullWidget::NullWidget;
+	}
+
+	if (TreeElement->Type == FCollectionCustomizationTreeElement::EElementType::ConditionGroup)
+	{
+		return GenerateConditionGroupColumnWidget(ColumnName, TreeElement);
+	}
+
+	if (TreeElement->Type == FCollectionCustomizationTreeElement::EElementType::Query)
+	{
+		return GenerateConditionGroupQueryColumnWidget(ColumnName, TreeElement);
+	}
+		
+	return GenerateQueryPropertyColumnWidget(ColumnName, TreeElement);
+}
+
+TSharedRef<IDetailCustomization> FMovieGraphCollectionsCustomization::MakeInstance()
+{
+	return MakeShared<FMovieGraphCollectionsCustomization>();
+}
+
+void FMovieGraphCollectionsCustomization::CustomizeDetails(IDetailLayoutBuilder& InDetailBuilder)
+{
+	TArray<TWeakObjectPtr<UMovieGraphCollectionNode>> CollectionNodes =
+		InDetailBuilder.GetObjectsOfTypeBeingCustomized<UMovieGraphCollectionNode>();
+	if (CollectionNodes.Num() != 1)
+	{
+		// Showing more than one collection node is not supported
+		return;
+	}
+
+	const TWeakObjectPtr<UMovieGraphCollectionNode> CollectionNode = CollectionNodes[0];
+	if (!CollectionNode.IsValid())
+	{
+		return;
+	}
+
+	const TObjectPtr<UMovieGraphCollection> Collection = CollectionNode->Collection;
+
+	// The Collection property is being completely replaced by a custom STreeView
+	const TSharedRef<IPropertyHandle> CollectionProperty = InDetailBuilder.GetProperty("Collection");
+	CollectionProperty->MarkHiddenByCustomization();
+
+	IDetailCategoryBuilder& CollectionCategory = InDetailBuilder.EditCategory("Collection");
+		
+	// Display the "Add Condition Group" button
+	CollectionCategory.AddCustomRow(FText())
+	[
+		SNew(SHorizontalBox)
+		+ SHorizontalBox::Slot()
+		.AutoWidth()
+		[
+			SNew(SPositiveActionButton)
+			.Icon(FAppStyle::Get().GetBrush("Icons.Plus"))
+			.Text(LOCTEXT("AddNewConditionGroup", "Condition Group"))
+			.OnClicked_Lambda([this]
+			{
+				if (TreeView.IsValid())
+				{
+					TreeView->AddConditionGroup();
+					TreeView->RefreshView();
+				}
+
+				return FReply::Handled();
+			})
+		]
+	];
+
+	// Add the collection name property
+	CollectionCategory.AddExternalObjectProperty({Collection}, FName(TEXT("CollectionName")));
+
+	FDetailWidgetRow& ConditionsWidgetRow = InDetailBuilder.EditCategory("Conditions")
+		.AddCustomRow(FText::FromString(TEXT("Conditions")));
+	
+	ConditionsWidgetRow.WholeRowWidget
+	[
+		SNew(SBox)
+		.VAlign(VAlign_Fill)
+		.Padding(0)
+		[
+			SAssignNew(TreeView, SMovieGraphCollectionTreeView)
+			.DetailBuilder(DetailBuilder)
+			.Collection(MakeWeakObjectPtr(Collection))
+		]
+	];
+
+	// Add the root elements to the tree; root elements will take care of adding children underneath them
+	TreeView->AddRootElement();
+}
+
+void FMovieGraphCollectionsCustomization::CustomizeDetails(const TSharedPtr<IDetailLayoutBuilder>& InDetailBuilder)
+{
+	DetailBuilder = InDetailBuilder;
+	CustomizeDetails(*InDetailBuilder);
+}
+
+#undef LOCTEXT_NAMESPACE
diff --git a/Engine/Plugins/MovieScene/MovieRenderPipeline/Source/MovieRenderPipelineEditor/Private/Customizations/Graph/MovieGraphCollectionsCustomization.h b/Engine/Plugins/MovieScene/MovieRenderPipeline/Source/MovieRenderPipelineEditor/Private/Customizations/Graph/MovieGraphCollectionsCustomization.h
new file mode 100644
--- /dev/null
+++ b/Engine/Plugins/MovieScene/MovieRenderPipeline/Source/MovieRenderPipelineEditor/Private/Customizations/Graph/MovieGraphCollectionsCustomization.h
@@ -0,0 +1,337 @@
+// Copyright Epic Games, Inc. All Rights Reserved.
+
+#pragma once
+
+#include "DragAndDrop/DecoratedDragDropOp.h"
+#include "IDetailCustomization.h"
+#include "Widgets/Views/STreeView.h"
+
+class FCollectionCustomizationTreeElement;
+class UMovieGraphCollection;
+class UMovieGraphConditionGroup;
+class UMovieGraphConditionGroupQueryBase;
+
+/**
+ * A widget which indents the specified number of indent levels, with the appropriate styling for the collection tree view.
+ */
+class SMovieGraphCollectionTreeViewIndent final : public SCompoundWidget
+{
+public:
+	SLATE_BEGIN_ARGS(SMovieGraphCollectionTreeViewIndent) { }
+		SLATE_ATTRIBUTE(int, NumIndents)
+	SLATE_END_ARGS()
+
+	void Construct(const FArguments& InArgs);
+};
+
+/** The drag-drop operator for use within the collections tree view. */
+class FMovieGraphCollectionTreeDragOp final : public FDecoratedDragDropOp
+{
+public:
+	DRAG_DROP_OPERATOR_TYPE(SMovieGraphCollectionTreeDragOp, FDragDropOperation)
+
+	static TSharedRef<FMovieGraphCollectionTreeDragOp> New(const TSharedPtr<FCollectionCustomizationTreeElement>& InElement);
+
+public:
+	/** The tree element associated with this drag-drop operation. */
+	TSharedPtr<FCollectionCustomizationTreeElement> Element;
+
+	/** Whether the drag-drop operation is currently valid. */
+	bool bIsValidOperation = false;
+};
+
+/** A tree view specifically for viewing the contents of a UMovieGraphCollection object. */
+class SMovieGraphCollectionTreeView final : public STreeView<TSharedPtr<FCollectionCustomizationTreeElement>>
+{
+public:
+	SLATE_BEGIN_ARGS(SMovieGraphCollectionTreeView) { }
+		SLATE_ARGUMENT(TWeakPtr<IDetailLayoutBuilder>, DetailBuilder)
+		SLATE_ARGUMENT(TWeakObjectPtr<UMovieGraphCollection>, Collection)
+	SLATE_END_ARGS()
+
+	void Construct(const FArguments& InArgs);
+
+	/** Called when an row in the tree view expands/collapses. Caches the expansion state of the tree so it can be restored at a later time. */
+	void OnExpansionChanged(TSharedPtr<FCollectionCustomizationTreeElement> InElement, bool bInExpanded);
+
+	/** Gets the children of the specified tree element. */
+	void GetChildrenForTree(TSharedPtr<FCollectionCustomizationTreeElement> InItem, TArray<TSharedPtr<FCollectionCustomizationTreeElement>>& OutChildren);
+
+	/** Restores the cached expansion state recursively through all tree elements. */
+	void RestoreExpansionStateRecursive(const TSharedPtr<FCollectionCustomizationTreeElement>& InElement);
+
+	/** Generates a tree row suitable for the type of tree element provided. */
+	TSharedRef<ITableRow> GenerateTreeRow(TSharedPtr<FCollectionCustomizationTreeElement> InTreeElement, const TSharedRef<STableViewBase>& OwnerTable);
+
+	/** Adds the root tree element. The root is not visible, but serves as the single entry point into the tree. */
+	void AddRootElement();
+
+	/** Adds condition group tree elements, one for each condition group in the provided collection. */
+	void AddRootConditionGroupElements(const TObjectPtr<UMovieGraphCollection>& InCollection) const;
+
+	/** Adds a new condition group element associated with the provided condition group. */
+	void AddNewConditionGroupElement(UMovieGraphConditionGroup* InConditionGroup) const;
+
+	/** Adds a new, empty condition group element to the tree, and adds a corresponding empty condition group to the underlying collection. */
+	void AddConditionGroup();
+
+	/** Removes the condition group associated with the provided element. The condition group is also removed from the underlying collection. */
+	void RemoveConditionGroup(const TWeakPtr<FCollectionCustomizationTreeElement>& InConditionGroupElement);
+
+	/** Refreshes the tree view while also taking care of restoring cached expansion state. Should always be called instead of RequestTreeRefresh(). */
+	void RefreshView();
+
+private:
+	/** The collection associated with this tree view. */
+	TWeakObjectPtr<UMovieGraphCollection> Collection = nullptr;
+	
+	/** The root-most elements in the tree. */
+	TArray<TSharedPtr<FCollectionCustomizationTreeElement>> RootElements;
+	
+	/** The hashes of elements that are currently expanded in the tree. */
+	TSet<uint32> ExpandedElements;
+
+	/** The detail builder associated with this tree. Some items in the tree need to add widgets that can be fetched from the detail builder. */
+	TWeakPtr<IDetailLayoutBuilder> DetailBuilder;
+};
+
+/**
+ * A widget which controls the op type for a condition group or condition group query.
+ */
+class SMovieGraphCollectionTreeOpTypeWidget final : public SCompoundWidget
+{
+public:
+	SLATE_BEGIN_ARGS(SMovieGraphCollectionTreeOpTypeWidget) { }
+		SLATE_ATTRIBUTE(bool, IsConditionGroup)
+		SLATE_ATTRIBUTE(TWeakPtr<FCollectionCustomizationTreeElement>, WeakTreeElement)
+	SLATE_END_ARGS()
+
+	void Construct(const FArguments& InArgs);
+
+private:
+	/** Whether this widget is associated with a condition group (else a condition group query). */
+	bool bIsConditionGroup = false;
+
+	/** The tree element this widget is being displayed in. */
+	TWeakPtr<FCollectionCustomizationTreeElement> WeakTreeElement = nullptr;
+
+private:
+	/** Gets the names of all op types available. */
+	template<typename T>
+	static TArray<FName>* GetOpTypes();
+
+	/** Gets a widget (icon + name) representing the given op. */
+	TSharedRef<SWidget> GetOpTypeContents(const FName& InOpName, const bool bIsConditionGroupOp) const;
+
+	/** Whether this widget is currently enabled. */
+	bool IsWidgetEnabled() const;
+
+	/** Sets the op type for the condition group or condition group query. */
+	void SetOpType(const FName InNewOpType, ESelectInfo::Type SelectInfo) const;
+
+	/** Gets the currently-assigned op type for the condition group. */
+	FName GetCurrentConditionGroupOpType() const;
+
+	/** Gets the currently-assigned op type for the condition group query. */
+	FName GetCurrentConditionGroupQueryOpType() const;
+};
+
+/**
+ * A widget which displays the "Add" menu for a condition group query, allowing the user to add content to the query (like an actor). The appearance
+ * of the widget and behavior of the menu is largely delegated to the query itself.
+ */
+class SMovieGraphCollectionTreeAddQueryContentWidget final : public SCompoundWidget
+{
+public:
+	SLATE_BEGIN_ARGS(SMovieGraphCollectionTreeAddQueryContentWidget) { }
+		SLATE_ATTRIBUTE(TWeakPtr<FCollectionCustomizationTreeElement>, WeakTreeElement)
+		SLATE_ATTRIBUTE(TWeakPtr<SMovieGraphCollectionTreeView>, WeakTreeView)
+	SLATE_END_ARGS()
+
+	void Construct(const FArguments& InArgs);
+};
+
+/**
+ * A ComboBox widget which displays the available condition group query types and updates the condition group when one is selected.
+ */
+class SMovieGraphCollectionTreeQueryTypeSelectorWidget final : public SCompoundWidget
+{
+public:
+	SLATE_BEGIN_ARGS(SMovieGraphCollectionTreeQueryTypeSelectorWidget) { }
+		SLATE_ATTRIBUTE(TWeakPtr<FCollectionCustomizationTreeElement>, WeakTreeElement)
+		SLATE_ATTRIBUTE(TWeakPtr<SMovieGraphCollectionTreeView>, WeakTreeView)
+	SLATE_END_ARGS()
+
+	void Construct(const FArguments& InArgs);
+
+private:
+	/** Gets the query types which are available to be selected and added to the condition group. */
+	static TArray<UClass*>* GetAvailableQueryTypes();
+
+	/** Gets the widget that displays an individual query type. */
+	TSharedRef<SWidget> GetQueryTypeContents(UClass* InTypeClass) const;
+
+	/** Gets the display name of the currently selected query. */
+	FText GetCurrentQueryTypeDisplayName() const;
+
+	/** Gets the UClass associated with the currently selected query. */
+	UClass* GetCurrentQueryType() const;
+
+	/** Updates the query to the query type which was selected. */
+	void SetQueryType(UClass* InNewQueryType, ESelectInfo::Type SelectInfo) const;
+
+private:
+	/** The tree element this widget is being displayed in. */
+	TWeakPtr<FCollectionCustomizationTreeElement> WeakTreeElement = nullptr;
+
+	/** The tree view associated with the tree element. */
+	TWeakPtr<SMovieGraphCollectionTreeView> WeakTreeView = nullptr;
+};
+
+/**
+ * Represents the data of one element displayed in the tree, which is one of the types in EElementType. Provides a way
+ * of generating the children under the element (and caching the children for fast access later).
+ */
+class FCollectionCustomizationTreeElement : public TSharedFromThis<FCollectionCustomizationTreeElement>
+{
+public:
+	/** The type of this element. */
+	enum class EElementType : uint8
+	{
+		Root,					///< The root element; does not represent any actual data in the collection
+		EmptyCollection,		///< Represents an "empty collection" notice within a collection
+		ConditionGroup,			///< Represents a condition group within the collection
+		EmptyConditionGroup,	///< Represents an "empty condition group" notice within a condition group
+		Query,					///< Represents a query within a condition group
+		QueryProperty			///< Represents a property within a query
+	};
+	
+	explicit FCollectionCustomizationTreeElement(const EElementType InType);
+
+	/** Whether the query associated with this element is currently enabled. */
+	bool IsQueryEnabled() const;
+
+	/**
+	 * Gets the child elements nested under this element. Returns a cached result if available, otherwise calculates the
+	 * children and returns the result. Call ClearCachedChildren() if the cache needs to be cleared (eg, when the tree
+	 * is being refreshed).
+	 */
+	const TArray<TSharedPtr<FCollectionCustomizationTreeElement>>& GetChildren() const;
+
+	/** Clears the cached result of GetChildren(). */
+	void ClearCachedChildren() const;
+
+	/** Gets the hash that uniquely identifies this element in the tree. */
+	uint32 GetHash() const;
+
+public:
+	/** The type of item in the tree this element represents. */
+	EElementType Type = EElementType::Root;
+
+	/** The collection associated with this element. */
+	TObjectPtr<UMovieGraphCollection> Collection = nullptr;
+
+	/** If this element represents a condition group, this points to the underlying condition group (else nullptr). */
+	TObjectPtr<UMovieGraphConditionGroup> ConditionGroup = nullptr;
+
+	/** If this element represents a condition group query, this points to the underlying query (else nullptr). */
+	TObjectPtr<UMovieGraphConditionGroupQueryBase> Query = nullptr;
+
+	/** If this element represents a query property, this points to the underlying property (else nullptr). */
+	FProperty* QueryProperty = nullptr;
+
+	/** This element's parent element in the tree.  */
+	TWeakPtr<const FCollectionCustomizationTreeElement> ParentElement = nullptr;
+
+private:
+	/** The (cached) child elements nested under this element. */
+	mutable TArray<TSharedPtr<FCollectionCustomizationTreeElement>> ChildrenCache;
+
+	/** The (cached) hash of this element. */
+	mutable uint32 ElementHash = 0;
+};
+
+/** Specialized row class that spans all columns within the tree view. */
+class SMovieGraphCollectionTreeItem_Stretch final : public STableRow<TSharedPtr<FCollectionCustomizationTreeElement>>
+{
+public:
+	/** Gets the widgets that should be displayed within the column (which stretches across all columns in the tree view). */
+	TArray<TSharedRef<SWidget>> GetColumnWidgets() const;
+	
+	void Construct(const FArguments& InArgs, const TSharedPtr<SMovieGraphCollectionTreeView>& InOwnerTable,
+		const TSharedPtr<FCollectionCustomizationTreeElement>& InTreeElement);
+
+private:
+	/** The tree element associated with this item. */
+	TWeakPtr<FCollectionCustomizationTreeElement> WeakTreeElement;
+};
+
+/**
+ * The widget that is responsible for generating the per-column content for each element in the tree.
+ */
+class SMovieGraphCollectionTreeItem final : public SMultiColumnTableRow<TSharedPtr<FCollectionCustomizationTreeElement>>
+{
+	SLATE_BEGIN_ARGS(SMovieGraphCollectionTreeItem) { }
+	SLATE_END_ARGS()
+
+	inline static const FName ColumnID_Name = TEXT("Name");
+	inline static const FName ColumnID_Value = TEXT("Value");
+
+	void Construct(const FArguments& InArgs, const TSharedPtr<SMovieGraphCollectionTreeView>& InOwnerTable,
+		const TSharedPtr<FCollectionCustomizationTreeElement>& InTreeElement, const TWeakPtr<IDetailLayoutBuilder>& InDetailBuilder);
+
+	/** DragDetected handler for this widget within the tree. */
+	FReply HandleOnDragDetected(const FGeometry& InGeometry, const FPointerEvent& MouseEvent) const;
+
+	/** CanAcceptDrop handler for this widget within the tree. */
+	TOptional<EItemDropZone> HandleOnCanAcceptDrop(const FDragDropEvent& DragDropEvent, const EItemDropZone DropZone, const TSharedPtr<FCollectionCustomizationTreeElement> TargetItem);
+
+	/** AcceptDrop handler for this widget within the tree. */
+	FReply HandleOnAcceptDrop(const FDragDropEvent& DragDropEvent, EItemDropZone DropZone, TSharedPtr<FCollectionCustomizationTreeElement> TargetItem) const;
+
+	/** Generates the widgets needed for condition group columns. */
+	TSharedRef<SWidget> GenerateConditionGroupColumnWidget(const FName& ColumnName, const TSharedPtr<FCollectionCustomizationTreeElement>& TreeElement);
+
+	/** Generates the widgets needed for condition group query columns. */
+	TSharedRef<SWidget> GenerateConditionGroupQueryColumnWidget(const FName& ColumnName, const TSharedPtr<FCollectionCustomizationTreeElement>& TreeElement);
+
+	/** Generates the widgets needed for condition group query property columns. */
+	TSharedRef<SWidget> GenerateQueryPropertyColumnWidget(const FName& ColumnName, const TSharedPtr<FCollectionCustomizationTreeElement>& TreeElement);
+
+	/** Convenience method for refreshing the tree view. */
+	void RefreshTreeView() const;
+
+	//~ Begin SMultiColumnTableRow overrides
+	virtual TSharedRef<SWidget> GenerateWidgetForColumn(const FName& ColumnName) override;
+	//~ End SMultiColumnTableRow overrides
+
+private:
+	/** The element that this tree item is displaying. */
+	TWeakPtr<FCollectionCustomizationTreeElement> WeakTreeElement;
+
+	/** The tree that owns this tree element. */
+	TWeakPtr<SMovieGraphCollectionTreeView> WeakTreeView;
+
+	/** The details builder associated with the customization. */
+	TWeakPtr<IDetailLayoutBuilder> DetailBuilder;
+};
+
+/** Customize how the Collection node appears in the details panel. */
+class FMovieGraphCollectionsCustomization final : public IDetailCustomization
+{
+public:
+	static TSharedRef<IDetailCustomization> MakeInstance();
+
+protected:
+	//~ Begin IDetailCustomization interface
+	virtual void CustomizeDetails(IDetailLayoutBuilder& InDetailBuilder) override;
+	virtual void CustomizeDetails(const TSharedPtr<IDetailLayoutBuilder>& InDetailBuilder) override;
+	//~ End IDetailCustomization interface
+
+private:
+	/** The tree view displaying condition groups and the queries contained in them. */
+	TSharedPtr<SMovieGraphCollectionTreeView> TreeView = nullptr;
+
+	/** The details builder associated with the customization. */
+	TWeakPtr<IDetailLayoutBuilder> DetailBuilder;
+};
diff --git a/Engine/Plugins/MovieScene/MovieRenderPipeline/Source/MovieRenderPipelineEditor/Private/Graph/MovieGraphAssetToolkit.cpp b/Engine/Plugins/MovieScene/MovieRenderPipeline/Source/MovieRenderPipelineEditor/Private/Graph/MovieGraphAssetToolkit.cpp
--- a/Engine/Plugins/MovieScene/MovieRenderPipeline/Source/MovieRenderPipelineEditor/Private/Graph/MovieGraphAssetToolkit.cpp
+++ b/Engine/Plugins/MovieScene/MovieRenderPipeline/Source/MovieRenderPipelineEditor/Private/Graph/MovieGraphAssetToolkit.cpp
@@ -2,6 +2,7 @@
 
 #include "MovieGraphAssetToolkit.h"
 
+#include "Customizations/Graph/MovieGraphCollectionsCustomization.h"
 #include "Customizations/Graph/MovieGraphMemberCustomization.h"
 #include "Customizations/Graph/MovieGraphNamedResolutionCustomization.h"
 #include "Customizations/Graph/MovieGraphNodeCustomization.h"
@@ -10,6 +11,7 @@
 #include "Graph/Renderers/MovieGraphShowFlags.h"
 
 #include "Graph/MovieGraphConfig.h"
+#include "Graph/Nodes/MovieGraphCollectionNode.h"
 #include "Graph/Nodes/MovieGraphSelectNode.h"
 #include "MovieEdGraphNode.h"
 #include "MovieGraphSchema.h"
@@ -250,6 +252,10 @@ TSharedRef<SDockTab> FMovieGraphAssetToolkit::SpawnTab_RenderGraphDetails(const
 		FMovieGraphNamedResolution::StaticStruct()->GetFName(),
 		FOnGetPropertyTypeCustomizationInstance::CreateStatic(&FMovieGraphNamedResolutionCustomization::MakeInstance));
 	
+	SelectedGraphObjectsDetailsWidget->RegisterInstancedCustomPropertyLayout(
+		UMovieGraphCollectionNode::StaticClass(),
+		FOnGetDetailCustomizationInstance::CreateStatic(&FMovieGraphCollectionsCustomization::MakeInstance));
+	
 	return SNew(SDockTab)
 		.TabColorScale(GetTabColorScale())
 		.Label(LOCTEXT("DetailsTab_Title", "Details"))
diff --git a/Engine/Plugins/MovieScene/MovieRenderPipeline/Source/MovieRenderPipelineEditor/Private/MovieRenderPipelineStyle.cpp b/Engine/Plugins/MovieScene/MovieRenderPipeline/Source/MovieRenderPipelineEditor/Private/MovieRenderPipelineStyle.cpp
--- a/Engine/Plugins/MovieScene/MovieRenderPipeline/Source/MovieRenderPipelineEditor/Private/MovieRenderPipelineStyle.cpp
+++ b/Engine/Plugins/MovieScene/MovieRenderPipeline/Source/MovieRenderPipelineEditor/Private/MovieRenderPipelineStyle.cpp
@@ -5,6 +5,8 @@
 #include "Styling/SlateTypes.h"
 #include "Styling/CoreStyle.h"
 #include "Styling/AppStyle.h"
+#include "Styling/SlateStyleMacros.h"
+#include "Styling/StyleColors.h"
 
 FName FMovieRenderPipelineStyle::StyleName("MovieRenderPipelineStyle");
 
@@ -101,6 +103,18 @@ FMovieRenderPipelineStyle::FMovieRenderPipelineStyle()
 
 	Set("MovieRenderPipeline.TabIcon", new FSlateImageBrush(RootToContentDir(TEXT("TabIcon_24x.png")), Icon16x16));
 
+	Set("MovieRenderGraph.CollectionsTree.Row", FTableRowStyle(FAppStyle::Get().GetWidgetStyle<FTableRowStyle>("SimpleTableView.Row"))
+		.SetEvenRowBackgroundBrush(FSlateColorBrush(FStyleColors::Header))
+		.SetEvenRowBackgroundHoveredBrush(FSlateColorBrush(FStyleColors::Header))
+		.SetOddRowBackgroundBrush(FSlateColorBrush(FStyleColors::Header))
+		.SetOddRowBackgroundHoveredBrush(FSlateColorBrush(FStyleColors::Header))
+		.SetTextColor(FAppStyle::GetSlateColor("DefaultForeground"))
+		.SetSelectedTextColor(FAppStyle::GetSlateColor("DefaultForeground")));
+
+	Set("MovieRenderGraph.CollectionsTree.BottomBorder", new FSlateBorderBrush(NAME_None, FMargin(0, 0, 0, 1.0), FStyleColors::Recessed));
+	Set("MovieRenderGraph.CollectionsTree.TopBottomBorder", new FSlateBorderBrush(NAME_None, FMargin(0, 1.0, 0, 1.0), FStyleColors::Recessed));
+	Set("MovieRenderGraph.CollectionsTree.BottomLeftBorder", new FSlateBorderBrush(NAME_None, FMargin(1.0, 0, 0, 1.0), FStyleColors::Recessed));
+	Set("MovieRenderGraph.CollectionsTree.IndentedRowBackground", new FSlateColorBrush(FAppStyle::Get().GetSlateColor("Colors.Panel")));
 
 	FSlateStyleRegistry::RegisterSlateStyle(*this);
 }
