MovieGraph: Fixed the PathTracer node not producing any images.

diff --git a/Engine/Plugins/MovieScene/MovieRenderPipeline/Source/MovieRenderPipelineCore/Private/Graph/Nodes/MovieGraphDeferredPassNode.cpp b/Engine/Plugins/MovieScene/MovieRenderPipeline/Source/MovieRenderPipelineCore/Private/Graph/Nodes/MovieGraphDeferredPassNode.cpp
--- a/Engine/Plugins/MovieScene/MovieRenderPipeline/Source/MovieRenderPipelineCore/Private/Graph/Nodes/MovieGraphDeferredPassNode.cpp
+++ b/Engine/Plugins/MovieScene/MovieRenderPipeline/Source/MovieRenderPipelineCore/Private/Graph/Nodes/MovieGraphDeferredPassNode.cpp
@@ -37,4 +37,24 @@ FString UMovieGraphDeferredRenderPassNode::GetRendererNameImpl() const
 EViewModeIndex UMovieGraphDeferredRenderPassNode::GetViewModeIndex() const
 {
 	return ViewModeIndex;
+}
+
+bool UMovieGraphDeferredRenderPassNode::GetWriteAllSamples() const
+{
+	return bWriteAllSamples;
+}
+
+int32 UMovieGraphDeferredRenderPassNode::GetNumSpatialSamples() const
+{
+	return SpatialSampleCount;
+}
+
+bool UMovieGraphDeferredRenderPassNode::GetDisableToneCurve() const
+{
+	return bDisableToneCurve;
+}
+
+EAntiAliasingMethod UMovieGraphDeferredRenderPassNode::GetAntiAliasingMethod() const
+{
+	return AntiAliasingMethod;
 }
\ No newline at end of file
diff --git a/Engine/Plugins/MovieScene/MovieRenderPipeline/Source/MovieRenderPipelineCore/Private/Graph/Nodes/MovieGraphPathTracedRenderPassNode.cpp b/Engine/Plugins/MovieScene/MovieRenderPipeline/Source/MovieRenderPipelineCore/Private/Graph/Nodes/MovieGraphPathTracedRenderPassNode.cpp
--- a/Engine/Plugins/MovieScene/MovieRenderPipeline/Source/MovieRenderPipelineCore/Private/Graph/Nodes/MovieGraphPathTracedRenderPassNode.cpp
+++ b/Engine/Plugins/MovieScene/MovieRenderPipeline/Source/MovieRenderPipelineCore/Private/Graph/Nodes/MovieGraphPathTracedRenderPassNode.cpp
@@ -1,10 +1,15 @@
-﻿// Copyright Epic Games, Inc. All Rights Reserved.
+// Copyright Epic Games, Inc. All Rights Reserved.
 
 #include "Graph/Nodes/MovieGraphPathTracedRenderPassNode.h"
-
+#include "Graph/Renderers/MovieGraphPathTracerPass.h"
 #include "Engine/EngineBaseTypes.h"
 #include "ShowFlags.h"
 
+TUniquePtr<UE::MovieGraph::Rendering::FMovieGraphImagePassBase> UMovieGraphPathTracedRenderPassNode::CreateInstance() const
+{
+	return MakeUnique<UE::MovieGraph::Rendering::FMovieGraphPathTracerPass>();
+}
+
 UMovieGraphPathTracedRenderPassNode::UMovieGraphPathTracedRenderPassNode()
 {
 	ShowFlags->ApplyDefaultShowFlagValue(VMI_PathTracing, true);
@@ -51,3 +56,18 @@ EViewModeIndex UMovieGraphPathTracedRenderPassNode::GetViewModeIndex() const
 {
 	return VMI_PathTracing;
 }
+
+bool UMovieGraphPathTracedRenderPassNode::GetWriteAllSamples() const
+{
+	return bWriteAllSamples;
+}
+
+int32 UMovieGraphPathTracedRenderPassNode::GetNumSpatialSamples() const
+{
+	return SpatialSampleCount;
+}
+
+bool UMovieGraphPathTracedRenderPassNode::GetDisableToneCurve() const
+{
+	return bDisableToneCurve;
+}
diff --git a/Engine/Plugins/MovieScene/MovieRenderPipeline/Source/MovieRenderPipelineCore/Private/Graph/Renderers/MovieGraphDeferredPass.cpp b/Engine/Plugins/MovieScene/MovieRenderPipeline/Source/MovieRenderPipelineCore/Private/Graph/Renderers/MovieGraphDeferredPass.cpp
--- a/Engine/Plugins/MovieScene/MovieRenderPipeline/Source/MovieRenderPipelineCore/Private/Graph/Renderers/MovieGraphDeferredPass.cpp
+++ b/Engine/Plugins/MovieScene/MovieRenderPipeline/Source/MovieRenderPipelineCore/Private/Graph/Renderers/MovieGraphDeferredPass.cpp
@@ -79,6 +79,21 @@ FName FMovieGraphDeferredPass::GetBranchName() const
 	return LayerData.BranchName;
 }
 
+UMovieGraphImagePassBaseNode* FMovieGraphDeferredPass::GetParentNode(UMovieGraphEvaluatedConfig* InConfig) const
+{
+	// This is a bit of a workaround for the fact that the pass doesn't have a strong pointer to the node it's supposed to be associated with,
+	// since that instance changes every frame. So instead we have a virtual function here so the node can look it up by type, and then we can
+	// call a bunch of virtual functions on the right instance to fetch values.
+	const bool bIncludeCDOs = true;
+	UMovieGraphDeferredRenderPassNode* ParentNode = InConfig->GetSettingForBranch<UMovieGraphDeferredRenderPassNode>(GetBranchName(), bIncludeCDOs);
+	if (!ensureMsgf(ParentNode, TEXT("DeferredPass should not exist without parent node in graph.")))
+	{
+		return nullptr;
+	}
+
+	return ParentNode;
+}
+
 void FMovieGraphDeferredPass::Render(const FMovieGraphTraversalContext& InFrameTraversalContext, const FMovieGraphTimeStepData& InTimeData)
 {
 	// ToDo: InFrameTraversalContext includes a copy of TimeData, but may be the one cached at the first temporal sample,
@@ -89,24 +104,16 @@ void FMovieGraphDeferredPass::Render(const FMovieGraphTraversalContext& InFrameT
 		return;
 	}
 
-	// ToDo: This means you have to use FMovieGraphDeferredPass with a DeferredNode (ie: implying that the Path Tracer node has to inherit from the Defered Node,
-	// which incorrectly picks up settings not supported by the path tracer. Perhaps these should be spun off to virtual functions that can be overwritten on 
-	// FMovieGraphPathTracerPass to cast the other way?
-	const bool bIncludeCDOs = false;
-	const UMovieGraphDeferredRenderPassNode* ParentNode = InTimeData.EvaluatedConfig->GetSettingForBranch<UMovieGraphDeferredRenderPassNode>(LayerData.BranchName, bIncludeCDOs);
-	if (!ensureMsgf(ParentNode, TEXT("DeferredPass should not exist without parent node in graph.")))
-	{
-		return;
-	}
-
-	const bool bWriteAllSamples = ParentNode->bWriteAllSamples;
-	int32 NumSpatialSamples = FMath::Max(1, ParentNode->SpatialSampleCount);
-	const bool bDisableToneCurve = ParentNode->bDisableToneCurve;
-	const EAntiAliasingMethod AntiAliasingMethod = ParentNode->AntiAliasingMethod;
+	UMovieGraphImagePassBaseNode* ParentNodeThisFrame = GetParentNode(InTimeData.EvaluatedConfig);
+	const bool bWriteAllSamples = ParentNodeThisFrame->GetWriteAllSamples();
+	int32 NumSpatialSamples = FMath::Max(1, ParentNodeThisFrame->GetNumSpatialSamples());
+	const bool bDisableToneCurve = ParentNodeThisFrame->GetDisableToneCurve();
+	const EAntiAliasingMethod AntiAliasingMethod = ParentNodeThisFrame->GetAntiAliasingMethod();
 	float OverscanFraction = 0.f;
 	const float TileOverlapPadRatio = 0.0f; // No tiling support right now
 
 	// Camera nodes are optional
+	const bool bIncludeCDOs = false;
 	const UMovieGraphCameraSettingNode* CameraNode = InTimeData.EvaluatedConfig->GetSettingForBranch<UMovieGraphCameraSettingNode>(LayerData.BranchName, bIncludeCDOs);
 	if (CameraNode)
 	{
@@ -170,6 +177,10 @@ void FMovieGraphDeferredPass::Render(const FMovieGraphTraversalContext& InFrameT
 		CameraInfo.TilingParams.OverlapPad = FVector2f(0.f, 0.f); // No tile support
 		CameraInfo.TilingParams.TileCount = FIntPoint(1, 1); // No tile support
 		CameraInfo.TilingParams.TileIndexes = FIntPoint(0, 0); // No tile support
+		CameraInfo.SamplingParams.TemporalSampleIndex = InTimeData.TemporalSampleIndex;
+		CameraInfo.SamplingParams.TemporalSampleCount = InTimeData.TemporalSampleCount;
+		CameraInfo.SamplingParams.SpatialSampleIndex = SpatialIndex;
+		CameraInfo.SamplingParams.SpatialSampleCount = NumSpatialSamples;
 		CameraInfo.OverscanFraction = OverscanFraction;
 		CameraInfo.ProjectionMatrixJitterAmount = FVector2D((SpatialShiftAmount.X) * 2.0f / (float)BackbufferResolution.X, SpatialShiftAmount.Y * -2.0f / (float)BackbufferResolution.Y);
 
@@ -199,6 +210,8 @@ void FMovieGraphDeferredPass::Render(const FMovieGraphTraversalContext& InFrameT
 		ViewFamilyInitData.bWorldIsPaused = bWorldIsPaused;
 		ViewFamilyInitData.FrameIndex = FrameIndex;
 		ViewFamilyInitData.AntiAliasingMethod = AntiAliasingMethod;
+		ViewFamilyInitData.ShowFlags = ParentNodeThisFrame->GetShowFlags();
+		ViewFamilyInitData.ViewModeIndex = ParentNodeThisFrame->GetViewModeIndex();
 		
 		TSharedRef<FSceneViewFamilyContext> ViewFamily = CreateSceneViewFamily(ViewFamilyInitData, CameraInfo);
 		 
@@ -241,9 +254,10 @@ void FMovieGraphDeferredPass::Render(const FMovieGraphTraversalContext& InFrameT
 		}
 
 		// If this was just to contribute to the history buffer, no need to go any further.
-		if (InTimeData.bDiscardOutput)
+		bool bDiscardOutput = InTimeData.bDiscardOutput || ShouldDiscardOutput(ViewFamily, CameraInfo);
+		if (bDiscardOutput)
 		{
-			return;
+			continue;
 		}
 
 		// Readback + Accumulate.
@@ -324,4 +338,3 @@ void FMovieGraphDeferredPass::PostRendererSubmission(
 	FMovieGraphImagePassBase::PostRendererSubmission(InSampleState, InRenderTargetInitParams, InCanvas, InCameraInfo);
 }
 } // namespace UE::MovieGraph::Rendering
-
diff --git a/Engine/Plugins/MovieScene/MovieRenderPipeline/Source/MovieRenderPipelineCore/Private/Graph/Renderers/MovieGraphImagePassBase.cpp b/Engine/Plugins/MovieScene/MovieRenderPipeline/Source/MovieRenderPipelineCore/Private/Graph/Renderers/MovieGraphImagePassBase.cpp
--- a/Engine/Plugins/MovieScene/MovieRenderPipeline/Source/MovieRenderPipelineCore/Private/Graph/Renderers/MovieGraphImagePassBase.cpp
+++ b/Engine/Plugins/MovieScene/MovieRenderPipeline/Source/MovieRenderPipelineCore/Private/Graph/Renderers/MovieGraphImagePassBase.cpp
@@ -124,24 +124,23 @@ void FMovieGraphImagePassBase::ApplyCameraManagerPostProcessBlends(FSceneView* I
 
 TSharedRef<FSceneViewFamilyContext> FMovieGraphImagePassBase::CreateSceneViewFamily(const FViewFamilyInitData& InInitData, const UE::MovieGraph::DefaultRenderer::FCameraInfo& InCameraInfo) const
 {
-	// FEngineShowFlags NewShowFlags = Cast<FMovieGraphImagePassBase>(LayerData.RenderPassNode)->GetShowFlags();
-	// const EViewModeIndex ViewModeIndex = Cast<FMovieGraphImagePassBase>(LayerData.RenderPassNode)->GetViewModeIndex();
-	EViewModeIndex ViewModeIndex = VMI_Lit;
-	FEngineShowFlags NewShowFlags = FEngineShowFlags(ESFIM_Game);
+
+	EViewModeIndex ViewModeIndex = InInitData.ViewModeIndex;
+	FEngineShowFlags ShowFlags = InInitData.ShowFlags;
 
 	const bool bIsPerspective = InCameraInfo.ViewInfo.ProjectionMode == ECameraProjectionMode::Type::Perspective;
 
 	// Allow the Engine Showflag system to override our engine showflags, based on our view mode index.
 	// This is required for certain debug view modes (to have matching show flags set for rendering).
-	ApplyViewMode(/*In*/ ViewModeIndex, bIsPerspective, /*InOut*/NewShowFlags);
+	ApplyViewMode(/*In*/ ViewModeIndex, bIsPerspective, /*InOut*/ShowFlags);
 
 	// And then we have to let another system override them again (based on cvars, etc.)
-	EngineShowFlagOverride(ESFIM_Game, ViewModeIndex, NewShowFlags, false);
+	EngineShowFlagOverride(ESFIM_Game, ViewModeIndex, ShowFlags, false);
 
 	TSharedRef<FSceneViewFamilyContext> OutViewFamily = MakeShared<FSceneViewFamilyContext>(FSceneViewFamily::ConstructionValues(
 		InInitData.RenderTarget,
 		InInitData.World->Scene,
-		NewShowFlags)
+		ShowFlags)
 		.SetTime(FGameTime::CreateUndilated(InInitData.TimeData.WorldSeconds, InInitData.TimeData.FrameDeltaTime))
 		.SetRealtimeUpdate(true));
 
@@ -227,39 +226,6 @@ void FMovieGraphImagePassBase::ApplyMovieGraphOverridesToSceneView(TSharedRef<FS
 		View->FinalPostProcessSettings.bOverride_MotionBlurMax = true;
 
 	}
-	
-	// Path Tracer Sampling
-	if (InOutFamily->EngineShowFlags.PathTracing)
-	{
-		// Override whatever settings came from PostProcessVolume or Camera
-
-		// If motion blur is enabled:
-		//    blend all spatial samples together while leaving the handling of temporal samples up to MRQ
-		//    each temporal sample will include denoising and post-process effects
-		// If motion blur is NOT enabled:
-		//    blend all temporal+spatial samples within the path tracer and only apply denoising on the last temporal sample
-		//    this way we minimize denoising cost and also allow a much higher number of temporal samples to be used which
-		//    can help reduce strobing
-
-		// NOTE: Tiling is not compatible with the reference motion blur mode because it changes the order of the loops over the image.
-		//const bool bHasTiles = InInitData.CameraInfo.TilingParams.TileCount.X > 1 || InInitData.CameraInfo.TilingParams.TileCount.Y;
-		//const bool bAccumulateSpatialSamplesOnly = InOutFamily->EngineShowFlags.MotionBlur || bHasTiles;
-		//
-		//const int32 SampleCount = bAccumulateSpatialSamplesOnly ? InOutSampleState.SpatialSampleCount : InOutSampleState.TemporalSampleCount * InOutSampleState.SpatialSampleCount;
-		//const int32 SampleIndex = bAccumulateSpatialSamplesOnly ? InOutSampleState.SpatialSampleIndex : InOutSampleState.TemporalSampleIndex * InOutSampleState.SpatialSampleCount + InOutSampleState.SpatialSampleIndex;
-		//
-		//// TODO: pass along FrameIndex (which includes SampleIndex) to make sure sampling is fully deterministic
-		//
-		//// Overwrite whatever sampling count came from the PostProcessVolume
-		//View->FinalPostProcessSettings.bOverride_PathTracingSamplesPerPixel = true;
-		//View->FinalPostProcessSettings.PathTracingSamplesPerPixel = SampleCount;
-		//
-		//// reset path tracer's accumulation at the start of each sample
-		//View->bForcePathTracerReset = SampleIndex == 0;
-		//
-		//// discard the result, unless its the last sample
-		//InOutSampleState.bDiscardResult |= !(SampleIndex == SampleCount - 1);
-	}
 
 	// Warn the user for invalid setting combinations / enforce hardware limitations
 	// Locked Exposure
diff --git a/Engine/Plugins/MovieScene/MovieRenderPipeline/Source/MovieRenderPipelineCore/Private/Graph/Renderers/MovieGraphPathTracerPass.cpp b/Engine/Plugins/MovieScene/MovieRenderPipeline/Source/MovieRenderPipelineCore/Private/Graph/Renderers/MovieGraphPathTracerPass.cpp
new file mode 100644
--- /dev/null
+++ b/Engine/Plugins/MovieScene/MovieRenderPipeline/Source/MovieRenderPipelineCore/Private/Graph/Renderers/MovieGraphPathTracerPass.cpp
@@ -0,0 +1,72 @@
+// Copyright Epic Games, Inc. All Rights Reserved.
+
+#include "Graph/Renderers/MovieGraphPathTracerPass.h"
+#include "Graph/Nodes/MovieGraphPathTracedRenderPassNode.h"
+
+namespace UE::MovieGraph::Rendering
+{
+
+UMovieGraphImagePassBaseNode* FMovieGraphPathTracerPass::GetParentNode(UMovieGraphEvaluatedConfig* InConfig) const
+{
+	// This is a bit of a workaround for the fact that the pass doesn't have a strong pointer to the node it's supposed to be associated with,
+	// since that instance changes every frame. So instead we have a virtual function here so the node can look it up by type, and then we can
+	// call a bunch of virtual functions on the right instance to fetch values.
+	const bool bIncludeCDOs = true;
+	UMovieGraphPathTracedRenderPassNode* ParentNode = InConfig->GetSettingForBranch<UMovieGraphPathTracedRenderPassNode>(GetBranchName(), bIncludeCDOs);
+	if (!ensureMsgf(ParentNode, TEXT("PathTracerPass should not exist without parent node in graph.")))
+	{
+		return nullptr;
+	}
+
+	return ParentNode;
+}
+namespace Private
+{
+	void GetSampleData(TSharedRef<FSceneViewFamilyContext> InOutFamily, const UE::MovieGraph::DefaultRenderer::FCameraInfo& InCameraInfo, int32& OutSampleCount, int32& OutSampleIndex)
+	{
+		// If motion blur is enabled:
+		//    blend all spatial samples together while leaving the handling of temporal samples up to MRQ
+		//    each temporal sample will include denoising and post-process effects
+		// If motion blur is NOT enabled:
+		//    blend all temporal+spatial samples within the path tracer and only apply denoising on the last temporal sample
+		//    this way we minimize denoising cost and also allow a much higher number of temporal samples to be used which
+		//    can help reduce strobing
+
+		// NOTE: Tiling is not compatible with the reference motion blur mode because it changes the order of the loops over the image.
+		const bool bHasTiles = InCameraInfo.TilingParams.TileCount.X > 1 || InCameraInfo.TilingParams.TileCount.Y;
+		const bool bAccumulateSpatialSamplesOnly = InOutFamily->EngineShowFlags.MotionBlur || bHasTiles;
+
+		OutSampleCount = bAccumulateSpatialSamplesOnly ? InCameraInfo.SamplingParams.SpatialSampleCount : InCameraInfo.SamplingParams.TemporalSampleCount * InCameraInfo.SamplingParams.SpatialSampleCount;
+		OutSampleIndex = bAccumulateSpatialSamplesOnly ? InCameraInfo.SamplingParams.SpatialSampleIndex : InCameraInfo.SamplingParams.TemporalSampleIndex * InCameraInfo.SamplingParams.SpatialSampleCount + InCameraInfo.SamplingParams.SpatialSampleIndex;
+	}
+}
+
+bool FMovieGraphPathTracerPass::ShouldDiscardOutput(const TSharedRef<FSceneViewFamilyContext>& InFamily, const UE::MovieGraph::DefaultRenderer::FCameraInfo& InCameraInfo) const
+{ 
+	int32 SampleCount, SampleIndex;
+	Private::GetSampleData(InFamily, InCameraInfo, SampleCount, SampleIndex);
+
+	// Discard the result, unless its the last sample
+	return !(SampleIndex == SampleCount - 1);
+}
+
+void FMovieGraphPathTracerPass::ApplyMovieGraphOverridesToSceneView(TSharedRef<FSceneViewFamilyContext> InOutFamily, const FViewFamilyInitData& InInitData, const UE::MovieGraph::DefaultRenderer::FCameraInfo& InCameraInfo) const
+{
+	FMovieGraphImagePassBase::ApplyMovieGraphOverridesToSceneView(InOutFamily, InInitData, InCameraInfo);
+
+	// Override whatever settings came from PostProcessVolume or Camera
+	int32 SampleCount, SampleIndex;
+	Private::GetSampleData(InOutFamily, InCameraInfo, SampleCount, SampleIndex);
+
+	// TODO: pass along FrameIndex (which includes SampleIndex) to make sure sampling is fully deterministic
+		
+	// Overwrite whatever sampling count came from the PostProcessVolume
+	FSceneView* View = const_cast<FSceneView*>(InOutFamily->Views[0]);
+	View->FinalPostProcessSettings.bOverride_PathTracingSamplesPerPixel = true;
+	View->FinalPostProcessSettings.PathTracingSamplesPerPixel = SampleCount;
+		
+	// reset path tracer's accumulation at the start of each sample
+	View->bForcePathTracerReset = SampleIndex == 0;
+}
+
+} // UE::MovieGraph::Rendering
\ No newline at end of file
diff --git a/Engine/Plugins/MovieScene/MovieRenderPipeline/Source/MovieRenderPipelineCore/Public/Graph/MovieGraphDefaultRenderer.h b/Engine/Plugins/MovieScene/MovieRenderPipeline/Source/MovieRenderPipelineCore/Public/Graph/MovieGraphDefaultRenderer.h
--- a/Engine/Plugins/MovieScene/MovieRenderPipeline/Source/MovieRenderPipelineCore/Public/Graph/MovieGraphDefaultRenderer.h
+++ b/Engine/Plugins/MovieScene/MovieRenderPipeline/Source/MovieRenderPipelineCore/Public/Graph/MovieGraphDefaultRenderer.h
@@ -35,6 +35,14 @@ namespace UE::MovieGraph::DefaultRenderer
 		FIntPoint TileCount;
 	};
 
+	struct FMovieGraphSamplingParams
+	{
+		int32 TemporalSampleIndex;
+		int32 TemporalSampleCount;
+		int32 SpatialSampleIndex;
+		int32 SpatialSampleCount;
+	};
+
 	struct FCameraInfo
 	{
 		FCameraInfo()
@@ -62,7 +70,9 @@ namespace UE::MovieGraph::DefaultRenderer
 		// questionable if these are fcamerainfo
 		/** When using tiling, we scale the sensor to counteract the view changes. This value comes from modifying the ProjectionMatrix. */
 		float DoFSensorScale;
+
 		FMovieGraphTilingParams TilingParams;
+		FMovieGraphSamplingParams SamplingParams;
 
 		// Sub-pixel jitter this camera should use. Only applied when using no AA.
 		FVector2D ProjectionMatrixJitterAmount;
diff --git a/Engine/Plugins/MovieScene/MovieRenderPipeline/Source/MovieRenderPipelineCore/Public/Graph/Nodes/MovieGraphDeferredPassNode.h b/Engine/Plugins/MovieScene/MovieRenderPipeline/Source/MovieRenderPipelineCore/Public/Graph/Nodes/MovieGraphDeferredPassNode.h
--- a/Engine/Plugins/MovieScene/MovieRenderPipeline/Source/MovieRenderPipelineCore/Public/Graph/Nodes/MovieGraphDeferredPassNode.h
+++ b/Engine/Plugins/MovieScene/MovieRenderPipeline/Source/MovieRenderPipelineCore/Public/Graph/Nodes/MovieGraphDeferredPassNode.h
@@ -25,6 +25,10 @@ protected:
 	// UMovieGraphRenderPassNode Interface
 	virtual FString GetRendererNameImpl() const override;
 	virtual TUniquePtr<UE::MovieGraph::Rendering::FMovieGraphImagePassBase> CreateInstance() const override;
+	virtual bool GetWriteAllSamples() const override;
+	virtual int32 GetNumSpatialSamples() const override;
+	virtual bool GetDisableToneCurve() const override;
+	virtual EAntiAliasingMethod GetAntiAliasingMethod() const override;
 	// ~UMovieGraphRenderPassNode Interface
 
 	// UMovieGraphImagePassBaseNode Interface
diff --git a/Engine/Plugins/MovieScene/MovieRenderPipeline/Source/MovieRenderPipelineCore/Public/Graph/Nodes/MovieGraphImagePassBaseNode.h b/Engine/Plugins/MovieScene/MovieRenderPipeline/Source/MovieRenderPipelineCore/Public/Graph/Nodes/MovieGraphImagePassBaseNode.h
--- a/Engine/Plugins/MovieScene/MovieRenderPipeline/Source/MovieRenderPipelineCore/Public/Graph/Nodes/MovieGraphImagePassBaseNode.h
+++ b/Engine/Plugins/MovieScene/MovieRenderPipeline/Source/MovieRenderPipelineCore/Public/Graph/Nodes/MovieGraphImagePassBaseNode.h
@@ -37,6 +37,18 @@ public:
 	/** Gets any show flags that should be applied as defaults, before user changes are applied. */
 	virtual TArray<uint32> GetDefaultShowFlags() const;
 
+	/** Should each individual sample rendered be written out for debugging? */
+	virtual bool GetWriteAllSamples() const { return false; }
+
+	/** How many spatial samples should be rendered each frame? */
+	virtual int32 GetNumSpatialSamples() const { return 1; }
+
+	/** Should the tone curve be disabled while rendering? Allows for linear values in exrs but changes the look of the final image. */
+	virtual bool GetDisableToneCurve() const { return false; }
+
+	/** Which AA Method should be used? */
+	virtual EAntiAliasingMethod GetAntiAliasingMethod() const { return EAntiAliasingMethod::AAM_None; }
+
 protected:
 	// Note: Since *individual* show flags are overridden instead of the entire ShowFlags property, manually set to
 	// overridden so the traversal picks the changes up (otherwise they will be ignored).
diff --git a/Engine/Plugins/MovieScene/MovieRenderPipeline/Source/MovieRenderPipelineCore/Public/Graph/Nodes/MovieGraphPathTracedRenderPassNode.h b/Engine/Plugins/MovieScene/MovieRenderPipeline/Source/MovieRenderPipelineCore/Public/Graph/Nodes/MovieGraphPathTracedRenderPassNode.h
--- a/Engine/Plugins/MovieScene/MovieRenderPipeline/Source/MovieRenderPipelineCore/Public/Graph/Nodes/MovieGraphPathTracedRenderPassNode.h
+++ b/Engine/Plugins/MovieScene/MovieRenderPipeline/Source/MovieRenderPipelineCore/Public/Graph/Nodes/MovieGraphPathTracedRenderPassNode.h
@@ -24,15 +24,59 @@ public:
 	virtual void TeardownImpl() override;
 	// ~UMovieGraphRenderPassNode Interface
 
+	// UMovieGraphImagePassBaseNode Interface
+	virtual bool GetWriteAllSamples() const override;
+	virtual int32 GetNumSpatialSamples() const override;
+	virtual bool GetDisableToneCurve() const override;
+	virtual TUniquePtr<UE::MovieGraph::Rendering::FMovieGraphImagePassBase> CreateInstance() const;
+	// ~UMovieGraphImagePassBaseNode Interface
+
 protected:
 	// UMovieGraphRenderPassNode Interface
 	virtual FString GetRendererNameImpl() const override;
+
 	// ~UMovieGraphRenderPassNode Interface
 
 	// UMovieGraphCoreRenderPassNode Interface
 	virtual EViewModeIndex GetViewModeIndex() const override;
 	// ~UMovieGraphCoreRenderPassNode Interface
 
+public:
+
+	UPROPERTY(EditAnywhere, BlueprintReadWrite, Category = Overrides, meta = (InlineEditConditionToggle))
+	uint8 bOverride_SpatialSampleCount : 1;
+
+	UPROPERTY(EditAnywhere, BlueprintReadWrite, Category = Overrides, meta = (InlineEditConditionToggle))
+	uint8 bOverride_bDisableToneCurve : 1;
+
+	UPROPERTY(EditAnywhere, BlueprintReadWrite, Category = Overrides, meta = (InlineEditConditionToggle))
+	uint8 bOverride_bWriteAllSamples : 1;
+
+	/**
+	* How many sub-pixel jitter renders should we do per temporal sample? This can be used to achieve high
+	* sample counts without Temporal Sub-Sampling (allowing high sample counts without motion blur being enabled),
+	* but we generally recommend using Temporal Sub-Samples when possible. It can also be combined with
+	* temporal samples and you will get SpatialSampleCount many renders per temporal sample.
+	*/
+	UPROPERTY(EditAnywhere, BlueprintReadWrite, meta = (UIMin = 1, ClampMin = 1), Category = "Settings", meta = (EditCondition = "bOverride_SpatialSampleCount"))
+	int32 SpatialSampleCount;
+
+	/**
+	* If true, the tone curve will be disabled for this render pass. This will result in values greater than 1.0 in final renders
+	* and can optionally be combined with OCIO profiles on the file output nodes to convert from Linear Values in Working Color Space
+	* (which is sRGB  (Rec. 709) by default, unless changed in the project settings).
+	*/
+	UPROPERTY(EditAnywhere, BlueprintReadWrite, Category = "Settings", meta = (EditCondition = "bOverride_bDisableToneCurve"))
+	bool bDisableToneCurve;
+
+	/**
+	* Debug Feature. Can use this to write out each individual Temporal and Spatial sample rendered by this render pass,
+	* which allows you to see which images are being accumulated together. Can be useful for debugging incorrect looking
+	* frames to see which sub-frame evaluations were incorrect.
+	*/
+	UPROPERTY(EditAnywhere, BlueprintReadWrite, Category = "Debug", meta = (EditCondition = "bOverride_bWriteAllsamples"))
+	bool bWriteAllSamples;
+
 private:
 	/**
 	 * The original value of the "r.PathTracing.ProgressDisplay" cvar before the render starts. The progress display
diff --git a/Engine/Plugins/MovieScene/MovieRenderPipeline/Source/MovieRenderPipelineCore/Public/Graph/Renderers/MovieGraphDeferredPass.h b/Engine/Plugins/MovieScene/MovieRenderPipeline/Source/MovieRenderPipelineCore/Public/Graph/Renderers/MovieGraphDeferredPass.h
--- a/Engine/Plugins/MovieScene/MovieRenderPipeline/Source/MovieRenderPipelineCore/Public/Graph/Renderers/MovieGraphDeferredPass.h
+++ b/Engine/Plugins/MovieScene/MovieRenderPipeline/Source/MovieRenderPipelineCore/Public/Graph/Renderers/MovieGraphDeferredPass.h
@@ -13,9 +13,10 @@ namespace UE::MovieGraph::Rendering
 		virtual void GatherOutputPasses(TArray<FMovieGraphRenderDataIdentifier>& OutExpectedPasses) const override;
 		virtual void AddReferencedObjects(FReferenceCollector& Collector) override;
 		virtual FName GetBranchName() const override;
+		virtual UMovieGraphImagePassBaseNode* GetParentNode(UMovieGraphEvaluatedConfig* InConfig) const override;
+			
 	protected:
 		virtual void PostRendererSubmission(const UE::MovieGraph::FMovieGraphSampleState& InSampleState, const UE::MovieGraph::DefaultRenderer::FRenderTargetInitParams& InRenderTargetInitParams, FCanvas& InCanvas, const UE::MovieGraph::DefaultRenderer::FCameraInfo& InCameraInfo) override;
-
 	protected:
 		FMovieGraphRenderPassLayerData LayerData;
 
diff --git a/Engine/Plugins/MovieScene/MovieRenderPipeline/Source/MovieRenderPipelineCore/Public/Graph/Renderers/MovieGraphImagePassBase.h b/Engine/Plugins/MovieScene/MovieRenderPipeline/Source/MovieRenderPipelineCore/Public/Graph/Renderers/MovieGraphImagePassBase.h
--- a/Engine/Plugins/MovieScene/MovieRenderPipeline/Source/MovieRenderPipelineCore/Public/Graph/Renderers/MovieGraphImagePassBase.h
+++ b/Engine/Plugins/MovieScene/MovieRenderPipeline/Source/MovieRenderPipelineCore/Public/Graph/Renderers/MovieGraphImagePassBase.h
@@ -34,6 +34,8 @@ namespace UE::MovieGraph::Rendering
 			, bWorldIsPaused(false)
 			, FrameIndex(-1)
 			, AntiAliasingMethod(EAntiAliasingMethod::AAM_None)
+			, ShowFlags(ESFIM_Game)
+			, ViewModeIndex(VMI_Lit)
 		{
 		}
 
@@ -44,6 +46,8 @@ namespace UE::MovieGraph::Rendering
 		bool bWorldIsPaused;
 		int32 FrameIndex;
 		EAntiAliasingMethod AntiAliasingMethod;
+		FEngineShowFlags ShowFlags;
+		EViewModeIndex ViewModeIndex;
 	};
 	
 
@@ -78,6 +82,9 @@ namespace UE::MovieGraph::Rendering
 		virtual void ApplyMovieGraphOverridesToSceneView(TSharedRef<FSceneViewFamilyContext> InOutFamily, const FViewFamilyInitData& InInitData, const UE::MovieGraph::DefaultRenderer::FCameraInfo& InCameraInfo) const;
 		virtual void ApplyMovieGraphOverridesToViewFamily(TSharedRef<FSceneViewFamilyContext> InOutFamily, const FViewFamilyInitData& InInitData) const;
 		virtual void PostRendererSubmission(const UE::MovieGraph::FMovieGraphSampleState& InSampleState, const UE::MovieGraph::DefaultRenderer::FRenderTargetInitParams& InRenderTargetInitParams, FCanvas& InCanvas, const UE::MovieGraph::DefaultRenderer::FCameraInfo& InCameraInfo);
+		virtual bool ShouldDiscardOutput(const TSharedRef<FSceneViewFamilyContext>& InFamily, const UE::MovieGraph::DefaultRenderer::FCameraInfo& InCameraInfo) const { return false; }
+		/** For this image pass, look up the associated node by type for the given config. */
+		virtual UMovieGraphImagePassBaseNode* GetParentNode(UMovieGraphEvaluatedConfig* InConfig) const { return nullptr; }
 
 	protected:
 		TWeakObjectPtr<UMovieGraphDefaultRenderer> WeakGraphRenderer;
diff --git a/Engine/Plugins/MovieScene/MovieRenderPipeline/Source/MovieRenderPipelineCore/Public/Graph/Renderers/MovieGraphPathTracerPass.h b/Engine/Plugins/MovieScene/MovieRenderPipeline/Source/MovieRenderPipelineCore/Public/Graph/Renderers/MovieGraphPathTracerPass.h
new file mode 100644
--- /dev/null
+++ b/Engine/Plugins/MovieScene/MovieRenderPipeline/Source/MovieRenderPipelineCore/Public/Graph/Renderers/MovieGraphPathTracerPass.h
@@ -0,0 +1,14 @@
+// Copyright Epic Games, Inc. All Rights Reserved.
+#pragma once
+
+#include "Graph/Renderers/MovieGraphDeferredPass.h"
+
+namespace UE::MovieGraph::Rendering
+{
+	struct FMovieGraphPathTracerPass : public FMovieGraphDeferredPass
+	{
+		virtual UMovieGraphImagePassBaseNode* GetParentNode(UMovieGraphEvaluatedConfig* InConfig) const override;
+		virtual void ApplyMovieGraphOverridesToSceneView(TSharedRef<FSceneViewFamilyContext> InOutFamily, const FViewFamilyInitData& InInitData, const UE::MovieGraph::DefaultRenderer::FCameraInfo& InCameraInfo) const override;
+		virtual bool ShouldDiscardOutput(const TSharedRef<FSceneViewFamilyContext>& InFamily, const UE::MovieGraph::DefaultRenderer::FCameraInfo& InCameraInfo) const override;
+	}; 
+}
\ No newline at end of file
