MRG: Add separate multi-layer EXR output node & fix OCIO on output.

#rnx

diff --git a/Engine/Plugins/MovieScene/MovieRenderPipeline/Source/MovieRenderPipelineCore/Private/Graph/Nodes/MovieGraphDeferredPassNode.cpp b/Engine/Plugins/MovieScene/MovieRenderPipeline/Source/MovieRenderPipelineCore/Private/Graph/Nodes/MovieGraphDeferredPassNode.cpp
--- a/Engine/Plugins/MovieScene/MovieRenderPipeline/Source/MovieRenderPipelineCore/Private/Graph/Nodes/MovieGraphDeferredPassNode.cpp
+++ b/Engine/Plugins/MovieScene/MovieRenderPipeline/Source/MovieRenderPipelineCore/Private/Graph/Nodes/MovieGraphDeferredPassNode.cpp
@@ -11,6 +11,9 @@ TUniquePtr<UE::MovieGraph::Rendering::FMovieGraphImagePassBase> UMovieGraphDefer
 UMovieGraphDeferredRenderPassNode::UMovieGraphDeferredRenderPassNode()
 	: SpatialSampleCount(1)
 	, AntiAliasingMethod(EAntiAliasingMethod::AAM_TSR)
+	, bDisableToneCurve(false)
+	, bAllowOCIO(true)
+	, bWriteAllSamples(false)
 	, ViewModeIndex(VMI_Lit)
 {
 }
@@ -54,6 +57,11 @@ bool UMovieGraphDeferredRenderPassNode::GetDisableToneCurve() const
 	return bDisableToneCurve;
 }
 
+bool UMovieGraphDeferredRenderPassNode::GetAllowOCIO() const
+{
+	return bAllowOCIO;
+}
+
 EAntiAliasingMethod UMovieGraphDeferredRenderPassNode::GetAntiAliasingMethod() const
 {
 	return AntiAliasingMethod;
diff --git a/Engine/Plugins/MovieScene/MovieRenderPipeline/Source/MovieRenderPipelineCore/Private/Graph/Nodes/MovieGraphPathTracedRenderPassNode.cpp b/Engine/Plugins/MovieScene/MovieRenderPipeline/Source/MovieRenderPipelineCore/Private/Graph/Nodes/MovieGraphPathTracedRenderPassNode.cpp
--- a/Engine/Plugins/MovieScene/MovieRenderPipeline/Source/MovieRenderPipelineCore/Private/Graph/Nodes/MovieGraphPathTracedRenderPassNode.cpp
+++ b/Engine/Plugins/MovieScene/MovieRenderPipeline/Source/MovieRenderPipelineCore/Private/Graph/Nodes/MovieGraphPathTracedRenderPassNode.cpp
@@ -11,6 +11,10 @@ TUniquePtr<UE::MovieGraph::Rendering::FMovieGraphImagePassBase> UMovieGraphPathT
 }
 
 UMovieGraphPathTracedRenderPassNode::UMovieGraphPathTracedRenderPassNode()
+	: SpatialSampleCount(1)
+	, bDisableToneCurve(false)
+	, bAllowOCIO(true)
+	, bWriteAllSamples(false)
 {
 	ShowFlags->ApplyDefaultShowFlagValue(VMI_PathTracing, true);
 	// TODO: Showflag for SetMotionBlur()?
@@ -71,3 +75,8 @@ bool UMovieGraphPathTracedRenderPassNode::GetDisableToneCurve() const
 {
 	return bDisableToneCurve;
 }
+
+bool UMovieGraphPathTracedRenderPassNode::GetAllowOCIO() const
+{
+	return bAllowOCIO;
+}
diff --git a/Engine/Plugins/MovieScene/MovieRenderPipeline/Source/MovieRenderPipelineCore/Private/Graph/Renderers/MovieGraphDeferredPass.cpp b/Engine/Plugins/MovieScene/MovieRenderPipeline/Source/MovieRenderPipelineCore/Private/Graph/Renderers/MovieGraphDeferredPass.cpp
--- a/Engine/Plugins/MovieScene/MovieRenderPipeline/Source/MovieRenderPipelineCore/Private/Graph/Renderers/MovieGraphDeferredPass.cpp
+++ b/Engine/Plugins/MovieScene/MovieRenderPipeline/Source/MovieRenderPipelineCore/Private/Graph/Renderers/MovieGraphDeferredPass.cpp
@@ -251,6 +251,7 @@ void FMovieGraphDeferredPass::Render(const FMovieGraphTraversalContext& InFrameT
 			SampleState.OverlappedOffset = OverlappedOffset;
 			SampleState.OverlappedSubpixelShift = OverlappedSubpixelShift;
 			SampleState.OverscanFraction = OverscanFraction;
+			SampleState.bAllowOCIO = ParentNodeThisFrame->GetAllowOCIO();
 		}
 
 		// If this was just to contribute to the history buffer, no need to go any further.
diff --git a/Engine/Plugins/MovieScene/MovieRenderPipeline/Source/MovieRenderPipelineCore/Public/Graph/MovieGraphDataTypes.h b/Engine/Plugins/MovieScene/MovieRenderPipeline/Source/MovieRenderPipelineCore/Public/Graph/MovieGraphDataTypes.h
--- a/Engine/Plugins/MovieScene/MovieRenderPipeline/Source/MovieRenderPipelineCore/Public/Graph/MovieGraphDataTypes.h
+++ b/Engine/Plugins/MovieScene/MovieRenderPipeline/Source/MovieRenderPipelineCore/Public/Graph/MovieGraphDataTypes.h
@@ -345,6 +345,7 @@ namespace UE::MovieGraph
 			, bCompositeOnOtherRenders(false)
 			, OverscanFraction(0.f)
 			, CompositingSortOrder(0)
+			, bAllowOCIO(true)
 		{}
 
 		/** The traversal context used to read graph values at the time of submission. */
@@ -388,6 +389,9 @@ namespace UE::MovieGraph
 		* Passes with a low sort order will composite on top of passes with a higher sort order.
 		*/
 		int32 CompositingSortOrder;
+
+		/** Allow OpenColorIO transform to be used on this render. */
+		bool bAllowOCIO;
 	};
 
 	/**
diff --git a/Engine/Plugins/MovieScene/MovieRenderPipeline/Source/MovieRenderPipelineCore/Public/Graph/Nodes/MovieGraphDeferredPassNode.h b/Engine/Plugins/MovieScene/MovieRenderPipeline/Source/MovieRenderPipelineCore/Public/Graph/Nodes/MovieGraphDeferredPassNode.h
--- a/Engine/Plugins/MovieScene/MovieRenderPipeline/Source/MovieRenderPipelineCore/Public/Graph/Nodes/MovieGraphDeferredPassNode.h
+++ b/Engine/Plugins/MovieScene/MovieRenderPipeline/Source/MovieRenderPipelineCore/Public/Graph/Nodes/MovieGraphDeferredPassNode.h
@@ -28,6 +28,7 @@ protected:
 	virtual bool GetWriteAllSamples() const override;
 	virtual int32 GetNumSpatialSamples() const override;
 	virtual bool GetDisableToneCurve() const override;
+	virtual bool GetAllowOCIO() const override;
 	virtual EAntiAliasingMethod GetAntiAliasingMethod() const override;
 	// ~UMovieGraphRenderPassNode Interface
 
@@ -44,6 +45,9 @@ public:
 	
 	UPROPERTY(EditAnywhere, BlueprintReadWrite, Category = Overrides, meta = (InlineEditConditionToggle))
 	uint8 bOverride_bDisableToneCurve : 1;
+
+	UPROPERTY(EditAnywhere, BlueprintReadWrite, Category = Overrides, meta = (InlineEditConditionToggle))
+	uint8 bOverride_bAllowOCIO : 1;
 	
 	UPROPERTY(EditAnywhere, BlueprintReadWrite, Category = Overrides, meta = (InlineEditConditionToggle))
 	uint8 bOverride_ViewModeIndex : 1;
@@ -78,6 +82,12 @@ public:
 	UPROPERTY(EditAnywhere, BlueprintReadWrite, Category = "Settings", meta = (EditCondition = "bOverride_bDisableToneCurve"))
 	bool bDisableToneCurve;
 
+	/**
+	* Allow the output file OpenColorIO transform to be used on this render.
+	*/
+	UPROPERTY(EditAnywhere, BlueprintReadWrite, Category = "Settings", meta = (EditCondition = "bOverride_bAllowOCIO"))
+	bool bAllowOCIO;
+
 	/**
 	* Debug Feature. Can use this to write out each individual Temporal and Spatial sample rendered by this render pass,
 	* which allows you to see which images are being accumulated together. Can be useful for debugging incorrect looking
diff --git a/Engine/Plugins/MovieScene/MovieRenderPipeline/Source/MovieRenderPipelineCore/Public/Graph/Nodes/MovieGraphImagePassBaseNode.h b/Engine/Plugins/MovieScene/MovieRenderPipeline/Source/MovieRenderPipelineCore/Public/Graph/Nodes/MovieGraphImagePassBaseNode.h
--- a/Engine/Plugins/MovieScene/MovieRenderPipeline/Source/MovieRenderPipelineCore/Public/Graph/Nodes/MovieGraphImagePassBaseNode.h
+++ b/Engine/Plugins/MovieScene/MovieRenderPipeline/Source/MovieRenderPipelineCore/Public/Graph/Nodes/MovieGraphImagePassBaseNode.h
@@ -46,6 +46,9 @@ public:
 	/** Should the tone curve be disabled while rendering? Allows for linear values in exrs but changes the look of the final image. */
 	virtual bool GetDisableToneCurve() const { return false; }
 
+	/** Should the output file be allowed to apply an OCIO transform on this render? */
+	virtual bool GetAllowOCIO() const { return true; }
+
 	/** Which AA Method should be used? */
 	virtual EAntiAliasingMethod GetAntiAliasingMethod() const { return EAntiAliasingMethod::AAM_None; }
 
diff --git a/Engine/Plugins/MovieScene/MovieRenderPipeline/Source/MovieRenderPipelineCore/Public/Graph/Nodes/MovieGraphPathTracedRenderPassNode.h b/Engine/Plugins/MovieScene/MovieRenderPipeline/Source/MovieRenderPipelineCore/Public/Graph/Nodes/MovieGraphPathTracedRenderPassNode.h
--- a/Engine/Plugins/MovieScene/MovieRenderPipeline/Source/MovieRenderPipelineCore/Public/Graph/Nodes/MovieGraphPathTracedRenderPassNode.h
+++ b/Engine/Plugins/MovieScene/MovieRenderPipeline/Source/MovieRenderPipelineCore/Public/Graph/Nodes/MovieGraphPathTracedRenderPassNode.h
@@ -28,6 +28,7 @@ public:
 	virtual bool GetWriteAllSamples() const override;
 	virtual int32 GetNumSpatialSamples() const override;
 	virtual bool GetDisableToneCurve() const override;
+	virtual bool GetAllowOCIO() const override;
 	virtual TUniquePtr<UE::MovieGraph::Rendering::FMovieGraphImagePassBase> CreateInstance() const;
 	// ~UMovieGraphImagePassBaseNode Interface
 
@@ -49,6 +50,9 @@ public:
 	UPROPERTY(EditAnywhere, BlueprintReadWrite, Category = Overrides, meta = (InlineEditConditionToggle))
 	uint8 bOverride_bDisableToneCurve : 1;
 
+	UPROPERTY(EditAnywhere, BlueprintReadWrite, Category = Overrides, meta = (InlineEditConditionToggle))
+	uint8 bOverride_bAllowOCIO : 1;
+
 	UPROPERTY(EditAnywhere, BlueprintReadWrite, Category = Overrides, meta = (InlineEditConditionToggle))
 	uint8 bOverride_bWriteAllSamples : 1;
 
@@ -69,6 +73,12 @@ public:
 	UPROPERTY(EditAnywhere, BlueprintReadWrite, Category = "Settings", meta = (EditCondition = "bOverride_bDisableToneCurve"))
 	bool bDisableToneCurve;
 
+	/**
+	* Allow the output file OpenColorIO transform to be used on this render.
+	*/
+	UPROPERTY(EditAnywhere, BlueprintReadWrite, Category = "Settings", meta = (EditCondition = "bOverride_bAllowOCIO"))
+	bool bAllowOCIO;
+
 	/**
 	* Debug Feature. Can use this to write out each individual Temporal and Spatial sample rendered by this render pass,
 	* which allows you to see which images are being accumulated together. Can be useful for debugging incorrect looking
diff --git a/Engine/Plugins/MovieScene/MovieRenderPipeline/Source/MovieRenderPipelineRenderPasses/Private/MovieGraphImageSequenceOutputNode.cpp b/Engine/Plugins/MovieScene/MovieRenderPipeline/Source/MovieRenderPipelineRenderPasses/Private/MovieGraphImageSequenceOutputNode.cpp
--- a/Engine/Plugins/MovieScene/MovieRenderPipeline/Source/MovieRenderPipelineRenderPasses/Private/MovieGraphImageSequenceOutputNode.cpp
+++ b/Engine/Plugins/MovieScene/MovieRenderPipeline/Source/MovieRenderPipelineRenderPasses/Private/MovieGraphImageSequenceOutputNode.cpp
@@ -163,6 +163,10 @@ void UMovieGraphImageSequenceOutputNode::OnReceiveImageDataImpl(UMovieGraphPipel
 			continue;
 		}
 
+		const UMovieGraphImageSequenceOutputNode* ParentNode = Cast<UMovieGraphImageSequenceOutputNode>(
+			InRawFrameData->EvaluatedConfig->GetSettingForBranch(GetClass(), RenderData.Key.RootBranchName, false, true));
+		checkf(ParentNode, TEXT("Image sequence output should not exist without a parent node in the graph."));
+
 		UE::MovieGraph::FMovieGraphSampleState* Payload = RenderData.Value->GetPayload<UE::MovieGraph::FMovieGraphSampleState>();
 		const TObjectPtr<UMoviePipelineExecutorShot>& Shot = InPipeline->GetActiveShotList()[Payload->TraversalContext.ShotIndex];
 
@@ -235,6 +239,7 @@ void UMovieGraphImageSequenceOutputNode::OnReceiveImageDataImpl(UMovieGraphPipel
 		FMovieGraphResolveArgs FinalResolvedKVPs;
 		const FString FileName = UMovieGraphBlueprintLibrary::ResolveFilenameFormatArguments(FileNameFormatString, Params, FinalResolvedKVPs);
 
+		//TODO: Support for single-layer EXR image write task (compression)...
 		TUniquePtr<FImageWriteTask> TileImageTask = MakeUnique<FImageWriteTask>();
 		TileImageTask->Format = OutputFormat;
 		TileImageTask->CompressionQuality = 100;
@@ -243,9 +248,9 @@ void UMovieGraphImageSequenceOutputNode::OnReceiveImageDataImpl(UMovieGraphPipel
 
 		bool bQuantizationEncodeSRGB = true;
 #if WITH_EDITOR
-		if (OCIOConfiguration.bIsEnabled)
+		if (ParentNode->OCIOConfiguration.bIsEnabled && Payload->bAllowOCIO)
 		{
-			FPixelPreProcessor OCIOPixelPreProcessor = UE::MovieGraph::Private::CreateOpenColorIOPixelPreProcessor(OCIOConfiguration.ColorConfiguration);
+			FPixelPreProcessor OCIOPixelPreProcessor = UE::MovieGraph::Private::CreateOpenColorIOPixelPreProcessor(ParentNode->OCIOConfiguration.ColorConfiguration);
 			if (OCIOPixelPreProcessor)
 			{
 				TileImageTask->PixelPreProcessors.Emplace(MoveTemp(OCIOPixelPreProcessor));
@@ -301,18 +306,14 @@ void UMovieGraphImageSequenceOutputNode::OnReceiveImageDataImpl(UMovieGraphPipel
 
 }
 
-void UMovieGraphImageSequenceOutputNode_EXR::OnReceiveImageDataImpl(UMovieGraphPipeline* InPipeline, UE::MovieGraph::FMovieGraphOutputMergerFrame* InRawFrameData, const TSet<FMovieGraphRenderDataIdentifier>& InMask)
+void UMovieGraphImageSequenceOutputNode_MultiLayerEXR::OnReceiveImageDataImpl(UMovieGraphPipeline* InPipeline, UE::MovieGraph::FMovieGraphOutputMergerFrame* InRawFrameData, const TSet<FMovieGraphRenderDataIdentifier>& InMask)
 {
-	if (!bMultilayer)
-	{
-		// Some software doesn't support multi-layer, so in that case we fall back to the single-layer-multiple-file
-		// codepath of our parent.
-		Super::OnReceiveImageDataImpl(InPipeline, InRawFrameData, InMask);
-		return;
-	}
-	
 	check(InRawFrameData);
 
+	const UMovieGraphImageSequenceOutputNode_MultiLayerEXR* ParentNode = Cast<UMovieGraphImageSequenceOutputNode_MultiLayerEXR>(
+		InRawFrameData->EvaluatedConfig->GetSettingForBranch(GetClass(), UMovieGraphNode::GlobalsPinName, false, true));
+	checkf(ParentNode, TEXT("Multi-Layer EXR should not exist without a parent node in the graph."));
+
 	// Ensure our OpenExrRTTI module gets loaded. This needs to happen from the main thread, if it's not loaded then metadata silently fails when writing.
 	static const FName RTTIExtensionModuleName("UEOpenExrRTTI");
 	FModuleManager::Get().LoadModule(RTTIExtensionModuleName);
@@ -333,7 +334,7 @@ void UMovieGraphImageSequenceOutputNode_EXR::OnReceiveImageDataImpl(UMovieGraphP
 		TUniquePtr<FEXRImageWriteTask> MultiLayerImageTask = MakeUnique<FEXRImageWriteTask>();
 		
 		MultiLayerImageTask->Filename = Filename;
-		MultiLayerImageTask->Compression = Compression;
+		MultiLayerImageTask->Compression = ParentNode->Compression;
 		// MultiLayerImageTask->CompressionLevel is intentionally skipped because it doesn't seem to make any practical difference
 		// so we don't expose it to the user because that will just cause confusion where the setting doesn't seem to do anything.
 
@@ -358,7 +359,6 @@ void UMovieGraphImageSequenceOutputNode_EXR::OnReceiveImageDataImpl(UMovieGraphP
 
 		// Add color space metadata to the output: xy chromaticity coordinates and/or the color space source/dest names.
 		// TODO: Support is also needed for regular exrs via the image wrapper module.
-		// TODO: No OCIO node yet
 		// {
 		// 	UMoviePipelineColorSetting* ColorSetting = InPipeline->GetPipelinePrimaryConfig()->FindSetting<UMoviePipelineColorSetting>();
 		//
@@ -421,6 +421,17 @@ void UMovieGraphImageSequenceOutputNode_EXR::OnReceiveImageDataImpl(UMovieGraphP
 				MultiLayerImageTask->LayerNames.FindOrAdd(PixelData.Get(), CombinedName);
 			}
 
+#if WITH_EDITOR
+			if (ParentNode->OCIOConfiguration.bIsEnabled && Payload->bAllowOCIO)
+			{
+				FPixelPreProcessor OCIOPixelPreProcessor = UE::MovieGraph::Private::CreateOpenColorIOPixelPreProcessor(ParentNode->OCIOConfiguration.ColorConfiguration);
+				if (OCIOPixelPreProcessor)
+				{
+					MultiLayerImageTask->PixelPreprocessors.FindOrAdd(LayerIndex).Emplace(MoveTemp(OCIOPixelPreProcessor));
+				}
+			}
+#endif
+
 			MultiLayerImageTask->Layers.Add(MoveTemp(PixelData));
 			LayerIndex++;
 		}
@@ -434,7 +445,7 @@ void UMovieGraphImageSequenceOutputNode_EXR::OnReceiveImageDataImpl(UMovieGraphP
 	}
 }
 
-void UMovieGraphImageSequenceOutputNode_EXR::GetFilenameToRenderIDMappings(
+void UMovieGraphImageSequenceOutputNode_MultiLayerEXR::GetFilenameToRenderIDMappings(
 	UMovieGraphPipeline* InPipeline, UE::MovieGraph::FMovieGraphOutputMergerFrame* InRawFrameData,
 	TMap<FString, TArray<FMovieGraphRenderDataIdentifier>>& OutFilenameToRenderIDs,
 	TMap<FString, FMovieGraphResolveArgs>& OutFilenameToResolveArgs) const
@@ -501,7 +512,7 @@ void UMovieGraphImageSequenceOutputNode_EXR::GetFilenameToRenderIDMappings(
 	}
 }
 
-FString UMovieGraphImageSequenceOutputNode_EXR::ResolveOutputFilename(
+FString UMovieGraphImageSequenceOutputNode_MultiLayerEXR::ResolveOutputFilename(
 	const UMovieGraphPipeline* InPipeline,
 	const int32 ResolutionIndex, const UE::MovieGraph::FMovieGraphOutputMergerFrame* InRawFrameData,
 	const FName& InBranchName, FMovieGraphResolveArgs& OutResolveArgs) const
diff --git a/Engine/Plugins/MovieScene/MovieRenderPipeline/Source/MovieRenderPipelineRenderPasses/Private/MoviePipelineEXROutput.cpp b/Engine/Plugins/MovieScene/MovieRenderPipeline/Source/MovieRenderPipelineRenderPasses/Private/MoviePipelineEXROutput.cpp
--- a/Engine/Plugins/MovieScene/MovieRenderPipeline/Source/MovieRenderPipelineRenderPasses/Private/MoviePipelineEXROutput.cpp
+++ b/Engine/Plugins/MovieScene/MovieRenderPipeline/Source/MovieRenderPipelineRenderPasses/Private/MoviePipelineEXROutput.cpp
@@ -122,6 +122,8 @@ bool FEXRImageWriteTask::WriteToDisk()
 
 	if (bSuccess)
 	{
+		PreProcess();
+
 		Imf::Compression FileCompression = Imf::Compression::NO_COMPRESSION;
 		switch(Compression)
 		{
@@ -406,6 +408,28 @@ void FEXRImageWriteTask::AddFileMetadata(Imf::Header& InHeader)
 	}
 }
 
+void FEXRImageWriteTask::PreProcess()
+{
+	if (PixelPreprocessors.IsEmpty())
+	{
+		return;
+	}
+
+	for (int32 LayerIdx = 0; LayerIdx < Layers.Num(); ++LayerIdx)
+	{
+		if (const TArray<FPixelPreProcessor>* LayerPixelPreprocessorsPtr = PixelPreprocessors.Find(LayerIdx))
+		{
+			for (const FPixelPreProcessor& PreProcessor : *LayerPixelPreprocessorsPtr)
+			{
+				// PreProcessors are assumed to be valid. Fetch the Data pointer each time
+				// in case a pre-processor changes our pixel data.
+				FImagePixelData* Data = Layers[LayerIdx].Get();
+				PreProcessor(Data);
+			}
+		}
+	}
+}
+
 #endif // WITH_UNREALEXR
 
 void UMoviePipelineImageSequenceOutput_EXR::OnReceiveImageDataImpl(FMoviePipelineMergerOutputFrame* InMergedOutputFrame)
diff --git a/Engine/Plugins/MovieScene/MovieRenderPipeline/Source/MovieRenderPipelineRenderPasses/Public/MovieGraphImageSequenceOutputNode.h b/Engine/Plugins/MovieScene/MovieRenderPipeline/Source/MovieRenderPipelineRenderPasses/Public/MovieGraphImageSequenceOutputNode.h
--- a/Engine/Plugins/MovieScene/MovieRenderPipeline/Source/MovieRenderPipelineRenderPasses/Public/MovieGraphImageSequenceOutputNode.h
+++ b/Engine/Plugins/MovieScene/MovieRenderPipeline/Source/MovieRenderPipelineRenderPasses/Public/MovieGraphImageSequenceOutputNode.h
@@ -63,7 +63,7 @@ protected:
 };
 
 /**
- * Image sequence output node that can write EXR files, which are optionally multi-layer.
+ * Image sequence output node that can write single-layer EXR files.
  */
 UCLASS()
 class UMovieGraphImageSequenceOutputNode_EXR : public UMovieGraphImageSequenceOutputNode
@@ -76,11 +76,8 @@ public:
 		OutputFormat = EImageFormat::EXR;
 		bQuantizeTo8Bit = false;
 		Compression = EEXRCompressionFormat::PIZ;
-		bMultilayer = true;
 	}
 
-	virtual void OnReceiveImageDataImpl(UMovieGraphPipeline* InPipeline, UE::MovieGraph::FMovieGraphOutputMergerFrame* InRawFrameData, const TSet<FMovieGraphRenderDataIdentifier>& InMask) override;
-
 #if WITH_EDITOR
 	virtual FText GetNodeTitle(const bool bGetDescriptive = false) const override 
 	{ 
@@ -111,22 +108,63 @@ public:
 public:
 	UPROPERTY(EditAnywhere, BlueprintReadWrite, Category = Overrides, meta = (InlineEditConditionToggle))
 	uint8 bOverride_Compression : 1;
-
-	UPROPERTY(EditAnywhere, BlueprintReadWrite, Category = Overrides, meta = (InlineEditConditionToggle))
-	uint8 bOverride_bMultilayer : 1;
 	
 	/**
 	 * Which compression method should the resulting EXR file be compressed with.
 	 */
 	UPROPERTY(EditAnywhere, BlueprintReadWrite, meta=(EditCondition="bOverride_Compression"), Category = "EXR")
 	EEXRCompressionFormat Compression;
+};
 
-	/**
-	 * Whether all renders should be written to the same EXR file via multi-layer support. Not all software
-	 * supports reading multi-layer EXR files.
-	 */
-	UPROPERTY(EditAnywhere, BlueprintReadWrite, meta=(EditCondition="bOverride_bMultilayer"), Category = "EXR")
-	bool bMultilayer;
+
+/**
+ * Image sequence output node that can write multi-layer EXR files.
+ */
+UCLASS()
+class UMovieGraphImageSequenceOutputNode_MultiLayerEXR : public UMovieGraphImageSequenceOutputNode_EXR
+{
+	GENERATED_BODY()
+
+public:
+
+	virtual EMovieGraphBranchRestriction GetBranchRestriction() const override
+	{
+		return EMovieGraphBranchRestriction::Globals;
+	}
+
+	virtual void OnReceiveImageDataImpl(UMovieGraphPipeline* InPipeline, UE::MovieGraph::FMovieGraphOutputMergerFrame* InRawFrameData, const TSet<FMovieGraphRenderDataIdentifier>& InMask) override;
+
+#if WITH_EDITOR
+	virtual FText GetNodeTitle(const bool bGetDescriptive = false) const override
+	{
+		static const FText EXRSequenceNodeName = NSLOCTEXT("MovieGraph", "NodeName_EXRSequenceMultilayer", ".exr Sequence (Multilayer)");
+		return EXRSequenceNodeName;
+	}
+
+	virtual FText GetKeywords() const override
+	{
+		static const FText Keywords = NSLOCTEXT("MovieGraph", "ImageSequenceOutputNode_EXRMultilayer_Keywords", ".exr image (Multilayer)");
+		return Keywords;
+	}
+		
+	virtual FLinearColor GetNodeTitleColor() const override
+	{
+		return FLinearColor(0.047f, 0.654f, 0.537f);
+	}
+
+	virtual FSlateIcon GetIconAndTint(FLinearColor& OutColor) const override
+	{
+		static const FSlateIcon ImageSequenceIcon = FSlateIcon(FAppStyle::GetAppStyleSetName(), "ClassIcon.Texture2D");
+
+		OutColor = FLinearColor::White;
+		return ImageSequenceIcon;
+	}
+
+	virtual FText GetMenuCategory() const override
+	{
+		return NSLOCTEXT("MovieGraphNodes", "FileOutputGraphNode_Category", "Output Type");
+	}
+#endif
 
 private:
 	/**
diff --git a/Engine/Plugins/MovieScene/MovieRenderPipeline/Source/MovieRenderPipelineRenderPasses/Public/MoviePipelineEXROutput.h b/Engine/Plugins/MovieScene/MovieRenderPipeline/Source/MovieRenderPipelineRenderPasses/Public/MoviePipelineEXROutput.h
--- a/Engine/Plugins/MovieScene/MovieRenderPipeline/Source/MovieRenderPipelineRenderPasses/Public/MoviePipelineEXROutput.h
+++ b/Engine/Plugins/MovieScene/MovieRenderPipeline/Source/MovieRenderPipelineRenderPasses/Public/MoviePipelineEXROutput.h
@@ -70,6 +70,9 @@ public:
 	/** The image data to write. Supports multiple layers of different bitdepths. */
 	TArray<TUniquePtr<FImagePixelData>> Layers;
 
+	/** Per-layer array of preprocessors to apply serially to the pixel data when this task is executed. */
+	TSortedMap<int32, TArray<FPixelPreProcessor>> PixelPreprocessors;
+
 	/** Optional. A mapping between the FImagePixelData and a name. The standard is that the default layer is nameless (at which point it would be omitted) and other layers are prefixed. */
 	TMap<FImagePixelData*, FString> LayerNames;
 
@@ -112,6 +115,11 @@ private:
 	*/
 	void AddFileMetadata(Imf::Header& InHeader);
 
+	/**
+	 * Run over all the processors for the pixel data
+	 */
+	void PreProcess();
+
 	template <Imf::PixelType OutputFormat>
 	int64 CompressRaw(Imf::Header& InHeader, Imf::FrameBuffer& InFrameBuffer, FImagePixelData* InLayer);
 };
