Movie Graph:

-New: Implement auto-versioning UX in the Output Setting node

diff --git a/Engine/Plugins/MovieScene/MovieRenderPipeline/Source/MovieRenderPipelineCore/Private/Graph/MovieGraphBlueprintLibrary.cpp b/Engine/Plugins/MovieScene/MovieRenderPipeline/Source/MovieRenderPipelineCore/Private/Graph/MovieGraphBlueprintLibrary.cpp
--- a/Engine/Plugins/MovieScene/MovieRenderPipeline/Source/MovieRenderPipelineCore/Private/Graph/MovieGraphBlueprintLibrary.cpp
+++ b/Engine/Plugins/MovieScene/MovieRenderPipeline/Source/MovieRenderPipelineCore/Private/Graph/MovieGraphBlueprintLibrary.cpp
@@ -228,9 +228,9 @@ int32 UMovieGraphBlueprintLibrary::ResolveVersionNumber(FMovieGraphFilenameResol
 	constexpr bool bIncludeCDOs = true;
 	constexpr bool bExactMatch = true;
 	const UMovieGraphOutputSettingNode* OutputSettingNode = InParams.EvaluatedConfig->GetSettingForBranch<UMovieGraphOutputSettingNode>(UMovieGraphNode::GlobalsPinName, bIncludeCDOs, bExactMatch);
-	if (!OutputSettingNode->bAutoVersion)
+	if (!OutputSettingNode->VersioningSettings.bAutoVersioning)
 	{
-		return OutputSettingNode->VersionNumber;
+		return OutputSettingNode->VersioningSettings.VersionNumber;
 	}
 
 	// Calculate a version number by looking at the output path and then scanning for a version token.
diff --git a/Engine/Plugins/MovieScene/MovieRenderPipeline/Source/MovieRenderPipelineCore/Private/Graph/Nodes/MovieGraphOutputSettingNode.cpp b/Engine/Plugins/MovieScene/MovieRenderPipeline/Source/MovieRenderPipelineCore/Private/Graph/Nodes/MovieGraphOutputSettingNode.cpp
--- a/Engine/Plugins/MovieScene/MovieRenderPipeline/Source/MovieRenderPipelineCore/Private/Graph/Nodes/MovieGraphOutputSettingNode.cpp
+++ b/Engine/Plugins/MovieScene/MovieRenderPipeline/Source/MovieRenderPipelineCore/Private/Graph/Nodes/MovieGraphOutputSettingNode.cpp
@@ -11,8 +11,7 @@ UMovieGraphOutputSettingNode::UMovieGraphOutputSettingNode()
 	, bOverwriteExistingOutput(true)
 	, ZeroPadFrameNumbers(4)
 	, FrameNumberOffset(0)
-	, bAutoVersion(true)
-	, VersionNumber(1)
+	, VersioningSettings(FMovieGraphVersioningSettings(true, 1))
 {
 	FileNameFormat = TEXT("{sequence_name}.{frame_number}");
 	OutputDirectory.Path = TEXT("{project_dir}/Saved/MovieRenders/");
diff --git a/Engine/Plugins/MovieScene/MovieRenderPipeline/Source/MovieRenderPipelineCore/Public/Graph/Nodes/MovieGraphOutputSettingNode.h b/Engine/Plugins/MovieScene/MovieRenderPipeline/Source/MovieRenderPipelineCore/Public/Graph/Nodes/MovieGraphOutputSettingNode.h
--- a/Engine/Plugins/MovieScene/MovieRenderPipeline/Source/MovieRenderPipelineCore/Public/Graph/Nodes/MovieGraphOutputSettingNode.h
+++ b/Engine/Plugins/MovieScene/MovieRenderPipeline/Source/MovieRenderPipelineCore/Public/Graph/Nodes/MovieGraphOutputSettingNode.h
@@ -9,6 +9,26 @@
 
 #include "MovieGraphOutputSettingNode.generated.h"
 
+USTRUCT(BlueprintType)
+struct FMovieGraphVersioningSettings
+{
+	GENERATED_BODY()
+
+	/**
+	 * If true, {version} tokens specified in the Output Directory and File Name Format properties will automatically
+	 * be incremented with each local render. If false, the version specified in Version Number will be used instead.
+	 *
+	 * Auto-versioning will search across all render branches and use the highest version found as the basis for the
+	 * next version used.
+	 */
+	UPROPERTY(EditAnywhere, BlueprintReadWrite, Category = "Versioning")
+	bool bAutoVersioning;
+	
+	/** The value to use for the version token if versions are not automatically incremented (Auto Version is off). */
+	UPROPERTY(EditAnywhere, BlueprintReadWrite, Category = "Versioning", meta = (UIMin = 1, UIMax = 50, ClampMin = 1))
+	int32 VersionNumber;
+};
+
 UCLASS()
 class MOVIERENDERPIPELINECORE_API UMovieGraphOutputSettingNode : public UMovieGraphSettingNode
 {
@@ -55,10 +75,7 @@ public:
 	uint8 bOverride_HandleFrameCount : 1;
 
 	UPROPERTY(EditAnywhere, BlueprintReadWrite, Category = Overrides, meta = (InlineEditConditionToggle))
-	uint8 bOverride_bAutoVersion : 1;
-
-	UPROPERTY(EditAnywhere, BlueprintReadWrite, Category = Overrides, meta = (InlineEditConditionToggle))
-	uint8 bOverride_VersionNumber : 1;
+	uint8 bOverride_VersioningSettings : 1;
 
 	/** What directory should all of our output files be relative to. */
 	UPROPERTY(EditAnywhere, BlueprintReadWrite, Category = "File Output", meta=(EditCondition="bOverride_OutputDirectory"))
@@ -100,19 +117,9 @@ public:
 	UPROPERTY(EditAnywhere, BlueprintReadWrite, Category = "Frames", meta = (UIMin = 0, ClampMin = 0, EditCondition = "bOverride_HandleFrameCount"))
 	int32 HandleFrameCount;
 
-	// TODO: This versioning property should be on the "global" node
 	/**
-	 * If true, {version} tokens specified in the Output Directory and File Name Format properties will automatically
-	 * be incremented with each local render. If false, the version specified in Version Number will be used instead.
-	 *
-	 * Auto-versioning will search across all render branches and use the highest version found as the basis for the
-	 * next version used.
+	 * Determines how versioning should be handled (Auto Version, Version Number, etc.).
 	 */
-	UPROPERTY(EditAnywhere, BlueprintReadWrite, Category = "Versioning", meta = (EditCondition = "bOverride_bAutoVersion"))
-	bool bAutoVersion;
-	
-	// TODO: This versioning property should be on the "global" node
-	/** The value to use for the version token if versions are not automatically incremented (Auto Version is off). */
-	UPROPERTY(EditAnywhere, BlueprintReadWrite, Category = "Versioning", meta = (EditCondition = "bOverride_VersionNumber", UIMin = 1, UIMax = 50, ClampMin = 1))
-	int32 VersionNumber;
+	UPROPERTY(EditAnywhere, BlueprintReadWrite, Category = "Versioning", meta = (EditCondition = "bOverride_VersioningSettings"))
+	FMovieGraphVersioningSettings VersioningSettings;
 };
diff --git a/Engine/Plugins/MovieScene/MovieRenderPipeline/Source/MovieRenderPipelineEditor/Private/Customizations/Graph/MovieGraphVersioningSettingsCustomization.h b/Engine/Plugins/MovieScene/MovieRenderPipeline/Source/MovieRenderPipelineEditor/Private/Customizations/Graph/MovieGraphVersioningSettingsCustomization.h
new file mode 100644
--- /dev/null
+++ b/Engine/Plugins/MovieScene/MovieRenderPipeline/Source/MovieRenderPipelineEditor/Private/Customizations/Graph/MovieGraphVersioningSettingsCustomization.h
@@ -0,0 +1,106 @@
+﻿// Copyright Epic Games, Inc. All Rights Reserved.
+#pragma once
+
+#include "Graph/Nodes/MovieGraphOutputSettingNode.h"
+
+#include "DetailLayoutBuilder.h"
+#include "DetailWidgetRow.h"
+#include "IPropertyTypeCustomization.h"
+#include "Widgets/SBoxPanel.h"
+#include "Widgets/Input/SEditableTextBox.h"
+#include "Widgets/Layout/SWidgetSwitcher.h"
+
+#define LOCTEXT_NAMESPACE "FMovieGraphVersioningSettingsCustomization"
+
+/* Customizes how Versioning Settings are displayed in the details pane. */
+class MOVIERENDERPIPELINEEDITOR_API FMovieGraphVersioningSettingsCustomization : public IPropertyTypeCustomization
+{
+public:
+	static TSharedRef<IPropertyTypeCustomization> MakeInstance()
+	{
+		return MakeShared<FMovieGraphVersioningSettingsCustomization>();
+	}
+
+protected:
+	
+	//~ Begin IPropertyTypeCustomization interface
+	virtual void CustomizeHeader(
+		TSharedRef<IPropertyHandle> InStructPropertyHandle,
+		FDetailWidgetRow& HeaderRow, IPropertyTypeCustomizationUtils& CustomizationUtils) override
+	{
+		const TSharedRef<IPropertyHandle> AutoVersioningProp = InStructPropertyHandle->GetChildHandle(
+			GET_MEMBER_NAME_CHECKED(FMovieGraphVersioningSettings, bAutoVersioning)).ToSharedRef();
+		
+		const TSharedRef<IPropertyHandle> VersionNumberProp = InStructPropertyHandle->GetChildHandle(
+			GET_MEMBER_NAME_CHECKED(FMovieGraphVersioningSettings, VersionNumber)).ToSharedRef();
+
+		HeaderRow
+		.NameContent()
+		[
+			// Show "Version Number" as the row name
+			VersionNumberProp->CreatePropertyNameWidget()
+		]
+		.ValueContent()
+		.HAlign(HAlign_Fill)
+		[
+			SNew(SHorizontalBox)
+			
+			+ SHorizontalBox::Slot()
+			.Padding(0, 2, 8, 2)
+			.HAlign(HAlign_Fill)
+			[
+				SNew(SWidgetSwitcher)
+				.WidgetIndex_Lambda([AutoVersioningProp]()
+				{
+					bool bAutoVersioning = false;
+					AutoVersioningProp->GetValue(bAutoVersioning);
+
+					return bAutoVersioning;
+				})
+
+				+SWidgetSwitcher::Slot()
+				[
+					VersionNumberProp->CreatePropertyValueWidget()
+				]
+
+				+SWidgetSwitcher::Slot()
+				[
+					SNew(SEditableTextBox)
+					.Font(IDetailLayoutBuilder::GetDetailFont())
+					.Text(LOCTEXT("SetByRender", "Set By Render"))
+					.IsEnabled(false)
+				]
+			]
+
+			+ SHorizontalBox::Slot()
+			.Padding(0, 2, 8, 2)
+			.AutoWidth()
+			.HAlign(HAlign_Fill)
+			[
+				SNew(SHorizontalBox)
+				
+				+ SHorizontalBox::Slot()
+				.Padding(8, 0)
+				.AutoWidth()
+				[
+					AutoVersioningProp->CreatePropertyValueWidget()
+				]
+
+				+ SHorizontalBox::Slot()
+				.Padding(0)
+				.AutoWidth()
+				[
+					AutoVersioningProp->CreatePropertyNameWidget()
+				]
+			]
+		];
+	}
+
+	// Skip customizing children
+	virtual void CustomizeChildren(
+		TSharedRef<IPropertyHandle> InStructPropertyHandle,
+		IDetailChildrenBuilder& StructBuilder, IPropertyTypeCustomizationUtils& StructCustomizationUtils) override { }
+	//~ End IPropertyTypeCustomization interface
+};
+
+#undef LOCTEXT_NAMESPACE
diff --git a/Engine/Plugins/MovieScene/MovieRenderPipeline/Source/MovieRenderPipelineEditor/Private/Graph/MovieGraphAssetToolkit.cpp b/Engine/Plugins/MovieScene/MovieRenderPipeline/Source/MovieRenderPipelineEditor/Private/Graph/MovieGraphAssetToolkit.cpp
--- a/Engine/Plugins/MovieScene/MovieRenderPipeline/Source/MovieRenderPipelineEditor/Private/Graph/MovieGraphAssetToolkit.cpp
+++ b/Engine/Plugins/MovieScene/MovieRenderPipeline/Source/MovieRenderPipelineEditor/Private/Graph/MovieGraphAssetToolkit.cpp
@@ -8,10 +8,12 @@
 #include "Customizations/Graph/MovieGraphNodeCustomization.h"
 #include "Customizations/Graph/MovieGraphSelectNodeCustomization.h"
 #include "Customizations/Graph/MovieGraphShowFlagsCustomization.h"
+#include "Customizations/Graph/MovieGraphVersioningSettingsCustomization.h"
 #include "Graph/Renderers/MovieGraphShowFlags.h"
 
 #include "Graph/MovieGraphConfig.h"
 #include "Graph/Nodes/MovieGraphCollectionNode.h"
+#include "Graph/Nodes/MovieGraphOutputSettingNode.h"
 #include "Graph/Nodes/MovieGraphSelectNode.h"
 #include "MovieEdGraphNode.h"
 #include "MovieGraphSchema.h"
@@ -251,6 +253,10 @@ TSharedRef<SDockTab> FMovieGraphAssetToolkit::SpawnTab_RenderGraphDetails(const
 	SelectedGraphObjectsDetailsWidget->RegisterInstancedCustomPropertyTypeLayout(
 		FMovieGraphNamedResolution::StaticStruct()->GetFName(),
 		FOnGetPropertyTypeCustomizationInstance::CreateStatic(&FMovieGraphNamedResolutionCustomization::MakeInstance));
+
+	SelectedGraphObjectsDetailsWidget->RegisterInstancedCustomPropertyTypeLayout(
+		FMovieGraphVersioningSettings::StaticStruct()->GetFName(),
+		FOnGetPropertyTypeCustomizationInstance::CreateStatic(&FMovieGraphVersioningSettingsCustomization::MakeInstance));
 	
 	SelectedGraphObjectsDetailsWidget->RegisterInstancedCustomPropertyLayout(
 		UMovieGraphCollectionNode::StaticClass(),
