Movie Graph:

-Fix: Issue 390808: Compile errors in Module.MovieRenderPipelineCore.6.cpp, MovieGraphOutputSettingNode.cpp and MovieGraphOutputSettingNode.h

diff --git a/Engine/Plugins/MovieScene/MovieRenderPipeline/Source/MovieRenderPipelineCore/Private/Graph/Nodes/MovieGraphOutputSettingNode.cpp b/Engine/Plugins/MovieScene/MovieRenderPipeline/Source/MovieRenderPipelineCore/Private/Graph/Nodes/MovieGraphOutputSettingNode.cpp
--- a/Engine/Plugins/MovieScene/MovieRenderPipeline/Source/MovieRenderPipelineCore/Private/Graph/Nodes/MovieGraphOutputSettingNode.cpp
+++ b/Engine/Plugins/MovieScene/MovieRenderPipeline/Source/MovieRenderPipelineCore/Private/Graph/Nodes/MovieGraphOutputSettingNode.cpp
@@ -11,7 +11,6 @@ UMovieGraphOutputSettingNode::UMovieGraphOutputSettingNode()
 	, bOverwriteExistingOutput(true)
 	, ZeroPadFrameNumbers(4)
 	, FrameNumberOffset(0)
-	, VersioningSettings(FMovieGraphVersioningSettings(true, 1))
 {
 	FileNameFormat = TEXT("{sequence_name}.{frame_number}");
 	OutputDirectory.Path = TEXT("{project_dir}/Saved/MovieRenders/");
diff --git a/Engine/Plugins/MovieScene/MovieRenderPipeline/Source/MovieRenderPipelineCore/Public/Graph/Nodes/MovieGraphOutputSettingNode.h b/Engine/Plugins/MovieScene/MovieRenderPipeline/Source/MovieRenderPipelineCore/Public/Graph/Nodes/MovieGraphOutputSettingNode.h
--- a/Engine/Plugins/MovieScene/MovieRenderPipeline/Source/MovieRenderPipelineCore/Public/Graph/Nodes/MovieGraphOutputSettingNode.h
+++ b/Engine/Plugins/MovieScene/MovieRenderPipeline/Source/MovieRenderPipelineCore/Public/Graph/Nodes/MovieGraphOutputSettingNode.h
@@ -22,11 +22,11 @@ struct FMovieGraphVersioningSettings
 	 * next version used.
 	 */
 	UPROPERTY(EditAnywhere, BlueprintReadWrite, Category = "Versioning")
-	bool bAutoVersioning;
+	bool bAutoVersioning = true;
 	
 	/** The value to use for the version token if versions are not automatically incremented (Auto Version is off). */
 	UPROPERTY(EditAnywhere, BlueprintReadWrite, Category = "Versioning", meta = (UIMin = 1, UIMax = 50, ClampMin = 1))
-	int32 VersionNumber;
+	int32 VersionNumber = 1;
 };
 
 UCLASS()
