MoviePipeline: Exposing all EXR compression formats.

diff --git a/Engine/Plugins/MovieScene/MovieRenderPipeline/Source/MovieRenderPipelineRenderPasses/Private/MoviePipelineEXROutput.cpp b/Engine/Plugins/MovieScene/MovieRenderPipeline/Source/MovieRenderPipelineRenderPasses/Private/MoviePipelineEXROutput.cpp
--- a/Engine/Plugins/MovieScene/MovieRenderPipeline/Source/MovieRenderPipelineRenderPasses/Private/MoviePipelineEXROutput.cpp
+++ b/Engine/Plugins/MovieScene/MovieRenderPipeline/Source/MovieRenderPipelineRenderPasses/Private/MoviePipelineEXROutput.cpp
@@ -124,22 +124,8 @@ bool FEXRImageWriteTask::WriteToDisk()
 	{
 		PreProcess();
 
-		Imf::Compression FileCompression = Imf::Compression::NO_COMPRESSION;
-		switch(Compression)
-		{
-			case EEXRCompressionFormat::None:
-				FileCompression = Imf::Compression::NO_COMPRESSION; break;
-			case EEXRCompressionFormat::ZIP:
-				FileCompression = Imf::Compression::ZIP_COMPRESSION; break;
-			case EEXRCompressionFormat::PIZ:
-				FileCompression = Imf::Compression::PIZ_COMPRESSION; break;
-			case EEXRCompressionFormat::DWAA:
-				FileCompression = Imf::Compression::DWAA_COMPRESSION; break;
-			case EEXRCompressionFormat::DWAB:
-				FileCompression = Imf::Compression::DWAB_COMPRESSION; break;
-			default: 
-				checkNoEntry();
-		}
+		static_assert(static_cast<uint8>(EEXRCompressionFormat::Max) == Imf::Compression::NUM_COMPRESSION_METHODS);
+		Imf::Compression FileCompression = static_cast<Imf::Compression>(Compression);
 		
 		// Data Window specifies how much data is in the actual file, ie: 1920x1080
 		IMATH_NAMESPACE::Box2i DataWindow = IMATH_NAMESPACE::Box2i(IMATH_NAMESPACE::V2i(0,0), IMATH_NAMESPACE::V2i(Width - 1, Height - 1));
diff --git a/Engine/Plugins/MovieScene/MovieRenderPipeline/Source/MovieRenderPipelineRenderPasses/Public/MoviePipelineEXROutput.h b/Engine/Plugins/MovieScene/MovieRenderPipeline/Source/MovieRenderPipelineRenderPasses/Public/MoviePipelineEXROutput.h
--- a/Engine/Plugins/MovieScene/MovieRenderPipeline/Source/MovieRenderPipelineRenderPasses/Public/MoviePipelineEXROutput.h
+++ b/Engine/Plugins/MovieScene/MovieRenderPipeline/Source/MovieRenderPipelineRenderPasses/Public/MoviePipelineEXROutput.h
@@ -24,19 +24,32 @@ THIRD_PARTY_INCLUDES_END
 
 class UMoviePipelineColorSetting;
 
+// Exr compression format options. Exactly matches the exr library Imf::Compression enum.
 UENUM(BlueprintType)
 enum class EEXRCompressionFormat : uint8
 {
 	/** No compression is applied. */
-	None,
+	None = 0,
+	/** This compression method is fast, and works well for images with large flat areas but yields worse results for grainy images. Lossless. */
+	RLE = 1,
+	/** This compression method is similar to ZIP but compresses only one image row at a time. Lossless. */
+	ZIPS = 2 UMETA(DisplayName = "ZIP (1 scanline)"),
+	/** Good compression quality for images with low amounts of noise. This compression method operates in in blocks of 16 scan lines. Lossless. */
+	ZIP = 3  UMETA(DisplayName = "ZIP (16 scanlines)"),
 	/** Good compression quality for grainy images. Lossless.*/
-	PIZ,
-	/** Good compression quality for images with low amounts of noise. Lossless. */
-	ZIP,
+	PIZ = 4,
+	/** This format only stores 24 bits of the 32 bit data and has subsequently a significant loss of precision. This method is only applied when saving in FLOAT color depth. HALF and UINT remain unchanged. Lossy. */
+	PXR24 = 5,
+	/** This compression method only applies to images stored in HALF color depth. Blocks of 4×4 pixels are stored with using only 14 byte each (instead of the 32 byte they would normally need). Each block is compressed to the exact same size. Different images with the same dimensions require the same storage space regardless of image content. Lossy. */
+	B44 = 6,
+	/** A modified version of B44. If all pixels in a 4*4 block have the same color it will use only 3 instead of 14 byte. */
+	B44A = 7,
 	/** Lossy DCT-based compression for RGB channels. Alpha and other channels are uncompressed. More efficient than DWAB for partial buffer access on read in 3rd party tools. */
-	DWAA,
+	DWAA = 8,
 	/** Similar to DWAA but goes in blocks of 256 scanlines instead of 32. More efficient disk space and faster to decode than DWAA. */
-	DWAB
+	DWAB = 9,
+
+	Max UMETA(Hidden)
 };
 
 #if WITH_UNREALEXR
