MRG: Adding support for EXR color space metadata.

diff --git a/Engine/Plugins/MovieScene/MovieRenderPipeline/Source/MovieRenderPipelineCore/Private/Graph/Renderers/MovieGraphDeferredPass.cpp b/Engine/Plugins/MovieScene/MovieRenderPipeline/Source/MovieRenderPipelineCore/Private/Graph/Renderers/MovieGraphDeferredPass.cpp
--- a/Engine/Plugins/MovieScene/MovieRenderPipeline/Source/MovieRenderPipelineCore/Private/Graph/Renderers/MovieGraphDeferredPass.cpp
+++ b/Engine/Plugins/MovieScene/MovieRenderPipeline/Source/MovieRenderPipelineCore/Private/Graph/Renderers/MovieGraphDeferredPass.cpp
@@ -107,7 +107,7 @@ void FMovieGraphDeferredPass::Render(const FMovieGraphTraversalContext& InFrameT
 	UMovieGraphImagePassBaseNode* ParentNodeThisFrame = GetParentNode(InTimeData.EvaluatedConfig);
 	const bool bWriteAllSamples = ParentNodeThisFrame->GetWriteAllSamples();
 	int32 NumSpatialSamples = FMath::Max(1, ParentNodeThisFrame->GetNumSpatialSamples());
-	const bool bDisableToneCurve = ParentNodeThisFrame->GetDisableToneCurve();
+	const ESceneCaptureSource SceneCaptureSource = ParentNodeThisFrame->GetDisableToneCurve() ? ESceneCaptureSource::SCS_FinalColorHDR : ESceneCaptureSource::SCS_FinalToneCurveHDR;
 	const EAntiAliasingMethod AntiAliasingMethod = ParentNodeThisFrame->GetAntiAliasingMethod();
 	float OverscanFraction = 0.f;
 	const float TileOverlapPadRatio = 0.0f; // No tiling support right now
@@ -206,7 +206,7 @@ void FMovieGraphDeferredPass::Render(const FMovieGraphTraversalContext& InFrameT
 		ViewFamilyInitData.RenderTarget = RenderTargetResource;
 		ViewFamilyInitData.World = GraphRenderer->GetWorld();
 		ViewFamilyInitData.TimeData = InTimeData;
-		ViewFamilyInitData.SceneCaptureSource = bDisableToneCurve ? ESceneCaptureSource::SCS_FinalColorHDR : ESceneCaptureSource::SCS_FinalToneCurveHDR;
+		ViewFamilyInitData.SceneCaptureSource = SceneCaptureSource;
 		ViewFamilyInitData.bWorldIsPaused = bWorldIsPaused;
 		ViewFamilyInitData.FrameIndex = FrameIndex;
 		ViewFamilyInitData.AntiAliasingMethod = AntiAliasingMethod;
@@ -252,6 +252,7 @@ void FMovieGraphDeferredPass::Render(const FMovieGraphTraversalContext& InFrameT
 			SampleState.OverlappedSubpixelShift = OverlappedSubpixelShift;
 			SampleState.OverscanFraction = OverscanFraction;
 			SampleState.bAllowOCIO = ParentNodeThisFrame->GetAllowOCIO();
+			SampleState.SceneCaptureSource = SceneCaptureSource;
 		}
 
 		// If this was just to contribute to the history buffer, no need to go any further.
diff --git a/Engine/Plugins/MovieScene/MovieRenderPipeline/Source/MovieRenderPipelineCore/Public/Graph/MovieGraphDataTypes.h b/Engine/Plugins/MovieScene/MovieRenderPipeline/Source/MovieRenderPipelineCore/Public/Graph/MovieGraphDataTypes.h
--- a/Engine/Plugins/MovieScene/MovieRenderPipeline/Source/MovieRenderPipelineCore/Public/Graph/MovieGraphDataTypes.h
+++ b/Engine/Plugins/MovieScene/MovieRenderPipeline/Source/MovieRenderPipelineCore/Public/Graph/MovieGraphDataTypes.h
@@ -392,6 +392,9 @@ namespace UE::MovieGraph
 
 		/** Allow OpenColorIO transform to be used on this render. */
 		bool bAllowOCIO;
+
+		/** Render scene capture source used for tracking the output color space (without OpenColorIO). */
+		ESceneCaptureSource SceneCaptureSource;
 	};
 
 	/**
diff --git a/Engine/Plugins/MovieScene/MovieRenderPipeline/Source/MovieRenderPipelineRenderPasses/Private/MovieGraphImageSequenceOutputNode.cpp b/Engine/Plugins/MovieScene/MovieRenderPipeline/Source/MovieRenderPipelineRenderPasses/Private/MovieGraphImageSequenceOutputNode.cpp
--- a/Engine/Plugins/MovieScene/MovieRenderPipeline/Source/MovieRenderPipelineRenderPasses/Private/MovieGraphImageSequenceOutputNode.cpp
+++ b/Engine/Plugins/MovieScene/MovieRenderPipeline/Source/MovieRenderPipelineRenderPasses/Private/MovieGraphImageSequenceOutputNode.cpp
@@ -357,25 +357,6 @@ void UMovieGraphImageSequenceOutputNode_MultiLayerEXR::OnReceiveImageDataImpl(UM
 			MultiLayerImageTask->FileMetadata.Add(Metadata.Key, Metadata.Value);
 		}
 
-		// Add color space metadata to the output: xy chromaticity coordinates and/or the color space source/dest names.
-		// TODO: Support is also needed for regular exrs via the image wrapper module.
-		// {
-		// 	UMoviePipelineColorSetting* ColorSetting = InPipeline->GetPipelinePrimaryConfig()->FindSetting<UMoviePipelineColorSetting>();
-		//
-		// 	FColorSpaceMetadata ColorSpaceMetadata = GetColorSpaceMetadata(ColorSetting);
-		//
-		// 	if (!ColorSpaceMetadata.SourceName.IsEmpty())
-		// 	{
-		// 		MultiLayerImageTask->FileMetadata.Add("unreal/colorSpace/source", ColorSpaceMetadata.SourceName);
-		// 	}
-		// 	if (!ColorSpaceMetadata.DestinationName.IsEmpty())
-		// 	{
-		// 		MultiLayerImageTask->FileMetadata.Add("unreal/colorSpace/destination", ColorSpaceMetadata.DestinationName);
-		// 	}
-		//
-		// 	MultiLayerImageTask->ColorSpaceChromaticities = ColorSpaceMetadata.Chromaticities;
-		// }
-
 		// Add each render pass as a layer to the EXR
 		int32 LayerIndex = 0;
 		int32 ShotIndex = 0;
@@ -388,6 +369,19 @@ void UMovieGraphImageSequenceOutputNode_MultiLayerEXR::OnReceiveImageDataImpl(UM
 			const UE::MovieGraph::FMovieGraphSampleState* Payload = ImageData->GetPayload<UE::MovieGraph::FMovieGraphSampleState>();
 			ShotIndex = Payload->TraversalContext.ShotIndex;
 
+			bool bEnabledOCIO = false;
+#if WITH_EDITOR
+			if (ParentNode->OCIOConfiguration.bIsEnabled && Payload->bAllowOCIO)
+			{
+				FPixelPreProcessor OCIOPixelPreProcessor = UE::MovieGraph::Private::CreateOpenColorIOPixelPreProcessor(ParentNode->OCIOConfiguration.ColorConfiguration);
+				if (OCIOPixelPreProcessor)
+				{
+					MultiLayerImageTask->PixelPreprocessors.FindOrAdd(LayerIndex).Emplace(MoveTemp(OCIOPixelPreProcessor));
+					bEnabledOCIO = true;
+				}
+			}
+#endif
+
 			// If there is more than one layer, then we will prefix the layer. The first layer is not prefixed (and gets inserted as RGBA)
 			// as most programs that handle EXRs expect the main image data to be in an unnamed layer.
 			if (LayerIndex == 0)
@@ -402,6 +396,16 @@ void UMovieGraphImageSequenceOutputNode_MultiLayerEXR::OnReceiveImageDataImpl(UM
 				MultiLayerImageTask->Height = Resolution.Y;
 				
 				MultiLayerImageTask->OverscanPercentage = Payload->OverscanFraction;
+#if WITH_EDITOR
+				if (bEnabledOCIO)
+				{
+					UE::MoviePipeline::UpdateColorSpaceMetadata(ParentNode->OCIOConfiguration.ColorConfiguration, *MultiLayerImageTask);
+				}
+				else
+#endif
+				{
+					UE::MoviePipeline::UpdateColorSpaceMetadata(Payload->SceneCaptureSource, *MultiLayerImageTask);
+				}
 			}
 			else
 			{
@@ -421,17 +425,6 @@ void UMovieGraphImageSequenceOutputNode_MultiLayerEXR::OnReceiveImageDataImpl(UM
 				MultiLayerImageTask->LayerNames.FindOrAdd(PixelData.Get(), CombinedName);
 			}
 
-#if WITH_EDITOR
-			if (ParentNode->OCIOConfiguration.bIsEnabled && Payload->bAllowOCIO)
-			{
-				FPixelPreProcessor OCIOPixelPreProcessor = UE::MovieGraph::Private::CreateOpenColorIOPixelPreProcessor(ParentNode->OCIOConfiguration.ColorConfiguration);
-				if (OCIOPixelPreProcessor)
-				{
-					MultiLayerImageTask->PixelPreprocessors.FindOrAdd(LayerIndex).Emplace(MoveTemp(OCIOPixelPreProcessor));
-				}
-			}
-#endif
-
 			MultiLayerImageTask->Layers.Add(MoveTemp(PixelData));
 			LayerIndex++;
 		}
diff --git a/Engine/Plugins/MovieScene/MovieRenderPipeline/Source/MovieRenderPipelineRenderPasses/Private/MoviePipelineEXROutput.cpp b/Engine/Plugins/MovieScene/MovieRenderPipeline/Source/MovieRenderPipelineRenderPasses/Private/MoviePipelineEXROutput.cpp
--- a/Engine/Plugins/MovieScene/MovieRenderPipeline/Source/MovieRenderPipelineRenderPasses/Private/MoviePipelineEXROutput.cpp
+++ b/Engine/Plugins/MovieScene/MovieRenderPipeline/Source/MovieRenderPipelineRenderPasses/Private/MoviePipelineEXROutput.cpp
@@ -418,6 +418,110 @@ void FEXRImageWriteTask::PreProcess()
 
 #endif // WITH_UNREALEXR
 
+namespace UE
+{
+	namespace MoviePipeline
+	{
+		static TPair<UE::Color::EColorSpace, FString> GetDisplayGamutType(EDisplayColorGamut InDisplayGamut)
+		{
+			switch (InDisplayGamut)
+			{
+			case EDisplayColorGamut::sRGB_D65: return { UE::Color::EColorSpace::sRGB, TEXT("sRGB") };
+			case EDisplayColorGamut::DCIP3_D65: return { UE::Color::EColorSpace::P3DCI, TEXT("P3DCI") };
+			case EDisplayColorGamut::Rec2020_D65: return { UE::Color::EColorSpace::Rec2020, TEXT("Rec2020") };
+			case EDisplayColorGamut::ACES_D60: return { UE::Color::EColorSpace::ACESAP0, TEXT("ACESAP0") };
+			case EDisplayColorGamut::ACEScg_D60: return { UE::Color::EColorSpace::ACESAP1, TEXT("ACESAP1") };
+
+			default:
+				checkNoEntry();
+				return {};
+			}
+		};
+
+		void UpdateColorSpaceMetadataImpl(const FEXRColorSpaceMetadata& InColorSpaceMetadata, FEXRImageWriteTask& InOutImageTask)
+		{
+			if (!InColorSpaceMetadata.SourceName.IsEmpty())
+			{
+				InOutImageTask.FileMetadata.Add("unreal/colorSpace/source", InColorSpaceMetadata.SourceName);
+			}
+			if (!InColorSpaceMetadata.DestinationName.IsEmpty())
+			{
+				InOutImageTask.FileMetadata.Add("unreal/colorSpace/destination", InColorSpaceMetadata.DestinationName);
+			}
+
+			InOutImageTask.ColorSpaceChromaticities = InColorSpaceMetadata.Chromaticities;
+		}
+
+		void UpdateColorSpaceMetadata(const FOpenColorIOColorConversionSettings& InConversionSettings, FEXRImageWriteTask& InOutImageTask)
+		{
+			FEXRColorSpaceMetadata ColorSpaceMetadata;
+
+			if (InConversionSettings.IsValid())
+			{
+				// Note: OpenColorIO does not expose chromaticity information so we only provide transform names.
+				if (InConversionSettings.IsDisplayView())
+				{
+					switch (InConversionSettings.DisplayViewDirection)
+					{
+					case EOpenColorIOViewTransformDirection::Forward:
+						ColorSpaceMetadata.SourceName = InConversionSettings.SourceColorSpace.ToString();
+						ColorSpaceMetadata.DestinationName = InConversionSettings.DestinationDisplayView.ToString();
+						break;
+					case EOpenColorIOViewTransformDirection::Inverse:
+						ColorSpaceMetadata.SourceName = InConversionSettings.DestinationDisplayView.ToString();
+						ColorSpaceMetadata.DestinationName = InConversionSettings.SourceColorSpace.ToString();
+						break;
+
+					default:
+						checkNoEntry();
+					}
+				}
+				else
+				{
+					ColorSpaceMetadata.SourceName = InConversionSettings.SourceColorSpace.ToString();
+					ColorSpaceMetadata.DestinationName = InConversionSettings.DestinationColorSpace.ToString();
+				}
+			}
+
+			UpdateColorSpaceMetadataImpl(ColorSpaceMetadata, InOutImageTask);
+		}
+
+		void UpdateColorSpaceMetadata(ESceneCaptureSource InSceneCaptureSource, FEXRImageWriteTask& InOutImageTask)
+		{
+			FEXRColorSpaceMetadata ColorSpaceMetadata;
+
+			switch (InSceneCaptureSource)
+			{
+			case SCS_FinalColorLDR:
+			case SCS_FinalToneCurveHDR:
+			{
+				// We are in output display space
+				TPair<UE::Color::EColorSpace, FString> ColorSpaceType = GetDisplayGamutType(HDRGetDefaultDisplayColorGamut());
+				UE::Color::FColorSpace OutputCS = UE::Color::FColorSpace(ColorSpaceType.Key);
+
+				ColorSpaceMetadata.DestinationName = ColorSpaceType.Value;
+				ColorSpaceMetadata.Chromaticities = { OutputCS.GetRedChromaticity(), OutputCS.GetGreenChromaticity(), OutputCS.GetBlueChromaticity(), OutputCS.GetWhiteChromaticity() };
+				break;
+			}
+			case SCS_SceneColorHDR:
+			case SCS_SceneColorHDRNoAlpha:
+			case SCS_FinalColorHDR:
+			case SCS_BaseColor:
+			{
+				// We are in working color space
+				const UE::Color::FColorSpace& WCS = UE::Color::FColorSpace::GetWorking();
+				ColorSpaceMetadata.Chromaticities = { WCS.GetRedChromaticity(), WCS.GetGreenChromaticity(), WCS.GetBlueChromaticity(), WCS.GetWhiteChromaticity() };
+				break;
+			}
+			default:
+				break;
+			}
+
+			UpdateColorSpaceMetadataImpl(ColorSpaceMetadata, InOutImageTask);
+		}
+	}
+}
+
 void UMoviePipelineImageSequenceOutput_EXR::OnReceiveImageDataImpl(FMoviePipelineMergerOutputFrame* InMergedOutputFrame)
 {
 	if (!bMultilayer)
@@ -524,21 +628,15 @@ void UMoviePipelineImageSequenceOutput_EXR::OnReceiveImageDataImpl(FMoviePipelin
 
 		// Add color space metadata to the output: xy chromaticity coordinates and/or the color space source/dest names.
 		// TODO: Support is also needed for regular exrs via the image wrapper module.
+		UMoviePipelineColorSetting* ColorSetting = GetPipeline()->GetPipelinePrimaryConfig()->FindSetting<UMoviePipelineColorSetting>();
+		if (ColorSetting && ColorSetting->OCIOConfiguration.bIsEnabled)
 		{
-			UMoviePipelineColorSetting* ColorSetting = GetPipeline()->GetPipelinePrimaryConfig()->FindSetting<UMoviePipelineColorSetting>();
-
-			FColorSpaceMetadata ColorSpaceMetadata = GetColorSpaceMetadata(ColorSetting);
-
-			if (!ColorSpaceMetadata.SourceName.IsEmpty())
-			{
-				MultiLayerImageTask->FileMetadata.Add("unreal/colorSpace/source", ColorSpaceMetadata.SourceName);
-			}
-			if (!ColorSpaceMetadata.DestinationName.IsEmpty())
-			{
-				MultiLayerImageTask->FileMetadata.Add("unreal/colorSpace/destination", ColorSpaceMetadata.DestinationName);
-			}
-
-			MultiLayerImageTask->ColorSpaceChromaticities = ColorSpaceMetadata.Chromaticities;
+			UE::MoviePipeline::UpdateColorSpaceMetadata(ColorSetting->OCIOConfiguration.ColorConfiguration, *MultiLayerImageTask);
+		}
+		else
+		{
+			ESceneCaptureSource SceneCaptureSource = (ColorSetting && ColorSetting->bDisableToneCurve) ? ESceneCaptureSource::SCS_FinalColorHDR : ESceneCaptureSource::SCS_FinalToneCurveHDR;
+			UE::MoviePipeline::UpdateColorSpaceMetadata(SceneCaptureSource, *MultiLayerImageTask);
 		}
 
 		int32 LayerIndex = 0;
@@ -604,75 +702,3 @@ void UMoviePipelineImageSequenceOutput_EXR::OnReceiveImageDataImpl(FMoviePipelin
 
 	}
 }
-
-UMoviePipelineImageSequenceOutput_EXR::FColorSpaceMetadata UMoviePipelineImageSequenceOutput_EXR::GetColorSpaceMetadata(UMoviePipelineColorSetting* InColorSettings)
-{
-	auto GetDisplayGamutTypeFn = [](EDisplayColorGamut InDisplayGamut) -> TPair<UE::Color::EColorSpace, FString>
-	{
-		switch (InDisplayGamut)
-		{
-		case EDisplayColorGamut::sRGB_D65: return { UE::Color::EColorSpace::sRGB, TEXT("sRGB") };
-		case EDisplayColorGamut::DCIP3_D65: return { UE::Color::EColorSpace::P3DCI, TEXT("P3DCI") };
-		case EDisplayColorGamut::Rec2020_D65: return { UE::Color::EColorSpace::Rec2020, TEXT("Rec2020") };
-		case EDisplayColorGamut::ACES_D60: return { UE::Color::EColorSpace::ACESAP0, TEXT("ACESAP0") };
-		case EDisplayColorGamut::ACEScg_D60: return { UE::Color::EColorSpace::ACESAP1, TEXT("ACESAP1") };
-
-		default:
-			checkNoEntry();
-			return {};
-		}
-	};
-
-	const UE::Color::FColorSpace& WCS = UE::Color::FColorSpace::GetWorking();
-	FColorSpaceMetadata OutMetadata;
-
-	if (InColorSettings)
-	{
-		if (InColorSettings->OCIOConfiguration.bIsEnabled)
-		{
-			// Note: OpenColorIO does not expose chromaticity information so we only provide transform names.
-			const FOpenColorIOColorConversionSettings& ConversionSettings = InColorSettings->OCIOConfiguration.ColorConfiguration;
-
-			if (ConversionSettings.IsDisplayView())
-			{
-				switch (ConversionSettings.DisplayViewDirection)
-				{
-				case EOpenColorIOViewTransformDirection::Forward:
-					OutMetadata.SourceName = ConversionSettings.SourceColorSpace.ToString();
-					OutMetadata.DestinationName = ConversionSettings.DestinationDisplayView.ToString();
-					break;
-				case EOpenColorIOViewTransformDirection::Inverse:
-					OutMetadata.SourceName = ConversionSettings.DestinationDisplayView.ToString();
-					OutMetadata.DestinationName = ConversionSettings.SourceColorSpace.ToString();
-					break;
-
-				default:
-					checkNoEntry();
-				}
-			}
-			else
-			{
-				OutMetadata.SourceName = ConversionSettings.SourceColorSpace.ToString();
-				OutMetadata.DestinationName = ConversionSettings.DestinationColorSpace.ToString();
-			}
-		}
-		else if (!InColorSettings->bDisableToneCurve)
-		{
-			TPair<UE::Color::EColorSpace, FString> ColorSpaceType = GetDisplayGamutTypeFn(HDRGetDefaultDisplayColorGamut());
-			UE::Color::FColorSpace OutputCS = UE::Color::FColorSpace(ColorSpaceType.Key);
-
-			OutMetadata.DestinationName = ColorSpaceType.Value;
-			OutMetadata.Chromaticities = { OutputCS.GetRedChromaticity(), OutputCS.GetGreenChromaticity(), OutputCS.GetBlueChromaticity(), OutputCS.GetWhiteChromaticity() };
-		}
-		else
-		{
-			OutMetadata.Chromaticities = { WCS.GetRedChromaticity(), WCS.GetGreenChromaticity(), WCS.GetBlueChromaticity(), WCS.GetWhiteChromaticity() };
-		}
-	}
-	else
-	{
-		OutMetadata.Chromaticities = { WCS.GetRedChromaticity(), WCS.GetGreenChromaticity(), WCS.GetBlueChromaticity(), WCS.GetWhiteChromaticity() };
-	}
-
-	return OutMetadata;
-}
diff --git a/Engine/Plugins/MovieScene/MovieRenderPipeline/Source/MovieRenderPipelineRenderPasses/Public/MoviePipelineEXROutput.h b/Engine/Plugins/MovieScene/MovieRenderPipeline/Source/MovieRenderPipelineRenderPasses/Public/MoviePipelineEXROutput.h
--- a/Engine/Plugins/MovieScene/MovieRenderPipeline/Source/MovieRenderPipelineRenderPasses/Public/MoviePipelineEXROutput.h
+++ b/Engine/Plugins/MovieScene/MovieRenderPipeline/Source/MovieRenderPipelineRenderPasses/Public/MoviePipelineEXROutput.h
@@ -23,6 +23,27 @@ THIRD_PARTY_INCLUDES_END
 #include "MoviePipelineEXROutput.generated.h"
 
 class UMoviePipelineColorSetting;
+class FEXRImageWriteTask;
+
+namespace UE
+{
+	namespace MoviePipeline
+	{
+		/** Collection of color space metadata for EXR. */
+		struct FEXRColorSpaceMetadata
+		{
+			FString SourceName;
+			FString DestinationName;
+			TArray<FVector2d> Chromaticities;
+		};
+
+		/** Update the image write task with color space metadata from the OpenColorIO transform. */
+		void UpdateColorSpaceMetadata(const FOpenColorIOColorConversionSettings& InConversionSettings, FEXRImageWriteTask& InOutImageTask);
+
+		/** Update the image write task with color space metadata from the render capture source mode. */
+		void UpdateColorSpaceMetadata(ESceneCaptureSource InSceneCaptureSource, FEXRImageWriteTask& InOutImageTask);
+	}
+}
 
 // Exr compression format options. Exactly matches the exr library Imf::Compression enum.
 UENUM(BlueprintType)
@@ -170,16 +191,4 @@ public:
 	UPROPERTY(EditAnywhere, BlueprintReadWrite, Category = "EXR")
 	bool bMultilayer;
 
-protected:
-	struct FColorSpaceMetadata
-	{
-		FString SourceName;
-		FString DestinationName;
-		TArray<FVector2d> Chromaticities;
-	};
-
-	/**
-	* Get color space chromaticities, source and destination names from the color settings or working color space.
-	*/
-	static FColorSpaceMetadata GetColorSpaceMetadata(UMoviePipelineColorSetting* InColorSettings);
-};
\ No newline at end of file
+};
