MoviePipeline: Added a details customization for the new modifiers UX.

Modifier nodes now support an array of modifiers under the hood, to pave the way for user-defined modifiers in the future.

#rnx

diff --git a/Engine/Plugins/MovieScene/MovieRenderPipeline/Source/MovieRenderPipelineCore/Private/Graph/MovieGraphConfig.cpp b/Engine/Plugins/MovieScene/MovieRenderPipeline/Source/MovieRenderPipelineCore/Private/Graph/MovieGraphConfig.cpp
--- a/Engine/Plugins/MovieScene/MovieRenderPipeline/Source/MovieRenderPipelineCore/Private/Graph/MovieGraphConfig.cpp
+++ b/Engine/Plugins/MovieScene/MovieRenderPipeline/Source/MovieRenderPipelineCore/Private/Graph/MovieGraphConfig.cpp
@@ -10,6 +10,7 @@
 #include "Graph/Nodes/MovieGraphInputNode.h"
 #include "Graph/Nodes/MovieGraphOutputNode.h"
 #include "Graph/Nodes/MovieGraphRemoveRenderSettingNode.h"
+#include "Graph/Nodes/MovieGraphSubgraphNode.h"
 #include "Graph/Nodes/MovieGraphVariableNode.h"
 #include "MovieGraphUtils.h"
 #include "MoviePipelineQueue.h"
@@ -986,6 +987,29 @@ TArray<FString> UMovieGraphConfig::GetUpstreamBranchNames(UMovieGraphNode* FromN
 	return BranchNames;
 }
 
+void UMovieGraphConfig::GetAllContainedSubgraphs(TSet<UMovieGraphConfig*>& OutSubgraphs) const
+{
+	for (const TObjectPtr<UMovieGraphNode>& Node : GetNodes())
+	{
+		if (!Node)
+		{
+			continue;
+		}
+		
+		if (const UMovieGraphSubgraphNode* SubgraphNode = Cast<UMovieGraphSubgraphNode>(Node))
+		{
+			UMovieGraphConfig* SubgraphConfig = SubgraphNode->GetSubgraphAsset();
+			
+			// Don't recurse into this graph if it was already added (to prevent infinite recursion)
+			if (!OutSubgraphs.Contains(SubgraphConfig))
+			{
+				OutSubgraphs.Add(SubgraphConfig);
+				SubgraphConfig->GetAllContainedSubgraphs(OutSubgraphs);
+			}
+		}
+	}
+}
+
 void UMovieGraphConfig::InitializeFlattenedNode(UMovieGraphNode* InNode)
 {
 	// We go through each of the bOverride_ properties on this new instance and set
diff --git a/Engine/Plugins/MovieScene/MovieRenderPipeline/Source/MovieRenderPipelineCore/Private/Graph/MovieGraphPipeline.cpp b/Engine/Plugins/MovieScene/MovieRenderPipeline/Source/MovieRenderPipelineCore/Private/Graph/MovieGraphPipeline.cpp
--- a/Engine/Plugins/MovieScene/MovieRenderPipeline/Source/MovieRenderPipelineCore/Private/Graph/MovieGraphPipeline.cpp
+++ b/Engine/Plugins/MovieScene/MovieRenderPipeline/Source/MovieRenderPipelineCore/Private/Graph/MovieGraphPipeline.cpp
@@ -227,33 +227,46 @@ void UMovieGraphPipeline::UpdateLayerContentsInRenderLayerSubsystem(const UMovie
 		
 		for (const UMovieGraphModifierNode* ModifierNode : ModifierNodes)
 		{
-			ModifierNode->ModifierClass->SetCollections({});
-			
-			// Find the collection that this modifier is modifying
-			for (const UMovieGraphCollectionNode* CollectionNode : CollectionNodes)
+			// For each modifier instance, find the collection(s) that it is modifying
+			for (UMoviePipelineCollectionModifier* ModifierInstance : ModifierNode->GetModifiers())
 			{
-				if (!CollectionNode || !CollectionNode->Collection)
+				ModifierInstance->SetCollections({});
+				TArray<UMovieGraphCollection*> ModifierCollections;
+
+				for (const FName& ModifiedCollectionName : ModifierNode->GetCollections())
 				{
-					continue;
+					bool bFoundModifiedCollection = false;
+					
+					for (const UMovieGraphCollectionNode* CollectionNode : CollectionNodes)
+					{
+						if (!CollectionNode || !CollectionNode->Collection)
+						{
+							continue;
+						}
+
+						if (CollectionNode->Collection->GetCollectionName() == ModifiedCollectionName)
+						{
+							bFoundModifiedCollection = true;
+							ModifierCollections.Add(CollectionNode->Collection);
+							break;
+						}
+					}
+
+					if (!bFoundModifiedCollection)
+					{
+						UE_LOG(LogMovieRenderPipeline, Warning, TEXT("The modifier '%s' specified a collection '%s', but the collection couldn't be found."),
+							*ModifierNode->ModifierName, *ModifiedCollectionName.ToString());
+					}
 				}
-				
-				if (CollectionNode->Collection->GetCollectionName() == ModifierNode->ModifiedCollectionName)
+
+				// If the modifier had valid collection(s) added to it, add the modifier to the render layer. The modifier
+				// won't have any effect without collection(s) to act on.
+				if (!ModifierCollections.IsEmpty())
 				{
-					ModifierNode->ModifierClass->SetCollections({CollectionNode->Collection});
-					break;
-				}
-			}
+					ModifierInstance->SetCollections(ModifierCollections);
 
-			// If the modifier had a valid collection added to it, add the modifier to the render layer. The modifier
-			// won't have any effect without a collection to act on.
-			if (ModifierNode->ModifierClass->GetCollections().IsEmpty())
-			{
-				UE_LOG(LogMovieRenderPipeline, Warning, TEXT("The modifier '%s' specified a collection '%s', but the collection couldn't be found."),
-					*ModifierNode->ModifierName, *ModifierNode->ModifiedCollectionName);
-			}
-			else
-			{
-				RenderLayer->AddModifier(ModifierNode->ModifierClass);
+					RenderLayer->AddModifier(ModifierInstance);
+				}
 			}
 		}
 	}
diff --git a/Engine/Plugins/MovieScene/MovieRenderPipeline/Source/MovieRenderPipelineCore/Private/Graph/MoviePipelineRenderLayerSubsystem.cpp b/Engine/Plugins/MovieScene/MovieRenderPipeline/Source/MovieRenderPipelineCore/Private/Graph/MoviePipelineRenderLayerSubsystem.cpp
--- a/Engine/Plugins/MovieScene/MovieRenderPipeline/Source/MovieRenderPipelineCore/Private/Graph/MoviePipelineRenderLayerSubsystem.cpp
+++ b/Engine/Plugins/MovieScene/MovieRenderPipeline/Source/MovieRenderPipelineCore/Private/Graph/MoviePipelineRenderLayerSubsystem.cpp
@@ -20,11 +20,11 @@
 #include "ActorTreeItem.h"
 #include "ClassViewerFilter.h"
 #include "ClassViewerModule.h"
+#include "Graph/MovieGraphSharedWidgets.h"
 #include "ISceneOutliner.h"
 #include "SceneOutlinerModule.h"
 #include "SceneOutlinerPublicTypes.h"
 #include "SClassViewer.h"
-#include "Widgets/SCompoundWidget.h"
 #endif
 
 #define LOCTEXT_NAMESPACE "MovieGraph"
@@ -32,120 +32,6 @@
 namespace UE::MovieGraph::Private
 {
 #if WITH_EDITOR
-	/**
-	 * A widget which lists items in rows w/ an alternating row color, where each item has an icon and text. A summary row appears at the end of the
-	 * list indicating how many items are in the list.
-	 */
-	template<typename ListType>
-	class SQueryContentsList final : public SCompoundWidget
-	{
-	public:
-		DECLARE_DELEGATE_RetVal_OneParam(const FSlateBrush*, FGetRowIcon, ListType);
-		DECLARE_DELEGATE_RetVal_OneParam(FText, FGetRowText, ListType);
-		
-		SLATE_BEGIN_ARGS(SQueryContentsList<ListType>)
-		{}
-			SLATE_ATTRIBUTE(TArray<ListType>*, DataSource)
-			SLATE_ATTRIBUTE(FText, DataType)
-			SLATE_ATTRIBUTE(FText, DataTypePlural)
-			SLATE_EVENT(FGetRowIcon, OnGetRowIcon);
-			SLATE_EVENT(FGetRowText, OnGetRowText);
-		SLATE_END_ARGS()
-
-		const FSlateBrush* GetRowIcon(const ListType& InListData) const
-		{
-			if (OnGetRowIcon.IsBound())
-			{
-				return OnGetRowIcon.Execute(InListData);
-			}
-			
-			return nullptr;
-		}
-
-		FText GetRowText(const ListType& InListData) const
-		{
-			if (OnGetRowText.IsBound())
-			{
-				return OnGetRowText.Execute(InListData);
-			}
-			
-			return FText();
-		}
-
-		void Construct(const FArguments& InArgs)
-		{
-			DataSource = InArgs._DataSource.Get();
-			DataType = InArgs._DataType.Get();
-			DataTypePlural = InArgs._DataTypePlural.Get();
-			OnGetRowIcon = InArgs._OnGetRowIcon;
-			OnGetRowText = InArgs._OnGetRowText;
-			
-			ChildSlot
-			.HAlign(HAlign_Fill)
-			.VAlign(VAlign_Top)
-			[
-				SNew(SVerticalBox)
-				+ SVerticalBox::Slot()
-				.AutoHeight()
-				[
-					SNew(SListView<ListType>)
-					.ListItemsSource(DataSource)
-					.OnGenerateRow_Lambda([this](const ListType& InListData, const TSharedRef<STableViewBase>& InOwnerTable)
-					{
-						return SNew(STableRow<ListType>, InOwnerTable)
-							.Style(FAppStyle::Get(), "TableView.AlternatingRow")
-							.ShowWires(false)
-							[
-								SNew(SHorizontalBox)
-								+ SHorizontalBox::Slot()
-								.VAlign(VAlign_Center)
-								.Padding(7.f, 5.f, 7.f, 5.f)
-								.AutoWidth()
-								[
-									SNew(SImage)
-									.Image(GetRowIcon(InListData))
-								]
-									
-								+ SHorizontalBox::Slot()
-								.VAlign(VAlign_Center)
-								.HAlign(HAlign_Fill)
-								[
-									SNew(STextBlock)
-									.Font(FAppStyle::GetFontStyle(TEXT("PropertyWindow.NormalFont")))
-									.Text(GetRowText(InListData))
-								]
-							];
-					})
-				]
-
-				+ SVerticalBox::Slot()
-				.AutoHeight()
-				[
-					SNew(SBorder)
-					.BorderImage(FAppStyle::Get().GetBrush("Brushes.Header"))
-					.VAlign(VAlign_Center)
-					.HAlign(HAlign_Fill)
-					.Padding(FMargin(14, 4))
-					[
-						SNew(STextBlock)
-						.Font(FAppStyle::GetFontStyle(TEXT("PropertyWindow.NormalFont")))
-						.Text_Lambda([this]()
-						{
-							return FText::FromString(FString::Printf(TEXT("%i %s"), DataSource->Num(), DataSource->Num() == 1 ? *DataType.ToString() : *DataTypePlural.ToString()));
-						})
-					]
-				]
-			];
-		}
-
-	private:
-		FGetRowIcon OnGetRowIcon;
-		FGetRowText OnGetRowText;
-		FText DataType;
-		FText DataTypePlural;
-		TArray<ListType>* DataSource = nullptr;
-	};
-
 	/**
 	 * A filter that can be used in the class viewer that appears in the Add menu. Filters out specified classes, and optionally filters out classes
 	 * that do not have a specific base class.
@@ -185,7 +71,7 @@ namespace UE::MovieGraph::Private
 
 void UMoviePipelineMaterialModifier::ApplyModifier(const UWorld* World)
 {
-	UMaterialInterface* NewMaterial = MaterialToApply.LoadSynchronous();
+	UMaterialInterface* NewMaterial = Material.LoadSynchronous();
 	if (!NewMaterial)
 	{
 		return;
@@ -613,7 +499,7 @@ TArray<TSharedRef<SWidget>> UMovieGraphConditionGroupQuery_Actor::GetWidgets()
 	}
 
 	Widgets.Add(
-		SNew(UE::MovieGraph::Private::SQueryContentsList<TSharedPtr<TSoftObjectPtr<AActor>>>)
+		SNew(SMovieGraphSimpleList<TSharedPtr<TSoftObjectPtr<AActor>>>)
 			.DataSource(&ListDataSource)
 			.DataType(FText::FromString("Actor"))
 			.DataTypePlural(FText::FromString("Actors"))
@@ -857,7 +743,7 @@ TArray<TSharedRef<SWidget>> UMovieGraphConditionGroupQuery_ActorType::GetWidgets
 	TArray<TSharedRef<SWidget>> Widgets;
 
 	Widgets.Add(
-		SNew(UE::MovieGraph::Private::SQueryContentsList<UClass*>)
+		SNew(SMovieGraphSimpleList<UClass*>)
 			.DataSource(&ActorTypes)
 			.DataType(FText::FromString("Actor Type"))
 			.DataTypePlural(FText::FromString("Actor Types"))
@@ -1036,7 +922,7 @@ TArray<TSharedRef<SWidget>> UMovieGraphConditionGroupQuery_ComponentType::GetWid
 	TArray<TSharedRef<SWidget>> Widgets;
 
 	Widgets.Add(
-		SNew(UE::MovieGraph::Private::SQueryContentsList<UClass*>)
+		SNew(SMovieGraphSimpleList<UClass*>)
 			.DataSource(&ComponentTypes)
 			.DataType(FText::FromString("Component Type"))
 			.DataTypePlural(FText::FromString("Component Types"))
@@ -1444,7 +1330,6 @@ void UMoviePipelineRenderLayerSubsystem::Initialize(FSubsystemCollectionBase& Co
 	VisualizationModifier_HideWorld = NewObject<UMoviePipelineVisibilityModifier>(GetTransientPackage(), NAME_None, RF_Transient);
 	VisualizationModifier_HideWorld->AddCollection(VisualizationEmptyCollection);
 	VisualizationModifier_HideWorld->SetHidden(true);
-	VisualizationModifier_HideWorld->bUseInvertedActors = true;
 
 	// Selectively show collections in the visualization
 	VisualizationModifier_VisibleCollections = NewObject<UMoviePipelineVisibilityModifier>(GetTransientPackage(), NAME_None, RF_Transient);
diff --git a/Engine/Plugins/MovieScene/MovieRenderPipeline/Source/MovieRenderPipelineCore/Private/Graph/Nodes/MovieGraphModifierNode.cpp b/Engine/Plugins/MovieScene/MovieRenderPipeline/Source/MovieRenderPipelineCore/Private/Graph/Nodes/MovieGraphModifierNode.cpp
--- a/Engine/Plugins/MovieScene/MovieRenderPipeline/Source/MovieRenderPipelineCore/Private/Graph/Nodes/MovieGraphModifierNode.cpp
+++ b/Engine/Plugins/MovieScene/MovieRenderPipeline/Source/MovieRenderPipelineCore/Private/Graph/Nodes/MovieGraphModifierNode.cpp
@@ -7,18 +7,113 @@
 
 #define LOCTEXT_NAMESPACE "MovieGraph"
 
+void UMovieGraphMergeableModifierContainer::Merge(const IMovieGraphTraversableObject* InSourceObject)
+{
+	const UMovieGraphMergeableModifierContainer* InSourceModifiers = Cast<UMovieGraphMergeableModifierContainer>(InSourceObject);
+	checkf(InSourceModifiers, TEXT("UMovieGraphMergeableModifierContainer cannot merge with null or an object of another type."));
+
+	for (TObjectPtr<UMoviePipelineCollectionModifier>& ModifierInstance : Modifiers)
+	{
+		// Find the equivalent modifier in the source modifier container
+		for (const TObjectPtr<UMoviePipelineCollectionModifier>& SourceModifier : InSourceModifiers->Modifiers)
+		{
+			if (ModifierInstance && SourceModifier && (ModifierInstance->GetClass() == SourceModifier->GetClass()))
+			{
+				MergeProperties(ModifierInstance, SourceModifier);
+				break;
+			}
+		}
+	}
+}
+
+TArray<TPair<FString, FString>> UMovieGraphMergeableModifierContainer::GetMergedProperties() const
+{
+	TArray<TPair<FString, FString>> MergedProperties;
+
+	// Include properties from modifiers if they have a bOverride_* counterpart
+	for (const TObjectPtr<UMoviePipelineCollectionModifier>& ModifierInstance : Modifiers)
+	{
+		for (TFieldIterator<FProperty> PropertyIterator(ModifierInstance->GetClass()); PropertyIterator; ++PropertyIterator)
+		{
+			const FProperty* ModifierProperty = *PropertyIterator;
+
+			// Exclude the bOverride_* properties
+			if (!UMovieGraphConfig::FindOverridePropertyForRealProperty(ModifierInstance->GetClass(), ModifierProperty))
+			{
+				continue;
+			}
+			
+			// Prefix the property name with the modifier name to distinguish it from properties on other modifiers
+#if WITH_EDITOR
+			const FString ModifierClassName = ModifierInstance->GetClass()->GetDisplayNameText().ToString();
+#else
+			const FString ModifierClassName = ModifierInstance->GetClass()->GetName();
+#endif
+			const FString PropertyName = FString::Printf(TEXT("%s / %s"), *ModifierClassName, *ModifierProperty->GetName());
+			
+			TPair<FString, FString>& NewMergedProperty = MergedProperties.Add_GetRef(TPair<FString, FString>(PropertyName, FString()));
+
+			// Use a string value of the property
+			ModifierProperty->ExportTextItem_InContainer(NewMergedProperty.Value, ModifierInstance, nullptr, nullptr, PPF_None);
+		}
+	}
+	
+	return MergedProperties;
+}
+
+void UMovieGraphMergeableModifierContainer::MergeProperties(const TObjectPtr<UMoviePipelineCollectionModifier>& InDestModifier, const TObjectPtr<UMoviePipelineCollectionModifier>& InSourceModifier)
+{
+	for (TFieldIterator<FProperty> PropertyIterator(InDestModifier->GetClass()); PropertyIterator; ++PropertyIterator)
+	{
+		const FProperty* ModifierProperty = *PropertyIterator;
+		const FBoolProperty* OverrideProp = UMovieGraphConfig::FindOverridePropertyForRealProperty(InSourceModifier->GetClass(), ModifierProperty);
+
+		// Only consider properties which have a bOverride_* counterpart
+		if (!OverrideProp)
+		{
+			continue;
+		}
+
+		// The graph is evaluated in reverse (Outputs to Inputs) so if this property has already been overridden on the dest, ignore the value.
+		if (OverrideProp->GetPropertyValue_InContainer(InDestModifier))
+		{
+			continue;
+		}
+
+		// Set the value if the source has an override
+		if (OverrideProp->GetPropertyValue_InContainer(InSourceModifier))
+		{
+			ModifierProperty->CopyCompleteValue_InContainer(InDestModifier, InSourceModifier);
+
+			// Set the override indicator on the destination
+			OverrideProp->SetPropertyValue_InContainer(InDestModifier, true);
+		}
+	}
+}
+
+UMovieGraphModifierNode::UMovieGraphModifierNode()
+{
+	ModifiersContainer = CreateDefaultSubobject<UMovieGraphMergeableModifierContainer>(TEXT("Modifiers"));
+	
+	// Add some default modifiers
+	if (!HasAnyFlags(RF_ClassDefaultObject))
+	{
+		AddModifier(UMoviePipelineVisibilityModifier::StaticClass());
+		AddModifier(UMoviePipelineMaterialModifier::StaticClass());
+	}
+}
+
 #if WITH_EDITOR
 FText UMovieGraphModifierNode::GetNodeTitle(const bool bGetDescriptive) const
 {
 	static const FText ModifierNodeName = LOCTEXT("NodeName_Modifier", "Modifier");
-	static const FText ModifierNodeDescription = LOCTEXT("NodeDescription_Modifier", "Modifier\n{0} - {1}");
-
-	const FString ModifiedCollectionNameDisp = ModifiedCollectionName.IsEmpty() ? TEXT("MISSING") : ModifiedCollectionName;
+	static const FText ModifierNodeDescription = LOCTEXT("NodeDescription_Modifier", "Modifier\n{0}");
+	
 	const FString ModifierNameDisp = ModifierName.IsEmpty() ? TEXT("MISSING") : ModifierName;
 
 	if (bGetDescriptive)
 	{
-		return FText::Format(ModifierNodeDescription, FText::FromString(ModifierNameDisp), FText::FromString(ModifiedCollectionNameDisp));
+		return FText::Format(ModifierNodeDescription, FText::FromString(ModifierNameDisp));
 	}
 
 	return ModifierNodeName;
@@ -49,14 +144,82 @@ void UMovieGraphModifierNode::PostEditChangeProperty(FPropertyChangedEvent& Prop
 
 	// Broadcast a node-changed delegate so that the node title's UI gets updated.
 	if ((PropertyChangedEvent.GetPropertyName() == GET_MEMBER_NAME_CHECKED(UMovieGraphModifierNode, ModifierName)) ||
-		(PropertyChangedEvent.GetPropertyName() == GET_MEMBER_NAME_CHECKED(UMovieGraphModifierNode, bOverride_ModifierName)) ||
-		(PropertyChangedEvent.GetPropertyName() == GET_MEMBER_NAME_CHECKED(UMovieGraphModifierNode, ModifiedCollectionName)) || 
-		(PropertyChangedEvent.GetPropertyName() == GET_MEMBER_NAME_CHECKED(UMovieGraphModifierNode, bOverride_ModifiedCollectionName)))
+		(PropertyChangedEvent.GetPropertyName() == GET_MEMBER_NAME_CHECKED(UMovieGraphModifierNode, bOverride_ModifierName)))
 	{
 		OnNodeChangedDelegate.Broadcast(this);
 	}
 }
-
 #endif // WITH_EDITOR
 
+UMoviePipelineCollectionModifier* UMovieGraphModifierNode::GetModifier(TSubclassOf<UMoviePipelineCollectionModifier> ModifierType) const
+{
+	const TObjectPtr<UMoviePipelineCollectionModifier>* FoundModifier =
+		ModifiersContainer->Modifiers.FindByPredicate([&ModifierType](const TObjectPtr<UMoviePipelineCollectionModifier>& Modifier)
+        {
+           	return Modifier && (Modifier->GetClass() == ModifierType);
+        });
+
+	return (FoundModifier == nullptr) ? nullptr : *FoundModifier;
+}
+
+const TArray<UMoviePipelineCollectionModifier*>& UMovieGraphModifierNode::GetModifiers() const
+{
+	return ModifiersContainer->Modifiers;
+}
+
+UMoviePipelineCollectionModifier* UMovieGraphModifierNode::AddModifier(TSubclassOf<UMoviePipelineCollectionModifier> ModifierType)
+{
+	const TObjectPtr<UMoviePipelineCollectionModifier>* FoundModifier =
+		ModifiersContainer->Modifiers.FindByPredicate([&ModifierType](const TObjectPtr<UMoviePipelineCollectionModifier>& Modifier)
+		{
+			return Modifier && (Modifier->GetClass() == ModifierType);
+		});
+
+	// Don't add duplicate modifiers
+	if (FoundModifier != nullptr)
+	{
+		return nullptr;
+	}
+
+	if (const UClass* ModifierClass = ModifierType.Get())
+	{
+		UMoviePipelineCollectionModifier* NewModifier = NewObject<UMoviePipelineCollectionModifier>(this, ModifierType, ModifierClass->GetFName());
+		ModifiersContainer->Modifiers.Add(NewModifier);
+
+		return NewModifier;
+	}
+
+	return nullptr;
+}
+
+bool UMovieGraphModifierNode::RemoveModifier(TSubclassOf<UMoviePipelineCollectionModifier> ModifierType)
+{
+	const int32 NumRemoved = ModifiersContainer->Modifiers.RemoveAll([&ModifierType](const TObjectPtr<UMoviePipelineCollectionModifier>& Modifier)
+	{
+		return Modifier && (Modifier->GetClass() == ModifierType);
+	});
+
+	return NumRemoved > 0;
+}
+
+void UMovieGraphModifierNode::AddCollection(const FName& InCollectionName)
+{
+	if (InCollectionName == NAME_None)
+	{
+		return;
+	}
+
+	Collections.AddUnique(InCollectionName);
+}
+
+bool UMovieGraphModifierNode::RemoveCollection(const FName& InCollectionName)
+{
+	return Collections.Remove(InCollectionName) > 0;
+}
+
+const TArray<FName>& UMovieGraphModifierNode::GetCollections() const
+{
+	return Collections;
+}
+
 #undef LOCTEXT_NAMESPACE // "MovieGraph"
\ No newline at end of file
diff --git a/Engine/Plugins/MovieScene/MovieRenderPipeline/Source/MovieRenderPipelineCore/Public/Graph/MovieGraphConfig.h b/Engine/Plugins/MovieScene/MovieRenderPipeline/Source/MovieRenderPipelineCore/Public/Graph/MovieGraphConfig.h
--- a/Engine/Plugins/MovieScene/MovieRenderPipeline/Source/MovieRenderPipelineCore/Public/Graph/MovieGraphConfig.h
+++ b/Engine/Plugins/MovieScene/MovieRenderPipeline/Source/MovieRenderPipelineCore/Public/Graph/MovieGraphConfig.h
@@ -627,6 +627,9 @@ public:
 	/** Determines the name(s) of the branches upstream from FromNode, starting at FromPin. */
 	TArray<FString> GetUpstreamBranchNames(UMovieGraphNode* FromNode, const UMovieGraphPin* FromPin) const;
 
+	/** Get all subgraphs that this graph contains, recursively (ie, subgraphs of subgraphs are included, etc). */
+	void GetAllContainedSubgraphs(TSet<UMovieGraphConfig*>& OutSubgraphs) const;
+
 protected:
 	/** Copies properties in FromNode that are marked for override into ToNode, but only if ToNode doesn't already override that value. */
 	void CopyOverriddenProperties(UMovieGraphNode* FromNode, UMovieGraphNode* ToNode, const FMovieGraphTraversalContext* InContext);
diff --git a/Engine/Plugins/MovieScene/MovieRenderPipeline/Source/MovieRenderPipelineCore/Public/Graph/MovieGraphSharedWidgets.h b/Engine/Plugins/MovieScene/MovieRenderPipeline/Source/MovieRenderPipelineCore/Public/Graph/MovieGraphSharedWidgets.h
new file mode 100644
--- /dev/null
+++ b/Engine/Plugins/MovieScene/MovieRenderPipeline/Source/MovieRenderPipelineCore/Public/Graph/MovieGraphSharedWidgets.h
@@ -0,0 +1,179 @@
+﻿// Copyright Epic Games, Inc. All Rights Reserved.
+
+#pragma once
+
+#include "CoreMinimal.h"
+#include "Framework/Application/SlateApplication.h"
+#include "Widgets/SCompoundWidget.h"
+#include "Widgets/Input/SComboButton.h"
+#include "Widgets/Text/STextBlock.h"
+#include "Widgets/Views/ITableRow.h"
+#include "Widgets/Views/SListView.h"
+
+// NOTE: There should not be widgetry defined in core. To fix this, the condition group queries in the render layer subsystem need to be refactored to
+// expose their custom widgetry in MovieRenderPipelineEditor.
+
+/**
+ * A widget which lists items in rows w/ an alternating row color, where each item has an icon and text. A summary row appears at the end of the
+ * list indicating how many items are in the list.
+ */
+template<typename ListType>
+class SMovieGraphSimpleList final : public SCompoundWidget
+{
+public:
+	DECLARE_DELEGATE_RetVal_OneParam(const FSlateBrush*, FGetRowIcon, ListType);
+	DECLARE_DELEGATE_RetVal_OneParam(FText, FGetRowText, ListType);
+	DECLARE_DELEGATE_OneParam(FOnDelete, ListType);
+	
+	SLATE_BEGIN_ARGS(SMovieGraphSimpleList<ListType>)
+		{}
+		/** The source of data that the list will display. */
+		SLATE_ATTRIBUTE(TArray<ListType>*, DataSource)
+
+		/** The name of the data type that will be shown in the summary row. */
+		SLATE_ATTRIBUTE(FText, DataType)
+
+		/** The plural of the DataType attribute. */
+		SLATE_ATTRIBUTE(FText, DataTypePlural)
+
+		/** Gets the icon for a row in the list. */
+		SLATE_EVENT(FGetRowIcon, OnGetRowIcon);
+
+		/** Gets the text for a row in the list. */
+		SLATE_EVENT(FGetRowText, OnGetRowText);
+
+		/** Invoked when a delete operation is performed. */
+		SLATE_EVENT(FOnDelete, OnDelete)
+	SLATE_END_ARGS()
+	
+	void Construct(const FArguments& InArgs)
+	{
+		DataSource = InArgs._DataSource.Get();
+		DataType = InArgs._DataType.Get();
+		DataTypePlural = InArgs._DataTypePlural.Get();
+		OnGetRowIcon = InArgs._OnGetRowIcon;
+		OnGetRowText = InArgs._OnGetRowText;
+		OnDelete = InArgs._OnDelete;
+
+		ChildSlot
+		.HAlign(HAlign_Fill)
+		.VAlign(VAlign_Top)
+		[
+			SNew(SVerticalBox)
+			+ SVerticalBox::Slot()
+			.AutoHeight()
+			[
+				SAssignNew(ListView, SListView<ListType>)
+				.ListItemsSource(DataSource)
+				.OnKeyDownHandler(this, &SMovieGraphSimpleList<ListType>::HandleDelete)
+				.OnGenerateRow(this, &SMovieGraphSimpleList<ListType>::GenerateRow)
+			]
+
+			+ SVerticalBox::Slot()
+			.AutoHeight()
+			[
+				SNew(SBorder)
+				.BorderImage(FAppStyle::Get().GetBrush("Brushes.Header"))
+				.VAlign(VAlign_Center)
+				.HAlign(HAlign_Fill)
+				.Padding(FMargin(14, 4))
+				[
+					SNew(STextBlock)
+					.Font(FAppStyle::GetFontStyle(TEXT("PropertyWindow.NormalFont")))
+					.Text_Lambda([this]()
+					{
+						return FText::FromString(FString::Printf(TEXT("%i %s"), DataSource->Num(), DataSource->Num() == 1 ? *DataType.ToString() : *DataTypePlural.ToString()));
+					})
+				]
+			]
+		];
+	}
+
+	/** Refreshes the contents of the list. */
+	void Refresh()
+	{
+		if (ListView)
+		{
+			ListView->RequestListRefresh();
+		}
+	}
+
+private:
+	/** Gets the row icon associated with the given list data. */
+	const FSlateBrush* GetRowIcon(const ListType& InListData) const
+	{
+		if (OnGetRowIcon.IsBound())
+		{
+			return OnGetRowIcon.Execute(InListData);
+		}
+	
+		return nullptr;
+	}
+
+	/** Gets the row text associated with the given list data. */
+	FText GetRowText(const ListType& InListData) const
+	{
+		if (OnGetRowText.IsBound())
+		{
+			return OnGetRowText.Execute(InListData);
+		}
+	
+		return FText();
+	}
+
+	/** Handles the delete operation. */
+	FReply HandleDelete(const FGeometry& InGeometry, const FKeyEvent& InKeyEvent) const
+	{
+		TArray<ListType> SelectedItems;
+		ListView->GetSelectedItems(SelectedItems);
+				
+		if ((InKeyEvent.GetKey() == EKeys::Delete) && (SelectedItems.Num() == 1))
+		{
+			if (OnDelete.IsBound())
+			{
+				OnDelete.Execute(SelectedItems[0]);
+				return FReply::Handled();
+			}
+		}
+
+		return FReply::Unhandled();
+	}
+
+	/** Generates a row in the list for the specified data. */
+	TSharedRef<ITableRow> GenerateRow(ListType InListData, const TSharedRef<STableViewBase>& InOwnerTable) const
+	{
+		return
+			SNew(STableRow<ListType>, InOwnerTable)
+			.Style(FAppStyle::Get(), "TableView.AlternatingRow")
+			.ShowWires(false)
+			[
+				SNew(SHorizontalBox)
+				+ SHorizontalBox::Slot()
+				.VAlign(VAlign_Center)
+				.Padding(7.f, 5.f, 7.f, 5.f)
+				.AutoWidth()
+				[
+					SNew(SImage)
+					.Image(GetRowIcon(InListData))
+				]
+				
+				+ SHorizontalBox::Slot()
+				.VAlign(VAlign_Center)
+				.HAlign(HAlign_Fill)
+				[
+					SNew(STextBlock)
+					.Font(FAppStyle::GetFontStyle(TEXT("PropertyWindow.NormalFont")))
+					.Text(GetRowText(InListData))
+				]
+			];
+	}
+
+private:
+	TSharedPtr<SListView<ListType>> ListView;
+	FGetRowIcon OnGetRowIcon;
+	FGetRowText OnGetRowText;
+	FOnDelete OnDelete;
+	FText DataType;
+	FText DataTypePlural;
+	TArray<ListType>* DataSource = nullptr;
+};
\ No newline at end of file
diff --git a/Engine/Plugins/MovieScene/MovieRenderPipeline/Source/MovieRenderPipelineCore/Public/Graph/MoviePipelineRenderLayerSubsystem.h b/Engine/Plugins/MovieScene/MovieRenderPipeline/Source/MovieRenderPipelineCore/Public/Graph/MoviePipelineRenderLayerSubsystem.h
--- a/Engine/Plugins/MovieScene/MovieRenderPipeline/Source/MovieRenderPipelineCore/Public/Graph/MoviePipelineRenderLayerSubsystem.h
+++ b/Engine/Plugins/MovieScene/MovieRenderPipeline/Source/MovieRenderPipelineCore/Public/Graph/MoviePipelineRenderLayerSubsystem.h
@@ -475,7 +475,7 @@ private:
  * Base class for providing actor modification functionality via collections.
  */
 UCLASS(Abstract)
-class UMoviePipelineCollectionModifier : public UObject
+class MOVIERENDERPIPELINECORE_API UMoviePipelineCollectionModifier : public UObject
 {
 	GENERATED_BODY()
 
@@ -494,11 +494,6 @@ public:
 	virtual void ApplyModifier(const UWorld* World) PURE_VIRTUAL(UMoviePipelineCollectionModifier::ApplyModifier, );
 	virtual void UndoModifier() PURE_VIRTUAL(UMoviePipelineCollectionModifier::UndoModifier, );
 
-public:
-	/** Whether an inverted collection of actors should be used instead. */
-	UPROPERTY(BlueprintReadWrite, EditAnywhere, Category = "Settings")
-	bool bUseInvertedActors = false;
-
 protected:
 	/** The collections which this modifier will operate on. */
 	UPROPERTY()
@@ -508,14 +503,14 @@ protected:
 /**
  * Modifies actor materials.
  */
-UCLASS(BlueprintType)
-class UMoviePipelineMaterialModifier : public UMoviePipelineCollectionModifier
+UCLASS(BlueprintType, DisplayName="Material")
+class MOVIERENDERPIPELINECORE_API UMoviePipelineMaterialModifier : public UMoviePipelineCollectionModifier
 {
 	GENERATED_BODY()
 
 public:
 	UFUNCTION(BlueprintCallable, Category = "Settings")
-	void SetMaterial(TSoftObjectPtr<UMaterialInterface> InMaterial) { MaterialToApply = InMaterial; }
+	void SetMaterial(TSoftObjectPtr<UMaterialInterface> InMaterial) { Material = InMaterial; }
 
 	UFUNCTION(BlueprintCallable, Category = "Settings")
 	virtual void ApplyModifier(const UWorld* World) override;
@@ -531,15 +526,18 @@ private:
 	FComponentToMaterialMap ModifiedComponents;
 
 public:
-	UPROPERTY(BlueprintReadWrite, EditAnywhere, Category = "Settings")
-	TSoftObjectPtr<UMaterialInterface> MaterialToApply;
+	UPROPERTY(EditAnywhere, BlueprintReadWrite, Category = Overrides, meta = (InlineEditConditionToggle))
+	uint8 bOverride_Material : 1;
+	
+	UPROPERTY(BlueprintReadWrite, EditAnywhere, Category = "Settings", meta = (EditCondition = "bOverride_Material"))
+	TSoftObjectPtr<UMaterialInterface> Material;
 };
 
 /**
  * Modifies actor visibility.
  */
-UCLASS(BlueprintType)
-class UMoviePipelineVisibilityModifier : public UMoviePipelineCollectionModifier
+UCLASS(BlueprintType, DisplayName="Visibility")
+class MOVIERENDERPIPELINECORE_API UMoviePipelineVisibilityModifier : public UMoviePipelineCollectionModifier
 {
 	GENERATED_BODY()
 
@@ -578,30 +576,45 @@ private:
 	TArray<FActorVisibilityState> ModifiedActors;
 
 public:
+	UPROPERTY(EditAnywhere, BlueprintReadWrite, Category = Overrides, meta = (InlineEditConditionToggle))
+	uint8 bOverride_bIsHidden : 1;
+
+	UPROPERTY(EditAnywhere, BlueprintReadWrite, Category = Overrides, meta = (InlineEditConditionToggle))
+	uint8 bOverride_bCastsShadows : 1;
+
+	UPROPERTY(EditAnywhere, BlueprintReadWrite, Category = Overrides, meta = (InlineEditConditionToggle))
+	uint8 bOverride_bCastShadowWhileHidden : 1;
+
+	UPROPERTY(EditAnywhere, BlueprintReadWrite, Category = Overrides, meta = (InlineEditConditionToggle))
+	uint8 bOverride_bAffectIndirectLightingWhileHidden : 1;
+
+	UPROPERTY(EditAnywhere, BlueprintReadWrite, Category = Overrides, meta = (InlineEditConditionToggle))
+	uint8 bOverride_bHoldout : 1;
+	
 	/**
 	 * If true, the actor will not be visible and will not contribute to any secondary effects (shadows, indirect
 	 * lighting) unless their respective flags are set below.
 	 */
-	UPROPERTY(BlueprintReadWrite, EditAnywhere, Category = "Settings")
+	UPROPERTY(BlueprintReadWrite, EditAnywhere, Category = "Settings", meta = (EditCondition = "bOverride_bIsHidden"))
 	uint8 bIsHidden : 1;
 	
 	/** If true, the primitive will cast shadows. */
-	UPROPERTY(BlueprintReadWrite, EditAnywhere, Category = "Primitive")
+	UPROPERTY(BlueprintReadWrite, EditAnywhere, Category = "Settings", meta = (EditCondition = "bOverride_bCastsShadows"))
 	uint8 bCastsShadows : 1;
 
 	/** If true, the primitive will cast shadows even if it is hidden. */
-	UPROPERTY(BlueprintReadWrite, EditAnywhere, Category = "Primitive")
+	UPROPERTY(BlueprintReadWrite, EditAnywhere, Category = "Settings", meta = (EditCondition = "bOverride_bCastShadowWhileHidden"))
 	uint8 bCastShadowWhileHidden : 1;
 
 	/** Controls whether the primitive should affect indirect lighting when hidden. */
-	UPROPERTY(BlueprintReadWrite, EditAnywhere, Category = "Primitive")
+	UPROPERTY(BlueprintReadWrite, EditAnywhere, Category = "Settings", meta = (EditCondition = "bOverride_bAffectIndirectLightingWhileHidden"))
 	uint8 bAffectIndirectLightingWhileHidden : 1;
 
 	/**
 	 * If true, the primitive will render black with an alpha of 0, but all secondary effects (shadows, reflections,
 	 * indirect lighting) remain. This feature is currently only implemented in the Path Tracer.
 	 */
-	UPROPERTY(BlueprintReadWrite, EditAnywhere, Category = "Primitive|Path Tracing")
+	UPROPERTY(BlueprintReadWrite, EditAnywhere, Category = "Settings", meta = (EditCondition = "bOverride_bHoldout"))
 	uint8 bHoldout : 1;
 };
 
diff --git a/Engine/Plugins/MovieScene/MovieRenderPipeline/Source/MovieRenderPipelineCore/Public/Graph/Nodes/MovieGraphModifierNode.h b/Engine/Plugins/MovieScene/MovieRenderPipeline/Source/MovieRenderPipelineCore/Public/Graph/Nodes/MovieGraphModifierNode.h
--- a/Engine/Plugins/MovieScene/MovieRenderPipeline/Source/MovieRenderPipelineCore/Public/Graph/Nodes/MovieGraphModifierNode.h
+++ b/Engine/Plugins/MovieScene/MovieRenderPipeline/Source/MovieRenderPipelineCore/Public/Graph/Nodes/MovieGraphModifierNode.h
@@ -2,12 +2,31 @@
 
 #pragma once
 
+#include "Graph/MovieGraphConfig.h"
 #include "Graph/MovieGraphNode.h"
+#include "Graph/MoviePipelineRenderLayerSubsystem.h"
 
 #include "MovieGraphModifierNode.generated.h"
 
-// Forward Declare
-class UMoviePipelineCollectionModifier;
+/** A container which allows an array of modifiers to be merged correctly. */
+UCLASS()
+class UMovieGraphMergeableModifierContainer final : public UObject, public IMovieGraphTraversableObject
+{
+	GENERATED_BODY()
+
+public:
+	// IMovieGraphTraversableObject interface
+	virtual void Merge(const IMovieGraphTraversableObject* InSourceObject) override;
+	virtual TArray<TPair<FString, FString>> GetMergedProperties() const override;
+	// ~IMovieGraphTraversableObject interface
+
+private:
+	void MergeProperties(const TObjectPtr<UMoviePipelineCollectionModifier>& InDestModifier, const TObjectPtr<UMoviePipelineCollectionModifier>& InSourceModifier);
+
+public:
+	UPROPERTY()
+	TArray<TObjectPtr<UMoviePipelineCollectionModifier>> Modifiers;
+};
 
 /** 
 * A collection node specifies an interface for doing dynamic scene queries for actors in the world. Collections work in tandem with
@@ -19,6 +38,8 @@ class MOVIERENDERPIPELINECORE_API UMovieGraphModifierNode : public UMovieGraphSe
 	GENERATED_BODY()
 
 public:
+	UMovieGraphModifierNode();
+
 #if WITH_EDITOR
 	virtual FText GetNodeTitle(const bool bGetDescriptive = false) const override;
 	virtual FText GetMenuCategory() const override;
@@ -32,25 +53,57 @@ public:
 
 	virtual FString GetNodeInstanceName() const override { return ModifierName; }
 
+	/** Gets the modifier of the specified type, or nullptr if one does not exist on this node. */
+	UFUNCTION(BlueprintCallable, Category = "Modifiers")
+	UMoviePipelineCollectionModifier* GetModifier(TSubclassOf<UMoviePipelineCollectionModifier> ModifierType) const;
+
+	/** Gets all modifiers currently added to the node. */
+	UFUNCTION(BlueprintCallable, Category = "Modifiers")
+	const TArray<UMoviePipelineCollectionModifier*>& GetModifiers() const;
+
+	/**
+	 * Adds a new modifier of the specified type. Returns a pointer to the new modifier, or nullptr if a modifier of the specified type already
+	 * exists on this node (only one modifier of each type can be added to the node).
+	 */
+	UFUNCTION(BlueprintCallable, Category = "Modifiers")
+	UMoviePipelineCollectionModifier* AddModifier(TSubclassOf<UMoviePipelineCollectionModifier> ModifierType);
+
+	/** Removes the modifier of the specified type. Returns true on success, or false if a modifier of the specified type does not exist on the node. */
+	UFUNCTION(BlueprintCallable, Category = "Modifiers")
+	bool RemoveModifier(TSubclassOf<UMoviePipelineCollectionModifier> ModifierType);
+
+	/** Add a collection identified by the given name which will be affected by the modifiers on this node. */
+	UFUNCTION(BlueprintCallable, Category = "Modifiers")
+	void AddCollection(const FName& InCollectionName);
+
+	/** Remove a collection identified by the given name. Returns true if the collection was found and removed successfully, else false. */
+	UFUNCTION(BlueprintCallable, Category = "Modifiers")
+	bool RemoveCollection(const FName& InCollectionName);
+
+	/** Gets all collections that will be affected by the modifiers on this node. */
+	UFUNCTION(BlueprintCallable, Category = "Modifiers")
+	const TArray<FName>& GetCollections() const;
+
 public:
-	UPROPERTY(EditAnywhere, BlueprintReadWrite, Category = Overrides, meta = (InlineEditConditionToggle))
-	uint8 bOverride_ModifierName : 1;
-	
-	UPROPERTY(EditAnywhere, BlueprintReadWrite, Category = Overrides, meta = (InlineEditConditionToggle))
-	uint8 bOverride_ModifiedCollectionName : 1;
-		
-	UPROPERTY(EditAnywhere, BlueprintReadWrite, Category = Overrides, meta = (InlineEditConditionToggle))
-	uint8 bOverride_ModifierClass : 1;
+	UPROPERTY()
+	uint8 bOverride_ModifierName : 1 = 1;	// Always merge the modifier name, no need for the user to do this explicitly
 
 	/** The name of this modifier. */
-	UPROPERTY(EditAnywhere, BlueprintReadWrite, Category = "General", meta=(EditCondition="bOverride_ModifierName"))
+	UPROPERTY(EditAnywhere, BlueprintReadWrite, Category = "Modifier")
 	FString ModifierName;
 
-	/** The name of the collection being modified. */
-	UPROPERTY(EditAnywhere, BlueprintReadWrite, Category = "General", meta=(EditCondition="bOverride_ModifiedCollectionName")) 
-	FString ModifiedCollectionName;
+private:
+	UPROPERTY()
+	uint8 bOverride_Collections : 1 = 1;
+
+	UPROPERTY()
+	uint8 bOverride_ModifiersContainer : 1 = 1;
+	
+	/** The names of collections being modified. */
+	UPROPERTY()
+	TArray<FName> Collections;
 	
-	/** The modifier this node should run. */
-	UPROPERTY(BlueprintReadWrite, EditAnywhere, Instanced, Category = "General", meta=(EditCondition="bOverride_ModifierClass", EditInline))
-	TObjectPtr<UMoviePipelineCollectionModifier> ModifierClass;
+	/** The modifiers this node should run. */
+	UPROPERTY(meta=(DisplayName="Modifiers"))
+	TObjectPtr<UMovieGraphMergeableModifierContainer> ModifiersContainer;
 };
\ No newline at end of file
diff --git a/Engine/Plugins/MovieScene/MovieRenderPipeline/Source/MovieRenderPipelineEditor/Private/Customizations/Graph/MovieGraphModifiersCustomization.cpp b/Engine/Plugins/MovieScene/MovieRenderPipeline/Source/MovieRenderPipelineEditor/Private/Customizations/Graph/MovieGraphModifiersCustomization.cpp
new file mode 100644
--- /dev/null
+++ b/Engine/Plugins/MovieScene/MovieRenderPipeline/Source/MovieRenderPipelineEditor/Private/Customizations/Graph/MovieGraphModifiersCustomization.cpp
@@ -0,0 +1,341 @@
+// Copyright Epic Games, Inc. All Rights Reserved.
+
+#include "MovieGraphModifiersCustomization.h"
+
+#include "DetailCategoryBuilder.h"
+#include "DetailLayoutBuilder.h"
+#include "DetailWidgetRow.h"
+#include "Graph/MovieGraphConfig.h"
+#include "Graph/MovieGraphSharedWidgets.h"
+#include "Graph/MoviePipelineRenderLayerSubsystem.h"
+#include "Graph/Nodes/MovieGraphCollectionNode.h"
+#include "Graph/Nodes/MovieGraphModifierNode.h"
+#include "Widgets/Images/SLayeredImage.h"
+#include "Widgets/Layout/SWidgetSwitcher.h"
+
+#define LOCTEXT_NAMESPACE "MovieGraphModifiersCustomization"
+
+/** Discovers collections that are pickable from a specific graph, and presents them in a list. */
+class SMovieGraphCollectionPicker final : public SCompoundWidget
+{
+public:
+	DECLARE_DELEGATE_OneParam(FOnCollectionPicked, FName);
+	DECLARE_DELEGATE_RetVal_OneParam(bool, FOnFilter, FName);
+
+	SLATE_BEGIN_ARGS(SMovieGraphCollectionPicker)
+		{}
+		/** The graph to begin discovering collections from. */
+		SLATE_ATTRIBUTE(UMovieGraphConfig*, Graph)
+
+		/** Called when a collection is picked in the list. */
+		SLATE_EVENT(FOnCollectionPicked, OnCollectionPicked);
+
+		/** Optional filter that can prevent discovered collections from showing up in the list. */
+		SLATE_EVENT(FOnFilter, OnFilter);
+	SLATE_END_ARGS()
+
+	void Construct(const FArguments& InArgs)
+	{
+		CurrentGraph = InArgs._Graph.Get();
+		OnCollectionPicked = InArgs._OnCollectionPicked;
+		OnFilter = InArgs._OnFilter;
+
+		UpdateDataSource();
+		
+		ChildSlot
+		.HAlign(HAlign_Fill)
+		.VAlign(VAlign_Fill)
+		[
+			SNew(SVerticalBox)
+			+ SVerticalBox::Slot()
+			.Padding(5.f)
+			.AutoHeight()
+			[
+				SNew(STextBlock)
+				.Text(LOCTEXT("PickCollectionHelpText", "Pick a Collection"))
+				.Font(FAppStyle::GetFontStyle("PropertyWindow.NormalFont"))
+			]
+			
+			+ SVerticalBox::Slot()
+			.VAlign(VAlign_Fill)
+			[
+				SNew(SWidgetSwitcher)
+				.WidgetIndex(DataSource.IsEmpty() ? 0 : 1)
+
+				+ SWidgetSwitcher::Slot()
+				.VAlign(VAlign_Center)
+				.HAlign(HAlign_Center)
+				[
+					SNew(STextBlock)
+					.Text(LOCTEXT("NoCollectionsFoundWarning", "No collections found."))
+				]
+
+				+ SWidgetSwitcher::Slot()
+				[
+					SNew(SListView<FName>)
+					.ListItemsSource(&DataSource)
+					.SelectionMode(ESelectionMode::Single)
+					.OnSelectionChanged(this, &SMovieGraphCollectionPicker::OnCollectionSelected)
+					.OnGenerateRow(this, &SMovieGraphCollectionPicker::GenerateRow)
+				]
+			]
+		];
+	}
+
+private:
+	/** Discovers all available collections, and refreshes the data the list is displaying. */
+	TArray<UMovieGraphCollectionNode*> UpdateDataSource()
+	{
+		TArray<UMovieGraphCollectionNode*> CollectionNodes;
+		
+		if (!CurrentGraph)
+		{
+			return CollectionNodes;
+		}
+
+		TSet<UMovieGraphConfig*> Subgraphs;
+		CurrentGraph->GetAllContainedSubgraphs(Subgraphs);
+
+		TArray<UMovieGraphConfig*> AllGraphs = { CurrentGraph };
+		for (UMovieGraphConfig* Subgraph : Subgraphs)
+		{
+			AllGraphs.Add(Subgraph);
+		}
+
+		for (const UMovieGraphConfig* Graph : AllGraphs)
+		{
+			for (const TObjectPtr<UMovieGraphNode>& Node : Graph->GetNodes())
+			{
+				if (const TObjectPtr<UMovieGraphCollectionNode> CollectionNode = Cast<UMovieGraphCollectionNode>(Node))
+				{
+					const FName CollectionName = FName(CollectionNode->Collection->GetCollectionName());
+					bool bIncludeCollection = true;
+
+					if (OnFilter.IsBound())
+					{
+						bIncludeCollection = OnFilter.Execute(CollectionName);
+					}
+
+					if (bIncludeCollection)
+					{
+						DataSource.Add(CollectionName);
+					}
+				}
+			}
+		}
+
+		return CollectionNodes;
+	}
+
+	/** Handles a collection selected event. */
+	void OnCollectionSelected(const FName CollectionName, ESelectInfo::Type Type) const
+	{
+		if (OnCollectionPicked.IsBound())
+		{
+			OnCollectionPicked.Execute(CollectionName);
+		}
+
+		FSlateApplication::Get().DismissAllMenus();
+	}
+
+	/** Generates a row which displays a single collection. */
+	TSharedRef<ITableRow> GenerateRow(const FName CollectionName, const TSharedRef<STableViewBase>& InOwnerTable) const
+	{
+		return
+			SNew(STableRow<FName>, InOwnerTable)
+			.Style(FAppStyle::Get(), "TableView.AlternatingRow")
+			.ShowWires(false)
+			[
+				SNew(SHorizontalBox)
+				+ SHorizontalBox::Slot()
+				.VAlign(VAlign_Center)
+				.Padding(7.f, 5.f, 7.f, 5.f)
+				.AutoWidth()
+				[
+					SNew(SImage)
+					.Image(FAppStyle::GetBrush("Icons.FilledCircle"))
+				]
+					
+				+ SHorizontalBox::Slot()
+				.VAlign(VAlign_Center)
+				.HAlign(HAlign_Fill)
+				[
+					SNew(STextBlock)
+					.Font(FAppStyle::GetFontStyle(TEXT("PropertyWindow.NormalFont")))
+					.Text(FText::FromName(CollectionName))
+				]
+			];
+	}
+
+private:
+	/** The data source for the list view widget. Names of the collection nodes that can be picked. */
+	TArray<FName> DataSource;
+
+	/** The current graph being viewed. The data source will be populated from this graph. */
+	UMovieGraphConfig* CurrentGraph = nullptr;
+	
+	FOnCollectionPicked OnCollectionPicked;
+	FOnFilter OnFilter;
+};
+
+TSharedRef<IDetailCustomization> FMovieGraphModifiersCustomization::MakeInstance()
+{
+	return MakeShared<FMovieGraphModifiersCustomization>();
+}
+
+void FMovieGraphModifiersCustomization::CustomizeDetails(IDetailLayoutBuilder& InDetailBuilder)
+{
+	TArray<TWeakObjectPtr<UMovieGraphModifierNode>> ModifierNodes =
+		InDetailBuilder.GetObjectsOfTypeBeingCustomized<UMovieGraphModifierNode>();
+	if (ModifierNodes.Num() != 1)
+	{
+		// Showing more than one modifier node is not supported
+		return;
+	}
+
+	const TWeakObjectPtr<UMovieGraphModifierNode> ModifierNode = ModifierNodes[0];
+	if (!ModifierNode.IsValid())
+	{
+		return;
+	}
+
+	// Update the data source
+	ListDataSource = ModifierNode->GetCollections();
+	
+	// Generate a (multi-layered) icon for the "Add" menu
+	const TSharedRef<SLayeredImage> AddIcon =
+		SNew(SLayeredImage)
+		.ColorAndOpacity(FSlateColor::UseForeground())
+		.Image(FAppStyle::GetBrush("LevelEditor.OpenAddContent.Background"));
+	AddIcon->AddLayer(FAppStyle::GetBrush("LevelEditor.OpenAddContent.Overlay"));
+
+	// Replace the "Collection" category row with a custom whole-row widget which includes an add-collection button
+	IDetailCategoryBuilder& CollectionCategory = InDetailBuilder.EditCategory(FName("Collection"), FText::GetEmpty(), ECategoryPriority::Uncommon);
+	CollectionCategory.HeaderContent
+	(
+		SNew(SHorizontalBox)
+		+ SHorizontalBox::Slot()
+		.HAlign(HAlign_Left)
+		.VAlign(VAlign_Center)
+		.AutoWidth()
+		[
+			SNew(STextBlock)
+			.Text(FText::FromString("Collections"))
+			.Font(FAppStyle::Get().GetFontStyle("DetailsView.CategoryFontStyle"))
+		]
+
+		+ SHorizontalBox::Slot()
+		.HAlign(HAlign_Left)
+		.VAlign(VAlign_Center)
+		.Padding(5.f, 0, 0, 0)
+		[
+			SNew(SComboButton)
+			.ToolTipText(LOCTEXT("AddCollectionToModifierTooltip", "Add a collection that will be affected by the configured modifiers."))
+			.ComboButtonStyle(FAppStyle::Get(), "SimpleComboButton")
+			.ContentPadding(0)
+			.HasDownArrow(false)
+			.OnGetMenuContent_Lambda([ModifierNode, this]()
+			{
+				return
+					SNew(SBox)
+					.WidthOverride(200.f)
+					.HeightOverride(200.f)
+					[
+						SNew(SMovieGraphCollectionPicker)
+						.Graph(ModifierNode->GetTypedOuter<UMovieGraphConfig>())
+						.OnFilter_Lambda([ModifierNode](const FName CollectionName)
+						{
+							if (ModifierNode.IsValid())
+							{
+								return !ModifierNode.Get()->GetCollections().Contains(CollectionName);
+							}
+
+							return false;
+						})
+						.OnCollectionPicked_Lambda([ModifierNode, this](const FName PickedCollectionName)
+						{
+							if (ModifierNode.IsValid())
+							{
+								ModifierNode->AddCollection(PickedCollectionName);
+								ListDataSource = ModifierNode->GetCollections();
+								CollectionsList->Refresh();
+							}
+						})
+					];
+			})
+			.ButtonContent()
+			[
+				SNew(SBox)
+				.WidthOverride(16)
+				.HeightOverride(16)
+				[
+					AddIcon
+				]
+			]
+		]
+	, /* bWholeRowContent */ true);
+
+	// Add a collections browser
+	CollectionCategory.AddCustomRow(FText::GetEmpty())
+	.WholeRowWidget
+	[
+		SAssignNew(CollectionsList, SMovieGraphSimpleList<FName>)
+		.DataSource(&ListDataSource)
+		.DataType(FText::FromString("Collection"))
+		.DataTypePlural(FText::FromString("Collections"))
+		.OnDelete_Lambda([this, ModifierNode](const FName DeletedCollectionName)
+		{
+			if (ModifierNode.IsValid())
+			{
+				ModifierNode.Get()->RemoveCollection(DeletedCollectionName);
+				ListDataSource = ModifierNode->GetCollections();
+				CollectionsList->Refresh();
+			}
+		})
+		.OnGetRowIcon_Static(&GetCollectionRowIcon)
+		.OnGetRowText_Static(&GetCollectionRowText)
+	];
+
+	// For all modifiers added to the node, add a category for each, and add each modifier's EditAnywhere properties to the category
+	for (UMoviePipelineCollectionModifier* Modifier : ModifierNode->GetModifiers())
+	{
+		if (!Modifier)
+		{
+			continue;
+		}
+		
+		const UClass* ModifierClass = Modifier->GetClass();
+		const FText DisplayName = ModifierClass->GetDisplayNameText();
+
+		// Add category as "Uncommon" to display after the general modifier properties
+		IDetailCategoryBuilder& Category = InDetailBuilder.EditCategory(FName(DisplayName.ToString()), DisplayName, ECategoryPriority::Uncommon);
+
+		for (TFieldIterator<FProperty> PropertyIterator(ModifierClass); PropertyIterator; ++PropertyIterator)
+		{
+			// Add any EditAnywhere properties, but skip the bOverride_* properties.
+			const FProperty* ModifierProperty = *PropertyIterator;
+			if (ModifierProperty && ModifierProperty->HasAnyPropertyFlags(CPF_Edit) && !ModifierProperty->HasMetaData(TEXT("InlineEditConditionToggle")))
+			{
+				Category.AddExternalObjectProperty({Modifier}, ModifierProperty->GetFName());
+			}
+		}
+	}
+}
+
+void FMovieGraphModifiersCustomization::CustomizeDetails(const TSharedPtr<IDetailLayoutBuilder>& InDetailBuilder)
+{
+	DetailBuilder = InDetailBuilder;
+	CustomizeDetails(*InDetailBuilder);
+}
+
+const FSlateBrush* FMovieGraphModifiersCustomization::GetCollectionRowIcon(const FName CollectionName)
+{
+	return FAppStyle::GetBrush("Icons.FilledCircle");
+}
+
+FText FMovieGraphModifiersCustomization::GetCollectionRowText(const FName CollectionName)
+{
+	return FText::FromName(CollectionName);
+}
+
+#undef LOCTEXT_NAMESPACE
\ No newline at end of file
diff --git a/Engine/Plugins/MovieScene/MovieRenderPipeline/Source/MovieRenderPipelineEditor/Private/Customizations/Graph/MovieGraphModifiersCustomization.h b/Engine/Plugins/MovieScene/MovieRenderPipeline/Source/MovieRenderPipelineEditor/Private/Customizations/Graph/MovieGraphModifiersCustomization.h
new file mode 100644
--- /dev/null
+++ b/Engine/Plugins/MovieScene/MovieRenderPipeline/Source/MovieRenderPipelineEditor/Private/Customizations/Graph/MovieGraphModifiersCustomization.h
@@ -0,0 +1,35 @@
+﻿// Copyright Epic Games, Inc. All Rights Reserved.
+
+#pragma once
+
+#include "IDetailCustomization.h"
+#include "Graph/MovieGraphSharedWidgets.h"
+
+class UMovieGraphCollectionNode;
+
+/** Customize how the Modifier node appears in the details panel. */
+class FMovieGraphModifiersCustomization final : public IDetailCustomization
+{
+public:
+	static TSharedRef<IDetailCustomization> MakeInstance();
+
+protected:
+	//~ Begin IDetailCustomization interface
+	virtual void CustomizeDetails(IDetailLayoutBuilder& InDetailBuilder) override;
+	virtual void CustomizeDetails(const TSharedPtr<IDetailLayoutBuilder>& InDetailBuilder) override;
+	//~ End IDetailCustomization interface
+
+private:
+	static const FSlateBrush* GetCollectionRowIcon(const FName CollectionName);
+	static FText GetCollectionRowText(const FName CollectionName);
+
+private:
+	/** The details builder associated with the customization. */
+	TWeakPtr<IDetailLayoutBuilder> DetailBuilder;
+
+	/** The data source for the list view. Modifiers don't expose a direct pointer to the array of collections, which is required by the list view. */
+	TArray<FName> ListDataSource;
+
+	/** Displays the collections which have been chosen. */
+	TSharedPtr<SMovieGraphSimpleList<FName>> CollectionsList;
+};
\ No newline at end of file
diff --git a/Engine/Plugins/MovieScene/MovieRenderPipeline/Source/MovieRenderPipelineEditor/Private/Graph/MovieGraphAssetToolkit.cpp b/Engine/Plugins/MovieScene/MovieRenderPipeline/Source/MovieRenderPipelineEditor/Private/Graph/MovieGraphAssetToolkit.cpp
--- a/Engine/Plugins/MovieScene/MovieRenderPipeline/Source/MovieRenderPipelineEditor/Private/Graph/MovieGraphAssetToolkit.cpp
+++ b/Engine/Plugins/MovieScene/MovieRenderPipeline/Source/MovieRenderPipelineEditor/Private/Graph/MovieGraphAssetToolkit.cpp
@@ -4,6 +4,7 @@
 
 #include "Customizations/Graph/MovieGraphCollectionsCustomization.h"
 #include "Customizations/Graph/MovieGraphMemberCustomization.h"
+#include "Customizations/Graph/MovieGraphModifiersCustomization.h"
 #include "Customizations/Graph/MovieGraphNamedResolutionCustomization.h"
 #include "Customizations/Graph/MovieGraphNodeCustomization.h"
 #include "Customizations/Graph/MovieGraphSelectNodeCustomization.h"
@@ -13,6 +14,7 @@
 
 #include "Graph/MovieGraphConfig.h"
 #include "Graph/Nodes/MovieGraphCollectionNode.h"
+#include "Graph/Nodes/MovieGraphModifierNode.h"
 #include "Graph/Nodes/MovieGraphOutputSettingNode.h"
 #include "Graph/Nodes/MovieGraphSelectNode.h"
 #include "MovieEdGraphNode.h"
@@ -261,6 +263,10 @@ TSharedRef<SDockTab> FMovieGraphAssetToolkit::SpawnTab_RenderGraphDetails(const
 	SelectedGraphObjectsDetailsWidget->RegisterInstancedCustomPropertyLayout(
 		UMovieGraphCollectionNode::StaticClass(),
 		FOnGetDetailCustomizationInstance::CreateStatic(&FMovieGraphCollectionsCustomization::MakeInstance));
+
+	SelectedGraphObjectsDetailsWidget->RegisterInstancedCustomPropertyLayout(
+		UMovieGraphModifierNode::StaticClass(),
+		FOnGetDetailCustomizationInstance::CreateStatic(&FMovieGraphModifiersCustomization::MakeInstance));
 	
 	return SNew(SDockTab)
 		.TabColorScale(GetTabColorScale())
diff --git a/Engine/Plugins/MovieScene/MovieRenderPipeline/Source/MovieRenderPipelineEditor/Private/Graph/SMovieGraphActiveRenderSettingsTabContent.cpp b/Engine/Plugins/MovieScene/MovieRenderPipeline/Source/MovieRenderPipelineEditor/Private/Graph/SMovieGraphActiveRenderSettingsTabContent.cpp
--- a/Engine/Plugins/MovieScene/MovieRenderPipeline/Source/MovieRenderPipelineEditor/Private/Graph/SMovieGraphActiveRenderSettingsTabContent.cpp
+++ b/Engine/Plugins/MovieScene/MovieRenderPipeline/Source/MovieRenderPipelineEditor/Private/Graph/SMovieGraphActiveRenderSettingsTabContent.cpp
@@ -180,7 +180,7 @@ const TArray<TSharedPtr<FActiveRenderSettingsTreeElement>>& FActiveRenderSetting
 		for (const FProperty* Property : SettingsNode->GetAllOverrideableProperties())
 		{
 			TSharedPtr<FActiveRenderSettingsTreeElement> Element =
-				MakeShared<FActiveRenderSettingsTreeElement>(Property->GetFName(), EElementType::Property);
+				MakeShared<FActiveRenderSettingsTreeElement>(FName(Property->GetDisplayNameText().ToString()), EElementType::Property);
 			Element->SettingsNode = SettingsNode;
 			Element->SettingsProperty = Property;
 			Element->FlattenedGraph = FlattenedGraph;
