MoviePipeline: Fixed crash which would occur when selecting multiple variable nodes in the graph.

Some details customizations (and nodes in general) aren't multiple-node-selection friendly. This was fixed by having the details panel display multiple top-level objects. Now the details panel will not attempt to merge objects together and display eg "Multiple Values", but instead display each object in its own section separately.

#rnx

diff --git a/Engine/Plugins/MovieScene/MovieRenderPipeline/Source/MovieRenderPipelineEditor/Private/Graph/MovieGraphAssetToolkit.cpp b/Engine/Plugins/MovieScene/MovieRenderPipeline/Source/MovieRenderPipelineEditor/Private/Graph/MovieGraphAssetToolkit.cpp
--- a/Engine/Plugins/MovieScene/MovieRenderPipeline/Source/MovieRenderPipelineEditor/Private/Graph/MovieGraphAssetToolkit.cpp
+++ b/Engine/Plugins/MovieScene/MovieRenderPipeline/Source/MovieRenderPipelineEditor/Private/Graph/MovieGraphAssetToolkit.cpp
@@ -23,6 +23,7 @@
 #include "SMovieGraphMembersTabContent.h"
 
 #include "Framework/Commands/GenericCommands.h"
+#include "IDetailRootObjectCustomization.h"
 #include "PropertyEditorModule.h"
 #include "ToolMenus.h"
 #include "Widgets/Docking/SDockTab.h"
@@ -36,6 +37,58 @@ const FName FMovieGraphAssetToolkit::DetailsTabId(TEXT("MovieGraphAssetToolkitDe
 const FName FMovieGraphAssetToolkit::MembersTabId(TEXT("MovieGraphAssetToolkitMembers"));
 const FName FMovieGraphAssetToolkit::ActiveRenderSettingsTabId(TEXT("MovieGraphAssetToolkitActiveRenderSettings"));
 
+/** Header customization for when multiple objects are displayed in the details panel. */
+class FMovieGraphDetailsRootObjectCustomization final : public IDetailRootObjectCustomization
+{
+public:
+	explicit FMovieGraphDetailsRootObjectCustomization(const TSharedRef<IDetailsView>& InDetailsView)
+		: DetailsView(InDetailsView)
+	{
+		
+	}
+
+	// IDetailRootObjectCustomization interface
+	virtual TSharedPtr<SWidget> CustomizeObjectHeader(const FDetailsObjectSet& InRootObjectSet, const TSharedPtr<ITableRow>& InTableRow) override
+	{
+		FText DisplayName = InRootObjectSet.CommonBaseClass->GetDisplayNameText();
+		if (const UMovieGraphNode* GraphNode = Cast<const UMovieGraphNode>(InRootObjectSet.RootObjects[0]))
+		{
+			DisplayName = GraphNode->GetNodeTitle();
+		}
+		
+		return
+			SNew(SBox)
+			.Padding(5.f)
+			.VAlign(VAlign_Center)
+			.HAlign(HAlign_Left)
+			[
+				SNew(STextBlock)
+				.Text(DisplayName)
+				.Font(FAppStyle::GetFontStyle(TEXT("PropertyWindow.BoldFont")))
+			];
+	}
+
+	virtual bool AreObjectsVisible(const FDetailsObjectSet& InRootObjectSet) const override
+	{
+		return true;
+	}
+	
+	virtual bool ShouldDisplayHeader(const FDetailsObjectSet& InRootObjectSet) const override
+	{
+		// Only display the header if multiple objects are selected
+		if (DetailsView.IsValid())
+		{
+			return DetailsView.Pin()->GetSelectedObjects().Num() > 1;
+		}
+		
+		return true;
+	}
+	// End IDetailRootObjectCustomization interface
+
+private:
+	TWeakPtr<IDetailsView> DetailsView;
+};
+
 // Temporary cvar to enable/disable upgrading to a graph-based configuration
 static TAutoConsoleVariable<bool> CVarMoviePipelineEnableRenderGraph(
 	TEXT("MoviePipeline.EnableRenderGraph"),
@@ -232,9 +285,12 @@ TSharedRef<SDockTab> FMovieGraphAssetToolkit::SpawnTab_RenderGraphDetails(const
 	DetailsViewArgs.bAllowSearch = false;
 	DetailsViewArgs.NameAreaSettings = FDetailsViewArgs::HideNameArea;
 	DetailsViewArgs.bHideSelectionTip = true;
+	DetailsViewArgs.bAllowMultipleTopLevelObjects = true;
 	DetailsViewArgs.ViewIdentifier = "MovieGraphSettings";
 
 	SelectedGraphObjectsDetailsWidget = PropertyEditorModule.CreateDetailView(DetailsViewArgs);
+	SelectedGraphObjectsDetailsWidget->SetRootObjectCustomizationInstance(
+		MakeShared<FMovieGraphDetailsRootObjectCustomization>(SelectedGraphObjectsDetailsWidget.ToSharedRef()));
 
 	SelectedGraphObjectsDetailsWidget->RegisterInstancedCustomPropertyLayout(
 		UMovieGraphMember::StaticClass(),
