MovieGraph: Add compression and metadata support to the single layer EXR output node, while maintaining support for composited passes.

diff --git a/Engine/Plugins/MovieScene/MovieRenderPipeline/Source/MovieRenderPipelineCore/Public/Graph/MovieGraphRenderDataIdentifier.h b/Engine/Plugins/MovieScene/MovieRenderPipeline/Source/MovieRenderPipelineCore/Public/Graph/MovieGraphRenderDataIdentifier.h
--- a/Engine/Plugins/MovieScene/MovieRenderPipeline/Source/MovieRenderPipelineCore/Public/Graph/MovieGraphRenderDataIdentifier.h
+++ b/Engine/Plugins/MovieScene/MovieRenderPipeline/Source/MovieRenderPipelineCore/Public/Graph/MovieGraphRenderDataIdentifier.h
@@ -52,6 +52,12 @@ struct FMovieGraphRenderDataIdentifier
 		return FString::Printf(TEXT("RootBranch: %s Renderer:%s SubResource: %s Camera: %s"), *InIdentifier.RootBranchName.ToString(), *InIdentifier.RendererName, *InIdentifier.SubResourceName, *InIdentifier.CameraName);
 	}
 
+	bool IsBranchAndCameraEqual(const FMovieGraphRenderDataIdentifier& InIdentifier) const
+	{
+		return CameraName == InIdentifier.CameraName &&
+			RootBranchName == InIdentifier.RootBranchName;
+	}
+
 public:
 	/** 
 	* The root branch name that this render layer exists on. Actual display name comes from a UMovieGraphRenderLayerNode (if found in the branch)
diff --git a/Engine/Plugins/MovieScene/MovieRenderPipeline/Source/MovieRenderPipelineRenderPasses/Private/MovieGraphImageSequenceOutputNode.cpp b/Engine/Plugins/MovieScene/MovieRenderPipeline/Source/MovieRenderPipelineRenderPasses/Private/MovieGraphImageSequenceOutputNode.cpp
--- a/Engine/Plugins/MovieScene/MovieRenderPipeline/Source/MovieRenderPipelineRenderPasses/Private/MovieGraphImageSequenceOutputNode.cpp
+++ b/Engine/Plugins/MovieScene/MovieRenderPipeline/Source/MovieRenderPipelineRenderPasses/Private/MovieGraphImageSequenceOutputNode.cpp
@@ -91,12 +91,12 @@ bool UMovieGraphImageSequenceOutputNode::IsFinishedWritingToDiskImpl() const
 	return Super::IsFinishedWritingToDiskImpl() && (!FinalizeFence.IsValid() || FinalizeFence.WaitFor(0));
 }
 
-void UMovieGraphImageSequenceOutputNode::OnReceiveImageDataImpl(UMovieGraphPipeline* InPipeline, UE::MovieGraph::FMovieGraphOutputMergerFrame* InRawFrameData, const TSet<FMovieGraphRenderDataIdentifier>& InMask)
+TArray<TPair<FMovieGraphRenderDataIdentifier, TUniquePtr<FImagePixelData>>> UMovieGraphImageSequenceOutputNode::GetCompositedPasses(
+	UE::MovieGraph::FMovieGraphOutputMergerFrame* InRawFrameData) const
 {
-	check(InRawFrameData);
-
 	// Gather the passes that need to be composited
 	TArray<TPair<FMovieGraphRenderDataIdentifier, TUniquePtr<FImagePixelData>>> CompositedPasses;
+
 	for (TPair<FMovieGraphRenderDataIdentifier, TUniquePtr<FImagePixelData>>& RenderData : InRawFrameData->ImageOutputData)
 	{
 		UE::MovieGraph::FMovieGraphSampleState* Payload = RenderData.Value->GetPayload<UE::MovieGraph::FMovieGraphSampleState>();
@@ -126,6 +126,84 @@ void UMovieGraphImageSequenceOutputNode::OnReceiveImageDataImpl(UMovieGraphPipel
 		return PayloadA->CompositingSortOrder > PayloadB->CompositingSortOrder;
 	});
 
+	return CompositedPasses;
+}
+
+FString UMovieGraphImageSequenceOutputNode::CreateFileName(
+	UE::MovieGraph::FMovieGraphOutputMergerFrame* InRawFrameData,
+	const UMovieGraphPipeline* InPipeline,
+	const TPair<FMovieGraphRenderDataIdentifier, TUniquePtr<FImagePixelData>>& InRenderData,
+	const EImageFormat InImageFormat,
+	FMovieGraphResolveArgs& OutMergedFormatArgs) const
+{
+	UE::MovieGraph::FMovieGraphSampleState* Payload = InRenderData.Value->GetPayload<UE::MovieGraph::FMovieGraphSampleState>();
+
+	const TObjectPtr<UMoviePipelineExecutorShot>& Shot = InPipeline->GetActiveShotList()[Payload->TraversalContext.ShotIndex];
+
+	const TCHAR* Extension = TEXT("");
+	switch (InImageFormat)
+	{
+	case EImageFormat::PNG: Extension = TEXT("png"); break;
+	case EImageFormat::JPEG: Extension = TEXT("jpeg"); break;
+	case EImageFormat::BMP: Extension = TEXT("bmp"); break;
+	case EImageFormat::EXR: Extension = TEXT("exr"); break;
+	}
+
+	UMovieGraphOutputSettingNode* OutputSettingNode = InRawFrameData->EvaluatedConfig->GetSettingForBranch<UMovieGraphOutputSettingNode>(InRenderData.Key.RootBranchName);
+	if (!OutputSettingNode)
+	{
+		return FString();
+	}
+
+	// Generate one string that puts the directory combined with the filename format.
+	FString FileNameFormatString = OutputSettingNode->OutputDirectory.Path / OutputSettingNode->FileNameFormat;
+
+	constexpr bool bIncludeRenderPass = false;
+	constexpr bool bTestFrameNumber = true;
+	UE::MoviePipeline::ValidateOutputFormatString(FileNameFormatString, bIncludeRenderPass, bTestFrameNumber);
+
+	// Map the .ext to be specific to our output data.
+	TMap<FString, FString> AdditionalFormatArgs;
+	AdditionalFormatArgs.Add(TEXT("ext"), Extension);
+
+	FMovieGraphFilenameResolveParams Params = FMovieGraphFilenameResolveParams();
+	Params.RenderDataIdentifier = InRenderData.Key;
+	//Params.RootFrameNumber = Payload->TraversalContext.Time.RootFrameNumber;
+	//Params.ShotFrameNumber = Payload->TraversalContext.Time.ShotFrameNumber;
+	Params.RootFrameNumberRel = Payload->TraversalContext.Time.OutputFrameNumber;
+	//Params.ShotFrameNumberRel = Payload->TraversalCOntext.Time.ShotFrameNumberRel
+	//Params.FileMetadata = ToDo: Track File Metadata
+	Params.ZeroPadFrameNumberCount = OutputSettingNode->ZeroPadFrameNumbers;
+	Params.FrameNumberOffset = OutputSettingNode->FrameNumberOffset;
+	Params.EvaluatedConfig = InRawFrameData->EvaluatedConfig.Get();
+	Params.Version = Shot->ShotInfo.VersionNumber;
+
+	// If time dilation is in effect, RootFrameNumber and ShotFrameNumber will contain duplicates and the files will overwrite each other, 
+	// so we force them into relative mode and then warn users we did that (as their numbers will jump from say 1001 -> 0000).
+	bool bForceRelativeFrameNumbers = true; // TODO: Use relative frame numbers until we track Root vs. Shot frame numbers. (Previously false);
+	//if (FileNameFormatString.Contains(TEXT("{frame")) Payload->TraversalContext.Time.IsTimeDilated() && !FileNameFormatString.Contains(TEXT("_rel}")))
+	//{
+	//	UE_LOG(LogMovieRenderPipeline, Warning, TEXT("Time Dilation was used but output format does not use relative time, forcing relative numbers. Change {frame_number} to {frame_number_rel} (or shot version) to remove this message."));
+	//	bForceRelativeFrameNumbers = true;
+	//}
+	Params.bForceRelativeFrameNumbers = bForceRelativeFrameNumbers;
+	Params.bEnsureAbsolutePath = true;
+	Params.FileNameFormatOverrides = AdditionalFormatArgs;
+	Params.InitializationTime = InPipeline->GetInitializationTime();
+	Params.Shot = Shot;
+	Params.Job = InPipeline->GetCurrentJob();
+	Params.EvaluatedConfig = InRawFrameData->EvaluatedConfig.Get();
+
+	// Take our string path from the Output Setting and resolve it.
+	return UMovieGraphBlueprintLibrary::ResolveFilenameFormatArguments(FileNameFormatString, Params, OutMergedFormatArgs);
+}
+
+void UMovieGraphImageSequenceOutputNode::OnReceiveImageDataImpl(UMovieGraphPipeline* InPipeline, UE::MovieGraph::FMovieGraphOutputMergerFrame* InRawFrameData, const TSet<FMovieGraphRenderDataIdentifier>& InMask)
+{
+	check(InRawFrameData);
+
+	TArray<TPair<FMovieGraphRenderDataIdentifier, TUniquePtr<FImagePixelData>>> CompositedPasses = GetCompositedPasses(InRawFrameData);
+
 	// ToDo:
 	// The ImageWriteQueue is set up in a fire-and-forget manner. This means that the data needs to be placed in the WriteQueue
 	// as a TUniquePtr (so it can free the data when its done). Unfortunately we can have multiple output formats at once,
@@ -142,17 +220,10 @@ void UMovieGraphImageSequenceOutputNode::OnReceiveImageDataImpl(UMovieGraphPipel
 	for (TPair<FMovieGraphRenderDataIdentifier, TUniquePtr<FImagePixelData>>& RenderData : InRawFrameData->ImageOutputData)
 	{
 		// If this pass is composited, skip it for now
-		bool bSkip = false;
-		for (TPair<FMovieGraphRenderDataIdentifier, TUniquePtr<FImagePixelData>>& CompositedPass : CompositedPasses)
-		{
-			if (CompositedPass.Key == RenderData.Key)
+		if (CompositedPasses.ContainsByPredicate([&RenderData](const TPair<FMovieGraphRenderDataIdentifier, TUniquePtr<FImagePixelData>>& CompositedPass)
 			{
-				bSkip = true;
-				break;
-			}
-		}
-
-		if (bSkip)
+				return CompositedPass.Key == RenderData.Key;
+			}))
 		{
 			continue;
 		}
@@ -163,89 +234,29 @@ void UMovieGraphImageSequenceOutputNode::OnReceiveImageDataImpl(UMovieGraphPipel
 			continue;
 		}
 
-		const UMovieGraphImageSequenceOutputNode* ParentNode = Cast<UMovieGraphImageSequenceOutputNode>(
-			InRawFrameData->EvaluatedConfig->GetSettingForBranch(GetClass(), RenderData.Key.RootBranchName, false, true));
-		checkf(ParentNode, TEXT("Image sequence output should not exist without a parent node in the graph."));
-
-		UE::MovieGraph::FMovieGraphSampleState* Payload = RenderData.Value->GetPayload<UE::MovieGraph::FMovieGraphSampleState>();
-		const TObjectPtr<UMoviePipelineExecutorShot>& Shot = InPipeline->GetActiveShotList()[Payload->TraversalContext.ShotIndex];
-
-		const bool bIncludeCDOs = true;
-		UMovieGraphOutputSettingNode* OutputSettingNode = InRawFrameData->EvaluatedConfig->GetSettingForBranch<UMovieGraphOutputSettingNode>(RenderData.Key.RootBranchName, bIncludeCDOs);
-		if (!ensure(OutputSettingNode))
-		{
-			continue;
-		}
-
-		FString RenderLayerName = RenderData.Key.RootBranchName.ToString();
-		UMovieGraphRenderLayerNode* RenderLayerNode = InRawFrameData->EvaluatedConfig->GetSettingForBranch<UMovieGraphRenderLayerNode>(RenderData.Key.RootBranchName, bIncludeCDOs);
-		if (RenderLayerNode)
-		{
-			RenderLayerName = RenderLayerNode->GetRenderLayerName();
-		}
 		// ToDo: Certain images may require transparency, at which point
 		// we write out a .png instead of a .jpeg.
 		EImageFormat PreferredOutputFormat = OutputFormat;
-
-		const TCHAR* Extension = TEXT("");
-		switch (PreferredOutputFormat)
+		
+		FMovieGraphResolveArgs FinalResolvedKVPs;
+		FString FileName = CreateFileName(InRawFrameData, InPipeline, RenderData, PreferredOutputFormat, FinalResolvedKVPs);
+		if (!ensureMsgf(!FileName.IsEmpty(), TEXT("Unexpected empty file name, skipping frame.")))
 		{
-		case EImageFormat::PNG: Extension = TEXT("png"); break;
-		case EImageFormat::JPEG: Extension = TEXT("jpeg"); break;
-		case EImageFormat::BMP: Extension = TEXT("bmp"); break;
-		case EImageFormat::EXR: Extension = TEXT("exr"); break;
+			continue;
 		}
 
-		// Generate one string that puts the directory combined with the filename format.
-		FString FileNameFormatString = OutputSettingNode->OutputDirectory.Path / OutputSettingNode->FileNameFormat;
-
-		constexpr bool bIncludeRenderPass = false;
-		constexpr bool bTestFrameNumber = true;
-		UE::MoviePipeline::ValidateOutputFormatString(FileNameFormatString, bIncludeRenderPass, bTestFrameNumber);
-
-		// Map the .ext to be specific to our output data.
-		TMap<FString, FString> AdditionalFormatArgs;
-		AdditionalFormatArgs.Add(TEXT("ext"), Extension);
-
-		FMovieGraphFilenameResolveParams Params = FMovieGraphFilenameResolveParams();
-		Params.RenderDataIdentifier = RenderData.Key;
-		//Params.RootFrameNumber = Payload->TraversalContext.Time.RootFrameNumber;
-		//Params.ShotFrameNumber = Payload->TraversalContext.Time.ShotFrameNumber;
-		Params.RootFrameNumberRel = Payload->TraversalContext.Time.OutputFrameNumber;
-		//Params.ShotFrameNumberRel = Payload->TraversalCOntext.Time.ShotFrameNumberRel
-		//Params.FileMetadata = ToDo: Track File Metadata
-		Params.ZeroPadFrameNumberCount = OutputSettingNode->ZeroPadFrameNumbers;
-		Params.FrameNumberOffset = OutputSettingNode->FrameNumberOffset;
-		Params.EvaluatedConfig = InRawFrameData->EvaluatedConfig.Get();
-		Params.Version = Shot->ShotInfo.VersionNumber;
-
-		// If time dilation is in effect, RootFrameNumber and ShotFrameNumber will contain duplicates and the files will overwrite each other, 
-		// so we force them into relative mode and then warn users we did that (as their numbers will jump from say 1001 -> 0000).
-		bool bForceRelativeFrameNumbers = true; // TODO: Use relative frame numbers until we track Root vs. Shot frame numbers. (Previously false);
-		//if (FileNameFormatString.Contains(TEXT("{frame")) Payload->TraversalContext.Time.IsTimeDilated() && !FileNameFormatString.Contains(TEXT("_rel}")))
-		//{
-		//	UE_LOG(LogMovieRenderPipeline, Warning, TEXT("Time Dilation was used but output format does not use relative time, forcing relative numbers. Change {frame_number} to {frame_number_rel} (or shot version) to remove this message."));
-		//	bForceRelativeFrameNumbers = true;
-		//}
-		Params.bForceRelativeFrameNumbers = bForceRelativeFrameNumbers;
-		Params.bEnsureAbsolutePath = true;
-		Params.FileNameFormatOverrides = AdditionalFormatArgs;
-		Params.InitializationTime = InPipeline->GetInitializationTime();
-		Params.Shot = Shot;
-		Params.Job = InPipeline->GetCurrentJob();
-		Params.EvaluatedConfig = InRawFrameData->EvaluatedConfig.Get();
-
-		// Take our string path from the Output Setting and resolve it.
-		FMovieGraphResolveArgs FinalResolvedKVPs;
-		const FString FileName = UMovieGraphBlueprintLibrary::ResolveFilenameFormatArguments(FileNameFormatString, Params, FinalResolvedKVPs);
-
-		//TODO: Support for single-layer EXR image write task (compression)...
 		TUniquePtr<FImageWriteTask> TileImageTask = MakeUnique<FImageWriteTask>();
-		TileImageTask->Format = OutputFormat;
+		TileImageTask->Format = PreferredOutputFormat;
 		TileImageTask->CompressionQuality = 100;
 		TileImageTask->Filename = FileName;
 		TileImageTask->PixelData = RenderData.Value->CopyImageData();
 
+		const UMovieGraphImageSequenceOutputNode* ParentNode = Cast<UMovieGraphImageSequenceOutputNode>(
+			InRawFrameData->EvaluatedConfig->GetSettingForBranch(GetClass(), RenderData.Key.RootBranchName, false, true));
+		checkf(ParentNode, TEXT("Image sequence output should not exist without a parent node in the graph."));
+
+		UE::MovieGraph::FMovieGraphSampleState* Payload = RenderData.Value->GetPayload<UE::MovieGraph::FMovieGraphSampleState>();
+		
 		bool bQuantizationEncodeSRGB = true;
 #if WITH_EDITOR
 		if (ParentNode->OCIOConfiguration.bIsEnabled && Payload->bAllowOCIO)
@@ -273,8 +284,7 @@ void UMovieGraphImageSequenceOutputNode::OnReceiveImageDataImpl(UMovieGraphPipel
 		for (TPair<FMovieGraphRenderDataIdentifier, TUniquePtr<FImagePixelData>>& CompositedPass : CompositedPasses)
 		{
 			// This composited pass will only composite on top of renders w/ the same branch and camera
-			const FMovieGraphRenderDataIdentifier& Id = CompositedPass.Key;
-			if ((Id.CameraName != RenderData.Key.CameraName) || (Id.RootBranchName != RenderData.Key.RootBranchName))
+			if (!CompositedPass.Key.IsBranchAndCameraEqual(RenderData.Key))
 			{
 				continue;
 			}
@@ -295,7 +305,7 @@ void UMovieGraphImageSequenceOutputNode::OnReceiveImageDataImpl(UMovieGraphPipel
 		}
 
 		UE::MovieGraph::FMovieGraphOutputFutureData OutputData;
-		OutputData.Shot = Shot;
+		OutputData.Shot = InPipeline->GetActiveShotList()[Payload->TraversalContext.ShotIndex];
 		OutputData.FilePath = FileName;
 		OutputData.DataIdentifier = RenderData.Key;
 
@@ -303,21 +313,176 @@ void UMovieGraphImageSequenceOutputNode::OnReceiveImageDataImpl(UMovieGraphPipel
 
 		InPipeline->AddOutputFuture(MoveTemp(Future), OutputData);
 	}
+}
+
+
+TUniquePtr<FEXRImageWriteTask> UMovieGraphImageSequenceOutputNode_EXR::CreateImageWriteTask(FString InFileName, EEXRCompressionFormat InCompression) const
+{
+	// Ensure our OpenExrRTTI module gets loaded.
+	UE_CALL_ONCE([]
+		{
+			check(IsInGameThread());
+			FModuleManager::Get().LoadModule(TEXT("UEOpenExrRTTI"));
+		});
+
+	TUniquePtr<FEXRImageWriteTask> ImageWriteTask = MakeUnique<FEXRImageWriteTask>();
+	ImageWriteTask->Filename = MoveTemp(InFileName);
+	ImageWriteTask->Compression = InCompression;
+	// ImageWriteTask->CompressionLevel is intentionally skipped and not exposed ("dwaCompressionLevel" is deprecated)
+
+	return MoveTemp(ImageWriteTask);
+}
+
+void UMovieGraphImageSequenceOutputNode_EXR::PrepareTaskGlobalMetadata(FEXRImageWriteTask& InOutImageTask, UE::MovieGraph::FMovieGraphOutputMergerFrame* InRawFrameData, TMap<FString, FString>& InMetadata) const
+{
+	// Add in hardware usage metadata
+	UE::MoviePipeline::GetHardwareUsageMetadata(InMetadata, FPaths::GetPath(InOutImageTask.Filename));
+
+	// Add passed in resolved metadata
+	for (const TPair<FString, FString>& Metadata : InMetadata)
+	{
+		InOutImageTask.FileMetadata.Emplace(Metadata.Key, Metadata.Value);
+	}
+
+	// Add in any metadata from the output merger frame
+	for (const TPair<FString, FString>& Metadata : InRawFrameData->FileMetadata)
+	{
+		InOutImageTask.FileMetadata.Add(Metadata.Key, Metadata.Value);
+	}
+}
+
+void UMovieGraphImageSequenceOutputNode_EXR::UpdateTaskPerLayer(
+	FEXRImageWriteTask& InOutImageTask,
+	const UMovieGraphImageSequenceOutputNode* InParentNode,
+	FImagePixelData* InImageData,
+	int32 InLayerIndex,
+	const FString& InLayerName) const
+{
+	const UE::MovieGraph::FMovieGraphSampleState* Payload = InImageData->GetPayload<UE::MovieGraph::FMovieGraphSampleState>();
+
+	// No quantization required, just copy the data as we will move it into the image write task.
+	TUniquePtr<FImagePixelData> PixelData = InImageData->CopyImageData();
+
+	bool bEnabledOCIO = false;
+#if WITH_EDITOR
+	if (InParentNode->OCIOConfiguration.bIsEnabled && Payload->bAllowOCIO)
+	{
+		FPixelPreProcessor OCIOPixelPreProcessor = UE::MovieGraph::Private::CreateOpenColorIOPixelPreProcessor(InParentNode->OCIOConfiguration.ColorConfiguration);
+		if (OCIOPixelPreProcessor)
+		{
+			InOutImageTask.PixelPreprocessors.FindOrAdd(InLayerIndex).Emplace(MoveTemp(OCIOPixelPreProcessor));
+			bEnabledOCIO = true;
+		}
+	}
+#endif
+
+	if (InLayerIndex == 0)
+	{
+		// Add task information that is common to all layers. This metadata may be redundant with unreal/* metadata,
+		// but these are "standard" fields in EXR metadata.
+		InOutImageTask.FileMetadata.Add("owner", UE::MoviePipeline::GetJobAuthor(Payload->TraversalContext.Job));
+		InOutImageTask.FileMetadata.Add("comments", Payload->TraversalContext.Job->Comment);
+
+		const FIntPoint& Resolution = PixelData->GetSize();
+		InOutImageTask.Width = Resolution.X;
+		InOutImageTask.Height = Resolution.Y;
+
+		InOutImageTask.OverscanPercentage = Payload->OverscanFraction;
+#if WITH_EDITOR
+		if (bEnabledOCIO)
+		{
+			UE::MoviePipeline::UpdateColorSpaceMetadata(InParentNode->OCIOConfiguration.ColorConfiguration, InOutImageTask);
+		}
+		else
+#endif
+		{
+			UE::MoviePipeline::UpdateColorSpaceMetadata(Payload->SceneCaptureSource, InOutImageTask);
+		}
+	}
+
+	if (!InLayerName.IsEmpty())
+	{
+		InOutImageTask.LayerNames.FindOrAdd(PixelData.Get(), InLayerName);
+	}
+
+	InOutImageTask.Layers.Add(MoveTemp(PixelData));
+}
+
+void UMovieGraphImageSequenceOutputNode_EXR::OnReceiveImageDataImpl(UMovieGraphPipeline* InPipeline, UE::MovieGraph::FMovieGraphOutputMergerFrame* InRawFrameData, const TSet<FMovieGraphRenderDataIdentifier>& InMask)
+{
+	check(InRawFrameData);
+
+	TArray<TPair<FMovieGraphRenderDataIdentifier, TUniquePtr<FImagePixelData>>> CompositedPasses = GetCompositedPasses(InRawFrameData);
+
+	for (TPair<FMovieGraphRenderDataIdentifier, TUniquePtr<FImagePixelData>>& RenderData : InRawFrameData->ImageOutputData)
+	{
+		// If this pass is composited, skip it for now
+		if (CompositedPasses.ContainsByPredicate([&RenderData](const TPair<FMovieGraphRenderDataIdentifier, TUniquePtr<FImagePixelData>>& CompositedPass)
+			{
+				return RenderData.Key == CompositedPass.Key;
+			}))
+		{
+			continue;
+		}
+
+		// A layer within this output data may have chosen to not be written to disk by this CDO node
+		if (!InMask.Contains(RenderData.Key))
+		{
+			continue;
+		}
+
+		FMovieGraphResolveArgs ResolvedFormatArgs;
+		FString FileName = CreateFileName(InRawFrameData, InPipeline, RenderData, OutputFormat, ResolvedFormatArgs);
+		if (!ensureMsgf(!FileName.IsEmpty(), TEXT("Unexpected empty file name, skipping frame.")))
+		{
+			continue;
+		}
+
+		UE::MovieGraph::FMovieGraphSampleState* Payload = RenderData.Value->GetPayload<UE::MovieGraph::FMovieGraphSampleState>();
 
+		const UMovieGraphImageSequenceOutputNode_EXR* ParentNode = InRawFrameData->EvaluatedConfig->GetSettingForBranch<UMovieGraphImageSequenceOutputNode_EXR>(
+			RenderData.Key.RootBranchName, false /*bIncludeCDOs*/, true /*bExactMatch*/);
+		checkf(ParentNode, TEXT("Single-layer EXR should not exist without a parent node in the graph."));
+		
+		TUniquePtr<FEXRImageWriteTask> ImageWriteTask = CreateImageWriteTask(FileName, ParentNode->Compression);
+		
+		PrepareTaskGlobalMetadata(*ImageWriteTask, InRawFrameData, ResolvedFormatArgs.FileMetadata);
+
+		// No layer is equivalent to a zero-index layer
+		constexpr int32 LayerIndex = 0;
+		UpdateTaskPerLayer(*ImageWriteTask, ParentNode, RenderData.Value.Get(), LayerIndex);
+
+		// Perform compositing if any composited passes were found earlier
+		for (TPair<FMovieGraphRenderDataIdentifier, TUniquePtr<FImagePixelData>>& CompositedPass : CompositedPasses)
+		{
+			// This composited pass will only composite on top of renders w/ the same branch and camera
+			if (CompositedPass.Key.IsBranchAndCameraEqual(RenderData.Key))
+			{
+				// There could be multiple renders within this branch using the composited pass, so we have to copy the image data
+				ImageWriteTask->PixelPreprocessors.FindOrAdd(LayerIndex).Add(TAsyncCompositeImage<FFloat16Color>(CompositedPass.Value->CopyImageData()));
+			}
+		}
+
+		UE::MovieGraph::FMovieGraphOutputFutureData OutputFutureData;
+		OutputFutureData.Shot = InPipeline->GetActiveShotList()[Payload->TraversalContext.ShotIndex];
+		OutputFutureData.FilePath = FileName;
+		OutputFutureData.DataIdentifier = RenderData.Key;
+
+		TFuture<bool> Future = ImageWriteQueue->Enqueue(MoveTemp(ImageWriteTask));
+
+		InPipeline->AddOutputFuture(MoveTemp(Future), OutputFutureData);
+	}
 }
 
+
 void UMovieGraphImageSequenceOutputNode_MultiLayerEXR::OnReceiveImageDataImpl(UMovieGraphPipeline* InPipeline, UE::MovieGraph::FMovieGraphOutputMergerFrame* InRawFrameData, const TSet<FMovieGraphRenderDataIdentifier>& InMask)
 {
 	check(InRawFrameData);
 
-	const UMovieGraphImageSequenceOutputNode_MultiLayerEXR* ParentNode = Cast<UMovieGraphImageSequenceOutputNode_MultiLayerEXR>(
-		InRawFrameData->EvaluatedConfig->GetSettingForBranch(GetClass(), UMovieGraphNode::GlobalsPinName, false, true));
+	const UMovieGraphImageSequenceOutputNode_MultiLayerEXR* ParentNode = InRawFrameData->EvaluatedConfig->GetSettingForBranch<UMovieGraphImageSequenceOutputNode_MultiLayerEXR>(
+		UMovieGraphNode::GlobalsPinName, false /*bIncludeCDOs*/, true /*bExactMatch*/);
 	checkf(ParentNode, TEXT("Multi-Layer EXR should not exist without a parent node in the graph."));
 
-	// Ensure our OpenExrRTTI module gets loaded. This needs to happen from the main thread, if it's not loaded then metadata silently fails when writing.
-	static const FName RTTIExtensionModuleName("UEOpenExrRTTI");
-	FModuleManager::Get().LoadModule(RTTIExtensionModuleName);
-
 	// Generate a mapping of resolved filename -> RenderIDs, and filename -> resolve args. The generated EXRs can only
 	// store layers with a common resolution, and this takes care of ensuring that only renderIDs with a common resolution
 	// map to the same filename.
@@ -331,31 +496,8 @@ void UMovieGraphImageSequenceOutputNode_MultiLayerEXR::OnReceiveImageDataImpl(UM
 		const FString& Filename = RenderIDsForFilename.Key;
 		const TArray<FMovieGraphRenderDataIdentifier>& RenderIDs = RenderIDsForFilename.Value;
 		
-		TUniquePtr<FEXRImageWriteTask> MultiLayerImageTask = MakeUnique<FEXRImageWriteTask>();
-		
-		MultiLayerImageTask->Filename = Filename;
-		MultiLayerImageTask->Compression = ParentNode->Compression;
-		// MultiLayerImageTask->CompressionLevel is intentionally skipped because it doesn't seem to make any practical difference
-		// so we don't expose it to the user because that will just cause confusion where the setting doesn't seem to do anything.
-
-		// Add in hardware usage metadata
-		UE::MoviePipeline::GetHardwareUsageMetadata(FilenameToResolveArgs[Filename].FileMetadata, FPaths::GetPath(Filename));
-
-		// FileMetadata has been generated by ResolveFilenameFormatArgs, but we need to convert from FString, FString
-		// (needed for BP/Python purposes) to a FStringFormatArg as we need to preserve numeric metadata types later in
-		// the image writing process (for compression level).
-		TMap<FString, FStringFormatArg> NewFileMetadataMap;
-		for (const TPair<FString, FString>& Metadata : FilenameToResolveArgs[Filename].FileMetadata)
-		{
-			NewFileMetadataMap.Add(Metadata.Key, Metadata.Value);
-		}
-		MultiLayerImageTask->FileMetadata = NewFileMetadataMap;
-
-		// Add in any metadata from the output merger frame
-		for (const TPair<FString, FString>& Metadata : InRawFrameData->FileMetadata)
-		{
-			MultiLayerImageTask->FileMetadata.Add(Metadata.Key, Metadata.Value);
-		}
+		TUniquePtr<FEXRImageWriteTask> MultiLayerImageTask = CreateImageWriteTask(Filename, ParentNode->Compression);
+		PrepareTaskGlobalMetadata(*MultiLayerImageTask, InRawFrameData, FilenameToResolveArgs[Filename].FileMetadata);
 
 		// Add each render pass as a layer to the EXR
 		int32 LayerIndex = 0;
@@ -363,51 +505,12 @@ void UMovieGraphImageSequenceOutputNode_MultiLayerEXR::OnReceiveImageDataImpl(UM
 		for (const FMovieGraphRenderDataIdentifier& RenderID : RenderIDs)
 		{
 			const TUniquePtr<FImagePixelData>& ImageData = InRawFrameData->ImageOutputData[RenderID];
-
-			// No quantization required, just copy the data as we will move it into the image write task.
-			TUniquePtr<FImagePixelData> PixelData = ImageData->CopyImageData();
 			const UE::MovieGraph::FMovieGraphSampleState* Payload = ImageData->GetPayload<UE::MovieGraph::FMovieGraphSampleState>();
 			ShotIndex = Payload->TraversalContext.ShotIndex;
 
-			bool bEnabledOCIO = false;
-#if WITH_EDITOR
-			if (ParentNode->OCIOConfiguration.bIsEnabled && Payload->bAllowOCIO)
-			{
-				FPixelPreProcessor OCIOPixelPreProcessor = UE::MovieGraph::Private::CreateOpenColorIOPixelPreProcessor(ParentNode->OCIOConfiguration.ColorConfiguration);
-				if (OCIOPixelPreProcessor)
-				{
-					MultiLayerImageTask->PixelPreprocessors.FindOrAdd(LayerIndex).Emplace(MoveTemp(OCIOPixelPreProcessor));
-					bEnabledOCIO = true;
-				}
-			}
-#endif
+			FString LayerName = {};
 
-			// If there is more than one layer, then we will prefix the layer. The first layer is not prefixed (and gets inserted as RGBA)
-			// as most programs that handle EXRs expect the main image data to be in an unnamed layer.
-			if (LayerIndex == 0)
-			{
-				// Add task information that is common to all layers. This metadata may be redundant with unreal/* metadata,
-				// but these are "standard" fields in EXR metadata.
-				MultiLayerImageTask->FileMetadata.Add("owner", UE::MoviePipeline::GetJobAuthor(Payload->TraversalContext.Job));
-				MultiLayerImageTask->FileMetadata.Add("comments", Payload->TraversalContext.Job->Comment);
-
-				const FIntPoint& Resolution = ImageData->GetSize();
-				MultiLayerImageTask->Width = Resolution.X;
-				MultiLayerImageTask->Height = Resolution.Y;
-				
-				MultiLayerImageTask->OverscanPercentage = Payload->OverscanFraction;
-#if WITH_EDITOR
-				if (bEnabledOCIO)
-				{
-					UE::MoviePipeline::UpdateColorSpaceMetadata(ParentNode->OCIOConfiguration.ColorConfiguration, *MultiLayerImageTask);
-				}
-				else
-#endif
-				{
-					UE::MoviePipeline::UpdateColorSpaceMetadata(Payload->SceneCaptureSource, *MultiLayerImageTask);
-				}
-			}
-			else
+			if (LayerIndex != 0)
 			{
 				// If there is more than one layer, then we will prefix the layer. The first layer is not prefixed (and gets inserted as RGBA)
 				// as most programs that handle EXRs expect the main image data to be in an unnamed layer. We only postfix with cameraname
@@ -415,17 +518,16 @@ void UMovieGraphImageSequenceOutputNode_MultiLayerEXR::OnReceiveImageDataImpl(UM
 				// TODO: The number of cameras may be inaccurate -- no camera setting in the graph yet
 				UMoviePipelineExecutorShot* CurrentShot = InPipeline->GetActiveShotList()[ShotIndex];
 				int32 NumCameras = CurrentShot->SidecarCameras.Num();
-				
-				FString CombinedName = FString::Printf(TEXT("%s_%s"), *RenderID.RootBranchName.ToString(), *RenderID.RendererName);
+
+				LayerName = FString::Printf(TEXT("%s_%s"), *RenderID.RootBranchName.ToString(), *RenderID.RendererName);
 				if (NumCameras > 1)
 				{
-					CombinedName = FString::Printf(TEXT("%s_%s"), *CombinedName, *RenderID.CameraName);
+					LayerName = FString::Printf(TEXT("%s_%s"), *LayerName, *RenderID.CameraName);
 				}
-				
-				MultiLayerImageTask->LayerNames.FindOrAdd(PixelData.Get(), CombinedName);
 			}
 
-			MultiLayerImageTask->Layers.Add(MoveTemp(PixelData));
+			UpdateTaskPerLayer(*MultiLayerImageTask, ParentNode, ImageData.Get(), LayerIndex, LayerName);
+
 			LayerIndex++;
 		}
 		
diff --git a/Engine/Plugins/MovieScene/MovieRenderPipeline/Source/MovieRenderPipelineRenderPasses/Public/MovieGraphImageSequenceOutputNode.h b/Engine/Plugins/MovieScene/MovieRenderPipeline/Source/MovieRenderPipelineRenderPasses/Public/MovieGraphImageSequenceOutputNode.h
--- a/Engine/Plugins/MovieScene/MovieRenderPipeline/Source/MovieRenderPipelineRenderPasses/Public/MovieGraphImageSequenceOutputNode.h
+++ b/Engine/Plugins/MovieScene/MovieRenderPipeline/Source/MovieRenderPipelineRenderPasses/Public/MovieGraphImageSequenceOutputNode.h
@@ -48,6 +48,19 @@ public:
 	FOpenColorIODisplayConfiguration OCIOConfiguration;
 #endif
 
+protected:
+	/** Convenience function to get the list of active composite passes from render data. */
+	TArray<TPair<FMovieGraphRenderDataIdentifier, TUniquePtr<FImagePixelData>>> GetCompositedPasses(
+		UE::MovieGraph::FMovieGraphOutputMergerFrame* InRawFrameData) const;
+
+	/** Convenience function to create the output file name. */
+	FString CreateFileName(
+		UE::MovieGraph::FMovieGraphOutputMergerFrame* InRawFrameData,
+		const UMovieGraphPipeline* InPipeline,
+		const TPair<FMovieGraphRenderDataIdentifier, TUniquePtr<FImagePixelData>>& InRenderData,
+		const EImageFormat InImageFormat,
+		FMovieGraphResolveArgs& OutMergedFormatArgs) const;
+
 protected:
 	/** The output format (as known used by the ImageWriteQueue) to output into. */
 	EImageFormat OutputFormat;
@@ -78,6 +91,8 @@ public:
 		Compression = EEXRCompressionFormat::PIZ;
 	}
 
+	virtual void OnReceiveImageDataImpl(UMovieGraphPipeline* InPipeline, UE::MovieGraph::FMovieGraphOutputMergerFrame* InRawFrameData, const TSet<FMovieGraphRenderDataIdentifier>& InMask) override;
+
 #if WITH_EDITOR
 	virtual FText GetNodeTitle(const bool bGetDescriptive = false) const override 
 	{ 
@@ -114,6 +129,30 @@ public:
 	 */
 	UPROPERTY(EditAnywhere, BlueprintReadWrite, meta=(EditCondition="bOverride_Compression"), Category = "EXR")
 	EEXRCompressionFormat Compression;
+
+protected:
+
+	/** Convenience function to create a new EXR image write task, given a file name and compression format. */
+	TUniquePtr<FEXRImageWriteTask> CreateImageWriteTask(
+		FString InFileName,
+		EEXRCompressionFormat InCompression
+	) const;
+	
+	/** Convenience function to prepare the image write task's global file metadata. */
+	void PrepareTaskGlobalMetadata(
+		FEXRImageWriteTask& InOutImageTask,
+		UE::MovieGraph::FMovieGraphOutputMergerFrame* InRawFrameData,
+		TMap<FString, FString>& InMetadata
+	) const;
+
+	/** Convenience function to update the image write task for layer data. */
+	void UpdateTaskPerLayer(
+		FEXRImageWriteTask& InOutImageTask,
+		const UMovieGraphImageSequenceOutputNode* InParentNode,
+		FImagePixelData* InImageData,
+		int32 InLayerIndex,
+		const FString& InLayerName = {}
+	) const;
 };
 
 
