MovieGraph: Changed the FMovieGraphNamedResolution API to make it easier to work with from scripting. Separated creating from a profile vs. a custom one into two separate functions to make it clear to the user what they were doing (instead of the previous behavior where trying to use a profile that didn't exist would silently try to create it).

diff --git a/Engine/Plugins/MovieScene/MovieRenderPipeline/Source/MovieRenderPipelineCore/Private/Graph/MovieGraphBlueprintLibrary.cpp b/Engine/Plugins/MovieScene/MovieRenderPipeline/Source/MovieRenderPipelineCore/Private/Graph/MovieGraphBlueprintLibrary.cpp
--- a/Engine/Plugins/MovieScene/MovieRenderPipeline/Source/MovieRenderPipelineCore/Private/Graph/MovieGraphBlueprintLibrary.cpp
+++ b/Engine/Plugins/MovieScene/MovieRenderPipeline/Source/MovieRenderPipelineCore/Private/Graph/MovieGraphBlueprintLibrary.cpp
@@ -3,6 +3,7 @@
 #include "Graph/MovieGraphBlueprintLibrary.h"
 
 #include "Graph/MovieGraphPipeline.h"
+#include "Graph/MovieGraphProjectSettings.h"
 #include "Graph/Nodes/MovieGraphOutputSettingNode.h"
 #include "Graph/Nodes/MovieGraphRenderLayerNode.h"
 #include "Graph/Nodes/MovieGraphCameraNode.h"
@@ -523,4 +524,44 @@ UCineCameraComponent* UMovieGraphBlueprintLibrary::GetCurrentCineCamera(const UM
 	}
 
 	return nullptr;
+}
+
+FMovieGraphNamedResolution UMovieGraphBlueprintLibrary::NamedResolutionFromProfile(const FName& InResolutionProfileName)
+{
+	// Find a matching custom entry from Project Settings
+	if (const UMovieGraphProjectSettings* MovieGraphProjectSettings =
+		GetDefault<UMovieGraphProjectSettings>())
+	{
+		if (const FMovieGraphNamedResolution* Match = MovieGraphProjectSettings->FindNamedResolutionForOption(InResolutionProfileName))
+		{
+			return *Match;
+		}
+	}
+
+	// If we couldn't find it, throw an exception
+	FFrame::KismetExecutionMessage(
+		*FString::Printf(
+			TEXT("%hs: Could not find named resolution with profile name %s in Project Settings."), __FUNCTION__, *InResolutionProfileName.ToString()),
+		ELogVerbosity::Error);
+
+	return FMovieGraphNamedResolution();
+}
+
+bool UMovieGraphBlueprintLibrary::IsNamedResolutionValid(const FName& InResolutionProfileName)
+{
+	if (const UMovieGraphProjectSettings* MovieGraphProjectSettings =
+		GetDefault<UMovieGraphProjectSettings>())
+	{
+		if (const FMovieGraphNamedResolution* Match = MovieGraphProjectSettings->FindNamedResolutionForOption(InResolutionProfileName))
+		{
+			return true;
+		}
+	}
+
+	return false;
+}
+
+FMovieGraphNamedResolution UMovieGraphBlueprintLibrary::NamedResolutionFromSize(const int32 InResX, const int32 InResY)
+{
+	return FMovieGraphNamedResolution(FMovieGraphNamedResolution::CustomEntryName, FIntPoint(InResX, InResY), FString());
 }
\ No newline at end of file
diff --git a/Engine/Plugins/MovieScene/MovieRenderPipeline/Source/MovieRenderPipelineCore/Private/Graph/MovieGraphNamedResolution.cpp b/Engine/Plugins/MovieScene/MovieRenderPipeline/Source/MovieRenderPipelineCore/Private/Graph/MovieGraphNamedResolution.cpp
deleted file mode 100644
diff --git a/Engine/Plugins/MovieScene/MovieRenderPipeline/Source/MovieRenderPipelineCore/Private/Graph/Nodes/MovieGraphOutputSettingNode.cpp b/Engine/Plugins/MovieScene/MovieRenderPipeline/Source/MovieRenderPipelineCore/Private/Graph/Nodes/MovieGraphOutputSettingNode.cpp
--- a/Engine/Plugins/MovieScene/MovieRenderPipeline/Source/MovieRenderPipelineCore/Private/Graph/Nodes/MovieGraphOutputSettingNode.cpp
+++ b/Engine/Plugins/MovieScene/MovieRenderPipeline/Source/MovieRenderPipelineCore/Private/Graph/Nodes/MovieGraphOutputSettingNode.cpp
@@ -1,17 +1,42 @@
 // Copyright Epic Games, Inc. All Rights Reserved.
 
 #include "Graph/Nodes/MovieGraphOutputSettingNode.h"
-
 #include "Graph/MovieGraphProjectSettings.h"
+#include "Graph/MovieGraphBlueprintLibrary.h"
 #include "Styling/AppStyle.h"
+#include "Algo/Find.h"
 
 UMovieGraphOutputSettingNode::UMovieGraphOutputSettingNode()
-	: OutputResolution(FMovieGraphNamedResolution(FMovieGraphNamedResolution::DefaultResolutionName))
-	, OutputFrameRate(FFrameRate(24, 1))
+	: OutputFrameRate(FFrameRate(24, 1))
 	, bOverwriteExistingOutput(true)
 	, ZeroPadFrameNumbers(4)
 	, FrameNumberOffset(0)
 {
+	// We prefer a 1080p resolution by default, but users may not have a preset that matches that. So we'll 
+	// look for a matching resolution if we can find one, otherwise we go to Custom set to 1920x1080.
+	const FMovieGraphNamedResolution* FoundResolution = nullptr;
+	if (const UMovieGraphProjectSettings* MovieGraphProjectSettings =
+		GetDefault<UMovieGraphProjectSettings>())
+	{
+		FoundResolution = Algo::FindByPredicate(
+			MovieGraphProjectSettings->DefaultNamedResolutions,
+			[](const FMovieGraphNamedResolution& Other)
+			{
+				return Other.Resolution == FIntPoint(1920, 1080);
+			});
+	}
+
+	// If we found one that was 1080p, regardless of name, use it.
+	if (FoundResolution)
+	{
+		OutputResolution = *FoundResolution;
+	}
+	else
+	{
+		// We didn't find one that was 1080p, just force a custom resolution.
+		OutputResolution = UMovieGraphBlueprintLibrary::NamedResolutionFromSize(1920, 1080);
+	}
+
 	FileNameFormat = TEXT("{sequence_name}.{frame_number}");
 	OutputDirectory.Path = TEXT("{project_dir}/Saved/MovieRenders/");
 }
diff --git a/Engine/Plugins/MovieScene/MovieRenderPipeline/Source/MovieRenderPipelineCore/Public/Graph/MovieGraphBlueprintLibrary.h b/Engine/Plugins/MovieScene/MovieRenderPipeline/Source/MovieRenderPipelineCore/Public/Graph/MovieGraphBlueprintLibrary.h
--- a/Engine/Plugins/MovieScene/MovieRenderPipeline/Source/MovieRenderPipelineCore/Public/Graph/MovieGraphBlueprintLibrary.h
+++ b/Engine/Plugins/MovieScene/MovieRenderPipeline/Source/MovieRenderPipelineCore/Public/Graph/MovieGraphBlueprintLibrary.h
@@ -5,6 +5,7 @@
 #include "Kismet/BlueprintFunctionLibrary.h"
 #include "Graph/MovieGraphFilenameResolveParams.h"
 #include "Misc/FrameRate.h"
+#include "MovieGraphNamedResolution.h"
 
 #include "MovieGraphBlueprintLibrary.generated.h"
 
@@ -192,4 +193,27 @@ public:
 	*/
 	UFUNCTION(BlueprintPure, Category = "Movie Graph")
 	static UCineCameraComponent* GetCurrentCineCamera(const UMovieGraphPipeline* InMovieGraphPipeline);
+
+
+
+	/**
+	* Create a Named Resolution from the profile name. Throws a Kismet Exception if the profile name isn't found.
+	* The known profiles can be found in UMovieGraphProjectSettings's CDO.
+	*/
+	UFUNCTION(BlueprintPure, Category = "Movie Graph")
+	static FMovieGraphNamedResolution NamedResolutionFromProfile(const FName& InResolutionProfileName);
+
+	/**
+	* Utility function for checking if a given resolution profile name is valid, since NamedResolutionFromProfile
+	* will throw a kismet exception, but blueprints can't actually try/catch them.
+	*/
+	UFUNCTION(BlueprintPure, Category = "Movie Graph")
+	static bool IsNamedResolutionValid(const FName& InResolutionProfileName);
+
+	/**
+	* Create a Named Resolution from the given resolution. Given named resolution will be named "Custom".
+	*/
+	UFUNCTION(BlueprintPure, Category = "Movie Graph")
+	static FMovieGraphNamedResolution NamedResolutionFromSize(const int32 InResX, const int32 InResY);
+
 };
\ No newline at end of file
diff --git a/Engine/Plugins/MovieScene/MovieRenderPipeline/Source/MovieRenderPipelineCore/Public/Graph/MovieGraphNamedResolution.h b/Engine/Plugins/MovieScene/MovieRenderPipeline/Source/MovieRenderPipelineCore/Public/Graph/MovieGraphNamedResolution.h
--- a/Engine/Plugins/MovieScene/MovieRenderPipeline/Source/MovieRenderPipelineCore/Public/Graph/MovieGraphNamedResolution.h
+++ b/Engine/Plugins/MovieScene/MovieRenderPipeline/Source/MovieRenderPipelineCore/Public/Graph/MovieGraphNamedResolution.h
@@ -24,7 +24,6 @@ public:
 		, Description(InDescription)
 	{}
 
-	FMovieGraphNamedResolution(const FName& InResolutionProfileName);
 	
 	/** The name of the resolution this links to */
 	UPROPERTY(config, EditAnywhere, BlueprintReadWrite, Category = "Movie Graph|Resolution")
@@ -47,4 +46,9 @@ public:
 	 * The default resolution name to use when one is not defined.
 	 */
 	inline static FName DefaultResolutionName = TEXT("1080p (FHD)");
+
+	/**
+	 * Predefined name for the 'custom' resolution option in the combobox
+	 */
+	inline static FName CustomEntryName = TEXT("Custom");
 };
\ No newline at end of file
diff --git a/Engine/Plugins/MovieScene/MovieRenderPipeline/Source/MovieRenderPipelineCore/Public/Graph/MovieGraphProjectSettings.h b/Engine/Plugins/MovieScene/MovieRenderPipeline/Source/MovieRenderPipelineCore/Public/Graph/MovieGraphProjectSettings.h
--- a/Engine/Plugins/MovieScene/MovieRenderPipeline/Source/MovieRenderPipelineCore/Public/Graph/MovieGraphProjectSettings.h
+++ b/Engine/Plugins/MovieScene/MovieRenderPipeline/Source/MovieRenderPipelineCore/Public/Graph/MovieGraphProjectSettings.h
@@ -13,7 +13,7 @@
 /**
  * Settings that apply to the Movie Graph.
  */
-UCLASS(config = MovieRenderPipeline, DefaultConfig, meta=(DisplayName = "Movie Graph Settings"))
+UCLASS(BlueprintType, config = MovieRenderPipeline, DefaultConfig, meta=(DisplayName = "Movie Graph Settings"))
 class MOVIERENDERPIPELINECORE_API UMovieGraphProjectSettings : public UDeveloperSettings
 {
 public:
diff --git a/Engine/Plugins/MovieScene/MovieRenderPipeline/Source/MovieRenderPipelineEditor/Private/Customizations/Graph/MovieGraphNamedResolutionCustomization.h b/Engine/Plugins/MovieScene/MovieRenderPipeline/Source/MovieRenderPipelineEditor/Private/Customizations/Graph/MovieGraphNamedResolutionCustomization.h
--- a/Engine/Plugins/MovieScene/MovieRenderPipeline/Source/MovieRenderPipelineEditor/Private/Customizations/Graph/MovieGraphNamedResolutionCustomization.h
+++ b/Engine/Plugins/MovieScene/MovieRenderPipeline/Source/MovieRenderPipelineEditor/Private/Customizations/Graph/MovieGraphNamedResolutionCustomization.h
@@ -76,11 +76,11 @@ protected:
 			}
 		}
 
-		ComboBoxOptions.AddUnique(CustomEntryName);
+		ComboBoxOptions.AddUnique(FMovieGraphNamedResolution::CustomEntryName);
 
 		if (!IsCurrentOptionValid())
 		{
-			SetCurrentOptionAndCacheTooltipText(CustomEntryName);
+			SetCurrentOptionAndCacheTooltipText(FMovieGraphNamedResolution::CustomEntryName);
 		}
 	}
 
@@ -108,7 +108,7 @@ protected:
 	const FMovieGraphNamedResolution* FindNamedResolutionForOption(const FName& InOption) const
 	{
 		// Return the custom entry
-		if (InOption.IsEqual(CustomEntryName))
+		if (InOption.IsEqual(FMovieGraphNamedResolution::CustomEntryName))
 		{
 			return &CustomEntry;
 		}
@@ -171,7 +171,7 @@ protected:
 			SNew(SHorizontalBox)
 			.Visibility_Lambda([this]()
 			{
-				return GetCurrentOptionAssigningIfNeeded().IsEqual(CustomEntryName) ?
+				return GetCurrentOptionAssigningIfNeeded().IsEqual(FMovieGraphNamedResolution::CustomEntryName) ?
 					EVisibility::Visible : EVisibility::Collapsed;
 			})
 
@@ -361,15 +361,12 @@ protected:
 	 */
 	FDelegateHandle OnProjectSettingsModifiedHandle;
 
-	/**
-	 * Predefined name for the 'custom' resolution option in the combobox
-	 */
-	inline static FName CustomEntryName = TEXT("Custom");
+
 	/**
 	 * Default FMovieGraphNamedResolution for the 'custom' option
 	 */
 	inline static FMovieGraphNamedResolution CustomEntry =
-		FMovieGraphNamedResolution(CustomEntryName, FIntPoint(1920, 1080), CustomEntryName.ToString());	
+		FMovieGraphNamedResolution(FMovieGraphNamedResolution::CustomEntryName, FIntPoint(1920, 1080), FMovieGraphNamedResolution::CustomEntryName.ToString());
 };
 
 #undef LOCTEXT_NAMESPACE
