MovieGraph: Move pixel data in the common case where there is only one output image sequence node.

#rnx

diff --git a/Engine/Plugins/MovieScene/MovieRenderPipeline/Source/MovieRenderPipelineCore/Private/Graph/Nodes/MovieGraphFileOutputNode.cpp b/Engine/Plugins/MovieScene/MovieRenderPipeline/Source/MovieRenderPipelineCore/Private/Graph/Nodes/MovieGraphFileOutputNode.cpp
new file mode 100644
--- /dev/null
+++ b/Engine/Plugins/MovieScene/MovieRenderPipeline/Source/MovieRenderPipelineCore/Private/Graph/Nodes/MovieGraphFileOutputNode.cpp
@@ -0,0 +1,10 @@
+// Copyright Epic Games, Inc. All Rights Reserved.
+
+#include "Graph/Nodes/MovieGraphFileOutputNode.h"
+
+#include "Graph/MovieGraphConfig.h"
+
+int32 UMovieGraphFileOutputNode::GetNumFileOutputNodes(const UMovieGraphEvaluatedConfig& InEvaluatedConfig, const FName& InBranchName)
+{
+	return InEvaluatedConfig.GetSettingsForBranch(UMovieGraphFileOutputNode::StaticClass(), InBranchName, false /*bIncludeCDOs*/, false /*bExactMatch*/).Num();
+}
diff --git a/Engine/Plugins/MovieScene/MovieRenderPipeline/Source/MovieRenderPipelineCore/Public/Graph/Nodes/MovieGraphFileOutputNode.h b/Engine/Plugins/MovieScene/MovieRenderPipeline/Source/MovieRenderPipelineCore/Public/Graph/Nodes/MovieGraphFileOutputNode.h
--- a/Engine/Plugins/MovieScene/MovieRenderPipeline/Source/MovieRenderPipelineCore/Public/Graph/Nodes/MovieGraphFileOutputNode.h
+++ b/Engine/Plugins/MovieScene/MovieRenderPipeline/Source/MovieRenderPipelineCore/Public/Graph/Nodes/MovieGraphFileOutputNode.h
@@ -56,4 +56,7 @@ protected:
 	virtual void OnReceiveImageDataImpl(UMovieGraphPipeline* InPipeline, UE::MovieGraph::FMovieGraphOutputMergerFrame* InRawFrameData, const TSet<FMovieGraphRenderDataIdentifier>& InMask) {}
 	virtual void OnAllFramesSubmittedImpl() {}
 	virtual bool IsFinishedWritingToDiskImpl() const { return true; }
+	
+	/** Returns the number of evaluated (active) file nodes on the specified branch. */
+	static int32 GetNumFileOutputNodes(const UMovieGraphEvaluatedConfig& InEvaluatedConfig, const FName& InBranchName);
 };
\ No newline at end of file
diff --git a/Engine/Plugins/MovieScene/MovieRenderPipeline/Source/MovieRenderPipelineRenderPasses/Private/MovieGraphImageSequenceOutputNode.cpp b/Engine/Plugins/MovieScene/MovieRenderPipeline/Source/MovieRenderPipelineRenderPasses/Private/MovieGraphImageSequenceOutputNode.cpp
--- a/Engine/Plugins/MovieScene/MovieRenderPipeline/Source/MovieRenderPipelineRenderPasses/Private/MovieGraphImageSequenceOutputNode.cpp
+++ b/Engine/Plugins/MovieScene/MovieRenderPipeline/Source/MovieRenderPipelineRenderPasses/Private/MovieGraphImageSequenceOutputNode.cpp
@@ -206,9 +206,9 @@ void UMovieGraphImageSequenceOutputNode::OnReceiveImageDataImpl(UMovieGraphPipel
 
 	// ToDo:
 	// The ImageWriteQueue is set up in a fire-and-forget manner. This means that the data needs to be placed in the WriteQueue
-	// as a TUniquePtr (so it can free the data when its done). Unfortunately we can have multiple output formats at once,
-	// so we can't MoveTemp the data into it, we need to make a copy (though we could optimize for the common case where there is
-	// only one output format).
+	// as a TUniquePtr (so it can free the data when its done). Unfortunately if we have multiple output formats at once,
+	// we can't MoveTemp the data so we need to make a copy.
+	// 
 	// Copying can be expensive (3ms @ 1080p, 12ms at 4k for a single layer image) so ideally we'd like to do it on the task graph
 	// but this isn't really compatible with the ImageWriteQueue API as we need the future returned by the ImageWriteQueue to happen
 	// in order, so that we push our futures to the main Movie Pipeline in order, otherwise when we encode files to videos they'll
@@ -234,6 +234,8 @@ void UMovieGraphImageSequenceOutputNode::OnReceiveImageDataImpl(UMovieGraphPipel
 			continue;
 		}
 
+		checkf(RenderData.Value.IsValid(), TEXT("Unexpected empty image data: incorrectly moved or its production failed?"));
+
 		// ToDo: Certain images may require transparency, at which point
 		// we write out a .png instead of a .jpeg.
 		EImageFormat PreferredOutputFormat = OutputFormat;
@@ -249,7 +251,16 @@ void UMovieGraphImageSequenceOutputNode::OnReceiveImageDataImpl(UMovieGraphPipel
 		TileImageTask->Format = PreferredOutputFormat;
 		TileImageTask->CompressionQuality = 100;
 		TileImageTask->Filename = FileName;
-		TileImageTask->PixelData = RenderData.Value->CopyImageData();
+
+		// Pixel data can only be moved if there are no other active output image sequence nodes on the branch
+		if (GetNumFileOutputNodes(*InRawFrameData->EvaluatedConfig, RenderData.Key.RootBranchName) > 1)
+		{
+			TileImageTask->PixelData = RenderData.Value->CopyImageData();
+		}
+		else
+		{
+			TileImageTask->PixelData = RenderData.Value->MoveImageDataToNew();
+		}
 
 		const UMovieGraphImageSequenceOutputNode* ParentNode = Cast<UMovieGraphImageSequenceOutputNode>(
 			InRawFrameData->EvaluatedConfig->GetSettingForBranch(GetClass(), RenderData.Key.RootBranchName, false, true));
@@ -354,15 +365,12 @@ void UMovieGraphImageSequenceOutputNode_EXR::PrepareTaskGlobalMetadata(FEXRImage
 void UMovieGraphImageSequenceOutputNode_EXR::UpdateTaskPerLayer(
 	FEXRImageWriteTask& InOutImageTask,
 	const UMovieGraphImageSequenceOutputNode* InParentNode,
-	FImagePixelData* InImageData,
+	TUniquePtr<FImagePixelData> InImageData,
 	int32 InLayerIndex,
 	const FString& InLayerName) const
 {
 	const UE::MovieGraph::FMovieGraphSampleState* Payload = InImageData->GetPayload<UE::MovieGraph::FMovieGraphSampleState>();
 
-	// No quantization required, just copy the data as we will move it into the image write task.
-	TUniquePtr<FImagePixelData> PixelData = InImageData->CopyImageData();
-
 	bool bEnabledOCIO = false;
 #if WITH_EDITOR
 	if (InParentNode->OCIOConfiguration.bIsEnabled && Payload->bAllowOCIO)
@@ -383,7 +391,7 @@ void UMovieGraphImageSequenceOutputNode_EXR::UpdateTaskPerLayer(
 		InOutImageTask.FileMetadata.Add("owner", UE::MoviePipeline::GetJobAuthor(Payload->TraversalContext.Job));
 		InOutImageTask.FileMetadata.Add("comments", Payload->TraversalContext.Job->Comment);
 
-		const FIntPoint& Resolution = PixelData->GetSize();
+		const FIntPoint& Resolution = InImageData->GetSize();
 		InOutImageTask.Width = Resolution.X;
 		InOutImageTask.Height = Resolution.Y;
 
@@ -402,10 +410,10 @@ void UMovieGraphImageSequenceOutputNode_EXR::UpdateTaskPerLayer(
 
 	if (!InLayerName.IsEmpty())
 	{
-		InOutImageTask.LayerNames.FindOrAdd(PixelData.Get(), InLayerName);
+		InOutImageTask.LayerNames.FindOrAdd(InImageData.Get(), InLayerName);
 	}
 
-	InOutImageTask.Layers.Add(MoveTemp(PixelData));
+	InOutImageTask.Layers.Add(MoveTemp(InImageData));
 }
 
 void UMovieGraphImageSequenceOutputNode_EXR::OnReceiveImageDataImpl(UMovieGraphPipeline* InPipeline, UE::MovieGraph::FMovieGraphOutputMergerFrame* InRawFrameData, const TSet<FMovieGraphRenderDataIdentifier>& InMask)
@@ -431,6 +439,8 @@ void UMovieGraphImageSequenceOutputNode_EXR::OnReceiveImageDataImpl(UMovieGraphP
 			continue;
 		}
 
+		checkf(RenderData.Value.IsValid(), TEXT("Unexpected empty image data: incorrectly moved or its production failed?"));
+
 		FMovieGraphResolveArgs ResolvedFormatArgs;
 		FString FileName = CreateFileName(InRawFrameData, InPipeline, RenderData, OutputFormat, ResolvedFormatArgs);
 		if (!ensureMsgf(!FileName.IsEmpty(), TEXT("Unexpected empty file name, skipping frame.")))
@@ -450,7 +460,16 @@ void UMovieGraphImageSequenceOutputNode_EXR::OnReceiveImageDataImpl(UMovieGraphP
 
 		// No layer is equivalent to a zero-index layer
 		constexpr int32 LayerIndex = 0;
-		UpdateTaskPerLayer(*ImageWriteTask, ParentNode, RenderData.Value.Get(), LayerIndex);
+		TUniquePtr<FImagePixelData> PixelData;
+		if (GetNumFileOutputNodes(*InRawFrameData->EvaluatedConfig, RenderData.Key.RootBranchName) > 1)
+		{
+			PixelData = RenderData.Value->CopyImageData();
+		}
+		else
+		{
+			PixelData = RenderData.Value->MoveImageDataToNew();
+		}
+		UpdateTaskPerLayer(*ImageWriteTask, ParentNode, MoveTemp(PixelData), LayerIndex);
 
 		// Perform compositing if any composited passes were found earlier
 		for (TPair<FMovieGraphRenderDataIdentifier, TUniquePtr<FImagePixelData>>& CompositedPass : CompositedPasses)
@@ -505,6 +524,8 @@ void UMovieGraphImageSequenceOutputNode_MultiLayerEXR::OnReceiveImageDataImpl(UM
 		for (const FMovieGraphRenderDataIdentifier& RenderID : RenderIDs)
 		{
 			const TUniquePtr<FImagePixelData>& ImageData = InRawFrameData->ImageOutputData[RenderID];
+			checkf(ImageData.IsValid(), TEXT("Unexpected empty image data: incorrectly moved or its production failed?"));
+
 			const UE::MovieGraph::FMovieGraphSampleState* Payload = ImageData->GetPayload<UE::MovieGraph::FMovieGraphSampleState>();
 			ShotIndex = Payload->TraversalContext.ShotIndex;
 
@@ -526,7 +547,16 @@ void UMovieGraphImageSequenceOutputNode_MultiLayerEXR::OnReceiveImageDataImpl(UM
 				}
 			}
 
-			UpdateTaskPerLayer(*MultiLayerImageTask, ParentNode, ImageData.Get(), LayerIndex, LayerName);
+			TUniquePtr<FImagePixelData> PixelData;
+			if (GetNumFileOutputNodes(*InRawFrameData->EvaluatedConfig, RenderID.RootBranchName) > 1)
+			{
+				PixelData = ImageData->CopyImageData();
+			}
+			else
+			{
+				PixelData = ImageData->MoveImageDataToNew();
+			}
+			UpdateTaskPerLayer(*MultiLayerImageTask, ParentNode, MoveTemp(PixelData), LayerIndex, LayerName);
 
 			LayerIndex++;
 		}
diff --git a/Engine/Plugins/MovieScene/MovieRenderPipeline/Source/MovieRenderPipelineRenderPasses/Public/MovieGraphImageSequenceOutputNode.h b/Engine/Plugins/MovieScene/MovieRenderPipeline/Source/MovieRenderPipelineRenderPasses/Public/MovieGraphImageSequenceOutputNode.h
--- a/Engine/Plugins/MovieScene/MovieRenderPipeline/Source/MovieRenderPipelineRenderPasses/Public/MovieGraphImageSequenceOutputNode.h
+++ b/Engine/Plugins/MovieScene/MovieRenderPipeline/Source/MovieRenderPipelineRenderPasses/Public/MovieGraphImageSequenceOutputNode.h
@@ -149,7 +149,7 @@ protected:
 	void UpdateTaskPerLayer(
 		FEXRImageWriteTask& InOutImageTask,
 		const UMovieGraphImageSequenceOutputNode* InParentNode,
-		FImagePixelData* InImageData,
+		TUniquePtr<FImagePixelData> InImageData,
 		int32 InLayerIndex,
 		const FString& InLayerName = {}
 	) const;
