Movie Graph:

-Named Resolutions customization polish

diff --git a/Engine/Plugins/MovieScene/MovieRenderPipeline/Source/MovieRenderPipelineEditor/Private/Customizations/Graph/MovieGraphNamedResolutionCustomization.h b/Engine/Plugins/MovieScene/MovieRenderPipeline/Source/MovieRenderPipelineEditor/Private/Customizations/Graph/MovieGraphNamedResolutionCustomization.h
--- a/Engine/Plugins/MovieScene/MovieRenderPipeline/Source/MovieRenderPipelineEditor/Private/Customizations/Graph/MovieGraphNamedResolutionCustomization.h
+++ b/Engine/Plugins/MovieScene/MovieRenderPipeline/Source/MovieRenderPipelineEditor/Private/Customizations/Graph/MovieGraphNamedResolutionCustomization.h
@@ -8,7 +8,9 @@
 
 #include "DetailLayoutBuilder.h"
 #include "DetailWidgetRow.h"
+#include "IDetailChildrenBuilder.h"
 #include "IPropertyTypeCustomization.h"
+#include "Widgets/Input/SButton.h"
 #include "Widgets/Input/SComboBox.h"
 #include "Widgets/SBoxPanel.h"
 
@@ -55,6 +57,10 @@ protected:
 
 	void SetCurrentOptionAndCacheTooltipText(const FName& InOption)
 	{
+		if (ComboBox)
+		{
+			ComboBox->SetSelectedItem(InOption);
+		}
 		CurrentOption = InOption;
 		CachedSelectedOptionTooltipText = GetOptionTooltipText(CurrentOption);
 	}
@@ -93,6 +99,10 @@ protected:
 	 */
 	FText GetOptionTooltipText(FName InOption) const
 	{
+		if (InOption.IsEqual(CurrentOption))
+		{
+			return CachedSelectedOptionTooltipText;
+		}
 		if (const FMovieGraphNamedResolution* Match = FindNamedResolutionForOption(InOption))
 		{
 			return FText::FromString(Match->Description);
@@ -124,6 +134,29 @@ protected:
 		return nullptr;
 	}
 
+	FText GetWidgetTextForOption(const FName& InOption) const
+	{
+		const FMovieGraphNamedResolution* NamedResolution = FindNamedResolutionForOption(InOption);
+		if (!ensureAlwaysMsgf(NamedResolution, TEXT("%hs: Failed to find FMovieGraphNamedResolution for option %s!"), __FUNCTION__, *InOption.ToString()))
+		{
+			return LOCTEXT("UnknownNamedResolutionComboboxItemWidgetText", "Unknown Named Resolution");
+		}
+
+
+		FNumberFormattingOptions FormattingOptions;
+		FormattingOptions.UseGrouping = false;
+		
+		return FText::Format(LOCTEXT("NamedResolutionComboboxItemWidgetText", "{0} - {1}x{2}"),
+				FText::FromName(InOption),
+				FText::AsNumber(NamedResolution->Resolution.X, &FormattingOptions),
+				FText::AsNumber(NamedResolution->Resolution.Y, &FormattingOptions));
+	}
+
+	FSlateFontInfo GetWidgetFontForOptions() const
+	{
+		return IDetailLayoutBuilder::GetDetailFont();
+	}
+
 	/**
 	 * Generates a widget for the given option with the option name as STextBlock with tooltip text .
 	 * @param InOption The name of the option.
@@ -133,15 +166,15 @@ protected:
 	{
 		return
 			SNew(STextBlock)
-			.Text(FText::FromName(InOption))
+			.Text(GetWidgetTextForOption(InOption))
 			.ToolTipText(this, &FMovieGraphNamedResolutionCustomization::GetOptionTooltipText, InOption)
-			.Font(IDetailLayoutBuilder::GetDetailFont())
+			.Font(this, &FMovieGraphNamedResolutionCustomization::GetWidgetFontForOptions);
 		;
 	}
 
 	TSharedRef<SWidget> CreateOptionComboBox()
 	{
-		return
+		ComboBox = 
 			SNew(SComboBox<FName>)
 			.OnComboBoxOpening(this, &FMovieGraphNamedResolutionCustomization::RepopulateOptions)
 			.OptionsSource(&ComboBoxOptions)
@@ -149,64 +182,20 @@ protected:
 			.OnGenerateWidget(this, &FMovieGraphNamedResolutionCustomization::MakeWidgetForOption)
 			.InitiallySelectedItem(GetCurrentOptionAssigningIfNeeded())
 			[
+				// We specifically don't use MakeWidgetForOption here for performance reasons
+				// as it will need to be called each frame to stay current (especially for 'Custom')
 				SNew(STextBlock)
 				.Text_Lambda([this]()
 				{
-					return FText::FromName(GetCurrentOptionAssigningIfNeeded());
+					return GetWidgetTextForOption(GetCurrentOptionAssigningIfNeeded());
 				})
 				.ToolTipText_Lambda([this]()
 				{
 					return CachedSelectedOptionTooltipText;
 				})
-				.Font(IDetailLayoutBuilder::GetDetailFont())
-			];
-	}
-
-	TSharedRef<SWidget> CreateCustomResolutionWidget(const TSharedRef<IPropertyHandle> InStructPropertyHandle)
-	{
-		const TSharedRef<IPropertyHandle> ResolutionPropHandle = InStructPropertyHandle->GetChildHandle(
-			GET_MEMBER_NAME_CHECKED(FMovieGraphNamedResolution, Resolution)).ToSharedRef();
-		
-		return
-			SNew(SHorizontalBox)
-			.Visibility_Lambda([this]()
-			{
-				return GetCurrentOptionAssigningIfNeeded().IsEqual(FMovieGraphNamedResolution::CustomEntryName) ?
-					EVisibility::Visible : EVisibility::Collapsed;
-			})
-
-			+SHorizontalBox::Slot()
-			.AutoWidth()
-			.VAlign(VAlign_Center)
-			.Padding(0, 0, 8, 0)
-			[
-				SNew(STextBlock)
-				.Text(FText::FromString(TEXT("Width")))
-				.Font(IDetailLayoutBuilder::GetDetailFont())
-			]
-
-			+SHorizontalBox::Slot()
-			.HAlign(HAlign_Fill)
-			.Padding(0, 0, 8, 0)
-			[
-				ResolutionPropHandle->GetChildHandle(0)->CreatePropertyValueWidget()
-			]
-
-			+SHorizontalBox::Slot()
-			.AutoWidth()
-			.VAlign(VAlign_Center)
-			.Padding(0, 0, 8, 0)
-			[
-				SNew(STextBlock)
-				.Text(FText::FromString(TEXT("Height")))
-				.Font(IDetailLayoutBuilder::GetDetailFont())
-			]
-
-			+SHorizontalBox::Slot()
-			.HAlign(HAlign_Fill)
-			[
-				ResolutionPropHandle->GetChildHandle(1)->CreatePropertyValueWidget()
+				.Font(this, &FMovieGraphNamedResolutionCustomization::GetWidgetFontForOptions)
 			];
+		return ComboBox.ToSharedRef(); 
 	}
 
 	void BindToOnProjectSettingsModified()
@@ -219,6 +208,18 @@ protected:
 					this, &FMovieGraphNamedResolutionCustomization::OnProjectSettingsChanged);
 		}
 	}
+
+	const FSlateBrush* GetLockBrush() const
+	{
+		return bLockedAspectRatio ? FAppStyle::GetBrush(TEXT("Icons.Link")) : FAppStyle::GetBrush(TEXT("Icons.Unlink"));
+	}
+
+	FText GetLockTooltipText() const
+	{
+		return bLockedAspectRatio ?
+			FText::Format(LOCTEXT("LockedAspectRatioTooltipFormat", "Aspect ratio is locked ({0})"), FText::AsNumber(CurrentAspectRatio)) :
+			LOCTEXT("UnlockedAspectRatioTooltip", "Aspect Ratio is unlocked");
+	}
 	
 	//~ Begin IPropertyTypeCustomization interface
 	virtual void CustomizeHeader(
@@ -226,6 +227,7 @@ protected:
 		FDetailWidgetRow& HeaderRow, IPropertyTypeCustomizationUtils& CustomizationUtils) override
 	{
 		StructPropertyHandle = InStructPropertyHandle;
+		CreateOptionComboBox();
 		
 		const TSharedRef<IPropertyHandle> NameProp = InStructPropertyHandle->GetChildHandle(
 			GET_MEMBER_NAME_CHECKED(FMovieGraphNamedResolution, ProfileName)).ToSharedRef();
@@ -243,6 +245,9 @@ protected:
 		// Repopulate options when the project settings change
 		BindToOnProjectSettingsModified();
 
+		// Prevent reset to default from being shown
+		HeaderRow.OverrideResetToDefault(FResetToDefaultOverride::Hide());
+			
 		HeaderRow
 		.NameContent()
 		[
@@ -252,30 +257,114 @@ protected:
 		.ValueContent()
 		.HAlign(HAlign_Fill)
 		[
-			SNew(SVerticalBox)
+			SNew(SHorizontalBox)
 			
-			+ SVerticalBox::Slot()
+			+ SHorizontalBox::Slot()
 			.Padding(0, 2, 8, 2)
-			.AutoHeight()
 			.HAlign(HAlign_Fill)
 			[
-				CreateOptionComboBox()
+				ComboBox.ToSharedRef()
 			]
+		];
+	}
 
-			+ SVerticalBox::Slot()
-			.Padding(0, 2, 8, 2)
-			.AutoHeight()
-			.HAlign(HAlign_Fill)
-			[
-				CreateCustomResolutionWidget(InStructPropertyHandle)
-			]
+	void UpdateLockImage()
+	{
+		LockImage->SetImage(GetLockBrush());
+		LockImage->SetToolTipText(GetLockTooltipText());
+	}
+
+	void ToggleAspectRatioLock()
+	{
+		bLockedAspectRatio = !bLockedAspectRatio;
+		if (bLockedAspectRatio) { CacheAspectRatio(); }
+
+		UpdateLockImage();
+	}
+
+	TSharedRef<SWidget> MakeLockExtensionWidget()
+	{
+		return SNew(SButton)
+		.OnClicked_Lambda([this]()
+		{
+			ToggleAspectRatioLock();
+			
+			return FReply::Handled();
+		})
+		.ContentPadding(FMargin(0, 0, 4, 0))
+		.ButtonStyle(FAppStyle::Get(), "NoBorder")
+		[
+			SAssignNew(LockImage, SImage)
+			.ColorAndOpacity(FSlateColor::UseForeground())
 		];
 	}
 
-	// Skip customizing children
 	virtual void CustomizeChildren(
 		TSharedRef<IPropertyHandle> InStructPropertyHandle,
-		IDetailChildrenBuilder& StructBuilder, IPropertyTypeCustomizationUtils& StructCustomizationUtils) override { }
+		IDetailChildrenBuilder& StructBuilder, IPropertyTypeCustomizationUtils& StructCustomizationUtils) override
+	{
+		const TSharedPtr<IPropertyHandle> ResolutionPropertyHandle =
+			InStructPropertyHandle->GetChildHandle(GET_MEMBER_NAME_CHECKED(FMovieGraphNamedResolution, Resolution));
+
+		ResolutionPropertyHandle->SetOnChildPropertyValuePreChange(
+			FSimpleDelegate::CreateSP(this, &FMovieGraphNamedResolutionCustomization::OnCustomResolutionPreManualChange));
+		ResolutionPropertyHandle->SetOnChildPropertyValueChangedWithData(TDelegate<void(const FPropertyChangedEvent&)>::CreateSP(
+			this, &FMovieGraphNamedResolutionCustomization::OnCustomResolutionManualChange));
+
+		if (const TSharedPtr<IPropertyHandle> PropertyXHandle = ResolutionPropertyHandle->GetChildHandle(0))
+		{			
+			IDetailPropertyRow& PropertyXRow = StructBuilder.AddProperty(PropertyXHandle.ToSharedRef());
+			ResolutionXPropertyHandle = PropertyXHandle;
+
+			// Prevent showing Reset to Default
+			PropertyXRow.OverrideResetToDefault(FResetToDefaultOverride::Hide());
+
+			PropertyXRow.CustomWidget()
+			.NameContent()
+			[
+				PropertyXHandle->CreatePropertyNameWidget()
+			]
+			.ValueContent()
+			.HAlign(HAlign_Fill)
+			[
+				PropertyXHandle->CreatePropertyValueWidget()
+			]
+			.ExtensionContent()
+			[
+				MakeLockExtensionWidget()
+			];
+
+			UpdateLockImage();
+		}
+
+		if (const TSharedPtr<IPropertyHandle> PropertyYHandle = ResolutionPropertyHandle->GetChildHandle(1))
+		{			
+			IDetailPropertyRow& PropertyYRow = StructBuilder.AddProperty(PropertyYHandle.ToSharedRef());
+			ResolutionYPropertyHandle = PropertyYHandle;
+
+			// Prevent showing Reset to Default
+			PropertyYRow.OverrideResetToDefault(FResetToDefaultOverride::Hide());
+
+			// Custom widget to enforce length and padding on value widget
+			PropertyYRow.CustomWidget()
+			.NameContent()
+			[
+				PropertyYHandle->CreatePropertyNameWidget()
+			]
+			.ValueContent()
+			.HAlign(HAlign_Fill)
+			[
+				SNew(SBox)
+				.HAlign(HAlign_Fill)
+				.Padding(FMargin(0, 0, 24, 0))
+				[
+					PropertyYHandle->CreatePropertyValueWidget()
+				]
+			];
+		}
+
+		CacheAspectRatio();
+	}
 	//~ End IPropertyTypeCustomization interface
 	
 	void OnProjectSettingsChanged(UObject*, FPropertyChangedEvent& PropertyChangedEvent)
@@ -286,6 +375,98 @@ protected:
 		}
 	}
 
+	void CacheAspectRatio()
+	{
+		CurrentAspectRatio = (double)CustomEntry.Resolution.X / CustomEntry.Resolution.Y;
+	}
+
+	void UpdateCustomEntryValues()
+	{
+		if (!ensureMsgf(ResolutionXPropertyHandle, TEXT("%hs: `ResolutionXPropertyHandle` is null."), __FUNCTION__))
+		{
+			return;
+		}
+		if (!ensureMsgf(ResolutionYPropertyHandle, TEXT("%hs: `ResolutionYPropertyHandle` is null."), __FUNCTION__))
+		{
+			return;
+		}
+			
+		// Set the custom resolution to the preset resolution
+		int32 ResolutionX = 0;
+		ResolutionXPropertyHandle->GetValue(ResolutionX);
+
+		CustomEntry.Resolution.X = ResolutionX;
+
+		int32 ResolutionY = 0;
+		ResolutionYPropertyHandle->GetValue(ResolutionY);
+
+		CustomEntry.Resolution.Y = ResolutionY;
+
+		// Cache the current aspect ratio
+		CacheAspectRatio();
+	}
+
+	void OnCustomResolutionPreManualChange()
+	{
+		// If the custom resolution is changed, switch to 'Custom' option
+		if (!CurrentOption.IsEqual(FMovieGraphNamedResolution::CustomEntryName))
+		{
+			UpdateCustomEntryValues();
+
+			// Switch to 'Custom' option
+			SetCurrentOptionAndCacheTooltipText(FMovieGraphNamedResolution::CustomEntryName);
+		}
+	}
+
+	int32 RoundToNearestEvenNumber(double InNumberToRound) const
+	{
+		// Round to nearest even
+		const int32 Floored = FMath::FloorToInt(InNumberToRound);
+		const int32 Ceiled = FMath::CeilToInt(InNumberToRound);
+
+		return Floored % 2 == 0 ? Floored : Ceiled;
+	}
+
+	void OnCustomResolutionManualChange(const FPropertyChangedEvent& Event) const
+	{
+		if (!ensureMsgf(ResolutionXPropertyHandle, TEXT("%hs: `ResolutionXPropertyHandle` is null."), __FUNCTION__))
+		{
+			return;
+		}
+		if (!ensureMsgf(ResolutionYPropertyHandle, TEXT("%hs: `ResolutionYPropertyHandle` is null."), __FUNCTION__))
+		{
+			return;
+		}
+
+		// Update CustomEntry values
+		int32 ResolutionX = 0;
+		ResolutionXPropertyHandle->GetValue(ResolutionX);
+
+		CustomEntry.Resolution.X = ResolutionX;
+
+		int32 ResolutionY = 0;
+		ResolutionYPropertyHandle->GetValue(ResolutionY);
+
+		CustomEntry.Resolution.Y = ResolutionY;
+
+		// Enforce aspect ratio if desired
+		if (bLockedAspectRatio && Event.ChangeType != EPropertyChangeType::Interactive)
+		{
+			if (Event.Property == ResolutionXPropertyHandle->GetProperty()) // X Changed
+			{
+				const double NewYResolution = (double)CustomEntry.Resolution.X / CurrentAspectRatio;
+
+				// Specify an interactive change to avoid feedback loop
+				ResolutionYPropertyHandle->SetValue(RoundToNearestEvenNumber(NewYResolution), EPropertyValueSetFlags::InteractiveChange);
+			}
+			else // Y Changed
+			{
+				const double NewXResolution = (double)CustomEntry.Resolution.Y * CurrentAspectRatio;
+				ResolutionXPropertyHandle->SetValue(RoundToNearestEvenNumber(NewXResolution), EPropertyValueSetFlags::InteractiveChange);
+			}
+		}
+	}
+
 	/**
 	 * Updates the struct which owns this customization with the data from the matched named resolution.
 	 */
@@ -334,22 +515,57 @@ protected:
 		SetCurrentOptionAndCacheTooltipText(NewValue);
 
 		UpdateOwningStruct();
+
+		CacheAspectRatio();
 	}
 
+	/**
+	 * The widget that allows the end user to select the output resolution
+	 */
+	TSharedPtr<SComboBox<FName>> ComboBox;
+
+	/**
+	 * The image that indicates whether the aspect ratio is locked
+	 */
+	TSharedPtr<SImage> LockImage;
+
 	/**
 	 * Weak ptr to the handle of the struct which owns this customization
 	 */
 	TWeakPtr<IPropertyHandle> StructPropertyHandle;
 
+	/**
+	 * Handle of the widget representing the X dimension of the 'Custom' resolution
+	 */
+	TSharedPtr<IPropertyHandle> ResolutionXPropertyHandle;
+	/**
+	 * Handle of the widget representing the Y dimension of the 'Custom' resolution
+	 */
+	TSharedPtr<IPropertyHandle> ResolutionYPropertyHandle;
+
 	/**
 	 * The option currently selected in the combobox
 	 */
 	FName CurrentOption = NAME_None;
+	
 	/**
 	 * The tooltip text for the currently selected option
 	 */
 	FText CachedSelectedOptionTooltipText;
 
+	/**
+	 * If true, changing one dimension of the custom resolution will also
+	 * change the other dimension to maintain the aspect ratio.
+	 */
+	bool bLockedAspectRatio = true;
+
+	/**
+	 * The aspect ratio of the currently selected resolution.
+	 * Updates when the aspect ratio lock is applied, a named resolution is selected,
+	 * or when the 'Custom' resolution values are changed.
+	 */
+	float CurrentAspectRatio = 0.0f;
+
 	/**
 	 * The options for the combo box, populated from UMovieGraphProjectSettings::DefaultNamedResolutions.
 	 * A 'custom' entry is added to the end of the list, representing an arbitrary user-defined resolution.
