MovieGraph: Dynamically allocate accumulators to avoid scenarios where you have too many needed for a single frame (which prevents any of them from actually starting, effectively deadlocking the rendering process).

diff --git a/Engine/Plugins/MovieScene/MovieRenderPipeline/Source/MovieRenderPipelineCore/Private/Graph/MovieGraphDefaultRenderer.cpp b/Engine/Plugins/MovieScene/MovieRenderPipeline/Source/MovieRenderPipelineCore/Private/Graph/MovieGraphDefaultRenderer.cpp
--- a/Engine/Plugins/MovieScene/MovieRenderPipeline/Source/MovieRenderPipelineCore/Private/Graph/MovieGraphDefaultRenderer.cpp
+++ b/Engine/Plugins/MovieScene/MovieRenderPipeline/Source/MovieRenderPipelineCore/Private/Graph/MovieGraphDefaultRenderer.cpp
@@ -437,45 +437,3 @@ UTexture* UMovieGraphDefaultRenderer::GetPreviewTexture() const
 
 	return nullptr;
 }
-
-
-namespace UE::MovieGraph::DefaultRenderer
-{
-FSurfaceAccumulatorPool::FInstancePtr FSurfaceAccumulatorPool::BlockAndGetAccumulator_GameThread(int32 InFrameNumber, const FMovieGraphRenderDataIdentifier& InPassIdentifier)
-{
-	FScopeLock ScopeLock(&CriticalSection);
-
-	int32 AvailableIndex = INDEX_NONE;
-	while (AvailableIndex == INDEX_NONE)
-	{
-		for (int32 Index = 0; Index < Accumulators.Num(); Index++)
-		{
-			if (InFrameNumber == Accumulators[Index]->ActiveFrameNumber && InPassIdentifier == Accumulators[Index]->ActivePassIdentifier)
-			{
-				AvailableIndex = Index;
-				break;
-			}
-		}
-
-		if (AvailableIndex == INDEX_NONE)
-		{
-			// If we don't have an accumulator already working on it let's look for a free one.
-			for (int32 Index = 0; Index < Accumulators.Num(); Index++)
-			{
-				if (!Accumulators[Index]->IsActive())
-				{
-					// Found a free one, tie it to this output frame.
-					Accumulators[Index]->ActiveFrameNumber = InFrameNumber;
-					Accumulators[Index]->ActivePassIdentifier = InPassIdentifier;
-					Accumulators[Index]->bIsActive = true;
-					Accumulators[Index]->TaskPrereq = UE::Tasks::FTask();
-					AvailableIndex = Index;
-					break;
-				}
-			}
-		}
-	}
-
-	return Accumulators[AvailableIndex];
-}
-}
diff --git a/Engine/Plugins/MovieScene/MovieRenderPipeline/Source/MovieRenderPipelineCore/Private/Graph/Renderers/MovieGraphImagePassBase.cpp b/Engine/Plugins/MovieScene/MovieRenderPipeline/Source/MovieRenderPipelineCore/Private/Graph/Renderers/MovieGraphImagePassBase.cpp
--- a/Engine/Plugins/MovieScene/MovieRenderPipeline/Source/MovieRenderPipelineCore/Private/Graph/Renderers/MovieGraphImagePassBase.cpp
+++ b/Engine/Plugins/MovieScene/MovieRenderPipeline/Source/MovieRenderPipelineCore/Private/Graph/Renderers/MovieGraphImagePassBase.cpp
@@ -479,12 +479,8 @@ void FMovieGraphImagePassBase::PostRendererSubmission(
 	}
 
 	FMoviePipelineAccumulatorPoolPtr SampleAccumulatorPool = GraphRenderer->GetOrCreateAccumulatorPool<FImageOverlappedAccumulator>();
-	UE::MovieGraph::DefaultRenderer::FSurfaceAccumulatorPool::FInstancePtr AccumulatorInstance = nullptr;
-	{
-		// SCOPE_CYCLE_COUNTER(STAT_MoviePipeline_WaitForAvailableAccumulator);
-		AccumulatorInstance = SampleAccumulatorPool->BlockAndGetAccumulator_GameThread(InSampleState.TraversalContext.Time.RenderedFrameNumber, InSampleState.TraversalContext.RenderDataIdentifier);
-	}
-
+	UE::MovieGraph::DefaultRenderer::FSurfaceAccumulatorPool::FInstancePtr AccumulatorInstance = SampleAccumulatorPool->GetAccumulatorInstance_GameThread<FImageOverlappedAccumulator>(InSampleState.TraversalContext.Time.RenderedFrameNumber, InSampleState.TraversalContext.RenderDataIdentifier);
+	
 	FMoviePipelineSurfaceQueuePtr LocalSurfaceQueue = GraphRenderer->GetOrCreateSurfaceQueue(InRenderTargetInitParams);
 	LocalSurfaceQueue->BlockUntilAnyAvailable();
 
diff --git a/Engine/Plugins/MovieScene/MovieRenderPipeline/Source/MovieRenderPipelineCore/Public/Graph/MovieGraphDefaultRenderer.h b/Engine/Plugins/MovieScene/MovieRenderPipeline/Source/MovieRenderPipelineCore/Public/Graph/MovieGraphDefaultRenderer.h
--- a/Engine/Plugins/MovieScene/MovieRenderPipeline/Source/MovieRenderPipelineCore/Public/Graph/MovieGraphDefaultRenderer.h
+++ b/Engine/Plugins/MovieScene/MovieRenderPipeline/Source/MovieRenderPipelineCore/Public/Graph/MovieGraphDefaultRenderer.h
@@ -129,13 +129,61 @@ namespace UE::MovieGraph::DefaultRenderer
 			FThreadSafeBool bIsActive;
 			UE::Tasks::FTask TaskPrereq;
 		};
+
 		typedef TSharedPtr<FInstance, ESPMode::ThreadSafe> FInstancePtr;
 
 		TArray<TSharedPtr<FInstance, ESPMode::ThreadSafe>> Accumulators;
 		FCriticalSection CriticalSection;
 
+		template <typename AccumulatorType>
+		FSurfaceAccumulatorPool::FInstancePtr GetAccumulatorInstance_GameThread(int32 InFrameNumber, const FMovieGraphRenderDataIdentifier& InPassIdentifier)
+		{
+			FScopeLock ScopeLock(&CriticalSection);
+
+			// Search for an existing accumulator for the given frame number and render data
+			int32 AvailableIndex = INDEX_NONE;
+			for (int32 Index = 0; Index < Accumulators.Num(); Index++)
+			{
+				if (InFrameNumber == Accumulators[Index]->ActiveFrameNumber && InPassIdentifier == Accumulators[Index]->ActivePassIdentifier)
+				{
+					AvailableIndex = Index;
+					break;
+				}
+			}
+
+			// If we didn't find one already in use for this frame, look to see if there's a previously
+			// allocated one which is no longer being used.
+			if (AvailableIndex == INDEX_NONE)
+			{
+				for (int32 Index = 0; Index < Accumulators.Num(); Index++)
+				{
+					if (!Accumulators[Index]->IsActive())
+					{
+						// Found a free one, tie it to this output frame.
+						Accumulators[Index]->ActiveFrameNumber = InFrameNumber;
+						Accumulators[Index]->ActivePassIdentifier = InPassIdentifier;
+						Accumulators[Index]->bIsActive = true;
+						Accumulators[Index]->TaskPrereq = UE::Tasks::FTask();
+						AvailableIndex = Index;
+						break;
+					}
+				}
+			}
+
+			// If we still don't have one, just allocate a new entry. The allocations are reasonably light-weight
+			// we don't actually allocate storage memory until they're used.
+			if (AvailableIndex == INDEX_NONE)
+			{
+				AvailableIndex = Accumulators.Num();
+
+				TSharedPtr<AccumulatorType> NewAccumulatorInstance = MakeShared<AccumulatorType>();
+				Accumulators.Add(MakeShared<UE::MovieGraph::DefaultRenderer::FSurfaceAccumulatorPool::FInstance>(NewAccumulatorInstance));
+
+				UE_LOG(LogMovieRenderPipeline, Log, TEXT("Allocated a Accumulator for Pool %s, New Pool Count: %d"), *AccumulatorType::GetName().ToString(), Accumulators.Num());
+			}
 
-		FInstancePtr BlockAndGetAccumulator_GameThread(int32 InFrameNumber, const FMovieGraphRenderDataIdentifier& InPassIdentifier);
+			return Accumulators[AvailableIndex];
+		}
 	};
 
 	class FMovieGraphAccumulationTask
@@ -213,13 +261,7 @@ public:
 			return *ExistingAccumulatorPool;
 		}
 
-		// ToDo: This will need to be dynamically sized based on the number of concurrently produced render datas.
 		FMoviePipelineAccumulatorPoolPtr NewAccumulatorPool = MakeShared<UE::MovieGraph::DefaultRenderer::FSurfaceAccumulatorPool, ESPMode::ThreadSafe>();
-		for (int32 Index = 0; Index < 3; Index++)
-		{
-			TSharedPtr<AccumulatorType> NewAccumulatorInstance = MakeShared<AccumulatorType>();
-			NewAccumulatorPool->Accumulators.Add(MakeShared<UE::MovieGraph::DefaultRenderer::FSurfaceAccumulatorPool::FInstance>(NewAccumulatorInstance));
-		}
 		UE_LOG(LogMovieRenderPipeline, Log, TEXT("Allocated a Accumulator Pool for AccumulatorType: %s"), *AccumulatorType::GetName().ToString());
 		
 		PooledAccumulators.Emplace(AccumulatorType::GetName(), NewAccumulatorPool);
