MovieGraph: UMovieJobVariableAssignmentContainer API now takes UMovieGraphVariableNode pointers directly (instead of a property name). The internal names of properties are actually fguids based on a non-visible name in UMovieGraphVariable, so the API was hard to work with due to requiring non-intuitive information.

diff --git a/Engine/Plugins/MovieScene/MovieRenderPipeline/Source/MovieRenderPipelineCore/Private/Graph/Nodes/MovieGraphVariableNode.cpp b/Engine/Plugins/MovieScene/MovieRenderPipeline/Source/MovieRenderPipelineCore/Private/Graph/Nodes/MovieGraphVariableNode.cpp
--- a/Engine/Plugins/MovieScene/MovieRenderPipeline/Source/MovieRenderPipelineCore/Private/Graph/Nodes/MovieGraphVariableNode.cpp
+++ b/Engine/Plugins/MovieScene/MovieRenderPipeline/Source/MovieRenderPipelineCore/Private/Graph/Nodes/MovieGraphVariableNode.cpp
@@ -1,4 +1,4 @@
-﻿// Copyright Epic Games, Inc. All Rights Reserved.
+// Copyright Epic Games, Inc. All Rights Reserved.
 
 #include "Graph/Nodes/MovieGraphVariableNode.h"
 
@@ -36,7 +36,7 @@ FString UMovieGraphVariableNode::GetResolvedValueForOutputPin(const FName& InPin
 	{
 		if (ContextHasEnabledAssignmentForVariable(InContext))
 		{
-			return InContext->Job->VariableAssignments->GetValueSerializedString(FName(GraphVariable->GetMemberName()));
+			return InContext->Job->VariableAssignments->GetValueSerializedString(GraphVariable);
 		}
 
 		// No valid variable assignment: just get the value from the variable
@@ -52,7 +52,7 @@ bool UMovieGraphVariableNode::GetResolvedValueForOutputPin(const FName& InPinNam
 	{
 		if (ContextHasEnabledAssignmentForVariable(InContext))
 		{
-			return InContext->Job->VariableAssignments->GetValueContainer(InPinName, OutValueContainer);
+			return InContext->Job->VariableAssignments->GetValueContainer(GraphVariable, OutValueContainer);
 		}
 
 		// No valid variable assignment: just get the value from the variable
diff --git a/Engine/Plugins/MovieScene/MovieRenderPipeline/Source/MovieRenderPipelineCore/Private/MovieJobVariableAssignmentContainer.cpp b/Engine/Plugins/MovieScene/MovieRenderPipeline/Source/MovieRenderPipelineCore/Private/MovieJobVariableAssignmentContainer.cpp
--- a/Engine/Plugins/MovieScene/MovieRenderPipeline/Source/MovieRenderPipelineCore/Private/MovieJobVariableAssignmentContainer.cpp
+++ b/Engine/Plugins/MovieScene/MovieRenderPipeline/Source/MovieRenderPipelineCore/Private/MovieJobVariableAssignmentContainer.cpp
@@ -32,165 +32,165 @@ uint32 UMovieJobVariableAssignmentContainer::GetNumAssignments() const
 	return NumAssignments;
 }
 
-bool UMovieJobVariableAssignmentContainer::GetValueBool(const FName& PropertyName, bool& bOutValue) const
+bool UMovieJobVariableAssignmentContainer::GetValueBool(const UMovieGraphVariable* InGraphVariable, bool& bOutValue) const
 {
-	TValueOrError<bool, EPropertyBagResult> Result = Value.GetValueBool(PropertyName);
+	TValueOrError<bool, EPropertyBagResult> Result = Value.GetValueBool(ConvertVariableToInternalName(InGraphVariable));
 	return UE::MovieGraph::Private::GetOptionalValue<bool>(Result, bOutValue);
 }
 
-bool UMovieJobVariableAssignmentContainer::GetValueByte(const FName& PropertyName, uint8& OutValue) const
+bool UMovieJobVariableAssignmentContainer::GetValueByte(const UMovieGraphVariable* InGraphVariable, uint8& OutValue) const
 {
-	TValueOrError<uint8, EPropertyBagResult> Result = Value.GetValueByte(PropertyName);
+	TValueOrError<uint8, EPropertyBagResult> Result = Value.GetValueByte(ConvertVariableToInternalName(InGraphVariable));
 	return UE::MovieGraph::Private::GetOptionalValue<uint8>(Result, OutValue);
 }
 
-bool UMovieJobVariableAssignmentContainer::GetValueInt32(const FName& PropertyName, int32& OutValue) const
+bool UMovieJobVariableAssignmentContainer::GetValueInt32(const UMovieGraphVariable* InGraphVariable, int32& OutValue) const
 {
-	TValueOrError<int32, EPropertyBagResult> Result = Value.GetValueInt32(PropertyName);
+	TValueOrError<int32, EPropertyBagResult> Result = Value.GetValueInt32(ConvertVariableToInternalName(InGraphVariable));
 	return UE::MovieGraph::Private::GetOptionalValue<int32>(Result, OutValue);
 }
 
-bool UMovieJobVariableAssignmentContainer::GetValueInt64(const FName& PropertyName, int64& OutValue) const
+bool UMovieJobVariableAssignmentContainer::GetValueInt64(const UMovieGraphVariable* InGraphVariable, int64& OutValue) const
 {
-	TValueOrError<int64, EPropertyBagResult> Result = Value.GetValueInt64(PropertyName);
+	TValueOrError<int64, EPropertyBagResult> Result = Value.GetValueInt64(ConvertVariableToInternalName(InGraphVariable));
 	return UE::MovieGraph::Private::GetOptionalValue<int64>(Result, OutValue);
 }
 
-bool UMovieJobVariableAssignmentContainer::GetValueFloat(const FName& PropertyName, float& OutValue) const
+bool UMovieJobVariableAssignmentContainer::GetValueFloat(const UMovieGraphVariable* InGraphVariable, float& OutValue) const
 {
-	TValueOrError<float, EPropertyBagResult> Result = Value.GetValueFloat(PropertyName);
+	TValueOrError<float, EPropertyBagResult> Result = Value.GetValueFloat(ConvertVariableToInternalName(InGraphVariable));
 	return UE::MovieGraph::Private::GetOptionalValue<float>(Result, OutValue);
 }
 
-bool UMovieJobVariableAssignmentContainer::GetValueDouble(const FName& PropertyName, double& OutValue) const
+bool UMovieJobVariableAssignmentContainer::GetValueDouble(const UMovieGraphVariable* InGraphVariable, double& OutValue) const
 {
-	TValueOrError<double, EPropertyBagResult> Result = Value.GetValueDouble(PropertyName);
+	TValueOrError<double, EPropertyBagResult> Result = Value.GetValueDouble(ConvertVariableToInternalName(InGraphVariable));
 	return UE::MovieGraph::Private::GetOptionalValue<double>(Result, OutValue);
 }
 
-bool UMovieJobVariableAssignmentContainer::GetValueName(const FName& PropertyName, FName& OutValue) const
+bool UMovieJobVariableAssignmentContainer::GetValueName(const UMovieGraphVariable* InGraphVariable, FName& OutValue) const
 {
-	TValueOrError<FName, EPropertyBagResult> Result = Value.GetValueName(PropertyName);
+	TValueOrError<FName, EPropertyBagResult> Result = Value.GetValueName(ConvertVariableToInternalName(InGraphVariable));
 	return UE::MovieGraph::Private::GetOptionalValue<FName>(Result, OutValue);
 }
 
-bool UMovieJobVariableAssignmentContainer::GetValueString(const FName& PropertyName, FString& OutValue) const
+bool UMovieJobVariableAssignmentContainer::GetValueString(const UMovieGraphVariable* InGraphVariable, FString& OutValue) const
 {
-	TValueOrError<FString, EPropertyBagResult> Result = Value.GetValueString(PropertyName);
+	TValueOrError<FString, EPropertyBagResult> Result = Value.GetValueString(ConvertVariableToInternalName(InGraphVariable));
 	return UE::MovieGraph::Private::GetOptionalValue<FString>(Result, OutValue);
 }
 
-bool UMovieJobVariableAssignmentContainer::GetValueText(const FName& PropertyName, FText& OutValue) const
+bool UMovieJobVariableAssignmentContainer::GetValueText(const UMovieGraphVariable* InGraphVariable, FText& OutValue) const
 {
-	TValueOrError<FText, EPropertyBagResult> Result = Value.GetValueText(PropertyName);
+	TValueOrError<FText, EPropertyBagResult> Result = Value.GetValueText(ConvertVariableToInternalName(InGraphVariable));
 	return UE::MovieGraph::Private::GetOptionalValue<FText>(Result, OutValue);
 }
 
-bool UMovieJobVariableAssignmentContainer::GetValueEnum(const FName& PropertyName, uint8& OutValue, const UEnum* RequestedEnum) const
+bool UMovieJobVariableAssignmentContainer::GetValueEnum(const UMovieGraphVariable* InGraphVariable, uint8& OutValue, const UEnum* RequestedEnum) const
 {
-	TValueOrError<uint8, EPropertyBagResult> Result = Value.GetValueEnum(PropertyName, RequestedEnum);
+	TValueOrError<uint8, EPropertyBagResult> Result = Value.GetValueEnum(ConvertVariableToInternalName(InGraphVariable), RequestedEnum);
 	return UE::MovieGraph::Private::GetOptionalValue<uint8>(Result, OutValue);
 }
 
-bool UMovieJobVariableAssignmentContainer::GetValueStruct(const FName& PropertyName, FStructView& OutValue, const UScriptStruct* RequestedStruct) const
+bool UMovieJobVariableAssignmentContainer::GetValueStruct(const UMovieGraphVariable* InGraphVariable, FStructView& OutValue, const UScriptStruct* RequestedStruct) const
 {
-	TValueOrError<FStructView, EPropertyBagResult> Result = Value.GetValueStruct(PropertyName, RequestedStruct);
+	TValueOrError<FStructView, EPropertyBagResult> Result = Value.GetValueStruct(ConvertVariableToInternalName(InGraphVariable), RequestedStruct);
 	return UE::MovieGraph::Private::GetOptionalValue<FStructView>(Result, OutValue);
 }
 
-bool UMovieJobVariableAssignmentContainer::GetValueObject(const FName& PropertyName, UObject* OutValue, const UClass* RequestedClass) const
+bool UMovieJobVariableAssignmentContainer::GetValueObject(const UMovieGraphVariable* InGraphVariable, UObject* OutValue, const UClass* RequestedClass) const
 {
-	TValueOrError<UObject*, EPropertyBagResult> Result = Value.GetValueObject(PropertyName, RequestedClass);
+	TValueOrError<UObject*, EPropertyBagResult> Result = Value.GetValueObject(ConvertVariableToInternalName(InGraphVariable), RequestedClass);
 	return UE::MovieGraph::Private::GetOptionalValue<UObject*>(Result, OutValue);
 }
 
-bool UMovieJobVariableAssignmentContainer::GetValueClass(const FName& PropertyName, UClass*& OutValue) const
+bool UMovieJobVariableAssignmentContainer::GetValueClass(const UMovieGraphVariable* InGraphVariable, UClass*& OutValue) const
 {
-	TValueOrError<UClass*, EPropertyBagResult> Result = Value.GetValueClass(PropertyName);
+	TValueOrError<UClass*, EPropertyBagResult> Result = Value.GetValueClass(ConvertVariableToInternalName(InGraphVariable));
 	return UE::MovieGraph::Private::GetOptionalValue<UClass*>(Result, OutValue);
 }
 
-FString UMovieJobVariableAssignmentContainer::GetValueSerializedString(const FName& PropertyName)
+FString UMovieJobVariableAssignmentContainer::GetValueSerializedString(const UMovieGraphVariable* InGraphVariable)
 {
-	TValueOrError<FString, EPropertyBagResult> Result = Value.GetValueSerializedString(PropertyName);
+	TValueOrError<FString, EPropertyBagResult> Result = Value.GetValueSerializedString(ConvertVariableToInternalName(InGraphVariable));
 	FString ResultString;
 	UE::MovieGraph::Private::GetOptionalValue<FString>(Result, ResultString);
 	return ResultString;
 }
 
-bool UMovieJobVariableAssignmentContainer::SetValueBool(const FName& PropertyName, const bool bInValue)
+bool UMovieJobVariableAssignmentContainer::SetValueBool(const UMovieGraphVariable* InGraphVariable, const bool bInValue)
 {
-	return Value.SetValueBool(PropertyName, bInValue) == EPropertyBagResult::Success;
+	return Value.SetValueBool(ConvertVariableToInternalName(InGraphVariable), bInValue) == EPropertyBagResult::Success;
 }
 
-bool UMovieJobVariableAssignmentContainer::SetValueByte(const FName& PropertyName, const uint8 InValue)
+bool UMovieJobVariableAssignmentContainer::SetValueByte(const UMovieGraphVariable* InGraphVariable, const uint8 InValue)
 {
-	return Value.SetValueByte(PropertyName, InValue) == EPropertyBagResult::Success;
+	return Value.SetValueByte(ConvertVariableToInternalName(InGraphVariable), InValue) == EPropertyBagResult::Success;
 }
 
-bool UMovieJobVariableAssignmentContainer::SetValueInt32(const FName& PropertyName, const int32 InValue)
+bool UMovieJobVariableAssignmentContainer::SetValueInt32(const UMovieGraphVariable* InGraphVariable, const int32 InValue)
 {
-	return Value.SetValueInt32(PropertyName, InValue) == EPropertyBagResult::Success;
+	return Value.SetValueInt32(ConvertVariableToInternalName(InGraphVariable), InValue) == EPropertyBagResult::Success;
 }
 
-bool UMovieJobVariableAssignmentContainer::SetValueInt64(const FName& PropertyName, const int64 InValue)
+bool UMovieJobVariableAssignmentContainer::SetValueInt64(const UMovieGraphVariable* InGraphVariable, const int64 InValue)
 {
-	return Value.SetValueInt64(PropertyName, InValue) == EPropertyBagResult::Success;
+	return Value.SetValueInt64(ConvertVariableToInternalName(InGraphVariable), InValue) == EPropertyBagResult::Success;
 }
 
-bool UMovieJobVariableAssignmentContainer::SetValueFloat(const FName& PropertyName, const float InValue)
+bool UMovieJobVariableAssignmentContainer::SetValueFloat(const UMovieGraphVariable* InGraphVariable, const float InValue)
 {
-	return Value.SetValueFloat(PropertyName, InValue) == EPropertyBagResult::Success;
+	return Value.SetValueFloat(ConvertVariableToInternalName(InGraphVariable), InValue) == EPropertyBagResult::Success;
 }
 
-bool UMovieJobVariableAssignmentContainer::SetValueDouble(const FName& PropertyName, const double InValue)
+bool UMovieJobVariableAssignmentContainer::SetValueDouble(const UMovieGraphVariable* InGraphVariable, const double InValue)
 {
-	return Value.SetValueDouble(PropertyName, InValue) == EPropertyBagResult::Success;
+	return Value.SetValueDouble(ConvertVariableToInternalName(InGraphVariable), InValue) == EPropertyBagResult::Success;
 }
 
-bool UMovieJobVariableAssignmentContainer::SetValueName(const FName& PropertyName, const FName InValue)
+bool UMovieJobVariableAssignmentContainer::SetValueName(const UMovieGraphVariable* InGraphVariable, const FName InValue)
 {
-	return Value.SetValueName(PropertyName, InValue) == EPropertyBagResult::Success;
+	return Value.SetValueName(ConvertVariableToInternalName(InGraphVariable), InValue) == EPropertyBagResult::Success;
 }
 
-bool UMovieJobVariableAssignmentContainer::SetValueString(const FName& PropertyName, const FString& InValue)
+bool UMovieJobVariableAssignmentContainer::SetValueString(const UMovieGraphVariable* InGraphVariable, const FString& InValue)
 {
-	return Value.SetValueString(PropertyName, InValue) == EPropertyBagResult::Success;
+	return Value.SetValueString(ConvertVariableToInternalName(InGraphVariable), InValue) == EPropertyBagResult::Success;
 }
 
-bool UMovieJobVariableAssignmentContainer::SetValueText(const FName& PropertyName, const FText& InValue)
+bool UMovieJobVariableAssignmentContainer::SetValueText(const UMovieGraphVariable* InGraphVariable, const FText& InValue)
 {
-	return Value.SetValueText(PropertyName, InValue) == EPropertyBagResult::Success;
+	return Value.SetValueText(ConvertVariableToInternalName(InGraphVariable), InValue) == EPropertyBagResult::Success;
 }
 
-bool UMovieJobVariableAssignmentContainer::SetValueEnum(const FName& PropertyName, const uint8 InValue, const UEnum* Enum)
+bool UMovieJobVariableAssignmentContainer::SetValueEnum(const UMovieGraphVariable* InGraphVariable, const uint8 InValue, const UEnum* Enum)
 {
-	return Value.SetValueEnum(PropertyName, InValue, Enum) == EPropertyBagResult::Success;
+	return Value.SetValueEnum(ConvertVariableToInternalName(InGraphVariable), InValue, Enum) == EPropertyBagResult::Success;
 }
 
-bool UMovieJobVariableAssignmentContainer::SetValueStruct(const FName& PropertyName, FConstStructView InValue)
+bool UMovieJobVariableAssignmentContainer::SetValueStruct(const UMovieGraphVariable* InGraphVariable, FConstStructView InValue)
 {
-	return Value.SetValueStruct(PropertyName, InValue) == EPropertyBagResult::Success;
+	return Value.SetValueStruct(ConvertVariableToInternalName(InGraphVariable), InValue) == EPropertyBagResult::Success;
 }
 
-bool UMovieJobVariableAssignmentContainer::SetValueObject(const FName& PropertyName, UObject* InValue)
+bool UMovieJobVariableAssignmentContainer::SetValueObject(const UMovieGraphVariable* InGraphVariable, UObject* InValue)
 {
-	return Value.SetValueObject(PropertyName, InValue) == EPropertyBagResult::Success;
+	return Value.SetValueObject(ConvertVariableToInternalName(InGraphVariable), InValue) == EPropertyBagResult::Success;
 }
 
-bool UMovieJobVariableAssignmentContainer::SetValueClass(const FName& PropertyName, UClass* InValue)
+bool UMovieJobVariableAssignmentContainer::SetValueClass(const UMovieGraphVariable* InGraphVariable, UClass* InValue)
 {
-	return Value.SetValueClass(PropertyName, InValue) == EPropertyBagResult::Success;
+	return Value.SetValueClass(ConvertVariableToInternalName(InGraphVariable), InValue) == EPropertyBagResult::Success;
 }
 
-bool UMovieJobVariableAssignmentContainer::SetValueSerializedString(const FName& PropertyName, const FString& NewValue)
+bool UMovieJobVariableAssignmentContainer::SetValueSerializedString(const UMovieGraphVariable* InGraphVariable, const FString& NewValue)
 {
-	return Value.SetValueSerializedString(PropertyName, NewValue) == EPropertyBagResult::Success;
+	return Value.SetValueSerializedString(ConvertVariableToInternalName(InGraphVariable), NewValue) == EPropertyBagResult::Success;
 }
 
-EMovieGraphValueType UMovieJobVariableAssignmentContainer::GetValueType(const FName& PropertyName) const
+EMovieGraphValueType UMovieJobVariableAssignmentContainer::GetValueType(const UMovieGraphVariable* InGraphVariable) const
 {
-	if (const FPropertyBagPropertyDesc* Desc = Value.FindPropertyDescByName(PropertyName))
+	if (const FPropertyBagPropertyDesc* Desc = Value.FindPropertyDescByName(ConvertVariableToInternalName(InGraphVariable)))
 	{
 		return static_cast<EMovieGraphValueType>(Desc->ValueType);
 	}
@@ -198,9 +198,9 @@ EMovieGraphValueType UMovieJobVariableAssignmentContainer::GetValueType(const FN
 	return EMovieGraphValueType::None;
 }
 
-const UObject* UMovieJobVariableAssignmentContainer::GetValueTypeObject(const FName& PropertyName) const
+const UObject* UMovieJobVariableAssignmentContainer::GetValueTypeObject(const UMovieGraphVariable* InGraphVariable) const
 {
-	if (const FPropertyBagPropertyDesc* Desc = Value.FindPropertyDescByName(PropertyName))
+	if (const FPropertyBagPropertyDesc* Desc = Value.FindPropertyDescByName(ConvertVariableToInternalName(InGraphVariable)))
 	{
 		return Desc->ValueTypeObject;
 	}
@@ -208,9 +208,9 @@ const UObject* UMovieJobVariableAssignmentContainer::GetValueTypeObject(const FN
 	return nullptr;
 }
 
-EMovieGraphContainerType UMovieJobVariableAssignmentContainer::GetValueContainerType(const FName& PropertyName) const
+EMovieGraphContainerType UMovieJobVariableAssignmentContainer::GetValueContainerType(const UMovieGraphVariable* InGraphVariable) const
 {
-	if (const FPropertyBagPropertyDesc* Desc = Value.FindPropertyDescByName(PropertyName))
+	if (const FPropertyBagPropertyDesc* Desc = Value.FindPropertyDescByName(ConvertVariableToInternalName(InGraphVariable)))
 	{
 		return static_cast<EMovieGraphContainerType>(Desc->ContainerTypes.GetFirstContainerType());
 	}
@@ -218,13 +218,14 @@ EMovieGraphContainerType UMovieJobVariableAssignmentContainer::GetValueContainer
 	return EMovieGraphContainerType::None;
 }
 
-bool UMovieJobVariableAssignmentContainer::GetValueContainer(const FName& PropertyName, TObjectPtr<UMovieGraphValueContainer>& OutValueContainer)
+bool UMovieJobVariableAssignmentContainer::GetValueContainer(const UMovieGraphVariable* InGraphVariable, TObjectPtr<UMovieGraphValueContainer>& OutValueContainer)
 {
+	const FName PropertyName = ConvertVariableToInternalName(InGraphVariable);
 	if (const FPropertyBagPropertyDesc* Desc = Value.FindPropertyDescByName(PropertyName))
 	{
 		// TODO: This object allocation is unfortunate -- it would be nice to find a way to avoid it
 		OutValueContainer = NewObject<UMovieGraphValueContainer>();
-		OutValueContainer->SetFromDesc(Desc, GetValueSerializedString(PropertyName));
+		OutValueContainer->SetFromDesc(Desc, GetValueSerializedString(InGraphVariable));
 		return true;
 	}
 
@@ -454,10 +455,12 @@ bool UMovieJobVariableAssignmentContainer::AddVariableAssignment(const UMovieGra
 
 bool UMovieJobVariableAssignmentContainer::SetVariableAssignmentEnableState(const UMovieGraphVariable* InGraphVariable, bool bIsEnabled)
 {
-	if (FindOrGenerateVariableOverride(InGraphVariable))
+	FPropertyBagPropertyDesc PropDesc;
+	FPropertyBagPropertyDesc EditConditionPropDesc;
+	const bool bAddIfNotExists = false;
+	if (FindOrGenerateVariableOverride(InGraphVariable, &PropDesc, &EditConditionPropDesc, bAddIfNotExists))
 	{
-		const FString EditCondPropName = FString::Format(TEXT("{0}{1}"), {EditConditionPrefix, InGraphVariable->GetMemberName()});
-		const EPropertyBagResult Result = Value.SetValueBool(FName(EditCondPropName), bIsEnabled);
+		const EPropertyBagResult Result = Value.SetValueBool(EditConditionPropDesc.Name, bIsEnabled);
 		
 		return Result == EPropertyBagResult::Success;
 	}
@@ -482,4 +485,21 @@ bool UMovieJobVariableAssignmentContainer::GetVariableAssignmentEnableState(cons
 	}
 
 	return false;
-}
\ No newline at end of file
+}
+
+FName UMovieJobVariableAssignmentContainer::ConvertVariableToInternalName(const UMovieGraphVariable* InGraphVariable) const
+{
+	FPropertyBagPropertyDesc PropDesc;
+	FPropertyBagPropertyDesc EditConditionPropDesc;
+	const bool bAddIfNotExists = false;
+
+	// FindOrGenerateVariableOverride is non-const (because of the bAddIfNotExists option), but since we never add,
+	// we're going to const-cast it so that we can call it from here.
+	UMovieJobVariableAssignmentContainer* NonConstThis = const_cast<UMovieJobVariableAssignmentContainer*>(this);
+	if (NonConstThis->FindOrGenerateVariableOverride(InGraphVariable, &PropDesc, &EditConditionPropDesc, bAddIfNotExists))
+	{
+		return PropDesc.Name;
+	}
+
+	return NAME_None;
+}
diff --git a/Engine/Plugins/MovieScene/MovieRenderPipeline/Source/MovieRenderPipelineCore/Public/MovieJobVariableAssignmentContainer.h b/Engine/Plugins/MovieScene/MovieRenderPipeline/Source/MovieRenderPipelineCore/Public/MovieJobVariableAssignmentContainer.h
--- a/Engine/Plugins/MovieScene/MovieRenderPipeline/Source/MovieRenderPipelineCore/Public/MovieJobVariableAssignmentContainer.h
+++ b/Engine/Plugins/MovieScene/MovieRenderPipeline/Source/MovieRenderPipelineCore/Public/MovieJobVariableAssignmentContainer.h
@@ -26,8 +26,7 @@ public:
 	UMovieJobVariableAssignmentContainer() = default;
 
 	/**
-	 * Sets the graph config associated with the variable assignments. Calls UpdateGraphVariableOverrides() to ensure
-	 * that the overrides reflect the specified graph config.
+	 * Sets the graph config associated with the variable assignments.
 	 */
 	UFUNCTION(BlueprintCallable, Category="Experimental")
 	void SetGraphConfig(const TSoftObjectPtr<UMovieGraphConfig>& InGraphConfig);
@@ -37,176 +36,176 @@ public:
 
 	/** Gets the bool value of the specified property. Returns true on success, else false. */
 	UFUNCTION(BlueprintCallable, Category="Experimental")
-	bool GetValueBool(const FName& PropertyName, bool& bOutValue) const;
+	bool GetValueBool(const UMovieGraphVariable* InGraphVariable, bool& bOutValue) const;
 
 	/** Gets the byte value of the specified property. Returns true on success, else false. */
 	UFUNCTION(BlueprintCallable, Category="Experimental")
-	bool GetValueByte(const FName& PropertyName, uint8& OutValue) const;
+	bool GetValueByte(const UMovieGraphVariable* InGraphVariable, uint8& OutValue) const;
 
 	/** Gets the int32 value of the specified property. Returns true on success, else false. */
 	UFUNCTION(BlueprintCallable, Category="Experimental")
-	bool GetValueInt32(const FName& PropertyName, int32& OutValue) const;
+	bool GetValueInt32(const UMovieGraphVariable* InGraphVariable, int32& OutValue) const;
 
 	/** Gets the int64 value of the specified property. Returns true on success, else false. */
 	UFUNCTION(BlueprintCallable, Category="Experimental")
-	bool GetValueInt64(const FName& PropertyName, int64& OutValue) const;
+	bool GetValueInt64(const UMovieGraphVariable* InGraphVariable, int64& OutValue) const;
 
 	/** Gets the float value of the specified property. Returns true on success, else false. */
 	UFUNCTION(BlueprintCallable, Category="Experimental")
-	bool GetValueFloat(const FName& PropertyName, float& OutValue) const;
+	bool GetValueFloat(const UMovieGraphVariable* InGraphVariable, float& OutValue) const;
 
 	/** Gets the double value of the specified property. Returns true on success, else false. */
 	UFUNCTION(BlueprintCallable, Category="Experimental")
-	bool GetValueDouble(const FName& PropertyName, double& OutValue) const;
+	bool GetValueDouble(const UMovieGraphVariable* InGraphVariable, double& OutValue) const;
 
 	/** Gets the FName value of the specified property. Returns true on success, else false. */
 	UFUNCTION(BlueprintCallable, Category="Experimental")
-	bool GetValueName(const FName& PropertyName, FName& OutValue) const;
+	bool GetValueName(const UMovieGraphVariable* InGraphVariable, FName& OutValue) const;
 
 	/** Gets the FString value of the specified property. Returns true on success, else false. */
 	UFUNCTION(BlueprintCallable, Category="Experimental")
-	bool GetValueString(const FName& PropertyName, FString& OutValue) const;
+	bool GetValueString(const UMovieGraphVariable* InGraphVariable, FString& OutValue) const;
 
 	/** Gets the FText value of the specified property. Returns true on success, else false. */
 	UFUNCTION(BlueprintCallable, Category="Experimental")
-	bool GetValueText(const FName& PropertyName, FText& OutValue) const;
+	bool GetValueText(const UMovieGraphVariable* InGraphVariable, FText& OutValue) const;
 
 	/** Gets the enum value (for a specific enum) of the specified property. Returns true on success, else false. */
 	UFUNCTION(BlueprintCallable, Category="Experimental")
-	bool GetValueEnum(const FName& PropertyName, uint8& OutValue, const UEnum* RequestedEnum = nullptr) const;
+	bool GetValueEnum(const UMovieGraphVariable* InGraphVariable, uint8& OutValue, const UEnum* RequestedEnum = nullptr) const;
 
 	/** Gets the struct value (for a specific struct) of the specified property. Returns true on success, else false. */
-	bool GetValueStruct(const FName& PropertyName, FStructView& OutValue, const UScriptStruct* RequestedStruct = nullptr) const;
+	bool GetValueStruct(const UMovieGraphVariable* InGraphVariable, FStructView& OutValue, const UScriptStruct* RequestedStruct = nullptr) const;
 
 	/** Gets the object value (for a specific class) of the specified property. Returns true on success, else false. */
 	UFUNCTION(BlueprintCallable, Category="Experimental")
-	bool GetValueObject(const FName& PropertyName, UObject* OutValue, const UClass* RequestedClass = nullptr) const;
+	bool GetValueObject(const UMovieGraphVariable* InGraphVariable, UObject* OutValue, const UClass* RequestedClass = nullptr) const;
 
 	/** Gets the UClass value of the specified property. Returns true on success, else false. */
 	UFUNCTION(BlueprintCallable, Category="Experimental")
-	bool GetValueClass(const FName& PropertyName, UClass*& OutValue) const;
+	bool GetValueClass(const UMovieGraphVariable* InGraphVariable, UClass*& OutValue) const;
 
 	/** Gets the serialized string value of the specified property. */
 	UFUNCTION(BlueprintCallable, Category="Experimental")
-	FString GetValueSerializedString(const FName& PropertyName);
+	FString GetValueSerializedString(const UMovieGraphVariable* InGraphVariable);
 
 	/** Gets the enum value of the specified property. Returns true on success, else false. */
 	template <typename T>
-	bool GetValueEnum(const FName& PropertyName, T& OutValue) const
+	bool GetValueEnum(const UMovieGraphVariable* InGraphVariable, T& OutValue) const
 	{
-		TValueOrError<T, EPropertyBagResult> Result = Value.GetValueEnum<T>(PropertyName);
+		TValueOrError<T, EPropertyBagResult> Result = Value.GetValueEnum<T>(ConvertVariableToInternalName(InGraphVariable));
 		return UE::MovieGraph::Private::GetOptionalValue<T>(Result, OutValue);
 	}
 
 	/** Gets the struct value of the specified property. Returns true on success, else false. */
 	template <typename T>
-	bool GetValueStruct(const FName& PropertyName, T* OutValue) const
+	bool GetValueStruct(const UMovieGraphVariable* InGraphVariable, T* OutValue) const
 	{
-		TValueOrError<T*, EPropertyBagResult> Result = Value.GetValueStruct<T*>(PropertyName);
+		TValueOrError<T*, EPropertyBagResult> Result = Value.GetValueStruct<T*>(ConvertVariableToInternalName(InGraphVariable));
 		return UE::MovieGraph::Private::GetOptionalValue<T*>(Result, OutValue);
 	}
 
 	/** Gets the object value of the specified property. Returns true on success, else false. */
 	template <typename T>
-	bool GetValueObject(const FName& PropertyName, T* OutValue) const
+	bool GetValueObject(const UMovieGraphVariable* InGraphVariable, T* OutValue) const
 	{
-		TValueOrError<T*, EPropertyBagResult> Result = Value.GetValueObject<T*>(PropertyName);
+		TValueOrError<T*, EPropertyBagResult> Result = Value.GetValueObject<T*>(ConvertVariableToInternalName(InGraphVariable));
 		return UE::MovieGraph::Private::GetOptionalValue<T*>(Result, OutValue);
 	}
 
 	/** Sets the bool value of the specified property. Returns true on success, else false. */
 	UFUNCTION(BlueprintCallable, Category="Experimental")
-	bool SetValueBool(const FName& PropertyName, const bool bInValue);
+	bool SetValueBool(const UMovieGraphVariable* InGraphVariable, const bool bInValue);
 
 	/** Sets the byte value of the specified property. Returns true on success, else false. */
 	UFUNCTION(BlueprintCallable, Category="Experimental")
-	bool SetValueByte(const FName& PropertyName, const uint8 InValue);
+	bool SetValueByte(const UMovieGraphVariable* InGraphVariable, const uint8 InValue);
 
 	/** Sets the int32 value of the specified property. Returns true on success, else false. */
 	UFUNCTION(BlueprintCallable, Category="Experimental")
-	bool SetValueInt32(const FName& PropertyName, const int32 InValue);
+	bool SetValueInt32(const UMovieGraphVariable* InGraphVariable, const int32 InValue);
 
 	/** Sets the int64 value of the specified property. Returns true on success, else false. */
 	UFUNCTION(BlueprintCallable, Category="Experimental")
-	bool SetValueInt64(const FName& PropertyName, const int64 InValue);
+	bool SetValueInt64(const UMovieGraphVariable* InGraphVariable, const int64 InValue);
 
 	/** Sets the float value of the specified property. Returns true on success, else false. */
 	UFUNCTION(BlueprintCallable, Category="Experimental")
-	bool SetValueFloat(const FName& PropertyName, const float InValue);
+	bool SetValueFloat(const UMovieGraphVariable* InGraphVariable, const float InValue);
 
 	/** Sets the double value of the specified property. Returns true on success, else false. */
 	UFUNCTION(BlueprintCallable, Category="Experimental")
-	bool SetValueDouble(const FName& PropertyName, const double InValue);
+	bool SetValueDouble(const UMovieGraphVariable* InGraphVariable, const double InValue);
 
 	/** Sets the FName value of the specified property. Returns true on success, else false. */
 	UFUNCTION(BlueprintCallable, Category="Experimental")
-	bool SetValueName(const FName& PropertyName, const FName InValue);
+	bool SetValueName(const UMovieGraphVariable* InGraphVariable, const FName InValue);
 
 	/** Sets the FString value of the specified property. Returns true on success, else false. */
 	UFUNCTION(BlueprintCallable, Category="Experimental")
-	bool SetValueString(const FName& PropertyName, const FString& InValue);
+	bool SetValueString(const UMovieGraphVariable* InGraphVariable, const FString& InValue);
 
 	/** Sets the FText value of the specified property. Returns true on success, else false. */
 	UFUNCTION(BlueprintCallable, Category="Experimental")
-	bool SetValueText(const FName& PropertyName, const FText& InValue);
+	bool SetValueText(const UMovieGraphVariable* InGraphVariable, const FText& InValue);
 
 	/** Sets the enum value of the specified property. Returns true on success, else false. */
 	UFUNCTION(BlueprintCallable, Category="Experimental")
-	bool SetValueEnum(const FName& PropertyName, const uint8 InValue, const UEnum* Enum);
+	bool SetValueEnum(const UMovieGraphVariable* InGraphVariable, const uint8 InValue, const UEnum* Enum);
 
 	/** Sets the struct value of the specified property. Returns true on success, else false. */
-	bool SetValueStruct(const FName& PropertyName, FConstStructView InValue);
+	bool SetValueStruct(const UMovieGraphVariable* InGraphVariable, FConstStructView InValue);
 
 	/** Sets the object value of the specified property. Returns true on success, else false. */
 	UFUNCTION(BlueprintCallable, Category="Experimental")
-	bool SetValueObject(const FName& PropertyName, UObject* InValue);
+	bool SetValueObject(const UMovieGraphVariable* InGraphVariable, UObject* InValue);
 
 	/** Sets the class value of the specified property. Returns true on success, else false. */
 	UFUNCTION(BlueprintCallable, Category="Experimental")
-	bool SetValueClass(const FName& PropertyName, UClass* InValue);
+	bool SetValueClass(const UMovieGraphVariable* InGraphVariable, UClass* InValue);
 
 	/** Sets the serialized value of this member. The string should be the serialized representation of the value. Returns true on success, else false.*/
 	UFUNCTION(BlueprintCallable, Category="Experimental")
-	bool SetValueSerializedString(const FName& PropertyName, const FString& NewValue);
+	bool SetValueSerializedString(const UMovieGraphVariable* InGraphVariable, const FString& NewValue);
 
 	/** Sets the enum value of the specified property. Returns true on success, else false. */
 	template <typename T>
-	bool SetValueEnum(const FName& PropertyName, const T InValue)
+	bool SetValueEnum(const UMovieGraphVariable* InGraphVariable, const T InValue)
 	{
-		return Value.SetValueEnum<T>(PropertyName, InValue) == EPropertyBagResult::Success;
+		return Value.SetValueEnum<T>(ConvertVariableToInternalName(InGraphVariable), InValue) == EPropertyBagResult::Success;
 	}
 
 	/** Sets the struct value of the specified property. Returns true on success, else false. */
 	template <typename T>
-	bool SetValueStruct(const FName& PropertyName, const T& InValue)
+	bool SetValueStruct(const UMovieGraphVariable* InGraphVariable, const T& InValue)
 	{
-		return Value.SetValueStruct<T>(PropertyName, InValue) == EPropertyBagResult::Success;
+		return Value.SetValueStruct<T>(ConvertVariableToInternalName(InGraphVariable), InValue) == EPropertyBagResult::Success;
 	}
 
 	/** Sets the object value of the specified property. Returns true on success, else false. */
 	template <typename T>
-	bool SetValueObject(const FName& PropertyName, T* InValue)
+	bool SetValueObject(const UMovieGraphVariable* InGraphVariable, T* InValue)
 	{
-		return Value.SetValueObject<T>(PropertyName, InValue) == EPropertyBagResult::Success;
+		return Value.SetValueObject<T>(ConvertVariableToInternalName(InGraphVariable), InValue) == EPropertyBagResult::Success;
 	}
 
 	/** Gets the type of the value stored in the specified property. */
 	UFUNCTION(BlueprintCallable, Category="Experimental")
-	EMovieGraphValueType GetValueType(const FName& PropertyName) const;
+	EMovieGraphValueType GetValueType(const UMovieGraphVariable* InGraphVariable) const;
 
 	/** Gets the object that defines the enum, struct, or class stored in the specified property. */
 	UFUNCTION(BlueprintCallable, Category="Experimental")
-	const UObject* GetValueTypeObject(const FName& PropertyName) const;
+	const UObject* GetValueTypeObject(const UMovieGraphVariable* InGraphVariable) const;
 
 	/** Gets the container type of the stored value in the specified property. */
 	UFUNCTION(BlueprintCallable, Category="Experimental")
-	EMovieGraphContainerType GetValueContainerType(const FName& PropertyName) const;
+	EMovieGraphContainerType GetValueContainerType(const UMovieGraphVariable* InGraphVariable) const;
 
 	/**
 	 * Gets a value container object rather than a strongly-typed value. Useful if the type of the value is not known
-	 * ahead of time. If no property with the specified name exists, OutValueContainer will not be modified.
+	 * ahead of time. If no property for the specified variable exists, OutValueContainer will not be modified.
 	 */
-	bool GetValueContainer(const FName& PropertyName, TObjectPtr<UMovieGraphValueContainer>& OutValueContainer);
+	bool GetValueContainer(const UMovieGraphVariable* InGraphVariable, TObjectPtr<UMovieGraphValueContainer>& OutValueContainer);
 
 	/**
 	 * Updates an existing variable assignment for the provided graph variable to a new enable state, or adds a new
@@ -228,10 +227,8 @@ public:
 	 * corresponding graph variable will be removed, and overrides will be created for graph variables which do not have
 	 * existing overrides.
 	 */
-#if WITH_EDITOR
 	UFUNCTION(BlueprintCallable, Category = "Experimental")
 	void UpdateGraphVariableOverrides();
-#endif
 
 protected:
 	/**
@@ -251,6 +248,8 @@ protected:
 	bool GenerateVariableOverride(const UMovieGraphVariable* InGraphVariable, FPropertyBagPropertyDesc* OutPropDesc = nullptr,
 		FPropertyBagPropertyDesc* OutEditConditionPropDesc = nullptr);
 
+	/** Given a graph variable, retrieve the hidden internal name used by our actual property bag. */
+	FName ConvertVariableToInternalName(const UMovieGraphVariable* InGraphVariable) const;
 private:
 	/** The properties managed by this object. */
 	UPROPERTY(EditAnywhere, meta=(ShowOnlyInnerProperties, FixedLayout), Category = "Value")
