MoviePipeline: Updated MRG collection query display names to play nicely in picker widgets. Eg, calling the Actor query "Actor" is highly confusing and conflicts with the *actual* Actor class.

#rnx

diff --git a/Engine/Plugins/MovieScene/MovieRenderPipeline/Source/MovieRenderPipelineCore/Public/Graph/MoviePipelineRenderLayerSubsystem.h b/Engine/Plugins/MovieScene/MovieRenderPipeline/Source/MovieRenderPipelineCore/Public/Graph/MoviePipelineRenderLayerSubsystem.h
--- a/Engine/Plugins/MovieScene/MovieRenderPipeline/Source/MovieRenderPipelineCore/Public/Graph/MoviePipelineRenderLayerSubsystem.h
+++ b/Engine/Plugins/MovieScene/MovieRenderPipeline/Source/MovieRenderPipelineCore/Public/Graph/MoviePipelineRenderLayerSubsystem.h
@@ -151,7 +151,7 @@ private:
 };
 
 /** Query type which filters actors via an explicit actor list. */
-UCLASS(BlueprintType, meta = (DisplayName = "Actor"))
+UCLASS(BlueprintType, meta = (DisplayName = "Movie Graph Actor Query"))
 class MOVIERENDERPIPELINECORE_API UMovieGraphConditionGroupQuery_Actor final : public UMovieGraphConditionGroupQueryBase
 {
 	GENERATED_BODY()
@@ -187,7 +187,7 @@ private:
 };
 
 /** Query type which filters actors via tags on actors. */
-UCLASS(BlueprintType, meta = (DisplayName = "Actor Tag Name"))
+UCLASS(BlueprintType, meta = (DisplayName = "Movie Graph Actor Tag Name Query"))
 class MOVIERENDERPIPELINECORE_API UMovieGraphConditionGroupQuery_ActorTagName final : public UMovieGraphConditionGroupQueryBase
 {
 	GENERATED_BODY()
@@ -219,7 +219,7 @@ public:
 };
 
 /** Query type which filters actors via their name (label). */
-UCLASS(BlueprintType, meta = (DisplayName = "Actor Name"))
+UCLASS(BlueprintType, meta = (DisplayName = "Movie Graph Actor Name Query"))
 class MOVIERENDERPIPELINECORE_API UMovieGraphConditionGroupQuery_ActorName final : public UMovieGraphConditionGroupQueryBase
 {
 	GENERATED_BODY()
@@ -253,7 +253,7 @@ public:
 };
 
 /** Query type which filters actors by type. */
-UCLASS(BlueprintType, meta = (DisplayName = "Actor Type"))
+UCLASS(BlueprintType, meta = (DisplayName = "Movie Graph Actor Type Query"))
 class MOVIERENDERPIPELINECORE_API UMovieGraphConditionGroupQuery_ActorType final : public UMovieGraphConditionGroupQueryBase
 {
 	GENERATED_BODY()
@@ -282,7 +282,7 @@ private:
 };
 
 /** Query type which filters actors by tags on their components. */
-UCLASS(BlueprintType, meta = (DisplayName = "Component Tag Name"))
+UCLASS(BlueprintType, meta = (DisplayName = "Movie Graph Component Tag Name Query"))
 class MOVIERENDERPIPELINECORE_API UMovieGraphConditionGroupQuery_ComponentTagName final : public UMovieGraphConditionGroupQueryBase
 {
 	GENERATED_BODY()
@@ -314,7 +314,7 @@ public:
 };
 
 /** Query type which filters actors via the components contained in them. */
-UCLASS(BlueprintType, meta = (DisplayName = "Component Type"))
+UCLASS(BlueprintType, meta = (DisplayName = "Movie Graph Component Type Query"))
 class MOVIERENDERPIPELINECORE_API UMovieGraphConditionGroupQuery_ComponentType final : public UMovieGraphConditionGroupQueryBase
 {
 	GENERATED_BODY()
