Movie Graph:

-Fix: Queue UI "Output" column is incorrect for graph configs

diff --git a/Engine/Plugins/MovieScene/MovieRenderPipeline/Source/MovieRenderPipelineCore/Private/Graph/MovieGraphConfig.cpp b/Engine/Plugins/MovieScene/MovieRenderPipeline/Source/MovieRenderPipelineCore/Private/Graph/MovieGraphConfig.cpp
--- a/Engine/Plugins/MovieScene/MovieRenderPipeline/Source/MovieRenderPipelineCore/Private/Graph/MovieGraphConfig.cpp
+++ b/Engine/Plugins/MovieScene/MovieRenderPipeline/Source/MovieRenderPipelineCore/Private/Graph/MovieGraphConfig.cpp
@@ -9,6 +9,7 @@
 #include "Graph/MovieGraphPipeline.h"
 #include "Graph/Nodes/MovieGraphInputNode.h"
 #include "Graph/Nodes/MovieGraphOutputNode.h"
+#include "Graph/Nodes/MovieGraphOutputSettingNode.h"
 #include "Graph/Nodes/MovieGraphRemoveRenderSettingNode.h"
 #include "Graph/Nodes/MovieGraphSubgraphNode.h"
 #include "Graph/Nodes/MovieGraphVariableNode.h"
@@ -1010,6 +1011,77 @@ void UMovieGraphConfig::GetAllContainedSubgraphs(TSet<UMovieGraphConfig*>& OutSu
 	}
 }
 
+void UMovieGraphConfig::RecurseUpGlobalsBranchToFindOutputDirectory(const UMovieGraphNode* InNode, FString& OutOutputDirectory) const
+{
+	// If there's no Node, no upstream pin or no downstream pin for whatever reason,
+	// there is no way to continue so we early out
+	if (!InNode) { return; }
+
+	// Only globals can connect to globals linearly, so we only need to look at the first connected input
+	// The only exception being subgraph nodes which will need to be separately evaluated
+	const UMovieGraphPin* DownstreamGlobalsPin = InNode->GetFirstConnectedInputPin();
+	if (!DownstreamGlobalsPin) { return; }
+
+	const UMovieGraphPin* UpstreamGlobalsPin = DownstreamGlobalsPin->GetFirstConnectedPin();
+	if (!UpstreamGlobalsPin) { return; }
+
+	UMovieGraphNode* ConnectedNode = UpstreamGlobalsPin->Node;
+
+	// Overrides can be set within Subgraphs
+	if (const UMovieGraphSubgraphNode* SubgraphNode = Cast<UMovieGraphSubgraphNode>(ConnectedNode))
+	{
+		const UMovieGraphConfig* SubgraphConfig = SubgraphNode->GetSubgraphAsset();
+		if (SubgraphConfig && OutOutputDirectory.IsEmpty())
+		{
+			const UMovieGraphPin* SubgraphGlobalsPin =
+				SubgraphConfig->GetOutputNode()->GetInputPin(UMovieGraphNode::GlobalsPinName);
+
+			if (SubgraphGlobalsPin && SubgraphGlobalsPin->IsConnected())
+			{
+				SubgraphConfig->RecurseUpGlobalsBranchToFindOutputDirectory(SubgraphGlobalsPin->Node, OutOutputDirectory);
+			}
+		}
+	}
+	else if (const UMovieGraphOutputSettingNode* SettingsNode = Cast<UMovieGraphOutputSettingNode>(InNode))
+	{
+		if (OutOutputDirectory.IsEmpty() && SettingsNode->bOverride_OutputDirectory)
+		{
+			OutOutputDirectory = SettingsNode->OutputDirectory.Path;
+		}
+	}
+
+	// Keep looking upstream if we haven't found any overrides
+	if (OutOutputDirectory.IsEmpty())
+	{
+		RecurseUpGlobalsBranchToFindOutputDirectory(UpstreamGlobalsPin->Node, OutOutputDirectory);
+	}
+};
+
+void UMovieGraphConfig::GetOutputDirectory(FString& OutOutputDirectory) const
+{
+	check (OutputNode);
+
+	// Clear out input strings
+	OutOutputDirectory = FString();
+
+	// We only traverse up the globals branch in order to find the output directory and file name format
+	const UMovieGraphPin* GlobalsPin = OutputNode->GetInputPin(UMovieGraphNode::GlobalsPinName);
+
+	if (GlobalsPin && GlobalsPin->IsConnected())
+	{
+		RecurseUpGlobalsBranchToFindOutputDirectory(OutputNode, OutOutputDirectory);
+
+		if (OutOutputDirectory.IsEmpty())
+		{
+			// If we didn't find any overrides, use the CDO values
+			UMovieGraphOutputSettingNode* CDO = Cast<UMovieGraphOutputSettingNode>(UMovieGraphOutputSettingNode::StaticClass()->ClassDefaultObject);
+			check(CDO);
+			
+			OutOutputDirectory = CDO->OutputDirectory.Path;
+		}
+	}
+}
+
 void UMovieGraphConfig::InitializeFlattenedNode(UMovieGraphNode* InNode)
 {
 	// We go through each of the bOverride_ properties on this new instance and set
diff --git a/Engine/Plugins/MovieScene/MovieRenderPipeline/Source/MovieRenderPipelineCore/Public/Graph/MovieGraphConfig.h b/Engine/Plugins/MovieScene/MovieRenderPipeline/Source/MovieRenderPipelineCore/Public/Graph/MovieGraphConfig.h
--- a/Engine/Plugins/MovieScene/MovieRenderPipeline/Source/MovieRenderPipelineCore/Public/Graph/MovieGraphConfig.h
+++ b/Engine/Plugins/MovieScene/MovieRenderPipeline/Source/MovieRenderPipelineCore/Public/Graph/MovieGraphConfig.h
@@ -630,7 +630,17 @@ public:
 	/** Get all subgraphs that this graph contains, recursively (ie, subgraphs of subgraphs are included, etc). */
 	void GetAllContainedSubgraphs(TSet<UMovieGraphConfig*>& OutSubgraphs) const;
 
+	/**
+	 * Walks the graph backward recursively from the output node searching for a UMovieGraphOutputSettings node. Traverses subgraphs as well.
+	 * If a node is not found with an override set, value is taken from the CDO of UMovieGraphOutputSettings.
+	 */
+	void GetOutputDirectory(FString& OutOutputDirectory) const;
+
 protected:
+
+	/** Look for the output directory in the UMovieGraphOutputSettings nodes found upstream of InNode. */
+	void RecurseUpGlobalsBranchToFindOutputDirectory(const UMovieGraphNode* InNode, FString& OutOutputDirectory) const;
+	
 	/** Copies properties in FromNode that are marked for override into ToNode, but only if ToNode doesn't already override that value. */
 	void CopyOverriddenProperties(UMovieGraphNode* FromNode, UMovieGraphNode* ToNode, const FMovieGraphTraversalContext* InContext);
 	
diff --git a/Engine/Plugins/MovieScene/MovieRenderPipeline/Source/MovieRenderPipelineEditor/Private/MoviePipelineEditorBlueprintLibrary.cpp b/Engine/Plugins/MovieScene/MovieRenderPipeline/Source/MovieRenderPipelineEditor/Private/MoviePipelineEditorBlueprintLibrary.cpp
--- a/Engine/Plugins/MovieScene/MovieRenderPipeline/Source/MovieRenderPipelineEditor/Private/MoviePipelineEditorBlueprintLibrary.cpp
+++ b/Engine/Plugins/MovieScene/MovieRenderPipeline/Source/MovieRenderPipelineEditor/Private/MoviePipelineEditorBlueprintLibrary.cpp
@@ -20,6 +20,7 @@
 #include "SequencerUtilities.h"
 #include "MoviePipelineBlueprintLibrary.h"
 #include "MoviePipelineOutputSetting.h"
+#include "Graph/MovieGraphBlueprintLibrary.h"
 
 #include UE_INLINE_GENERATED_CPP_BY_NAME(MoviePipelineEditorBlueprintLibrary)
 
@@ -214,36 +215,76 @@ void UMoviePipelineEditorBlueprintLibrary::EnsureJobHasDefaultSettings(UMoviePip
 	}
 }
 
-FString UMoviePipelineEditorBlueprintLibrary::ResolveOutputDirectoryFromJob(UMoviePipelineExecutorJob* InJob)
+bool UMoviePipelineEditorBlueprintLibrary::GetDisplayOutputPathFromJob(UMoviePipelineExecutorJob* InJob, FString& OutOutputPath)
 {
-	UMoviePipelineOutputSetting* OutputSetting = InJob->GetConfiguration()->FindSetting<UMoviePipelineOutputSetting>();
-	check(OutputSetting);
+	if (ensureMsgf(InJob, TEXT("%hs: Could not get format string because provided UMoviePipelineExecutorJob was invalid."), __FUNCTION__))
+	{
+		if (InJob->IsUsingGraphConfiguration())
+		{
+			if (InJob->GetGraphPreset())
+			{
+				FString OutputDirectory;
+				InJob->GetGraphPreset()->GetOutputDirectory(OutputDirectory);
+				OutOutputPath = OutputDirectory;
+				return true;
+			}
+		}
 
-	// Set up as many parameters as we can to try and resolve most of the string.
-	FString FormatString = OutputSetting->OutputDirectory.Path / OutputSetting->FileNameFormat;
+		if (InJob->GetConfiguration())
+		{
+			const UMoviePipelineOutputSetting* OutputSetting = InJob->GetConfiguration()->FindSetting<UMoviePipelineOutputSetting>();
+			check(OutputSetting);
 
-	// If they've set up any folders within the filename portion of it, let's be nice and resolve that.
-	FPaths::NormalizeFilename(FormatString);
-	int32 LastSlashIndex;
-	if (FormatString.FindLastChar(TEXT('/'), LastSlashIndex))
-	{
-		FormatString.LeftInline(LastSlashIndex + 1);
+			OutOutputPath = OutputSetting->OutputDirectory.Path;
+			return true;
+		}
 	}
 
-	// By having it swap {camera_name} and {shot_name} with an unresolvable tag, it will
-	// stay in the resolved path and can be removed using the code below.
-	static const FString DummyTag = TEXT("{dontresolvethis}");
-	const bool bGetNextVersion = false;
-	FMoviePipelineFilenameResolveParams Params;
-	Params.Job = InJob;
-	Params.ShotNameOverride = DummyTag;
-	Params.CameraNameOverride = DummyTag;
-	Params.InitializationTime = FDateTime::UtcNow();
-	Params.InitializationVersion = UMoviePipelineBlueprintLibrary::ResolveVersionNumber(Params, bGetNextVersion);
+	return false;
+}
+
+FString UMoviePipelineEditorBlueprintLibrary::ResolveOutputDirectoryFromJob(UMoviePipelineExecutorJob* InJob)
+{
+	// Set up as many parameters as we can to try and resolve most of the string.
+	FString FormatString;
+	GetDisplayOutputPathFromJob(InJob, FormatString);
 
 	FString OutResolvedPath;
-	FMoviePipelineFormatArgs Dummy;
-	UMoviePipelineBlueprintLibrary::ResolveFilenameFormatArguments(FormatString, Params, OutResolvedPath, Dummy);
+
+	FPaths::NormalizeFilename(FormatString);
+
+	if (InJob->IsUsingGraphConfiguration() && InJob->GetGraphPreset())
+	{
+		const bool bGetNextVersion = false;
+		FMovieGraphFilenameResolveParams Params;
+		Params.Job = InJob;
+		Params.InitializationTime = FDateTime::UtcNow();
+		Params.Version = UMovieGraphBlueprintLibrary::ResolveVersionNumber(Params, bGetNextVersion);
+		Params.RenderDataIdentifier.RootBranchName = UMovieGraphNode::GlobalsPinName;
+		
+		FMovieGraphTraversalContext TraversalContext;
+		TraversalContext.Job = InJob;
+		Params.EvaluatedConfig = InJob->GetGraphPreset()->CreateFlattenedGraph(TraversalContext);
+
+		FMovieGraphResolveArgs Dummy;
+		OutResolvedPath = UMovieGraphBlueprintLibrary::ResolveFilenameFormatArguments(FormatString, Params, Dummy);
+	}
+	else
+	{
+		// By having it swap {camera_name} and {shot_name} with an unresolvable tag, it will
+		// stay in the resolved path and can be removed using the code below.
+		static const FString DummyTag = TEXT("{dontresolvethis}");
+		const bool bGetNextVersion = false;
+		FMoviePipelineFilenameResolveParams Params;
+		Params.Job = InJob;
+		Params.ShotNameOverride = DummyTag;
+		Params.CameraNameOverride = DummyTag;
+		Params.InitializationTime = FDateTime::UtcNow();
+		Params.InitializationVersion = UMoviePipelineBlueprintLibrary::ResolveVersionNumber(Params, bGetNextVersion);
+
+		FMoviePipelineFormatArgs Dummy;
+		UMoviePipelineBlueprintLibrary::ResolveFilenameFormatArguments(FormatString, Params, OutResolvedPath, Dummy);
+	}
 
 	// Drop the .{ext} resolving always puts on.
 	OutResolvedPath.LeftChopInline(6);
@@ -261,7 +302,7 @@ FString UMoviePipelineEditorBlueprintLibrary::ResolveOutputDirectoryFromJob(UMov
 		// Just as a last bit of saftey, we'll trim anything between the { and the preceeding /. This is
 		// in case they did something like Render_{Date}, we wouldn't want to make a folder named Render_.
 		// We search backwards from where we found the first { brace, so that will get us the last usable slash.
-		LastSlashIndex = OutResolvedPath.Find(TEXT("/"), ESearchCase::IgnoreCase, ESearchDir::FromEnd, FormatStringToken);
+		int32 LastSlashIndex = OutResolvedPath.Find(TEXT("/"), ESearchCase::IgnoreCase, ESearchDir::FromEnd, FormatStringToken);
 		if (LastSlashIndex != INDEX_NONE)
 		{
 			OutResolvedPath.LeftInline(LastSlashIndex + 1);
diff --git a/Engine/Plugins/MovieScene/MovieRenderPipeline/Source/MovieRenderPipelineEditor/Private/Widgets/SMoviePipelineQueueEditor.cpp b/Engine/Plugins/MovieScene/MovieRenderPipeline/Source/MovieRenderPipelineEditor/Private/Widgets/SMoviePipelineQueueEditor.cpp
--- a/Engine/Plugins/MovieScene/MovieRenderPipeline/Source/MovieRenderPipelineEditor/Private/Widgets/SMoviePipelineQueueEditor.cpp
+++ b/Engine/Plugins/MovieScene/MovieRenderPipeline/Source/MovieRenderPipelineEditor/Private/Widgets/SMoviePipelineQueueEditor.cpp
@@ -344,12 +344,10 @@ public:
 	FText GetOutputLabel() const
 	{
 		UMoviePipelineExecutorJob* Job = WeakJob.Get();
-		if (Job && Job->GetConfiguration())
+		FString OutputDir;
+		if (UMoviePipelineEditorBlueprintLibrary::GetDisplayOutputPathFromJob(Job, OutputDir))
 		{
-			UMoviePipelineOutputSetting* OutputSetting = Job->GetConfiguration()->FindSetting<UMoviePipelineOutputSetting>();
-			check(OutputSetting);
-
-			return FText::FromString(OutputSetting->OutputDirectory.Path);
+			return FText::FromString(OutputDir);
 		}
 
 		return LOCTEXT("MissingConfigOutput_Label", "[No Config Set]");
@@ -357,16 +355,18 @@ public:
 
 	void BrowseToOutputFolder()
 	{
-		UMoviePipelineExecutorJob* Job = WeakJob.Get();
-		if (Job && Job->GetConfiguration())
+		if (UMoviePipelineExecutorJob* Job = WeakJob.Get())
 		{
-			FString ResolvedOutputDir = UMoviePipelineEditorBlueprintLibrary::ResolveOutputDirectoryFromJob(Job);
+			const FString ResolvedOutputDir = UMoviePipelineEditorBlueprintLibrary::ResolveOutputDirectoryFromJob(Job);
 
-			// Attempt to make the directory. The user can see the output folder before they render so the folder
-			// may not have been created yet and the ExploreFolder call will fail.
-			IFileManager::Get().MakeDirectory(*ResolvedOutputDir, true);
+			if (!ResolvedOutputDir.IsEmpty())
+			{
+				// Attempt to make the directory. The user can see the output folder before they render so the folder
+				// may not have been created yet and the ExploreFolder call will fail.
+				IFileManager::Get().MakeDirectory(*ResolvedOutputDir, true);
 
-			FPlatformProcess::ExploreFolder(*ResolvedOutputDir);
+				FPlatformProcess::ExploreFolder(*ResolvedOutputDir);
+			}
 		}
 	}
 
@@ -641,6 +641,7 @@ TSharedRef<SWidget> SQueueJobListRow::GenerateWidgetForColumn(const FName& Colum
 			[
 				SNew(SHyperlink)
 					.Text(Item.Get(), &FMoviePipelineQueueJobTreeItem::GetOutputLabel)
+					.ToolTipText(Item.Get(), &FMoviePipelineQueueJobTreeItem::GetOutputLabel)
 					.OnNavigate(Item.Get(), &FMoviePipelineQueueJobTreeItem::BrowseToOutputFolder)
 			];
 
diff --git a/Engine/Plugins/MovieScene/MovieRenderPipeline/Source/MovieRenderPipelineEditor/Public/MoviePipelineEditorBlueprintLibrary.h b/Engine/Plugins/MovieScene/MovieRenderPipeline/Source/MovieRenderPipelineEditor/Public/MoviePipelineEditorBlueprintLibrary.h
--- a/Engine/Plugins/MovieScene/MovieRenderPipeline/Source/MovieRenderPipelineEditor/Public/MoviePipelineEditorBlueprintLibrary.h
+++ b/Engine/Plugins/MovieScene/MovieRenderPipeline/Source/MovieRenderPipelineEditor/Public/MoviePipelineEditorBlueprintLibrary.h
@@ -49,6 +49,10 @@ public:
 	UFUNCTION(BlueprintCallable, Category = "Movie Render Pipeline")
 	static void EnsureJobHasDefaultSettings(UMoviePipelineExecutorJob* InJob);
 
+	/** Returns display string for output directory for this job. Does not resolve the full path from tokens. */
+	UFUNCTION(BlueprintCallable, Category = "Movie Render Pipeline")
+	static bool GetDisplayOutputPathFromJob(UMoviePipelineExecutorJob* InJob, FString& OutOutputPath);
+
 	/** Resolves as much of the output directory for this job into a usable directory path as possible. Cannot resolve anything that relies on shot name, frame numbers, etc. */
 	UFUNCTION(BlueprintCallable, Category = "Movie Render Pipeline")
 	static FString ResolveOutputDirectoryFromJob(UMoviePipelineExecutorJob* InJob);
