MovieGraph: Fix for multi-layer EXR not using the appropriate output frame number.

#rnx

diff --git a/Engine/Plugins/MovieScene/MovieRenderPipeline/Source/MovieRenderPipelineCore/Private/Graph/MovieGraphOutputMerger.cpp b/Engine/Plugins/MovieScene/MovieRenderPipeline/Source/MovieRenderPipelineCore/Private/Graph/MovieGraphOutputMerger.cpp
--- a/Engine/Plugins/MovieScene/MovieRenderPipeline/Source/MovieRenderPipelineCore/Private/Graph/MovieGraphOutputMerger.cpp
+++ b/Engine/Plugins/MovieScene/MovieRenderPipeline/Source/MovieRenderPipelineCore/Private/Graph/MovieGraphOutputMerger.cpp
@@ -19,21 +19,21 @@ namespace UE::MovieGraph
 		PendingData.Empty();
 	}
 
-	FMovieGraphOutputMergerFrame& FMovieGraphOutputMerger::AllocateNewOutputFrame_GameThread(const int32 InFrameNumber) 
+	FMovieGraphOutputMergerFrame& FMovieGraphOutputMerger::AllocateNewOutputFrame_GameThread(const int32 InRenderedFrameNumber)
 	{
 		FScopeLock ScopeLock(&PendingDataMutex);
 
 		// Ensure this frame hasn't already been entered somehow.
-		check(!PendingData.Find(InFrameNumber));
+		check(!PendingData.Find(InRenderedFrameNumber));
 
-		FMovieGraphOutputMergerFrame& NewFrame = PendingData.Add(InFrameNumber);
+		FMovieGraphOutputMergerFrame& NewFrame = PendingData.Add(InRenderedFrameNumber);
 		return NewFrame;
 	}
 
-	FMovieGraphOutputMergerFrame& FMovieGraphOutputMerger::GetOutputFrame_GameThread(int32 InFrameNumber)
+	FMovieGraphOutputMergerFrame& FMovieGraphOutputMerger::GetOutputFrame_GameThread(const int32 InRenderedFrameNumber)
 	{
-		check(PendingData.Find(InFrameNumber));
-		return *PendingData.Find(InFrameNumber);
+		check(PendingData.Find(InRenderedFrameNumber));
+		return *PendingData.Find(InRenderedFrameNumber);
 	}
 
 	
@@ -65,12 +65,12 @@ namespace UE::MovieGraph
 		FMovieGraphSampleState* Payload = InData->GetPayload<FMovieGraphSampleState>();
 		check(Payload);
 		
-		const int32 FrameNumber = Payload->TraversalContext.Time.OutputFrameNumber;
+		const int32 RenderedFrameNumber = Payload->TraversalContext.Time.RenderedFrameNumber;
 
 		// See if we can find the frame this data is for. This should always be valid, if it's not
 		// valid it means they either forgot to declare they were going to produce it, or this is
 		// coming in after the system already thinks it's finished that frame.
-		FMovieGraphOutputMergerFrame* OutputFrame = PendingData.Find(FrameNumber);
+		FMovieGraphOutputMergerFrame* OutputFrame = PendingData.Find(RenderedFrameNumber);
 		
 		// Make sure we expected this frame number
 		if (!ensureAlwaysMsgf(OutputFrame, TEXT("Received data for unknown frame. Frame was either already processed or not queued yet!")))
@@ -106,7 +106,7 @@ namespace UE::MovieGraph
 			
 			// Move this frame into our FinishedFrames array so the Game Thread can read it at its leisure
 			FMovieGraphOutputMergerFrame FinalFrame;
-			ensureMsgf(PendingData.RemoveAndCopyValue(FrameNumber, FinalFrame), TEXT("Could not find frame in pending data, output will be skipped!"));
+			ensureMsgf(PendingData.RemoveAndCopyValue(RenderedFrameNumber, FinalFrame), TEXT("Could not find frame in pending data, output will be skipped!"));
 			
 			// TQueue is thread safe so it's okay to just push the data into it.
 			FinishedFrames.Enqueue(MoveTemp(FinalFrame));
diff --git a/Engine/Plugins/MovieScene/MovieRenderPipeline/Source/MovieRenderPipelineCore/Private/Graph/MovieGraphOutputMerger.h b/Engine/Plugins/MovieScene/MovieRenderPipeline/Source/MovieRenderPipelineCore/Private/Graph/MovieGraphOutputMerger.h
--- a/Engine/Plugins/MovieScene/MovieRenderPipeline/Source/MovieRenderPipelineCore/Private/Graph/MovieGraphOutputMerger.h
+++ b/Engine/Plugins/MovieScene/MovieRenderPipeline/Source/MovieRenderPipelineCore/Private/Graph/MovieGraphOutputMerger.h
@@ -22,8 +22,8 @@ namespace UE::MovieGraph
 		FMovieGraphOutputMerger(UMovieGraphPipeline* InOwningMoviePipeline);
 
 		// IMovieGraphOutputMerger Interface
-		virtual FMovieGraphOutputMergerFrame& AllocateNewOutputFrame_GameThread(const int32 InFrameNumber) override;
-		virtual FMovieGraphOutputMergerFrame& GetOutputFrame_GameThread(int32 InFrameNumber) override;
+		virtual FMovieGraphOutputMergerFrame& AllocateNewOutputFrame_GameThread(const int32 InRenderedFrameNumber) override;
+		virtual FMovieGraphOutputMergerFrame& GetOutputFrame_GameThread(const int32 InRenderedFrameNumber) override;
 		virtual void OnCompleteRenderPassDataAvailable_AnyThread(TUniquePtr<FImagePixelData>&& InData) override;
 		virtual void OnSingleSampleDataAvailable_AnyThread(TUniquePtr<FImagePixelData>&& InData) override;
 		virtual void AbandonOutstandingWork() override;
diff --git a/Engine/Plugins/MovieScene/MovieRenderPipeline/Source/MovieRenderPipelineCore/Public/Graph/MovieGraphDataTypes.h b/Engine/Plugins/MovieScene/MovieRenderPipeline/Source/MovieRenderPipelineCore/Public/Graph/MovieGraphDataTypes.h
--- a/Engine/Plugins/MovieScene/MovieRenderPipeline/Source/MovieRenderPipelineCore/Public/Graph/MovieGraphDataTypes.h
+++ b/Engine/Plugins/MovieScene/MovieRenderPipeline/Source/MovieRenderPipelineCore/Public/Graph/MovieGraphDataTypes.h
@@ -440,10 +440,10 @@ namespace UE::MovieGraph
 		virtual ~IMovieGraphOutputMerger() {};
 
 		/** Call once per output frame to generate a struct to hold data. */
-		virtual FMovieGraphOutputMergerFrame& AllocateNewOutputFrame_GameThread(int32 InFrameNumber) = 0;
+		virtual FMovieGraphOutputMergerFrame& AllocateNewOutputFrame_GameThread(const int32 InRenderedFrameNumber) = 0;
 
 		/** Getter for the Output Frame. Make sure you call AllocateNewOutputFrame_GameThread first, otherwise this trips a check. */
-		virtual FMovieGraphOutputMergerFrame& GetOutputFrame_GameThread(int32 InFrameNumber) = 0;
+		virtual FMovieGraphOutputMergerFrame& GetOutputFrame_GameThread(const int32 InRenderedFrameNumber) = 0;
 
 		/** When a final accumulated render pass is available, call this function with the pixel data. */
 		virtual void OnCompleteRenderPassDataAvailable_AnyThread(TUniquePtr<FImagePixelData>&& InData) = 0;
