MoviePipeline: Added the ability to create a shot subgraph easily from the MRQ queue UI. This shows up as "Make Subgraph", instead of "Edit" which appears for shot jobs that are not using graphs.

"Make Subgraph" will create a new, saved graph asset which contains a single subgraph node which points to the primary graph.

Also fixed a bug which caused subgraph nodes to not be connectable to Globals and render layer branches simultaneously.

#rnx

diff --git a/Engine/Plugins/MovieScene/MovieRenderPipeline/Source/MovieRenderPipelineCore/Private/Graph/MovieGraphConfig.cpp b/Engine/Plugins/MovieScene/MovieRenderPipeline/Source/MovieRenderPipelineCore/Private/Graph/MovieGraphConfig.cpp
--- a/Engine/Plugins/MovieScene/MovieRenderPipeline/Source/MovieRenderPipelineCore/Private/Graph/MovieGraphConfig.cpp
+++ b/Engine/Plugins/MovieScene/MovieRenderPipeline/Source/MovieRenderPipelineCore/Private/Graph/MovieGraphConfig.cpp
@@ -286,7 +286,7 @@ UMovieGraphConfig::UMovieGraphConfig()
 
 		// Offset the default output node so it doesn't overlap the default input node
 #if WITH_EDITOR
-		constexpr int32 OutputNodeOffset = 300;
+		constexpr int32 OutputNodeOffset = 900;
 		OutputNode->SetNodePosX(OutputNodeOffset);
 #endif
 	}
@@ -944,45 +944,69 @@ void UMovieGraphConfig::VisitDownstreamNodes(UMovieGraphNode* FromNode, const FV
 	VisitDownstreamNodes_Recursive(FromNode, VisitCallback, VisitedNodes);
 }
 
-TArray<FString> UMovieGraphConfig::GetDownstreamBranchNames(UMovieGraphNode* FromNode, const UMovieGraphPin* FromPin) const
+TArray<FString> UMovieGraphConfig::GetDownstreamBranchNames(UMovieGraphNode* FromNode, const UMovieGraphPin* FromPin, const bool bStopAtSubgraph) const
 {
 	TArray<FString> BranchNames;
 
 	// FromNode itself might be the Outputs node, so check before visiting the downstream nodes
 	if (FromNode->IsA<UMovieGraphOutputNode>() && FromPin)
 	{
-		BranchNames.Add(FromPin->Properties.Label.ToString());
+		BranchNames.AddUnique(FromPin->Properties.Label.ToString());
 	}
 
 	VisitDownstreamNodes(FromNode, FVisitNodesCallback::CreateLambda(
-		[&BranchNames](UMovieGraphNode* VisitedNode, const UMovieGraphPin* VisitedPin)
+		[&BranchNames, bStopAtSubgraph](const UMovieGraphNode* VisitedNode, const UMovieGraphPin* VisitedPin)
 		{
+			if (VisitedNode->IsA<UMovieGraphSubgraphNode>() && VisitedPin)
+			{
+				BranchNames.AddUnique(VisitedPin->Properties.Label.ToString());
+
+				if (bStopAtSubgraph)
+				{
+					return false;	// Stop traversing nodes
+				}
+			}
+			
 			if (VisitedNode->IsA<UMovieGraphOutputNode>() && VisitedPin)
 			{
-				BranchNames.Add(VisitedPin->Properties.Label.ToString());
+				BranchNames.AddUnique(VisitedPin->Properties.Label.ToString());
 			}
+
+			return true;
 		}));
 
 	return BranchNames;
 }
 
-TArray<FString> UMovieGraphConfig::GetUpstreamBranchNames(UMovieGraphNode* FromNode, const UMovieGraphPin* FromPin) const
+TArray<FString> UMovieGraphConfig::GetUpstreamBranchNames(UMovieGraphNode* FromNode, const UMovieGraphPin* FromPin, const bool bStopAtSubgraph) const
 {
 	TArray<FString> BranchNames;
 
 	// FromNode itself might be the Inputs node, so check before visiting the upstream nodes
 	if (FromNode->IsA<UMovieGraphInputNode>() && FromPin)
 	{
-		BranchNames.Add(FromPin->Properties.Label.ToString());
+		BranchNames.AddUnique(FromPin->Properties.Label.ToString());
 	}
 
 	VisitUpstreamNodes(FromNode, FVisitNodesCallback::CreateLambda(
-		[&BranchNames](UMovieGraphNode* VisitedNode, const UMovieGraphPin* VisitedPin)
+		[&BranchNames, bStopAtSubgraph](const UMovieGraphNode* VisitedNode, const UMovieGraphPin* VisitedPin)
 		{
+			if (VisitedNode->IsA<UMovieGraphSubgraphNode>() && VisitedPin)
+			{
+				BranchNames.AddUnique(VisitedPin->Properties.Label.ToString());
+
+				if (bStopAtSubgraph)
+				{
+					return false;	// Stop traversing nodes
+				}
+			}
+			
 			if (VisitedNode->IsA<UMovieGraphInputNode>() && VisitedPin)
 			{
-				BranchNames.Add(VisitedPin->Properties.Label.ToString());
+				BranchNames.AddUnique(VisitedPin->Properties.Label.ToString());
 			}
+
+			return true;
 		}));
 
 	return BranchNames;
@@ -1365,8 +1389,16 @@ void UMovieGraphConfig::VisitUpstreamNodes_Recursive(UMovieGraphNode* FromNode,
 		{
 			if (ConnectedPin->Properties.bIsBranch)
 			{
-				VisitCallback.ExecuteIfBound(ConnectedPin->Node, ConnectedPin);
-				VisitUpstreamNodes_Recursive(ConnectedPin->Node, VisitCallback, VisitedNodes);
+				bool bContinueVisiting = true;
+				if (VisitCallback.IsBound())
+				{
+					bContinueVisiting = VisitCallback.Execute(ConnectedPin->Node, ConnectedPin);
+				}
+
+				if (bContinueVisiting)
+				{
+					VisitUpstreamNodes_Recursive(ConnectedPin->Node, VisitCallback, VisitedNodes);
+				}
 			}
 		}
 	}
@@ -1393,8 +1425,16 @@ void UMovieGraphConfig::VisitDownstreamNodes_Recursive(UMovieGraphNode* FromNode
 		{
 			if (ConnectedPin->Properties.bIsBranch)
 			{
-				VisitCallback.ExecuteIfBound(ConnectedPin->Node, ConnectedPin);
-				VisitDownstreamNodes_Recursive(ConnectedPin->Node, VisitCallback, VisitedNodes);
+				bool bContinueVisiting = true;
+				if (VisitCallback.IsBound())
+				{
+					bContinueVisiting = VisitCallback.Execute(ConnectedPin->Node, ConnectedPin);
+				}
+
+				if (bContinueVisiting)
+				{
+					VisitDownstreamNodes_Recursive(ConnectedPin->Node, VisitCallback, VisitedNodes);
+				}
 			}
 		}
 	}
diff --git a/Engine/Plugins/MovieScene/MovieRenderPipeline/Source/MovieRenderPipelineCore/Private/Graph/MovieGraphPin.cpp b/Engine/Plugins/MovieScene/MovieRenderPipeline/Source/MovieRenderPipelineCore/Private/Graph/MovieGraphPin.cpp
--- a/Engine/Plugins/MovieScene/MovieRenderPipeline/Source/MovieRenderPipelineCore/Private/Graph/MovieGraphPin.cpp
+++ b/Engine/Plugins/MovieScene/MovieRenderPipeline/Source/MovieRenderPipelineCore/Private/Graph/MovieGraphPin.cpp
@@ -6,6 +6,7 @@
 #include "EdGraph/EdGraphSchema.h"
 #include "Graph/MovieGraphNode.h"
 #include "Graph/MovieGraphEdge.h"
+#include "Graph/Nodes/MovieGraphSubgraphNode.h"
 #include "MovieRenderPipelineCoreModule.h"
 #include "Graph/MovieGraphConfig.h"
 
@@ -236,6 +237,9 @@ bool UMovieGraphPin::IsConnectionToBranchAllowed(const UMovieGraphPin* OtherPin,
 	check(ToNode && FromNode);
 	const UMovieGraphConfig* GraphConfig = ToNode->GetGraph();
 
+	const bool ToNodeIsSubgraph = ToNode->IsA<UMovieGraphSubgraphNode>();
+	const bool FromNodeIsSubgraph = FromNode->IsA<UMovieGraphSubgraphNode>();
+
 	// Test High-Level Node Restrictions
 	const EMovieGraphBranchRestriction FromNodeRestriction = FromNode->GetBranchRestriction();
 	const EMovieGraphBranchRestriction ToNodeRestriction = ToNode->GetBranchRestriction();
@@ -249,34 +253,90 @@ bool UMovieGraphPin::IsConnectionToBranchAllowed(const UMovieGraphPin* OtherPin,
 
 	// Get all upstream/downstream nodes that occur on the connection -- these are the nodes that need to be checked for branch restrictions.
 	// FromNode/ToNode themselves also needs to be part of the validation checks.
+	//
+	// If the FromNode is a subgraph, there's no need to visit upstream nodes. The subgraph node will enforce branch restrictions, since it
+	// effectively represents an Inputs node. The same logic applies to the ToNode behaving like an Outputs node.
 	TArray<UMovieGraphNode*> NodesToCheck = {FromNode, ToNode};
-	GraphConfig->VisitUpstreamNodes(FromNode, UMovieGraphConfig::FVisitNodesCallback::CreateLambda(
-		[&NodesToCheck](UMovieGraphNode* VisitedNode, const UMovieGraphPin* VisitedPin)
-		{
-			NodesToCheck.Add(VisitedNode);
-		}));
+	if (!FromNodeIsSubgraph)
+	{
+		GraphConfig->VisitUpstreamNodes(FromNode, UMovieGraphConfig::FVisitNodesCallback::CreateLambda(
+			[&NodesToCheck](UMovieGraphNode* VisitedNode, const UMovieGraphPin* VisitedPin)
+			{
+				if (VisitedNode->IsA<UMovieGraphSubgraphNode>())
+				{
+					return false;	// Don't visit more upstream nodes
+				}
+			
+				NodesToCheck.Add(VisitedNode);
+				return true;
+			}));
+	}
 
-	GraphConfig->VisitDownstreamNodes(ToNode, UMovieGraphConfig::FVisitNodesCallback::CreateLambda(
-		[&NodesToCheck](UMovieGraphNode* VisitedNode, const UMovieGraphPin* VisitedPin)
-		{
-			NodesToCheck.Add(VisitedNode);
-		}));
-
-	// Determine which branch(es) are connected to this node up/downstream.
-	const TArray<FString> DownstreamBranchNames = GraphConfig->GetDownstreamBranchNames(ToNode, InputPin);
-	const TArray<FString> UpstreamBranchNames = GraphConfig->GetUpstreamBranchNames(FromNode, OutputPin);
-	const bool bGlobalsIsDownstream = DownstreamBranchNames.Contains(UMovieGraphNode::GlobalsPinNameString);
-	const bool bGlobalsIsUpstream = UpstreamBranchNames.Contains(UMovieGraphNode::GlobalsPinNameString);
+	if (!ToNodeIsSubgraph)
+	{
+		GraphConfig->VisitDownstreamNodes(ToNode, UMovieGraphConfig::FVisitNodesCallback::CreateLambda(
+			[&NodesToCheck](UMovieGraphNode* VisitedNode, const UMovieGraphPin* VisitedPin)
+			{
+				if (VisitedNode->IsA<UMovieGraphSubgraphNode>())
+				{
+					return false;	// Don't visit more downstream nodes
+				}
+				
+				NodesToCheck.Add(VisitedNode);
+				return true;
+			}));
+	}
+
+	const FName InputName = InputPin->Properties.Label;
+	const FName OutputName = OutputPin->Properties.Label;
+	const bool bInputIsGlobals = (InputName == UMovieGraphNode::GlobalsPinName);
+	const bool bOutputIsGlobals = (OutputName == UMovieGraphNode::GlobalsPinName);
+	constexpr bool bStopAtSubgraph = true;
+
+	// Determine which branch(es) are connected to this node up/downstream. If the To/From node is a subgraph, skip trying to traverse the graph past
+	// the subgraph, because for the purposes of determining connection validity, the subgraph's input/output pin is enough.
+	const TArray<FString> DownstreamBranchNames = ToNodeIsSubgraph
+		? TArray{InputName.ToString()}
+		: GraphConfig->GetDownstreamBranchNames(ToNode, InputPin, bStopAtSubgraph);
+	
+	const TArray<FString> UpstreamBranchNames = FromNodeIsSubgraph
+		? TArray{OutputName.ToString()}
+		: GraphConfig->GetUpstreamBranchNames(FromNode, OutputPin, bStopAtSubgraph);
+	
+	const bool bGlobalsIsDownstream = bInputIsGlobals || DownstreamBranchNames.Contains(UMovieGraphNode::GlobalsPinNameString);
+	const bool bGlobalsIsUpstream = bOutputIsGlobals || UpstreamBranchNames.Contains(UMovieGraphNode::GlobalsPinNameString);
 	const bool bDownstreamBranchExistsAndIsntOnlyGlobals =
 		!DownstreamBranchNames.IsEmpty() && ((DownstreamBranchNames.Num() != 1) || (DownstreamBranchNames[0] != UMovieGraphNode::GlobalsPinNameString));
 	const bool bUpstreamBranchExistsAndIsntOnlyGlobals =
 		!UpstreamBranchNames.IsEmpty() && ((UpstreamBranchNames.Num() != 1) || (UpstreamBranchNames[0] != UMovieGraphNode::GlobalsPinNameString));
 
-	// Globals branches can only be connected to Globals branches
-	if ((bGlobalsIsDownstream && bUpstreamBranchExistsAndIsntOnlyGlobals) || (bGlobalsIsUpstream && bDownstreamBranchExistsAndIsntOnlyGlobals))
+	// Subgraph nodes are a special case -- they can be connected to both Globals and render layer branches at the same time
+	if (ToNodeIsSubgraph || FromNodeIsSubgraph)
 	{
-		OutError = NSLOCTEXT("MovieGraph", "GlobalsBranchMismatchError", "Globals branches can only be connected to other Globals branches.");
-		return false;
+		// Only allow Globals -> Globals connections
+		if ((ToNodeIsSubgraph && bInputIsGlobals && !bGlobalsIsUpstream) ||
+			(FromNodeIsSubgraph && bOutputIsGlobals && !bGlobalsIsDownstream))
+		{
+			OutError = NSLOCTEXT("MovieGraph", "SubgraphGlobalsBranchMismatchError", "A subgraph Globals branch can only be connected to another Globals branch.");
+			return false;
+		}
+
+		// Only allow non-Globals -> non-Globals connections
+		if ((ToNodeIsSubgraph && !bInputIsGlobals && bGlobalsIsUpstream) ||
+			(FromNodeIsSubgraph && !bOutputIsGlobals && bGlobalsIsDownstream))
+		{
+			OutError = NSLOCTEXT("MovieGraph", "SubgraphNonGlobalsBranchMismatchError", "A subgraph non-Globals branch can not be connected to the Globals branch.");
+			return false;
+		}
+	}
+	else
+	{
+		// Globals branches can only be connected to Globals branches
+		if ((bGlobalsIsDownstream && bUpstreamBranchExistsAndIsntOnlyGlobals) || (bGlobalsIsUpstream && bDownstreamBranchExistsAndIsntOnlyGlobals))
+		{
+			OutError = NSLOCTEXT("MovieGraph", "GlobalsBranchMismatchError", "Globals branches can only be connected to other Globals branches.");
+			return false;
+		}
 	}
 
 	// Error out if any of the nodes that are part of the connection cannot be connected to the upstream/downstream branches.
diff --git a/Engine/Plugins/MovieScene/MovieRenderPipeline/Source/MovieRenderPipelineCore/Private/Graph/Nodes/MovieGraphSubgraphNode.cpp b/Engine/Plugins/MovieScene/MovieRenderPipeline/Source/MovieRenderPipelineCore/Private/Graph/Nodes/MovieGraphSubgraphNode.cpp
--- a/Engine/Plugins/MovieScene/MovieRenderPipeline/Source/MovieRenderPipelineCore/Private/Graph/Nodes/MovieGraphSubgraphNode.cpp
+++ b/Engine/Plugins/MovieScene/MovieRenderPipeline/Source/MovieRenderPipelineCore/Private/Graph/Nodes/MovieGraphSubgraphNode.cpp
@@ -161,6 +161,7 @@ FText UMovieGraphSubgraphNode::GetMenuCategory() const
 void UMovieGraphSubgraphNode::SetSubGraphAsset(const TSoftObjectPtr<UMovieGraphConfig>& InSubgraphAsset)
 {
 	SubgraphAsset = InSubgraphAsset;
+	UpdatePins();
 }
 
 UMovieGraphConfig* UMovieGraphSubgraphNode::GetSubgraphAsset() const
diff --git a/Engine/Plugins/MovieScene/MovieRenderPipeline/Source/MovieRenderPipelineCore/Private/MoviePipelineQueue.cpp b/Engine/Plugins/MovieScene/MovieRenderPipeline/Source/MovieRenderPipelineCore/Private/MoviePipelineQueue.cpp
--- a/Engine/Plugins/MovieScene/MovieRenderPipeline/Source/MovieRenderPipelineCore/Private/MoviePipelineQueue.cpp
+++ b/Engine/Plugins/MovieScene/MovieRenderPipeline/Source/MovieRenderPipelineCore/Private/MoviePipelineQueue.cpp
@@ -327,6 +327,13 @@ void UMoviePipelineExecutorShot::SetShotOverridePresetOrigin(UMoviePipelineShotC
 	}
 }
 
+bool UMoviePipelineExecutorShot::IsUsingGraphConfiguration() const
+{
+	const UMoviePipelineExecutorJob* PrimaryJob = GetTypedOuter<UMoviePipelineExecutorJob>();
+	
+	return GraphPreset.IsValid() || (PrimaryJob && PrimaryJob->IsUsingGraphConfiguration());
+}
+
 void UMoviePipelineExecutorShot::PreSave(FObjectPreSaveContext ObjectSaveContext)
 {
 	Super::PreSave(ObjectSaveContext);
diff --git a/Engine/Plugins/MovieScene/MovieRenderPipeline/Source/MovieRenderPipelineCore/Public/Graph/MovieGraphConfig.h b/Engine/Plugins/MovieScene/MovieRenderPipeline/Source/MovieRenderPipelineCore/Public/Graph/MovieGraphConfig.h
--- a/Engine/Plugins/MovieScene/MovieRenderPipeline/Source/MovieRenderPipelineCore/Public/Graph/MovieGraphConfig.h
+++ b/Engine/Plugins/MovieScene/MovieRenderPipeline/Source/MovieRenderPipelineCore/Public/Graph/MovieGraphConfig.h
@@ -527,9 +527,10 @@ public:
 	/**
 	 * Callback for when a node is visited. The node is the node being visited, and the pin is the pin which the node
 	 * was accessed by (eg, if visiting downstream nodes, the pin will be the input pin that connects to the node that
-	 * the traversal started from, or the node that was previously visited).
+	 * the traversal started from, or the node that was previously visited). Return true to continue traversal, or false
+	 * to stop traversal.
 	 */
-	DECLARE_DELEGATE_TwoParams(FVisitNodesCallback, UMovieGraphNode*, const UMovieGraphPin*);
+	DECLARE_DELEGATE_RetVal_TwoParams(bool, FVisitNodesCallback, UMovieGraphNode*, const UMovieGraphPin*);
 
 	//~ UObject interface
 	virtual void PostLoad() override;
@@ -621,11 +622,17 @@ public:
 	 */
 	void VisitDownstreamNodes(UMovieGraphNode* FromNode, const FVisitNodesCallback& VisitCallback) const;
 
-	/** Determines the name(s) of the branches downstream from FromNode, starting at FromPin. */
-	TArray<FString> GetDownstreamBranchNames(UMovieGraphNode* FromNode, const UMovieGraphPin* FromPin) const;
+	/**
+	 * Determines the name(s) of the branches downstream from FromNode, starting at FromPin. Optionally, subgraph nodes can halt graph traversal
+	 * if bStopAtSubgraph is set to true.
+	 */
+	TArray<FString> GetDownstreamBranchNames(UMovieGraphNode* FromNode, const UMovieGraphPin* FromPin, const bool bStopAtSubgraph = false) const;
 
-	/** Determines the name(s) of the branches upstream from FromNode, starting at FromPin. */
-	TArray<FString> GetUpstreamBranchNames(UMovieGraphNode* FromNode, const UMovieGraphPin* FromPin) const;
+	/**
+	 * Determines the name(s) of the branches upstream from FromNode, starting at FromPin. Optionally, subgraph nodes can halt graph traversal
+	 * if bStopAtSubgraph is set to true.
+	 */
+	TArray<FString> GetUpstreamBranchNames(UMovieGraphNode* FromNode, const UMovieGraphPin* FromPin, const bool bStopAtSubgraph = false) const;
 
 	/** Get all subgraphs that this graph contains, recursively (ie, subgraphs of subgraphs are included, etc). */
 	void GetAllContainedSubgraphs(TSet<UMovieGraphConfig*>& OutSubgraphs) const;
diff --git a/Engine/Plugins/MovieScene/MovieRenderPipeline/Source/MovieRenderPipelineCore/Public/MoviePipelineQueue.h b/Engine/Plugins/MovieScene/MovieRenderPipeline/Source/MovieRenderPipelineCore/Public/MoviePipelineQueue.h
--- a/Engine/Plugins/MovieScene/MovieRenderPipeline/Source/MovieRenderPipelineCore/Public/MoviePipelineQueue.h
+++ b/Engine/Plugins/MovieScene/MovieRenderPipeline/Source/MovieRenderPipelineCore/Public/MoviePipelineQueue.h
@@ -138,10 +138,7 @@ public:
 	 * Returns true if this job is using graph-style configuration, else false.
 	 */
 	UFUNCTION(BlueprintPure, Category = "Movie Render Pipeline")
-	bool IsUsingGraphConfiguration() const
-	{
-		return GraphPreset.IsValid() || GraphConfig != nullptr;
-	}
+	bool IsUsingGraphConfiguration() const;
 
 	/**
 	 * Gets the graph-style preset that this job is using. If the job is not using a graph-style preset, returns nullptr.
diff --git a/Engine/Plugins/MovieScene/MovieRenderPipeline/Source/MovieRenderPipelineEditor/Private/Graph/MovieGraphConfigFactory.cpp b/Engine/Plugins/MovieScene/MovieRenderPipeline/Source/MovieRenderPipelineEditor/Private/Graph/MovieGraphConfigFactory.cpp
--- a/Engine/Plugins/MovieScene/MovieRenderPipeline/Source/MovieRenderPipelineEditor/Private/Graph/MovieGraphConfigFactory.cpp
+++ b/Engine/Plugins/MovieScene/MovieRenderPipeline/Source/MovieRenderPipelineEditor/Private/Graph/MovieGraphConfigFactory.cpp
@@ -3,6 +3,7 @@
 #include "Graph/MovieGraphConfigFactory.h"
 
 #include "Graph/MovieGraphConfig.h"
+#include "Graph/Nodes/MovieGraphSubgraphNode.h"
 #include "HAL/IConsoleManager.h"
 #include "MovieRenderPipelineCoreModule.h"
 #include "MovieRenderPipelineSettings.h"
@@ -18,6 +19,14 @@ UMovieGraphConfigFactory::UMovieGraphConfigFactory()
 
 UObject* UMovieGraphConfigFactory::FactoryCreateNew(UClass* Class, UObject* InParent, FName Name, EObjectFlags Flags, UObject* Context, FFeedbackContext* Warn)
 {
+	// If requesting a graph w/ a subgraph, skip the normal graph creation which looks for an existing graph asset to use as a template
+	if (InitialSubgraphAsset)
+	{
+		UMovieGraphConfig* NewGraph = NewObject<UMovieGraphConfig>(InParent, Class, Name, Flags);
+		AddSubgraphNodeToGraph(NewGraph);
+		return NewGraph;
+	}
+	
 	const UMovieRenderPipelineProjectSettings* ProjectSettings = GetDefault<UMovieRenderPipelineProjectSettings>();
 	const TSoftObjectPtr<UMovieGraphConfig> ProjectDefaultGraph = ProjectSettings->DefaultGraph;
 
@@ -48,3 +57,50 @@ bool UMovieGraphConfigFactory::ShouldShowInNewMenu() const
 	IConsoleVariable* RenderGraphCVar = IConsoleManager::Get().FindConsoleVariable(TEXT("MoviePipeline.EnableRenderGraph"));
 	return RenderGraphCVar && RenderGraphCVar->GetBool();
 }
+
+void UMovieGraphConfigFactory::AddSubgraphNodeToGraph(UMovieGraphConfig* InTargetGraph) const
+{
+	auto ExposeSubgraphInputOutputOnTargetGraph = [InTargetGraph](
+		const UMovieGraphInterfaceBase* InSubgraphInputOrOutput, const bool bIsInput, UMovieGraphSubgraphNode* InSubgraphNode)
+	{
+		// Add a new member to the target graph if this is not the Globals pin (which already exists)
+		if (InSubgraphInputOrOutput->GetMemberName() != UMovieGraphNode::GlobalsPinName)
+		{
+			UMovieGraphInterfaceBase* NewMember = bIsInput
+				? static_cast<UMovieGraphInterfaceBase*>(InTargetGraph->AddInput())
+				: static_cast<UMovieGraphInterfaceBase*>(InTargetGraph->AddOutput());
+			NewMember->SetMemberName(InSubgraphInputOrOutput->GetMemberName());
+		}
+		
+		UMovieGraphNode* FromNode = bIsInput ? InTargetGraph->GetInputNode() : InSubgraphNode;
+		UMovieGraphNode* ToNode = bIsInput ? InSubgraphNode : InTargetGraph->GetOutputNode();
+		const FName FromPinLabel = *InSubgraphInputOrOutput->GetMemberName();
+		const FName ToPinLabel = *InSubgraphInputOrOutput->GetMemberName();
+
+		// Connect up the new member to the subgraph node
+		InTargetGraph->AddLabeledEdge(FromNode, FromPinLabel, ToNode, ToPinLabel);
+	};
+	
+	UMovieGraphSubgraphNode* NewSubgraphNode = InTargetGraph->ConstructRuntimeNode<UMovieGraphSubgraphNode>();
+	NewSubgraphNode->SetSubGraphAsset(InitialSubgraphAsset);
+
+	// Roughly center the subgraph node between the Inputs and Outputs nodes.
+	const int32 InputNodePosX = InTargetGraph->GetInputNode()->GetNodePosX();
+	const int32 OutputNodePosX = InTargetGraph->GetOutputNode()->GetNodePosX();
+	NewSubgraphNode->SetNodePosX(InputNodePosX + ((OutputNodePosX - InputNodePosX) / 2.f));
+	NewSubgraphNode->SetNodePosY(InTargetGraph->GetInputNode()->GetNodePosY());
+
+	// Expose the subgraph's inputs on the target graph and connect
+	for (const UMovieGraphInput* SubgraphInput : InitialSubgraphAsset->GetInputs())
+	{
+		constexpr bool bIsInput = true;
+		ExposeSubgraphInputOutputOnTargetGraph(SubgraphInput, bIsInput, NewSubgraphNode);
+	}
+
+	// Expose the subgraph's outputs on the target graph and connect
+	for (const UMovieGraphOutput* SubgraphOutput : InitialSubgraphAsset->GetOutputs())
+	{
+		constexpr bool bIsInput = false;
+		ExposeSubgraphInputOutputOnTargetGraph(SubgraphOutput, bIsInput, NewSubgraphNode);
+	}
+}
diff --git a/Engine/Plugins/MovieScene/MovieRenderPipeline/Source/MovieRenderPipelineEditor/Private/Widgets/SMoviePipelineQueueEditor.cpp b/Engine/Plugins/MovieScene/MovieRenderPipeline/Source/MovieRenderPipelineEditor/Private/Widgets/SMoviePipelineQueueEditor.cpp
--- a/Engine/Plugins/MovieScene/MovieRenderPipeline/Source/MovieRenderPipelineEditor/Private/Widgets/SMoviePipelineQueueEditor.cpp
+++ b/Engine/Plugins/MovieScene/MovieRenderPipeline/Source/MovieRenderPipelineEditor/Private/Widgets/SMoviePipelineQueueEditor.cpp
@@ -926,7 +926,12 @@ struct FMoviePipelineShotItem : IMoviePipelineQueueTreeItem
 		{
 			if (Shot->IsUsingGraphConfiguration())
 			{
-				return LOCTEXT("QueueEditorDefaultShotGraph_Text", "Default Graph");
+				if (const UMovieGraphConfig* GraphConfig = Shot->GetGraphPreset())
+				{
+					return FText::FromString(GraphConfig->GetName());
+				}
+				
+				return LOCTEXT("QueueEditorMakeNewShotSubgraph", "Make Subgraph");
 			}
 			
 			// If the shot has a preset origin (ie, its config is based off a preset w/o any modifications), use its
diff --git a/Engine/Plugins/MovieScene/MovieRenderPipeline/Source/MovieRenderPipelineEditor/Private/Widgets/SMoviePipelineQueuePanel.cpp b/Engine/Plugins/MovieScene/MovieRenderPipeline/Source/MovieRenderPipelineEditor/Private/Widgets/SMoviePipelineQueuePanel.cpp
--- a/Engine/Plugins/MovieScene/MovieRenderPipeline/Source/MovieRenderPipelineEditor/Private/Widgets/SMoviePipelineQueuePanel.cpp
+++ b/Engine/Plugins/MovieScene/MovieRenderPipeline/Source/MovieRenderPipelineEditor/Private/Widgets/SMoviePipelineQueuePanel.cpp
@@ -14,6 +14,7 @@
 #include "MoviePipelineQueueSubsystem.h"
 #include "Graph/MovieGraphConfig.h"
 #include "Graph/MovieGraphAssetToolkit.h"
+#include "Graph/MovieGraphConfigFactory.h"
 
 // Slate Includes
 #include "Widgets/SBoxPanel.h"
@@ -370,7 +371,19 @@ void SMoviePipelineQueuePanel::OnEditJobConfigRequested(TWeakObjectPtr<UMoviePip
 	UMovieGraphConfig* GraphToEdit = nullptr;
 	if (InShot.IsValid() && InShot->IsUsingGraphConfiguration())
 	{
-		GraphToEdit = (InShot->GetGraphPreset() != nullptr) ? InShot->GetGraphPreset() : InShot->GetGraphConfig();
+		GraphToEdit = InShot->GetGraphPreset();
+
+		// If the graph preset is not valid, create a new graph
+		if (!GraphToEdit)
+		{
+			GraphToEdit = GenerateNewShotSubgraph(InJob.Get(), InShot.Get());
+
+			// Don't do anything if the user canceled out of the graph creation process, or there was an error creating the graph
+			if (!GraphToEdit)
+			{
+				return;
+			}
+		}
 	}
 	else if (InJob.IsValid() && InJob->IsUsingGraphConfiguration())
 	{
@@ -379,8 +392,7 @@ void SMoviePipelineQueuePanel::OnEditJobConfigRequested(TWeakObjectPtr<UMoviePip
 
 	if (GraphToEdit)
 	{
-		const TSharedRef<FMovieGraphAssetToolkit> MovieGraphEditor(new FMovieGraphAssetToolkit());
-		MovieGraphEditor->InitMovieGraphAssetToolkit(EToolkitMode::Standalone, TSharedPtr<IToolkitHost>(), GraphToEdit);
+		GEditor->GetEditorSubsystem<UAssetEditorSubsystem>()->OpenEditorForAsset(GraphToEdit);
 
 		return;
 	}
@@ -793,4 +805,28 @@ FString SMoviePipelineQueuePanel::GetQueueOriginName() const
 	return FString();
 }
 
+UMovieGraphConfig* SMoviePipelineQueuePanel::GenerateNewShotSubgraph(const UMoviePipelineExecutorJob* InJob, UMoviePipelineExecutorShot* InShot) const
+{
+	UMovieGraphConfigFactory* GraphFactory = NewObject<UMovieGraphConfigFactory>();
+	GraphFactory->InitialSubgraphAsset = InJob->GetGraphPreset();
+			
+	// Make the new graph via save dialog
+	const FAssetToolsModule& AssetToolsModule = FAssetToolsModule::GetModule();
+	UObject* NewAsset = AssetToolsModule.Get().CreateAssetWithDialog(GraphFactory->GetSupportedClass(), GraphFactory);
+
+	// Don't ensure here because a "cancel" in the dialog can cause the returned asset to be null
+	if (UMovieGraphConfig* NewGraph = Cast<UMovieGraphConfig>(NewAsset))
+	{
+		// Save out the new graph and assign it to the shot
+		constexpr bool bOnlyDirty = false;
+		UEditorLoadingAndSavingUtils::SavePackages({ NewGraph->GetPackage() }, bOnlyDirty);
+		
+		InShot->SetGraphPreset(NewGraph);
+		
+		return NewGraph;
+	}
+
+	return nullptr;
+}
+
 #undef LOCTEXT_NAMESPACE // SMoviePipelineQueuePanel
\ No newline at end of file
diff --git a/Engine/Plugins/MovieScene/MovieRenderPipeline/Source/MovieRenderPipelineEditor/Private/Widgets/SMoviePipelineQueuePanel.h b/Engine/Plugins/MovieScene/MovieRenderPipeline/Source/MovieRenderPipelineEditor/Private/Widgets/SMoviePipelineQueuePanel.h
--- a/Engine/Plugins/MovieScene/MovieRenderPipeline/Source/MovieRenderPipelineEditor/Private/Widgets/SMoviePipelineQueuePanel.h
+++ b/Engine/Plugins/MovieScene/MovieRenderPipeline/Source/MovieRenderPipelineEditor/Private/Widgets/SMoviePipelineQueuePanel.h
@@ -64,6 +64,12 @@ private:
 	UMoviePipelineQueue* GetQueueOrigin() const;
 	FString GetQueueOriginName() const;
 
+	/**
+	 * Generates and returns a new shot subgraph based on the given job and shot. Returns nullptr if there was an issue creating the graph. If the
+	 * graph was successfully created, it will be assigned to the shot.
+	 */
+	class UMovieGraphConfig* GenerateNewShotSubgraph(const UMoviePipelineExecutorJob* InJob, UMoviePipelineExecutorShot* InShot) const;
+
 private:
 	/** Allocates a transient preset so that the user can use the pipeline without saving it to an asset first. */
 	//UMoviePipelineConfigBase* AllocateTransientPreset();
diff --git a/Engine/Plugins/MovieScene/MovieRenderPipeline/Source/MovieRenderPipelineEditor/Public/Graph/MovieGraphConfigFactory.h b/Engine/Plugins/MovieScene/MovieRenderPipeline/Source/MovieRenderPipelineEditor/Public/Graph/MovieGraphConfigFactory.h
--- a/Engine/Plugins/MovieScene/MovieRenderPipeline/Source/MovieRenderPipelineEditor/Public/Graph/MovieGraphConfigFactory.h
+++ b/Engine/Plugins/MovieScene/MovieRenderPipeline/Source/MovieRenderPipelineEditor/Public/Graph/MovieGraphConfigFactory.h
@@ -5,14 +5,32 @@
 #include "Factories/Factory.h"
 #include "MovieGraphConfigFactory.generated.h"
 
+class UMovieGraphConfig;
+
 UCLASS(BlueprintType)
 class MOVIERENDERPIPELINEEDITOR_API UMovieGraphConfigFactory : public UFactory
 {
     GENERATED_BODY()
 public:
 	UMovieGraphConfigFactory();
+	
 	//~ Begin UFactory Interface
 	virtual UObject* FactoryCreateNew(UClass* Class, UObject* InParent, FName Name, EObjectFlags Flags, UObject* Context, FFeedbackContext* Warn) override;
 	virtual bool ShouldShowInNewMenu() const override;
 	//~ End UFactory Interface
+
+private:
+	/**
+	 * If the factory was created with an InitialSubgraphAsset, adds a subgraph node pointing to InitialSubgraphAsset to the target graph, and
+	 * connects it up to the Inputs and Outputs nodes.
+	 */
+	void AddSubgraphNodeToGraph(UMovieGraphConfig* InTargetGraph) const;
+
+public:
+	/**
+	 * If set, creates an empty graph with one subgraph node (where the subgraph points to this graph). If not set, the factory will create a graph
+	 * from the template like normal.
+	 */
+	UPROPERTY(Transient)
+	UMovieGraphConfig* InitialSubgraphAsset = nullptr;
 };
\ No newline at end of file
