MovieGraph: First pass at exposing OpenColorIO context.

Also adjusting property display names to better match Ryan's sketch.

#rnx

diff --git a/Engine/Plugins/Compositing/OpenColorIO/Source/OpenColorIO/Private/OpenColorIOColorTransform.cpp b/Engine/Plugins/Compositing/OpenColorIO/Source/OpenColorIO/Private/OpenColorIOColorTransform.cpp
--- a/Engine/Plugins/Compositing/OpenColorIO/Source/OpenColorIO/Private/OpenColorIOColorTransform.cpp
+++ b/Engine/Plugins/Compositing/OpenColorIO/Source/OpenColorIO/Private/OpenColorIOColorTransform.cpp
@@ -567,6 +567,11 @@ bool UOpenColorIOColorTransform::IsTransform(const FString& InSourceColorSpace,
 
 #if WITH_EDITOR
 bool UOpenColorIOColorTransform::GetTransformProcessor(FOpenColorIOWrapperProcessor& OutProcessor) const
+{
+	return GetTransformProcessor(GetContextKeyValues(), OutProcessor);
+}
+
+bool UOpenColorIOColorTransform::GetTransformProcessor(const TMap<FString, FString>& InLocalContext, FOpenColorIOWrapperProcessor& OutProcessor) const
 {
 	const UOpenColorIOConfiguration* ConfigurationOwner = Cast<UOpenColorIOConfiguration>(GetOuter());
 	const FOpenColorIOWrapperConfig* ConfigWrapper = GetTransformConfigWrapper(ConfigurationOwner);
@@ -581,11 +586,11 @@ bool UOpenColorIOColorTransform::GetTransformProcessor(FOpenColorIOWrapperProces
 	EOpenColorIOViewTransformDirection CurrentDisplayViewDirection;
 	if (GetDisplayViewDirection(CurrentDisplayViewDirection))
 	{
-		OutProcessor = FOpenColorIOWrapperProcessor(ConfigWrapper, SourceColorSpace, Display, View, CurrentDisplayViewDirection == EOpenColorIOViewTransformDirection::Inverse, GetContextKeyValues());
+		OutProcessor = FOpenColorIOWrapperProcessor(ConfigWrapper, SourceColorSpace, Display, View, CurrentDisplayViewDirection == EOpenColorIOViewTransformDirection::Inverse, InLocalContext);
 	}
 	else
 	{
-		OutProcessor = FOpenColorIOWrapperProcessor(ConfigWrapper, SourceColorSpace, DestinationColorSpace, GetContextKeyValues());
+		OutProcessor = FOpenColorIOWrapperProcessor(ConfigWrapper, SourceColorSpace, DestinationColorSpace, InLocalContext);
 	}
 
 	return OutProcessor.IsValid();
diff --git a/Engine/Plugins/Compositing/OpenColorIO/Source/OpenColorIO/Public/OpenColorIOColorSpace.h b/Engine/Plugins/Compositing/OpenColorIO/Source/OpenColorIO/Public/OpenColorIOColorSpace.h
--- a/Engine/Plugins/Compositing/OpenColorIO/Source/OpenColorIO/Public/OpenColorIOColorSpace.h
+++ b/Engine/Plugins/Compositing/OpenColorIO/Source/OpenColorIO/Public/OpenColorIOColorSpace.h
@@ -225,7 +225,7 @@ public:
 	 *  Since display look are applied on viewports, this will 
 	 * dictate whether it's applied or not to it
 	 */
-	UPROPERTY(EditAnywhere, BlueprintReadWrite, Category = ColorSpace)
+	UPROPERTY(EditAnywhere, BlueprintReadWrite, Category = ColorSpace, DisplayName = "Enable OCIO")
 	bool bIsEnabled = false;
 	
 	/** Conversion to apply when this display is enabled */
diff --git a/Engine/Plugins/Compositing/OpenColorIO/Source/OpenColorIO/Public/OpenColorIOColorTransform.h b/Engine/Plugins/Compositing/OpenColorIO/Source/OpenColorIO/Public/OpenColorIOColorTransform.h
--- a/Engine/Plugins/Compositing/OpenColorIO/Source/OpenColorIO/Public/OpenColorIOColorTransform.h
+++ b/Engine/Plugins/Compositing/OpenColorIO/Source/OpenColorIO/Public/OpenColorIOColorTransform.h
@@ -91,10 +91,20 @@ public:
 	/**
 	 * Get the transform processor.
 	 *
+	 * @param OutProcessor Processor wrapper for the current transform.
 	 * @return True if the processor was created and is valid.
 	 */
 	bool GetTransformProcessor(FOpenColorIOWrapperProcessor& OutProcessor) const;
 
+	/**
+	 * Get the transform processor.
+	 *
+	 * @param InLocalContext Context key-value string pairs.
+	 * @param OutProcessor Processor wrapper for the current transform, with a specified context.
+	 * @return True if the processor was created and is valid.
+	 */
+	bool GetTransformProcessor(const TMap<FString, FString>& InLocalContext, FOpenColorIOWrapperProcessor& OutProcessor) const;
+
 	/** Apply the color transform in-place to the specified color. */
 	bool TransformColor(FLinearColor& InOutColor) const;
 		
diff --git a/Engine/Plugins/MovieScene/MovieRenderPipeline/Source/MovieRenderPipelineRenderPasses/Private/MovieGraphImageSequenceOutputNode.cpp b/Engine/Plugins/MovieScene/MovieRenderPipeline/Source/MovieRenderPipelineRenderPasses/Private/MovieGraphImageSequenceOutputNode.cpp
--- a/Engine/Plugins/MovieScene/MovieRenderPipeline/Source/MovieRenderPipelineRenderPasses/Private/MovieGraphImageSequenceOutputNode.cpp
+++ b/Engine/Plugins/MovieScene/MovieRenderPipeline/Source/MovieRenderPipelineRenderPasses/Private/MovieGraphImageSequenceOutputNode.cpp
@@ -45,7 +45,7 @@ namespace UE::MovieGraph::Private
 	 *
 	 * @return The pixel preprocessor if successful, nullptr otherwise.
 	*/
-	static FPixelPreProcessor CreateOpenColorIOPixelPreProcessor(const FOpenColorIOColorConversionSettings& InConversionSettings)
+	static FPixelPreProcessor CreateOpenColorIOPixelPreProcessor(const FOpenColorIOColorConversionSettings& InConversionSettings, const TMap<FString, FString>& InContext)
 	{
 		const TObjectPtr<UOpenColorIOConfiguration>& ConfigurationSource = InConversionSettings.ConfigurationSource;
 		if (IsValid(ConfigurationSource))
@@ -53,8 +53,9 @@ namespace UE::MovieGraph::Private
 			TObjectPtr<const UOpenColorIOColorTransform> ColorTransform = ConfigurationSource->FindTransform(InConversionSettings);
 			if (IsValid(ColorTransform))
 			{
+				//TODO: We need further processing of the context values to support dynamic tokens such as {shot_name}.
 				FOpenColorIOWrapperProcessor Processor;
-				if (ColorTransform->GetTransformProcessor(Processor))
+				if (ColorTransform->GetTransformProcessor(InContext, Processor))
 				{
 					return FOpenColorIOPixelPreProcessor(MoveTemp(Processor));
 				}
@@ -272,7 +273,10 @@ void UMovieGraphImageSequenceOutputNode::OnReceiveImageDataImpl(UMovieGraphPipel
 #if WITH_EDITOR
 		if (ParentNode->OCIOConfiguration.bIsEnabled && Payload->bAllowOCIO)
 		{
-			FPixelPreProcessor OCIOPixelPreProcessor = UE::MovieGraph::Private::CreateOpenColorIOPixelPreProcessor(ParentNode->OCIOConfiguration.ColorConfiguration);
+			FPixelPreProcessor OCIOPixelPreProcessor = UE::MovieGraph::Private::CreateOpenColorIOPixelPreProcessor(
+				ParentNode->OCIOConfiguration.ColorConfiguration,
+				ParentNode->OCIOContext
+			);
 			if (OCIOPixelPreProcessor)
 			{
 				TileImageTask->PixelPreProcessors.Emplace(MoveTemp(OCIOPixelPreProcessor));
@@ -375,7 +379,10 @@ void UMovieGraphImageSequenceOutputNode_EXR::UpdateTaskPerLayer(
 #if WITH_EDITOR
 	if (InParentNode->OCIOConfiguration.bIsEnabled && Payload->bAllowOCIO)
 	{
-		FPixelPreProcessor OCIOPixelPreProcessor = UE::MovieGraph::Private::CreateOpenColorIOPixelPreProcessor(InParentNode->OCIOConfiguration.ColorConfiguration);
+		FPixelPreProcessor OCIOPixelPreProcessor = UE::MovieGraph::Private::CreateOpenColorIOPixelPreProcessor(
+			InParentNode->OCIOConfiguration.ColorConfiguration,
+			InParentNode->OCIOContext
+		);
 		if (OCIOPixelPreProcessor)
 		{
 			InOutImageTask.PixelPreprocessors.FindOrAdd(InLayerIndex).Emplace(MoveTemp(OCIOPixelPreProcessor));
diff --git a/Engine/Plugins/MovieScene/MovieRenderPipeline/Source/MovieRenderPipelineRenderPasses/Public/MovieGraphImageSequenceOutputNode.h b/Engine/Plugins/MovieScene/MovieRenderPipeline/Source/MovieRenderPipelineRenderPasses/Public/MovieGraphImageSequenceOutputNode.h
--- a/Engine/Plugins/MovieScene/MovieRenderPipeline/Source/MovieRenderPipelineRenderPasses/Public/MovieGraphImageSequenceOutputNode.h
+++ b/Engine/Plugins/MovieScene/MovieRenderPipeline/Source/MovieRenderPipelineRenderPasses/Public/MovieGraphImageSequenceOutputNode.h
@@ -37,6 +37,9 @@ public:
 	UPROPERTY(EditAnywhere, BlueprintReadWrite, Category = Overrides, meta = (InlineEditConditionToggle))
 	uint8 bOverride_OCIOConfiguration : 1;
 
+	UPROPERTY(EditAnywhere, BlueprintReadWrite, Category = Overrides, meta = (InlineEditConditionToggle))
+	uint8 bOverride_OCIOContext : 1;
+
 	/**
 	* OCIO configuration/transform settings.
 	*
@@ -44,8 +47,14 @@ public:
 	* 1) This feature only works in editor-mode, since the OpenColorIO library is currently unavailable in game builds.
 	* 2) Users are now responsible for setting the renderer output space to Final Color (HDR) in Linear Working Color Space (SCS_FinalColorHDR).
 	*/
-	UPROPERTY(EditAnywhere, BlueprintReadWrite, Category = "Color", DisplayName="OpenColorIO Configuration", meta = (EditCondition = "bOverride_OCIOConfiguration"))
+	UPROPERTY(EditAnywhere, BlueprintReadWrite, Category = "OCIO", DisplayName="OCIO Configuration", meta = (EditCondition = "bOverride_OCIOConfiguration"))
 	FOpenColorIODisplayConfiguration OCIOConfiguration;
+
+	/**
+	* OCIO context of key-value string pairs, typically used to apply shot-specific looks (such as a CDL color correction, or a 1D grade LUT).
+	*/
+	UPROPERTY(EditAnywhere, BlueprintReadWrite, Category = "OCIO", DisplayName = "OCIO Context", meta = (EditCondition = "bOverride_OCIOContext"))
+	TMap<FString, FString> OCIOContext;
 #endif
 
 protected:
