MoviePipeline: Graph-related queue UI fixes.

- Don't open multiple editors for the same graph.
- Show graph name in the queue.
- Correctly dirty queue when reassigning graph.

#rnx

diff --git a/Engine/Plugins/MovieScene/MovieRenderPipeline/Source/MovieRenderPipelineEditor/Private/Widgets/SMoviePipelineQueueEditor.cpp b/Engine/Plugins/MovieScene/MovieRenderPipeline/Source/MovieRenderPipelineEditor/Private/Widgets/SMoviePipelineQueueEditor.cpp
--- a/Engine/Plugins/MovieScene/MovieRenderPipeline/Source/MovieRenderPipelineEditor/Private/Widgets/SMoviePipelineQueueEditor.cpp
+++ b/Engine/Plugins/MovieScene/MovieRenderPipeline/Source/MovieRenderPipelineEditor/Private/Widgets/SMoviePipelineQueueEditor.cpp
@@ -230,9 +230,16 @@ public:
 		{
 			if (Job->IsUsingGraphConfiguration())
 			{
-				return LOCTEXT("QueueEditorDefaultJobGraph_Text", "Default Graph");
+				if (Job->GetGraphPreset())
+				{
+					return FText::FromString(Job->GetGraphPreset()->GetName());
+				}
+				else
+				{
+					return LOCTEXT("QueueEditorDefaultJobGraph_Text", "Default Graph");
+				}
 			}
-			
+
 			// If the job has a preset origin (ie, its config is based off a preset w/o any modifications), use its
 			// display name. If the config has a preset origin (ie, it's based off a preset, but has modifications), use
 			// that display name. Otherwise, fall back to the config's display name.
@@ -264,6 +271,9 @@ public:
 		UMoviePipelineExecutorJob* Job = WeakJob.Get();
 		if (Job)
 		{
+			FScopedTransaction Transaction(LOCTEXT("PickJobPresetAsset_Transaction", "Set Job Configuration Asset"));
+			Job->Modify();
+
 			if (Job->IsUsingGraphConfiguration())
 			{
 				Job->SetGraphPreset(CastChecked<UMovieGraphConfig>(AssetData.GetAsset()));
