MoviePipeline: Added support to graph nodes to allow them to specify game-mode command line arguments (used when render remote is active). This is similar to how the old system allowed settings to specify arguments.

This concept may need to be reworked in the future though, because jobs could have nodes that provide conflicting arguments, just like in the old system. But for now, this provides parity with the old system.

#rnx

diff --git a/Engine/Plugins/MovieScene/MovieRenderPipeline/Source/MovieRenderPipelineCore/Private/Graph/Nodes/MovieGraphAudioOutputNode.cpp b/Engine/Plugins/MovieScene/MovieRenderPipeline/Source/MovieRenderPipelineCore/Private/Graph/Nodes/MovieGraphAudioOutputNode.cpp
--- a/Engine/Plugins/MovieScene/MovieRenderPipeline/Source/MovieRenderPipelineCore/Private/Graph/Nodes/MovieGraphAudioOutputNode.cpp
+++ b/Engine/Plugins/MovieScene/MovieRenderPipeline/Source/MovieRenderPipelineCore/Private/Graph/Nodes/MovieGraphAudioOutputNode.cpp
@@ -19,6 +19,12 @@ static TAutoConsoleVariable<float> CVarWaveOutputDelay(
 	TEXT("encoder runs, the encoder may fail.\n"),
 	ECVF_Default);
 
+void UMovieGraphAudioOutputNode::BuildNewProcessCommandLineArgsImpl(TArray<FString>& InOutUnrealURLParams, TArray<FString>& InOutCommandLineArgs, TArray<FString>& InOutDeviceProfileCvars, TArray<FString>& InOutExecCmds) const
+{
+	// Always add this so that audio is muted, it'll never line up during preview anyways.
+	InOutCommandLineArgs.Add("-deterministicaudio");
+}
+
 #if WITH_EDITOR
 FText UMovieGraphAudioOutputNode::GetNodeTitle(const bool bGetDescriptive) const
 {
diff --git a/Engine/Plugins/MovieScene/MovieRenderPipeline/Source/MovieRenderPipelineCore/Private/Graph/Nodes/MovieGraphGlobalGameOverrides.cpp b/Engine/Plugins/MovieScene/MovieRenderPipeline/Source/MovieRenderPipelineCore/Private/Graph/Nodes/MovieGraphGlobalGameOverrides.cpp
--- a/Engine/Plugins/MovieScene/MovieRenderPipeline/Source/MovieRenderPipelineCore/Private/Graph/Nodes/MovieGraphGlobalGameOverrides.cpp
+++ b/Engine/Plugins/MovieScene/MovieRenderPipeline/Source/MovieRenderPipelineCore/Private/Graph/Nodes/MovieGraphGlobalGameOverrides.cpp
@@ -21,6 +21,72 @@ UMovieGraphGlobalGameOverridesNode::UMovieGraphGlobalGameOverridesNode()
 {
 }
 
+void UMovieGraphGlobalGameOverridesNode::BuildNewProcessCommandLineArgsImpl(TArray<FString>& InOutUnrealURLParams, TArray<FString>& InOutCommandLineArgs, TArray<FString>& InOutDeviceProfileCvars, TArray<FString>& InOutExecCmds) const
+{
+	// We don't provide the GameMode on the command line argument as we expect NewProcess to boot into an empty map and then it will
+	// transition into the correct map which will then use the GameModeOverride setting.
+	
+	{
+		Scalability::FQualityLevels QualityLevels;
+		QualityLevels.SetFromSingleQualityLevel(static_cast<int32>(ScalabilityQualityLevel));
+
+		InOutDeviceProfileCvars.Add(FString::Format(TEXT("sg.ViewDistanceQuality={0}"), {QualityLevels.ViewDistanceQuality}));
+		InOutDeviceProfileCvars.Add(FString::Format(TEXT("sg.AntiAliasingQuality={0}"), {QualityLevels.AntiAliasingQuality}));
+		InOutDeviceProfileCvars.Add(FString::Format(TEXT("sg.ShadowQuality={0}"), {QualityLevels.ShadowQuality}));
+		InOutDeviceProfileCvars.Add(FString::Format(TEXT("sg.GlobalIlluminationQuality={0}"), {QualityLevels.GlobalIlluminationQuality}));
+		InOutDeviceProfileCvars.Add(FString::Format(TEXT("sg.ReflectionQuality={0}"), {QualityLevels.ReflectionQuality}));
+		InOutDeviceProfileCvars.Add(FString::Format(TEXT("sg.PostProcessQuality={0}"), {QualityLevels.PostProcessQuality}));
+		InOutDeviceProfileCvars.Add(FString::Format(TEXT("sg.TextureQuality={0}"), {QualityLevels.TextureQuality}));
+		InOutDeviceProfileCvars.Add(FString::Format(TEXT("sg.EffectsQuality={0}"), {QualityLevels.EffectsQuality}));
+		InOutDeviceProfileCvars.Add(FString::Format(TEXT("sg.FoliageQuality={0}"), {QualityLevels.FoliageQuality}));
+		InOutDeviceProfileCvars.Add(FString::Format(TEXT("sg.ShadingQuality={0}"), {QualityLevels.ShadingQuality}));
+		InOutDeviceProfileCvars.Add(FString::Format(TEXT("sg.LandscapeQuality={0}"), {QualityLevels.LandscapeQuality}));
+	}
+
+	if (bDisableTextureStreaming)
+	{
+		InOutDeviceProfileCvars.Add(TEXT("r.TextureStreaming=0"));
+	}
+
+	if (bDisableLODs)
+	{
+		InOutDeviceProfileCvars.Add(TEXT("r.ForceLOD=0"));
+		InOutDeviceProfileCvars.Add(TEXT("r.SkeletalMeshLODBias=-10"));
+		InOutDeviceProfileCvars.Add(TEXT("r.ParticleLODBias=-10"));
+		InOutDeviceProfileCvars.Add(TEXT("foliage.DitheredLOD=0"));
+		InOutDeviceProfileCvars.Add(TEXT("foliage.ForceLOD=0"));
+	}
+
+	if (bDisableHLODs)
+	{
+		// It's a command and not an integer cvar (despite taking 1/0)
+		InOutExecCmds.Add(TEXT("r.HLOD 0"));
+	}
+
+	if (bFlushStreamingManagers)
+	{
+		InOutDeviceProfileCvars.Add(TEXT("r.Streaming.SyncStatesWhenBlocking=1"));
+	}
+
+	// Like the extra cvars applied in ApplySettings(), the below are applied to allow MRQ to function correctly.
+
+#if WITH_EDITOR
+	{
+		InOutDeviceProfileCvars.Add(TEXT("GeometryCache.Streamer.BlockTillFinishStreaming=1"));
+		InOutDeviceProfileCvars.Add(TEXT("GeometryCache.Streamer.ShowNotification=0"));
+	}
+#endif
+
+	InOutDeviceProfileCvars.Add(FString::Printf(TEXT("a.URO.Enable=%d"), 0));
+	InOutDeviceProfileCvars.Add(FString::Printf(TEXT("r.SkyLight.RealTimeReflectionCapture.TimeSlice=%d"), 0));
+	InOutDeviceProfileCvars.Add(FString::Printf(TEXT("r.VolumetricRenderTarget=%d"), 1));
+	InOutDeviceProfileCvars.Add(FString::Printf(TEXT("r.VolumetricRenderTarget.Mode=%d"), 3));
+	InOutDeviceProfileCvars.Add(FString::Printf(TEXT("wp.Runtime.BlockOnSlowStreaming=%d"), 0));
+	InOutDeviceProfileCvars.Add(FString::Printf(TEXT("p.Chaos.ImmPhys.MinStepTime=%d"), 0));
+	InOutDeviceProfileCvars.Add(FString::Printf(TEXT("r.SkipRedundantTransformUpdate=%d"), 0));
+	InOutDeviceProfileCvars.Add(FString::Printf(TEXT("p.ChaosCloth.UseTimeStepSmoothing=%d"), 0));
+}
+
 void UMovieGraphGlobalGameOverridesNode::ApplySettings(const bool bOverrideValues, UWorld* InWorld)
 {
 	// Apply the scalability settings
diff --git a/Engine/Plugins/MovieScene/MovieRenderPipeline/Source/MovieRenderPipelineCore/Public/Graph/MovieGraphNode.h b/Engine/Plugins/MovieScene/MovieRenderPipeline/Source/MovieRenderPipelineCore/Public/Graph/MovieGraphNode.h
--- a/Engine/Plugins/MovieScene/MovieRenderPipeline/Source/MovieRenderPipelineCore/Public/Graph/MovieGraphNode.h
+++ b/Engine/Plugins/MovieScene/MovieRenderPipeline/Source/MovieRenderPipelineCore/Public/Graph/MovieGraphNode.h
@@ -354,6 +354,7 @@ UCLASS(Abstract, BlueprintType)
 class MOVIERENDERPIPELINECORE_API UMovieGraphSettingNode : public UMovieGraphNode
 {
 	GENERATED_BODY()
+	
 public:
 	// UMovieGraphNode Interface
 	virtual TArray<FMovieGraphPinProperties> GetInputPinProperties() const override;
@@ -377,6 +378,9 @@ public:
 	* resolving the settings of the graph, the node only needs to read its own values.
 	*/
 	virtual void GetFormatResolveArgs(FMovieGraphResolveArgs& OutMergedFormatArgs, const FMovieGraphRenderDataIdentifier& InRenderDataIdentifier) const {}
+	
+	/** Modify the Unreal URL and command line arguments when the node will be run in a new process. Only applies to nodes in the Globals branch. */
+	virtual void BuildNewProcessCommandLineArgsImpl(TArray<FString>& InOutUnrealURLParams, TArray<FString>& InOutCommandLineArgs, TArray<FString>& InOutDeviceProfileCvars, TArray<FString>& InOutExecCmds) const { }
 };
 
 /**
diff --git a/Engine/Plugins/MovieScene/MovieRenderPipeline/Source/MovieRenderPipelineCore/Public/Graph/Nodes/MovieGraphAudioOutputNode.h b/Engine/Plugins/MovieScene/MovieRenderPipeline/Source/MovieRenderPipelineCore/Public/Graph/Nodes/MovieGraphAudioOutputNode.h
--- a/Engine/Plugins/MovieScene/MovieRenderPipeline/Source/MovieRenderPipelineCore/Public/Graph/Nodes/MovieGraphAudioOutputNode.h
+++ b/Engine/Plugins/MovieScene/MovieRenderPipeline/Source/MovieRenderPipelineCore/Public/Graph/Nodes/MovieGraphAudioOutputNode.h
@@ -19,16 +19,18 @@ public:
 	UMovieGraphAudioOutputNode() = default;
 
 	// UMovieGraphSettingNode Interface
+	virtual void BuildNewProcessCommandLineArgsImpl(TArray<FString>& InOutUnrealURLParams, TArray<FString>& InOutCommandLineArgs, TArray<FString>& InOutDeviceProfileCvars, TArray<FString>& InOutExecCmds) const override;
+	// ~UMovieGraphSettingNode Interface
+
+	// UMovieGraphNode Interface
+	virtual EMovieGraphBranchRestriction GetBranchRestriction() const override;
+
 #if WITH_EDITOR
 	virtual FText GetNodeTitle(const bool bGetDescriptive) const override;
 	virtual FText GetKeywords() const override;
 	virtual FLinearColor GetNodeTitleColor() const override;
 	virtual FSlateIcon GetIconAndTint(FLinearColor& OutColor) const override;
 #endif
-	// ~UMovieGraphSettingNode Interface
-
-	// UMovieGraphNode Interface
-	virtual EMovieGraphBranchRestriction GetBranchRestriction() const override;
 	// ~UMovieGraphNode Interface
 
 	/**
diff --git a/Engine/Plugins/MovieScene/MovieRenderPipeline/Source/MovieRenderPipelineCore/Public/Graph/Nodes/MovieGraphGlobalGameOverrides.h b/Engine/Plugins/MovieScene/MovieRenderPipeline/Source/MovieRenderPipelineCore/Public/Graph/Nodes/MovieGraphGlobalGameOverrides.h
--- a/Engine/Plugins/MovieScene/MovieRenderPipeline/Source/MovieRenderPipelineCore/Public/Graph/Nodes/MovieGraphGlobalGameOverrides.h
+++ b/Engine/Plugins/MovieScene/MovieRenderPipeline/Source/MovieRenderPipelineCore/Public/Graph/Nodes/MovieGraphGlobalGameOverrides.h
@@ -27,7 +27,13 @@ class MOVIERENDERPIPELINECORE_API UMovieGraphGlobalGameOverridesNode : public UM
 public:
 	UMovieGraphGlobalGameOverridesNode();
 
+	// UMovieGraphNode interface
 	virtual EMovieGraphBranchRestriction GetBranchRestriction() const override { return EMovieGraphBranchRestriction::Globals; }
+	// ~UMovieGraphNode interface
+
+	// UMovieGraphSettingNode interface
+	virtual void BuildNewProcessCommandLineArgsImpl(TArray<FString>& InOutUnrealURLParams, TArray<FString>& InOutCommandLineArgs, TArray<FString>& InOutDeviceProfileCvars, TArray<FString>& InOutExecCmds) const override;
+	// ~UMovieGraphSettingNode interface
 
 	/**
 	 * Applies any cvars, scalability settings, etc. to reflect the properties set on the node. Remembers what the
diff --git a/Engine/Plugins/MovieScene/MovieRenderPipeline/Source/MovieRenderPipelineEditor/Private/MoviePipelineNewProcessExecutor.cpp b/Engine/Plugins/MovieScene/MovieRenderPipeline/Source/MovieRenderPipelineEditor/Private/MoviePipelineNewProcessExecutor.cpp
--- a/Engine/Plugins/MovieScene/MovieRenderPipeline/Source/MovieRenderPipelineEditor/Private/MoviePipelineNewProcessExecutor.cpp
+++ b/Engine/Plugins/MovieScene/MovieRenderPipeline/Source/MovieRenderPipelineEditor/Private/MoviePipelineNewProcessExecutor.cpp
@@ -17,6 +17,7 @@
 #include "Misc/FileHelper.h"
 #include "MoviePipelineEditorBlueprintLibrary.h"
 #include "MoviePipelineGameMode.h"
+#include "MoviePipelineUtils.h"
 
 #include UE_INLINE_GENERATED_CPP_BY_NAME(MoviePipelineNewProcessExecutor)
 
@@ -85,18 +86,51 @@ void UMoviePipelineNewProcessExecutor::Execute_Implementation(UMoviePipelineQueu
 	// Boot into our custom Game Mode (to go with our custom map). Once booted the in-process Executor will load the correct map with correct gamemode.
 	UnrealURLParams += FString::Printf(TEXT("?game=%s"), *AMoviePipelineGameMode::StaticClass()->GetPathName());
 
-	// Loop through our settings in the job and let them modify the command line arguments/params. Because we could have multiple jobs,
-	// we go through all jobs and all settings and hope the user doesn't have conflicting settings.
 	TArray<FString> OutUrlParams;
 	TArray<FString> OutCommandLineArgs;
 	TArray<FString> OutDeviceProfileCVars;
 	TArray<FString> OutExecCmds;
-	for (const UMoviePipelineExecutorJob* Job : DuplicatedQueue->GetJobs())
+
+	// Loop through our settings in the job and let them modify the command line arguments/params. Because we could have multiple jobs,
+	// we go through all jobs and all settings and hope the user doesn't have conflicting settings.
+	for (UMoviePipelineExecutorJob* Job : DuplicatedQueue->GetJobs())
 	{
-		Job->GetConfiguration()->InitializeTransientSettings();
-		for (const UMoviePipelineSetting* Setting : Job->GetConfiguration()->GetAllSettings())
+		if (Job->IsUsingGraphConfiguration())
+		{
+			UMovieGraphConfig* GraphConfig = Job->GetGraphPreset();
+			if (!GraphConfig)
+			{
+				continue;
+			}
+			
+			FMovieGraphTraversalContext TraversalContext;
+			TraversalContext.Job = Job;
+
+			UMovieGraphEvaluatedConfig* EvaluatedGraph = GraphConfig->CreateFlattenedGraph(TraversalContext);
+
+			// Only nodes on the Globals branch can apply command line args. Get CDOs here as well, so all nodes which define command line args
+			// will be found, even if they were not instantiated in the graph.
+			TArray<UClass*> AllSettingsNodeClasses = UE::MovieRenderPipeline::FindMoviePipelineSettingClasses(UMovieGraphSettingNode::StaticClass());
+			for (UClass* NodeClass : AllSettingsNodeClasses)
+			{
+				constexpr bool bIncludeCDOs = true;
+				constexpr bool bExactMatch = true;
+				UMovieGraphSettingNode* GlobalsNode = EvaluatedGraph->GetSettingForBranch(NodeClass, UMovieGraphNode::GlobalsPinName, bIncludeCDOs, bExactMatch);
+				if (!GlobalsNode || (GlobalsNode->GetBranchRestriction() != EMovieGraphBranchRestriction::Globals))
+				{
+					continue;
+				}
+
+				GlobalsNode->BuildNewProcessCommandLineArgsImpl(OutUrlParams, OutCommandLineArgs, OutDeviceProfileCVars, OutExecCmds);
+			}
+		}
+		else
 		{
-			Setting->BuildNewProcessCommandLineArgs(OutUrlParams, OutCommandLineArgs, OutDeviceProfileCVars, OutExecCmds);
+			Job->GetConfiguration()->InitializeTransientSettings();
+			for (const UMoviePipelineSetting* Setting : Job->GetConfiguration()->GetAllSettings())
+			{
+				Setting->BuildNewProcessCommandLineArgs(OutUrlParams, OutCommandLineArgs, OutDeviceProfileCVars, OutExecCmds);
+			}
 		}
 	}
 
