MovieGraph: Add support for showing multiple render layers in the preview widget, and visualizing transparency of the individual layers using an AlphaComposite (A+B(1-Alpha)) blend to correctly show when there is data in the RGB channels even if there is no alpha for those pixels.

diff --git a/Engine/Plugins/MovieScene/MovieRenderPipeline/Source/MovieRenderPipelineCore/Private/Graph/MovieGraphDefaultRenderer.cpp b/Engine/Plugins/MovieScene/MovieRenderPipeline/Source/MovieRenderPipelineCore/Private/Graph/MovieGraphDefaultRenderer.cpp
--- a/Engine/Plugins/MovieScene/MovieRenderPipeline/Source/MovieRenderPipelineCore/Private/Graph/MovieGraphDefaultRenderer.cpp
+++ b/Engine/Plugins/MovieScene/MovieRenderPipeline/Source/MovieRenderPipelineCore/Private/Graph/MovieGraphDefaultRenderer.cpp
@@ -173,7 +173,7 @@ void UMovieGraphDefaultRenderer::AddReferencedObjects(UObject* InThis, FReferenc
 	UMovieGraphDefaultRenderer* This = CastChecked<UMovieGraphDefaultRenderer>(InThis);
 
 	// Can't be a const& due to AddStableReference API
-	for (TPair<UE::MovieGraph::DefaultRenderer::FRenderTargetInitParams, TObjectPtr<UTextureRenderTarget2D>>& KVP : This->PooledViewRenderTargets)
+	for (TPair<UE::MovieGraph::DefaultRenderer::FMovieGraphImagePreviewDataPoolParams, TObjectPtr<UTextureRenderTarget2D>>& KVP : This->PooledViewRenderTargets)
 	{
 		Collector.AddStableReference(&KVP.Value);
 	}
@@ -373,15 +373,19 @@ UE::MovieGraph::DefaultRenderer::FCameraInfo UMovieGraphDefaultRenderer::GetCame
 	return CameraInfo;
 }
 
-UTextureRenderTarget2D* UMovieGraphDefaultRenderer::GetOrCreateViewRenderTarget(const UE::MovieGraph::DefaultRenderer::FRenderTargetInitParams& InInitParams)
+UTextureRenderTarget2D* UMovieGraphDefaultRenderer::GetOrCreateViewRenderTarget(const UE::MovieGraph::DefaultRenderer::FRenderTargetInitParams& InInitParams, const FMovieGraphRenderDataIdentifier& InIdentifier)
 {
-	if (const TObjectPtr<UTextureRenderTarget2D>* ExistingViewRenderTarget = PooledViewRenderTargets.Find(InInitParams))
+	UE::MovieGraph::DefaultRenderer::FMovieGraphImagePreviewDataPoolParams CombinedParams;
+	CombinedParams.RenderInitParams = InInitParams;
+	CombinedParams.Identifier = InIdentifier;
+
+	if (const TObjectPtr<UTextureRenderTarget2D>* ExistingViewRenderTarget = PooledViewRenderTargets.Find(CombinedParams))
 	{
 		return *ExistingViewRenderTarget;
 	}
 
 	const TObjectPtr<UTextureRenderTarget2D> NewViewRenderTarget = CreateViewRenderTarget(InInitParams);
-	PooledViewRenderTargets.Emplace(InInitParams, NewViewRenderTarget);
+	PooledViewRenderTargets.Emplace(CombinedParams, NewViewRenderTarget);
 
 	return NewViewRenderTarget.Get();
 }
@@ -425,15 +429,18 @@ UE::MovieGraph::FRenderTimeStatistics* UMovieGraphDefaultRenderer::GetRenderTime
 	return &RenderTimeStatistics.FindOrAdd(InFrameNumber);
 }
 
-UTexture* UMovieGraphDefaultRenderer::GetPreviewTexture() const
-{
+TArray<FMovieGraphImagePreviewData> UMovieGraphDefaultRenderer::GetPreviewData() const
+{ 
+	TArray<FMovieGraphImagePreviewData> Results;
+
 	TArray<TObjectPtr<UTextureRenderTarget2D>> PooledTargets;
 	PooledViewRenderTargets.GenerateValueArray(PooledTargets);
 
-	if (PooledTargets.Num() > 0)
+	for (const TObjectPtr<UTextureRenderTarget2D>& Target : PooledTargets)
 	{
-		return PooledTargets[0];
+		FMovieGraphImagePreviewData& Data = Results.AddDefaulted_GetRef();
+		Data.Texture = Target.Get();
 	}
 
-	return nullptr;
+	return Results;
 }
diff --git a/Engine/Plugins/MovieScene/MovieRenderPipeline/Source/MovieRenderPipelineCore/Private/Graph/Nodes/MovieGraphWidgetRendererBaseNode.cpp b/Engine/Plugins/MovieScene/MovieRenderPipeline/Source/MovieRenderPipelineCore/Private/Graph/Nodes/MovieGraphWidgetRendererBaseNode.cpp
--- a/Engine/Plugins/MovieScene/MovieRenderPipeline/Source/MovieRenderPipelineCore/Private/Graph/Nodes/MovieGraphWidgetRendererBaseNode.cpp
+++ b/Engine/Plugins/MovieScene/MovieRenderPipeline/Source/MovieRenderPipelineCore/Private/Graph/Nodes/MovieGraphWidgetRendererBaseNode.cpp
@@ -111,7 +111,7 @@ void UMovieGraphWidgetRendererBaseNode::FMovieGraphWidgetPass::Render(const FMov
 	RenderTargetInitParams.Size = OutputResolution;
 	RenderTargetInitParams.TargetGamma = FOpenColorIORendering::DefaultDisplayGamma;
 	RenderTargetInitParams.PixelFormat = PF_B8G8R8A8;
-	UTextureRenderTarget2D* RenderTarget = Renderer->GetOrCreateViewRenderTarget(RenderTargetInitParams);
+	UTextureRenderTarget2D* RenderTarget = Renderer->GetOrCreateViewRenderTarget(RenderTargetInitParams, RenderDataIdentifier);
 
 	if (InFrameTraversalContext.Time.bIsFirstTemporalSampleForFrame)
 	{
diff --git a/Engine/Plugins/MovieScene/MovieRenderPipeline/Source/MovieRenderPipelineCore/Private/Graph/Renderers/MovieGraphDeferredPass.cpp b/Engine/Plugins/MovieScene/MovieRenderPipeline/Source/MovieRenderPipelineCore/Private/Graph/Renderers/MovieGraphDeferredPass.cpp
--- a/Engine/Plugins/MovieScene/MovieRenderPipeline/Source/MovieRenderPipelineCore/Private/Graph/Renderers/MovieGraphDeferredPass.cpp
+++ b/Engine/Plugins/MovieScene/MovieRenderPipeline/Source/MovieRenderPipelineCore/Private/Graph/Renderers/MovieGraphDeferredPass.cpp
@@ -137,7 +137,7 @@ void FMovieGraphDeferredPass::Render(const FMovieGraphTraversalContext& InFrameT
 		RenderTargetInitParams.PixelFormat = EPixelFormat::PF_FloatRGBA;
 	}
 
-	UTextureRenderTarget2D* RenderTarget = GraphRenderer->GetOrCreateViewRenderTarget(RenderTargetInitParams);
+	UTextureRenderTarget2D* RenderTarget = GraphRenderer->GetOrCreateViewRenderTarget(RenderTargetInitParams, RenderDataIdentifier);
 	FRenderTarget* RenderTargetResource = RenderTarget->GameThread_GetRenderTargetResource();
 	check(RenderTargetResource);
 
diff --git a/Engine/Plugins/MovieScene/MovieRenderPipeline/Source/MovieRenderPipelineCore/Public/Graph/MovieGraphDataTypes.h b/Engine/Plugins/MovieScene/MovieRenderPipeline/Source/MovieRenderPipelineCore/Public/Graph/MovieGraphDataTypes.h
--- a/Engine/Plugins/MovieScene/MovieRenderPipeline/Source/MovieRenderPipelineCore/Public/Graph/MovieGraphDataTypes.h
+++ b/Engine/Plugins/MovieScene/MovieRenderPipeline/Source/MovieRenderPipelineCore/Public/Graph/MovieGraphDataTypes.h
@@ -30,6 +30,24 @@ namespace UE::MovieGraph
 	struct FRenderTimeStatistics;
 }
 
+USTRUCT(BlueprintType)
+struct FMovieGraphImagePreviewData
+{
+	GENERATED_BODY()
+
+	FMovieGraphImagePreviewData()
+	: Texture(nullptr)
+	{}
+
+	/** The texture this preview image was rendered to. */
+	UPROPERTY(BlueprintReadOnly, Category = "Movie Graph")
+	class UTexture* Texture;
+	
+	/** The identifier for the image, containing the render layer name, renderer, etc. */
+	UPROPERTY(BlueprintReadOnly, Category = "Movie Graph")
+	FMovieGraphRenderDataIdentifier Identifier;
+};
+
 USTRUCT(BlueprintType)
 struct MOVIERENDERPIPELINECORE_API FMovieGraphInitConfig
 {
@@ -190,10 +208,10 @@ class MOVIERENDERPIPELINECORE_API UMovieGraphRendererBase : public UObject
 {
 	GENERATED_BODY()
 public:
-	/** Get an texture that can be used to see what the Renderer is producing for the in-game overlay. */
-	UFUNCTION(BlueprintPure, Category = "Movie Graph")
-	virtual UTexture* GetPreviewTexture() const { return nullptr; }
-
+	/** Get an array of image previews that are valid this frame. */
+	UFUNCTION(BlueprintCallable, Category = "Movie Graph")
+	virtual TArray<FMovieGraphImagePreviewData> GetPreviewData() const { return TArray<FMovieGraphImagePreviewData>(); }
+		
 	virtual void Render(const FMovieGraphTimeStepData& InTimeData) {}
 	virtual void SetupRenderingPipelineForShot(UMoviePipelineExecutorShot* InShot) {}
 	virtual void TeardownRenderingPipelineForShot(UMoviePipelineExecutorShot* InShot) {}
diff --git a/Engine/Plugins/MovieScene/MovieRenderPipeline/Source/MovieRenderPipelineCore/Public/Graph/MovieGraphDefaultRenderer.h b/Engine/Plugins/MovieScene/MovieRenderPipeline/Source/MovieRenderPipelineCore/Public/Graph/MovieGraphDefaultRenderer.h
--- a/Engine/Plugins/MovieScene/MovieRenderPipeline/Source/MovieRenderPipelineCore/Public/Graph/MovieGraphDefaultRenderer.h
+++ b/Engine/Plugins/MovieScene/MovieRenderPipeline/Source/MovieRenderPipelineCore/Public/Graph/MovieGraphDefaultRenderer.h
@@ -108,6 +108,29 @@ namespace UE::MovieGraph::DefaultRenderer
 		}
 	};
 
+	struct FMovieGraphImagePreviewDataPoolParams
+	{
+		/** The size of the render target, and the depth requested. */
+		FRenderTargetInitParams RenderInitParams;
+		/** Identifier for which render resource this is used by. */
+		FMovieGraphRenderDataIdentifier Identifier;
+
+		bool operator == (const FMovieGraphImagePreviewDataPoolParams& InRHS) const
+		{
+			return RenderInitParams == InRHS.RenderInitParams && Identifier == InRHS.Identifier;
+		}
+
+		bool operator != (const FMovieGraphImagePreviewDataPoolParams& InRHS) const
+		{
+			return !(*this == InRHS);
+		}
+
+		friend uint32 GetTypeHash(FMovieGraphImagePreviewDataPoolParams Params)
+		{
+			return HashCombineFast(GetTypeHash(Params.RenderInitParams), GetTypeHash(Params.Identifier));
+		}
+	};
+
 	struct FSurfaceAccumulatorPool : public TSharedFromThis<FSurfaceAccumulatorPool>
 	{
 		struct FInstance
@@ -233,7 +256,7 @@ class MOVIERENDERPIPELINECORE_API UMovieGraphDefaultRenderer : public UMovieGrap
 
 public:
 	// UMovieGraphRendererBase Interface
-	virtual UTexture* GetPreviewTexture() const override;
+	virtual TArray<FMovieGraphImagePreviewData> GetPreviewData() const override;
 	virtual void Render(const FMovieGraphTimeStepData& InTimeData) override;
 	virtual void SetupRenderingPipelineForShot(UMoviePipelineExecutorShot* InShot) override;
 	virtual void TeardownRenderingPipelineForShot(UMoviePipelineExecutorShot* InShot) override;
@@ -250,7 +273,7 @@ public:
 	void SetHasRenderedFirstViewThisFrame(bool bInValue) { bHasRenderedFirstViewThisFrame = bInValue; }
 	bool GetHasRenderedFirstViewThisFrame() const { return bHasRenderedFirstViewThisFrame; }
 public:
-	UTextureRenderTarget2D* GetOrCreateViewRenderTarget(const UE::MovieGraph::DefaultRenderer::FRenderTargetInitParams& InInitParams);
+	UTextureRenderTarget2D* GetOrCreateViewRenderTarget(const UE::MovieGraph::DefaultRenderer::FRenderTargetInitParams& InInitParams, const FMovieGraphRenderDataIdentifier& InIdentifier);
 	FMoviePipelineSurfaceQueuePtr GetOrCreateSurfaceQueue(const UE::MovieGraph::DefaultRenderer::FRenderTargetInitParams& InInitParams);
 	
 	template <typename AccumulatorType>
@@ -291,7 +314,7 @@ protected:
 	* because we queue a copy onto a FMoviePipelineSurface so the only reason it is needed
 	* is for display in the UI.
 	*/
-	TMap<UE::MovieGraph::DefaultRenderer::FRenderTargetInitParams, TObjectPtr<UTextureRenderTarget2D>> PooledViewRenderTargets;
+	TMap<UE::MovieGraph::DefaultRenderer::FMovieGraphImagePreviewDataPoolParams, TObjectPtr<UTextureRenderTarget2D>> PooledViewRenderTargets;
 
 	/*
 	* This is a list of Surfaces that we can copy our render target to immediately after
