MoviePipeline: Fixed graph comments not being deletable.

#rnx

diff --git a/Engine/Plugins/MovieScene/MovieRenderPipeline/Source/MovieRenderPipelineEditor/Private/Widgets/Graph/SMovieGraphConfigPanel.cpp b/Engine/Plugins/MovieScene/MovieRenderPipeline/Source/MovieRenderPipelineEditor/Private/Widgets/Graph/SMovieGraphConfigPanel.cpp
--- a/Engine/Plugins/MovieScene/MovieRenderPipeline/Source/MovieRenderPipelineEditor/Private/Widgets/Graph/SMovieGraphConfigPanel.cpp
+++ b/Engine/Plugins/MovieScene/MovieRenderPipeline/Source/MovieRenderPipelineEditor/Private/Widgets/Graph/SMovieGraphConfigPanel.cpp
@@ -273,24 +273,38 @@ bool SMoviePipelineGraphPanel::CanDeleteSelectedNodes() const
 void SMoviePipelineGraphPanel::DeleteSelectedNodes()
 {
 	TArray<UMovieGraphNode*> NodesToDelete;
+	TArray<UEdGraphNode*> EdNodesToDelete;
+	
 	for (UObject* Object : GraphEditorWidget->GetSelectedNodes())
 	{
-		UMoviePipelineEdGraphNodeBase* GraphNode = Cast<UMoviePipelineEdGraphNodeBase>(Object);
+		const UMoviePipelineEdGraphNodeBase* GraphNode = Cast<UMoviePipelineEdGraphNodeBase>(Object);
 		if (GraphNode && GraphNode->CanUserDeleteNode())
 		{
 			NodesToDelete.Add(GraphNode->GetRuntimeNode());
+			continue;
+		}
+
+		UEdGraphNode* EdGraphNode = Cast<UEdGraphNode>(Object);
+		if (EdGraphNode && EdGraphNode->CanUserDeleteNode())
+		{
+			EdNodesToDelete.Add(EdGraphNode);
 		}
 	}
 
-	if (NodesToDelete.IsEmpty())
+	// Remove all runtime graph nodes
+	const UMoviePipelineEdGraph* Graph = Cast<UMoviePipelineEdGraph>(GraphEditorWidget->GetCurrentGraph());
+	if (Graph && !NodesToDelete.IsEmpty())
 	{
-		return;
+		Graph->GetPipelineGraph()->RemoveNodes(NodesToDelete);
 	}
 
-	UMoviePipelineEdGraph* Graph = Cast<UMoviePipelineEdGraph>(GraphEditorWidget->GetCurrentGraph());
-	if (Graph)
+	// Remove all editor nodes (nodes not backed by a runtime node, like comments)
+	if (Graph && !EdNodesToDelete.IsEmpty())
 	{
-		Graph->GetPipelineGraph()->RemoveNodes(NodesToDelete);
+		for (UEdGraphNode* EdGraphNode : EdNodesToDelete)
+		{
+			GraphEditorWidget->GetCurrentGraph()->RemoveNode(EdGraphNode);
+		}
 	}
 
 	ClearGraphSelection();
