MoviePipeline: Fixed an editor lock-up that could occur when renaming members in the graph (inputs, outputs, variables).

#rnx

diff --git a/Engine/Plugins/MovieScene/MovieRenderPipeline/Source/MovieRenderPipelineEditor/Private/Graph/SMovieGraphMembersTabContent.cpp b/Engine/Plugins/MovieScene/MovieRenderPipeline/Source/MovieRenderPipelineEditor/Private/Graph/SMovieGraphMembersTabContent.cpp
--- a/Engine/Plugins/MovieScene/MovieRenderPipeline/Source/MovieRenderPipelineEditor/Private/Graph/SMovieGraphMembersTabContent.cpp
+++ b/Engine/Plugins/MovieScene/MovieRenderPipeline/Source/MovieRenderPipelineEditor/Private/Graph/SMovieGraphMembersTabContent.cpp
@@ -153,7 +153,7 @@ void SMovieGraphMembersTabContent::ClearSelection() const
 	}
 }
 
-void SMovieGraphMembersTabContent::DeleteSelectedMembers() const
+void SMovieGraphMembersTabContent::DeleteSelectedMembers()
 {
 	if (!ActionMenu.IsValid() || !CurrentGraph)
 	{
@@ -166,6 +166,8 @@ void SMovieGraphMembersTabContent::DeleteSelectedMembers() const
 	{
 		if (UMovieGraphMember* GraphMember = UE::MovieGraph::Private::GetMemberFromAction(SelectedAction.Get()))
 		{
+			MemberChangedHandles.Remove(GraphMember);
+			
 			CurrentGraph->DeleteMember(GraphMember);
 		}
 	}
@@ -212,7 +214,7 @@ void SMovieGraphMembersTabContent::CollectAllActions(FGraphActionListBuilderBase
 	FGraphActionMenuBuilder ActionMenuBuilder;
 
 	// Creates a new action in the action menu under a specific section w/ the provided action target
-	auto AddToActionMenu = [&ActionMenuBuilder](UMovieGraphMember* ActionTarget, const EActionSection Section, const FText& Category) -> void
+	auto AddToActionMenu = [&ActionMenuBuilder, this](UMovieGraphMember* ActionTarget, const EActionSection Section, const FText& Category) -> void
 	{
 		const FText MemberActionDesc = FText::FromString(ActionTarget->GetMemberName());
 		const FText MemberActionTooltip;
@@ -221,6 +223,31 @@ void SMovieGraphMembersTabContent::CollectAllActions(FGraphActionListBuilderBase
 		const TSharedPtr<FMovieGraphSchemaAction> MemberAction(new FMovieGraphSchemaAction(Category, MemberActionDesc, MemberActionTooltip, 0, MemberActionKeywords, MemberActionSectionID));
 		MemberAction->ActionTarget = ActionTarget;
 		ActionMenuBuilder.AddAction(MemberAction);
+
+		// Update actions when a member is updated (renamed, etc). Only subscribe to the delegate once.
+		if (!MemberChangedHandles.Contains(ActionTarget))
+		{
+			FDelegateHandle MemberChangedDelegate;
+			if (UMovieGraphInput* InputMember = Cast<UMovieGraphInput>(ActionTarget))
+			{
+				MemberChangedDelegate = InputMember->OnMovieGraphInputChangedDelegate.AddSP(this, &SMovieGraphMembersTabContent::RefreshMemberActions);
+			}
+			else if (UMovieGraphOutput* OutputMember = Cast<UMovieGraphOutput>(ActionTarget))
+			{
+				MemberChangedDelegate = OutputMember->OnMovieGraphOutputChangedDelegate.AddSP(this, &SMovieGraphMembersTabContent::RefreshMemberActions);
+			}
+			else if (UMovieGraphVariable* VariableMember = Cast<UMovieGraphVariable>(ActionTarget))
+			{
+				MemberChangedDelegate = VariableMember->OnMovieGraphVariableChangedDelegate.AddSP(this, &SMovieGraphMembersTabContent::RefreshMemberActions);
+			}
+			else
+			{
+				checkf(false, TEXT("Found an unsupported member type when adding it to the action menu."));
+				return;
+			}
+			
+			MemberChangedHandles.Add(ActionTarget, MemberChangedDelegate);
+		}
 	};
 
 	for (UMovieGraphInput* Input : CurrentGraph->GetInputs())
@@ -228,9 +255,6 @@ void SMovieGraphMembersTabContent::CollectAllActions(FGraphActionListBuilderBase
 		if (Input && Input->IsDeletable())
 		{
 			AddToActionMenu(Input, EActionSection::Inputs, EmptyCategory);
-            
-            // Update actions when an input is updated (renamed, etc)
-            Input->OnMovieGraphInputChangedDelegate.AddSP(this, &SMovieGraphMembersTabContent::RefreshMemberActions);
 		}
 	}
 
@@ -239,9 +263,6 @@ void SMovieGraphMembersTabContent::CollectAllActions(FGraphActionListBuilderBase
 		if (Output && Output->IsDeletable())
 		{
 			AddToActionMenu(Output, EActionSection::Outputs, EmptyCategory);
-
-			// Update actions when an output is updated (renamed, etc)
-			Output->OnMovieGraphOutputChangedDelegate.AddSP(this, &SMovieGraphMembersTabContent::RefreshMemberActions);
 		}
 	}
 
@@ -254,9 +275,6 @@ void SMovieGraphMembersTabContent::CollectAllActions(FGraphActionListBuilderBase
 		if (Variable && !Variable->IsGlobal())
 		{
 			AddToActionMenu(Variable, EActionSection::Variables, UserVariablesCategory);
-
-			// Update actions when a variable is updated (renamed, etc)
-			Variable->OnMovieGraphVariableChangedDelegate.AddSP(this, &SMovieGraphMembersTabContent::RefreshMemberActions);
 		}
 	}
 
diff --git a/Engine/Plugins/MovieScene/MovieRenderPipeline/Source/MovieRenderPipelineEditor/Private/Graph/SMovieGraphMembersTabContent.h b/Engine/Plugins/MovieScene/MovieRenderPipeline/Source/MovieRenderPipelineEditor/Private/Graph/SMovieGraphMembersTabContent.h
--- a/Engine/Plugins/MovieScene/MovieRenderPipeline/Source/MovieRenderPipelineEditor/Private/Graph/SMovieGraphMembersTabContent.h
+++ b/Engine/Plugins/MovieScene/MovieRenderPipeline/Source/MovieRenderPipelineEditor/Private/Graph/SMovieGraphMembersTabContent.h
@@ -50,7 +50,7 @@ public:
 	void ClearSelection() const;
 
 	/** Deletes the member(s) which are currently selected from the graph and the UI. */
-	void DeleteSelectedMembers() const;
+	void DeleteSelectedMembers();
 
 	/** Determines if all selected member(s) can be deleted. */
 	bool CanDeleteSelectedMembers() const;
@@ -100,6 +100,9 @@ private:
 	
 	/** Delegate to call when an action is selected */
 	FOnActionSelected OnActionSelected;
+
+	/** Handles to delegates handling member changes. */
+	TMap<TWeakObjectPtr<UMovieGraphMember>, FDelegateHandle> MemberChangedHandles;
 };
 
 /* Drag-and-drop action which handles variable members. */
