MoviePipeline: Updated the condition group queries to display shorter names in the UI.

#rnx

diff --git a/Engine/Plugins/MovieScene/MovieRenderPipeline/Source/MovieRenderPipelineCore/Private/Graph/MoviePipelineRenderLayerSubsystem.cpp b/Engine/Plugins/MovieScene/MovieRenderPipeline/Source/MovieRenderPipelineCore/Private/Graph/MoviePipelineRenderLayerSubsystem.cpp
--- a/Engine/Plugins/MovieScene/MovieRenderPipeline/Source/MovieRenderPipelineCore/Private/Graph/MoviePipelineRenderLayerSubsystem.cpp
+++ b/Engine/Plugins/MovieScene/MovieRenderPipeline/Source/MovieRenderPipelineCore/Private/Graph/MoviePipelineRenderLayerSubsystem.cpp
@@ -427,6 +427,12 @@ const FSlateIcon& UMovieGraphConditionGroupQueryBase::GetIcon() const
 	return EmptyIcon;
 }
 
+const FText& UMovieGraphConditionGroupQueryBase::GetDisplayName() const
+{
+	static const FText DisplayName = LOCTEXT("ConditionGroupQueryDisplayName", "Query Base");
+	return DisplayName;
+}
+
 #if WITH_EDITOR
 TArray<TSharedRef<SWidget>> UMovieGraphConditionGroupQueryBase::GetWidgets()
 {
@@ -486,6 +492,12 @@ const FSlateIcon& UMovieGraphConditionGroupQuery_Actor::GetIcon() const
 	return ActorIcon;
 }
 
+const FText& UMovieGraphConditionGroupQuery_Actor::GetDisplayName() const
+{
+	static const FText DisplayName = LOCTEXT("ConditionGroupQueryDisplayName_Actor", "Actor");
+	return DisplayName;
+}
+
 #if WITH_EDITOR
 TArray<TSharedRef<SWidget>> UMovieGraphConditionGroupQuery_Actor::GetWidgets()
 {
@@ -635,6 +647,12 @@ const FSlateIcon& UMovieGraphConditionGroupQuery_ActorTagName::GetIcon() const
 	return ActorTagIcon;
 }
 
+const FText& UMovieGraphConditionGroupQuery_ActorTagName::GetDisplayName() const
+{
+	static const FText DisplayName = LOCTEXT("ConditionGroupQueryDisplayName_ActorTagName", "Actor Tag Name");
+	return DisplayName;
+}
+
 #if WITH_EDITOR
 TArray<TSharedRef<SWidget>> UMovieGraphConditionGroupQuery_ActorTagName::GetWidgets()
 {
@@ -691,6 +709,12 @@ const FSlateIcon& UMovieGraphConditionGroupQuery_ActorName::GetIcon() const
 	return ActorTagIcon;
 }
 
+const FText& UMovieGraphConditionGroupQuery_ActorName::GetDisplayName() const
+{
+	static const FText DisplayName = LOCTEXT("ConditionGroupQueryDisplayName_ActorName", "Actor Name");
+	return DisplayName;
+}
+
 #if WITH_EDITOR
 TArray<TSharedRef<SWidget>> UMovieGraphConditionGroupQuery_ActorName::GetWidgets()
 {
@@ -737,6 +761,12 @@ const FSlateIcon& UMovieGraphConditionGroupQuery_ActorType::GetIcon() const
 	return ActorTagIcon;
 }
 
+const FText& UMovieGraphConditionGroupQuery_ActorType::GetDisplayName() const
+{
+	static const FText DisplayName = LOCTEXT("ConditionGroupQueryDisplayName_ActorType", "Actor Type");
+	return DisplayName;
+}
+
 #if WITH_EDITOR
 TArray<TSharedRef<SWidget>> UMovieGraphConditionGroupQuery_ActorType::GetWidgets()
 {
@@ -867,6 +897,12 @@ const FSlateIcon& UMovieGraphConditionGroupQuery_ComponentTagName::GetIcon() con
 	return ActorTagIcon;
 }
 
+const FText& UMovieGraphConditionGroupQuery_ComponentTagName::GetDisplayName() const
+{
+	static const FText DisplayName = LOCTEXT("ConditionGroupQueryDisplayName_ComponentTagName", "Component Tag Name");
+	return DisplayName;
+}
+
 #if WITH_EDITOR
 TArray<TSharedRef<SWidget>> UMovieGraphConditionGroupQuery_ComponentTagName::GetWidgets()
 {
@@ -916,6 +952,12 @@ const FSlateIcon& UMovieGraphConditionGroupQuery_ComponentType::GetIcon() const
 	return ActorTagIcon;
 }
 
+const FText& UMovieGraphConditionGroupQuery_ComponentType::GetDisplayName() const
+{
+	static const FText DisplayName = LOCTEXT("ConditionGroupQueryDisplayName_ComponentType", "Component Type");
+	return DisplayName;
+}
+
 #if WITH_EDITOR
 TArray<TSharedRef<SWidget>> UMovieGraphConditionGroupQuery_ComponentType::GetWidgets()
 {
diff --git a/Engine/Plugins/MovieScene/MovieRenderPipeline/Source/MovieRenderPipelineCore/Public/Graph/MoviePipelineRenderLayerSubsystem.h b/Engine/Plugins/MovieScene/MovieRenderPipeline/Source/MovieRenderPipelineCore/Public/Graph/MoviePipelineRenderLayerSubsystem.h
--- a/Engine/Plugins/MovieScene/MovieRenderPipeline/Source/MovieRenderPipelineCore/Public/Graph/MoviePipelineRenderLayerSubsystem.h
+++ b/Engine/Plugins/MovieScene/MovieRenderPipeline/Source/MovieRenderPipelineCore/Public/Graph/MoviePipelineRenderLayerSubsystem.h
@@ -110,6 +110,9 @@ public:
 	/** Gets the icon that represents this query class in the UI. */
 	virtual const FSlateIcon& GetIcon() const;
 
+	/** Gets the display name for this query class as shown in the UI. */
+	virtual const FText& GetDisplayName() const;
+
 #if WITH_EDITOR
 	/**
 	 * Gets the widgets that should be displayed for this query. If no custom widgets are specified (returning an empty array), the default
@@ -151,7 +154,7 @@ private:
 };
 
 /** Query type which filters actors via an explicit actor list. */
-UCLASS(BlueprintType, meta = (DisplayName = "Movie Graph Actor Query"))
+UCLASS(BlueprintType)
 class MOVIERENDERPIPELINECORE_API UMovieGraphConditionGroupQuery_Actor final : public UMovieGraphConditionGroupQueryBase
 {
 	GENERATED_BODY()
@@ -159,6 +162,7 @@ class MOVIERENDERPIPELINECORE_API UMovieGraphConditionGroupQuery_Actor final : p
 public:
 	virtual void Evaluate(const TArray<AActor*>& InActorsToQuery, TSet<AActor*>& OutMatchingActors) const override;
 	virtual const FSlateIcon& GetIcon() const override;
+	virtual const FText& GetDisplayName() const override;
 
 #if WITH_EDITOR
 	virtual TArray<TSharedRef<SWidget>> GetWidgets() override;
@@ -187,7 +191,7 @@ private:
 };
 
 /** Query type which filters actors via tags on actors. */
-UCLASS(BlueprintType, meta = (DisplayName = "Movie Graph Actor Tag Name Query"))
+UCLASS(BlueprintType)
 class MOVIERENDERPIPELINECORE_API UMovieGraphConditionGroupQuery_ActorTagName final : public UMovieGraphConditionGroupQueryBase
 {
 	GENERATED_BODY()
@@ -195,6 +199,7 @@ class MOVIERENDERPIPELINECORE_API UMovieGraphConditionGroupQuery_ActorTagName fi
 public:
 	virtual void Evaluate(const TArray<AActor*>& InActorsToQuery, TSet<AActor*>& OutMatchingActors) const override;
 	virtual const FSlateIcon& GetIcon() const override;
+	virtual const FText& GetDisplayName() const override;
 
 #if WITH_EDITOR
 	virtual TArray<TSharedRef<SWidget>> GetWidgets() override;
@@ -219,7 +224,7 @@ public:
 };
 
 /** Query type which filters actors via their name (label). */
-UCLASS(BlueprintType, meta = (DisplayName = "Movie Graph Actor Name Query"))
+UCLASS(BlueprintType)
 class MOVIERENDERPIPELINECORE_API UMovieGraphConditionGroupQuery_ActorName final : public UMovieGraphConditionGroupQueryBase
 {
 	GENERATED_BODY()
@@ -227,6 +232,7 @@ class MOVIERENDERPIPELINECORE_API UMovieGraphConditionGroupQuery_ActorName final
 public:
 	virtual void Evaluate(const TArray<AActor*>& InActorsToQuery, TSet<AActor*>& OutMatchingActors) const override;
 	virtual const FSlateIcon& GetIcon() const override;
+	virtual const FText& GetDisplayName() const override;
 
 #if WITH_EDITOR
 	virtual TArray<TSharedRef<SWidget>> GetWidgets() override;
@@ -253,7 +259,7 @@ public:
 };
 
 /** Query type which filters actors by type. */
-UCLASS(BlueprintType, meta = (DisplayName = "Movie Graph Actor Type Query"))
+UCLASS(BlueprintType)
 class MOVIERENDERPIPELINECORE_API UMovieGraphConditionGroupQuery_ActorType final : public UMovieGraphConditionGroupQueryBase
 {
 	GENERATED_BODY()
@@ -261,6 +267,7 @@ class MOVIERENDERPIPELINECORE_API UMovieGraphConditionGroupQuery_ActorType final
 public:
 	virtual void Evaluate(const TArray<AActor*>& InActorsToQuery, TSet<AActor*>& OutMatchingActors) const override;
 	virtual const FSlateIcon& GetIcon() const override;
+	virtual const FText& GetDisplayName() const override;
 
 #if WITH_EDITOR
 	virtual TArray<TSharedRef<SWidget>> GetWidgets() override;
@@ -282,7 +289,7 @@ private:
 };
 
 /** Query type which filters actors by tags on their components. */
-UCLASS(BlueprintType, meta = (DisplayName = "Movie Graph Component Tag Name Query"))
+UCLASS(BlueprintType)
 class MOVIERENDERPIPELINECORE_API UMovieGraphConditionGroupQuery_ComponentTagName final : public UMovieGraphConditionGroupQueryBase
 {
 	GENERATED_BODY()
@@ -290,6 +297,7 @@ class MOVIERENDERPIPELINECORE_API UMovieGraphConditionGroupQuery_ComponentTagNam
 public:
 	virtual void Evaluate(const TArray<AActor*>& InActorsToQuery, TSet<AActor*>& OutMatchingActors) const override;
 	virtual const FSlateIcon& GetIcon() const override;
+	virtual const FText& GetDisplayName() const override;
 
 #if WITH_EDITOR
 	virtual TArray<TSharedRef<SWidget>> GetWidgets() override;
@@ -314,7 +322,7 @@ public:
 };
 
 /** Query type which filters actors via the components contained in them. */
-UCLASS(BlueprintType, meta = (DisplayName = "Movie Graph Component Type Query"))
+UCLASS(BlueprintType)
 class MOVIERENDERPIPELINECORE_API UMovieGraphConditionGroupQuery_ComponentType final : public UMovieGraphConditionGroupQueryBase
 {
 	GENERATED_BODY()
@@ -322,6 +330,7 @@ class MOVIERENDERPIPELINECORE_API UMovieGraphConditionGroupQuery_ComponentType f
 public:
 	virtual void Evaluate(const TArray<AActor*>& InActorsToQuery, TSet<AActor*>& OutMatchingActors) const override;
 	virtual const FSlateIcon& GetIcon() const override;
+	virtual const FText& GetDisplayName() const override;
 
 #if WITH_EDITOR
 	virtual TArray<TSharedRef<SWidget>> GetWidgets() override;
diff --git a/Engine/Plugins/MovieScene/MovieRenderPipeline/Source/MovieRenderPipelineEditor/Private/Customizations/Graph/MovieGraphCollectionsCustomization.cpp b/Engine/Plugins/MovieScene/MovieRenderPipeline/Source/MovieRenderPipelineEditor/Private/Customizations/Graph/MovieGraphCollectionsCustomization.cpp
--- a/Engine/Plugins/MovieScene/MovieRenderPipeline/Source/MovieRenderPipelineEditor/Private/Customizations/Graph/MovieGraphCollectionsCustomization.cpp
+++ b/Engine/Plugins/MovieScene/MovieRenderPipeline/Source/MovieRenderPipelineEditor/Private/Customizations/Graph/MovieGraphCollectionsCustomization.cpp
@@ -534,7 +534,16 @@ TSharedRef<SWidget> SMovieGraphCollectionTreeQueryTypeSelectorWidget::GetQueryTy
 	else
 	{
 		SAssignNew(QueryTypeTextBlock, STextBlock)
-		.Text(InTypeClass->GetDisplayNameText())
+		.Text_Lambda([InTypeClass]()
+		{
+			if (const UMovieGraphConditionGroupQueryBase* Query = Cast<UMovieGraphConditionGroupQueryBase>(InTypeClass->GetDefaultObject()))
+			{
+				return Query->GetDisplayName();
+			}
+
+			// Show the UClass display name as a backup
+			return InTypeClass->GetDisplayNameText();
+		})
 		.Font(FAppStyle::GetFontStyle(TEXT("PropertyWindow.NormalFont")));
 	}
 
@@ -568,7 +577,9 @@ FText SMovieGraphCollectionTreeQueryTypeSelectorWidget::GetCurrentQueryTypeDispl
 	{
 		if (const UClass* CurrentQueryType = GetCurrentQueryType())
 		{
-			return CurrentQueryType->GetDisplayNameText();
+			check(CurrentQueryType->IsChildOf(UMovieGraphConditionGroupQueryBase::StaticClass()));
+			
+			return Cast<UMovieGraphConditionGroupQueryBase>(CurrentQueryType->GetDefaultObject())->GetDisplayName();
 		}
 	}
 		
