MoviePipeline: Fixed crash which can occur when connecting to one of the dynamic pins on a Select node.

This was caused by the "bIsBuiltInPin" argument to GetInputPin(). By default this was set to "true", which meant that the dynamic pins on the Select node were never being found, leading to a crash. This argument was originally meant for the Select node to distinguish between the select options and the built-in pins (Default, Selected Option, etc) because a select option could theoretically have the same name as one of the built-ins.

Somewhere along the line this behavior broke, so the argument was upgraded to be more permissive by default so *all* pins are searched. Callers need to opt-in to more restrictive dynamic or builtin-only search behavior.

#rnx

diff --git a/Engine/Plugins/MovieScene/MovieRenderPipeline/Source/MovieRenderPipelineCore/Private/Graph/MovieGraphNode.cpp b/Engine/Plugins/MovieScene/MovieRenderPipeline/Source/MovieRenderPipelineCore/Private/Graph/MovieGraphNode.cpp
--- a/Engine/Plugins/MovieScene/MovieRenderPipeline/Source/MovieRenderPipelineCore/Private/Graph/MovieGraphNode.cpp
+++ b/Engine/Plugins/MovieScene/MovieRenderPipeline/Source/MovieRenderPipelineCore/Private/Graph/MovieGraphNode.cpp
@@ -286,11 +286,27 @@ UMovieGraphConfig* UMovieGraphNode::GetGraph() const
 	return Cast<UMovieGraphConfig>(GetOuter());
 }
 
-UMovieGraphPin* UMovieGraphNode::GetInputPin(const FName& Label, const bool bIsBuiltInPin) const
+UMovieGraphPin* UMovieGraphNode::GetInputPin(const FName& InPinLabel, const EMovieGraphPinQueryRequirement PinRequirement) const
 {
 	for (UMovieGraphPin* InputPin : InputPins)
 	{
-		if ((InputPin->Properties.Label == Label) && (InputPin->Properties.bIsBuiltIn == bIsBuiltInPin))
+		if (InputPin->Properties.Label != InPinLabel)
+		{
+			continue;
+		}
+
+		if (PinRequirement == EMovieGraphPinQueryRequirement::BuiltInOrDynamic)
+		{
+			return InputPin;
+		}
+
+		const bool bIsBuiltIn = InputPin->Properties.bIsBuiltIn;
+		if ((PinRequirement == EMovieGraphPinQueryRequirement::BuiltIn) && bIsBuiltIn)
+		{
+			return InputPin;
+		}
+		
+		if ((PinRequirement == EMovieGraphPinQueryRequirement::Dynamic) && !bIsBuiltIn)
 		{
 			return InputPin;
 		}
diff --git a/Engine/Plugins/MovieScene/MovieRenderPipeline/Source/MovieRenderPipelineCore/Private/Graph/Nodes/MovieGraphSelectNode.cpp b/Engine/Plugins/MovieScene/MovieRenderPipeline/Source/MovieRenderPipelineCore/Private/Graph/Nodes/MovieGraphSelectNode.cpp
--- a/Engine/Plugins/MovieScene/MovieRenderPipeline/Source/MovieRenderPipelineCore/Private/Graph/Nodes/MovieGraphSelectNode.cpp
+++ b/Engine/Plugins/MovieScene/MovieRenderPipeline/Source/MovieRenderPipelineCore/Private/Graph/Nodes/MovieGraphSelectNode.cpp
@@ -184,7 +184,7 @@ TArray<UMovieGraphPin*> UMovieGraphSelectNode::EvaluatePinsToFollow(FMovieGraphE
 
 	// Try getting the value from a connection first
 	bool bGotValueFromConnection = false;
-	if (const UMovieGraphPin* SelectPin = GetInputPin(UE::MovieGraph::SelectNode::SelectedOption))
+	if (const UMovieGraphPin* SelectPin = GetInputPin(UE::MovieGraph::SelectNode::SelectedOption, EMovieGraphPinQueryRequirement::BuiltIn))
 	{
 		if (const UMovieGraphPin* OtherPin = SelectPin->GetFirstConnectedPin())
 		{
@@ -330,7 +330,7 @@ TArray<UMovieGraphPin*> UMovieGraphSelectNode::EvaluatePinsToFollow(FMovieGraphE
 	// Follow the Default branch if no pins match
 	if (PinsToFollow.IsEmpty())
 	{
-		PinsToFollow.Add(GetInputPin(UE::MovieGraph::SelectNode::DefaultBranch));
+		PinsToFollow.Add(GetInputPin(UE::MovieGraph::SelectNode::DefaultBranch, EMovieGraphPinQueryRequirement::BuiltIn));
 	}	
 
 	return PinsToFollow;
diff --git a/Engine/Plugins/MovieScene/MovieRenderPipeline/Source/MovieRenderPipelineCore/Public/Graph/MovieGraphNode.h b/Engine/Plugins/MovieScene/MovieRenderPipeline/Source/MovieRenderPipelineCore/Public/Graph/MovieGraphNode.h
--- a/Engine/Plugins/MovieScene/MovieRenderPipeline/Source/MovieRenderPipelineCore/Public/Graph/MovieGraphNode.h
+++ b/Engine/Plugins/MovieScene/MovieRenderPipeline/Source/MovieRenderPipelineCore/Public/Graph/MovieGraphNode.h
@@ -223,8 +223,8 @@ public:
 	 * Gets the input pin with the specified name, or nullptr if one could not be found. Most pins on a node are
 	 * "built-in", meaning they ship with the node. Dynamic pins (pins which are not built-in) can potentially have the
 	 * same name as a built-in (eg, the option pins on the Select node). To disambiguate between built-in and dynamic
-	 * pins, specify bIsBuiltInPin = false if trying to fetch a pin that is not built-in. */
-	UMovieGraphPin* GetInputPin(const FName& InPinLabel, const bool bIsBuiltInPin = true) const;
+	 * pins, PinRequirement can be specified. */
+	UMovieGraphPin* GetInputPin(const FName& InPinLabel, const EMovieGraphPinQueryRequirement PinRequirement = EMovieGraphPinQueryRequirement::BuiltInOrDynamic) const;
 
 	/** Gets the output pin with the specified name, or nullptr if one could not be found. */
 	UMovieGraphPin* GetOutputPin(const FName& InPinLabel) const;
diff --git a/Engine/Plugins/MovieScene/MovieRenderPipeline/Source/MovieRenderPipelineCore/Public/Graph/MovieGraphPin.h b/Engine/Plugins/MovieScene/MovieRenderPipeline/Source/MovieRenderPipelineCore/Public/Graph/MovieGraphPin.h
--- a/Engine/Plugins/MovieScene/MovieRenderPipeline/Source/MovieRenderPipelineCore/Public/Graph/MovieGraphPin.h
+++ b/Engine/Plugins/MovieScene/MovieRenderPipeline/Source/MovieRenderPipelineCore/Public/Graph/MovieGraphPin.h
@@ -68,6 +68,20 @@ struct FMovieGraphPinProperties
 	}
 };
 
+/** Specifies a restriction on pin properties when searching for a pin on a node. */
+UENUM(BlueprintType)
+enum class EMovieGraphPinQueryRequirement : uint8
+{
+	/** The pin must be built-in, meaning that it is always present on the node. */
+	BuiltIn,
+
+	/** The pin must be dynamic, meaning that it may not always exist on the node. These are typically user-created pins. */
+	Dynamic,
+
+	/** The pin can be either built-in or dynamic. */
+	BuiltInOrDynamic
+};
+
 
 UCLASS()
 class MOVIERENDERPIPELINECORE_API UMovieGraphPin : public UObject
