MovieGraph: Support tokens in OpenColorIO context.

Also regrouped FMovieGraphFilenameResolveParams creation into a single convenience function.

diff --git a/Engine/Plugins/MovieScene/MovieRenderPipeline/Source/MovieRenderPipelineRenderPasses/Private/MovieGraphImageSequenceOutputNode.cpp b/Engine/Plugins/MovieScene/MovieRenderPipeline/Source/MovieRenderPipelineRenderPasses/Private/MovieGraphImageSequenceOutputNode.cpp
--- a/Engine/Plugins/MovieScene/MovieRenderPipeline/Source/MovieRenderPipelineRenderPasses/Private/MovieGraphImageSequenceOutputNode.cpp
+++ b/Engine/Plugins/MovieScene/MovieRenderPipeline/Source/MovieRenderPipelineRenderPasses/Private/MovieGraphImageSequenceOutputNode.cpp
@@ -22,9 +22,55 @@
 #include "OpenColorIOConfiguration.h"
 #include "OpenColorIOColorTransform.h"
 #include "OpenColorIOWrapper.h"
+#endif
 
 namespace UE::MovieGraph::Private
 {
+	/** Convenience function to make filename resolve parameters for output files, or OCIO contexts. */
+	FMovieGraphFilenameResolveParams MakeResolveParams(
+		const FMovieGraphRenderDataIdentifier& InRenderId,
+		const UMovieGraphPipeline* InPipeline,
+		const TObjectPtr<UMovieGraphEvaluatedConfig>& InEvaluatedConfig,
+		const FMovieGraphTraversalContext& InTraversalContext,
+		const TMap<FString, FString>& InAdditionalFormatArgs = {})
+	{
+		const TObjectPtr<UMoviePipelineExecutorShot>& Shot = InPipeline->GetActiveShotList()[InTraversalContext.ShotIndex];
+
+		FMovieGraphFilenameResolveParams Params = FMovieGraphFilenameResolveParams();
+		Params.RenderDataIdentifier = InRenderId;
+		//Params.RootFrameNumber = InTraversalContext.Time.RootFrameNumber;
+		//Params.ShotFrameNumber = InTraversalContext.Time.ShotFrameNumber;
+		Params.RootFrameNumberRel = InTraversalContext.Time.OutputFrameNumber;
+		//Params.ShotFrameNumberRel = InTraversalContext.Time.ShotFrameNumberRel
+		//Params.FileMetadata = ToDo: Track File Metadata
+		const UMovieGraphOutputSettingNode* OutputSettingNode = InEvaluatedConfig->GetSettingForBranch<UMovieGraphOutputSettingNode>(InRenderId.RootBranchName);
+		if (IsValid(OutputSettingNode))
+		{
+			Params.ZeroPadFrameNumberCount = OutputSettingNode->ZeroPadFrameNumbers;
+			Params.FrameNumberOffset = OutputSettingNode->FrameNumberOffset;
+		}
+		Params.EvaluatedConfig = InEvaluatedConfig;
+		Params.Version = Shot->ShotInfo.VersionNumber;
+
+		// If time dilation is in effect, RootFrameNumber and ShotFrameNumber will contain duplicates and the files will overwrite each other, 
+		// so we force them into relative mode and then warn users we did that (as their numbers will jump from say 1001 -> 0000).
+		bool bForceRelativeFrameNumbers = true; // TODO: Use relative frame numbers until we track Root vs. Shot frame numbers. (Previously false);
+		//if (FileNameFormatString.Contains(TEXT("{frame")) InTraversalContext.Time.IsTimeDilated() && !FileNameFormatString.Contains(TEXT("_rel}")))
+		//{
+		//	UE_LOG(LogMovieRenderPipeline, Warning, TEXT("Time Dilation was used but output format does not use relative time, forcing relative numbers. Change {frame_number} to {frame_number_rel} (or shot version) to remove this message."));
+		//	bForceRelativeFrameNumbers = true;
+		//}
+		Params.bForceRelativeFrameNumbers = bForceRelativeFrameNumbers;
+		Params.bEnsureAbsolutePath = true;
+		Params.FileNameFormatOverrides = InAdditionalFormatArgs;
+		Params.InitializationTime = InPipeline->GetInitializationTime();
+		Params.Shot = Shot;
+		Params.Job = InPipeline->GetCurrentJob();
+
+		return Params;
+	}
+	
+#if WITH_EDITOR
 	struct FOpenColorIOPixelPreProcessor
 	{
 		FOpenColorIOPixelPreProcessor(FOpenColorIOWrapperProcessor&& InProcessor)
@@ -40,6 +86,42 @@ namespace UE::MovieGraph::Private
 		FOpenColorIOWrapperProcessor Processor;
 	};
 
+	/**
+	 * Convenience function to resolve an OpenColorIO context with supported tokens. Editor-only.
+	 *
+	 * @return The resolved key/value context.
+	*/
+	TMap<FString, FString> ResolveOpenColorIOContext(
+		const TMap<FString, FString>& InContext,
+		const FMovieGraphRenderDataIdentifier& InRenderId,
+		const UMovieGraphPipeline* InPipeline,
+		TObjectPtr<UMovieGraphEvaluatedConfig> InEvaluatedConfig,
+		const FMovieGraphTraversalContext& InTraversalContext
+	)
+	{
+		TMap<FString, FString> OutContext;
+		OutContext.Reserve(InContext.Num());
+
+		FMovieGraphFilenameResolveParams Params = MakeResolveParams(InRenderId, InPipeline, InEvaluatedConfig, InTraversalContext);
+
+		for (const TPair<FString, FString>& Pair : InContext)
+		{
+			FMovieGraphResolveArgs FormatArgs;
+			UMovieGraphBlueprintLibrary::ResolveFilenameFormatArguments(Pair.Value, Params, FormatArgs);
+
+			FStringFormatNamedArguments NamedArgs;
+			for (const TPair<FString, FString>& Argument : FormatArgs.FilenameArguments)
+			{
+				NamedArgs.Add(Argument.Key, Argument.Value);
+			}
+
+			const FString& ResolvedValue = OutContext.Add(Pair.Key, FString::Format(*Pair.Value, NamedArgs));
+			UE_LOG(LogMovieRenderPipeline, VeryVerbose, TEXT("OCIO Context Key/Value: %s / %s"), *Pair.Key, *ResolvedValue);
+		}
+
+		return OutContext;
+	}
+
 	/**
 	 * Convenience function to create an OpenColorIO CPU processor based on the specified conversion settings. Editor-only.
 	 *
@@ -72,9 +154,8 @@ namespace UE::MovieGraph::Private
 
 		return {};
 	}
-} //end namespace UE::MovieGraph::Private
-
 #endif
+} //end namespace UE::MovieGraph::Private
 
 UMovieGraphImageSequenceOutputNode::UMovieGraphImageSequenceOutputNode()
 {
@@ -137,10 +218,6 @@ FString UMovieGraphImageSequenceOutputNode::CreateFileName(
 	const EImageFormat InImageFormat,
 	FMovieGraphResolveArgs& OutMergedFormatArgs) const
 {
-	UE::MovieGraph::FMovieGraphSampleState* Payload = InRenderData.Value->GetPayload<UE::MovieGraph::FMovieGraphSampleState>();
-
-	const TObjectPtr<UMoviePipelineExecutorShot>& Shot = InPipeline->GetActiveShotList()[Payload->TraversalContext.ShotIndex];
-
 	const TCHAR* Extension = TEXT("");
 	switch (InImageFormat)
 	{
@@ -167,33 +244,14 @@ FString UMovieGraphImageSequenceOutputNode::CreateFileName(
 	TMap<FString, FString> AdditionalFormatArgs;
 	AdditionalFormatArgs.Add(TEXT("ext"), Extension);
 
-	FMovieGraphFilenameResolveParams Params = FMovieGraphFilenameResolveParams();
-	Params.RenderDataIdentifier = InRenderData.Key;
-	//Params.RootFrameNumber = Payload->TraversalContext.Time.RootFrameNumber;
-	//Params.ShotFrameNumber = Payload->TraversalContext.Time.ShotFrameNumber;
-	Params.RootFrameNumberRel = Payload->TraversalContext.Time.OutputFrameNumber;
-	//Params.ShotFrameNumberRel = Payload->TraversalCOntext.Time.ShotFrameNumberRel
-	//Params.FileMetadata = ToDo: Track File Metadata
-	Params.ZeroPadFrameNumberCount = OutputSettingNode->ZeroPadFrameNumbers;
-	Params.FrameNumberOffset = OutputSettingNode->FrameNumberOffset;
-	Params.EvaluatedConfig = InRawFrameData->EvaluatedConfig.Get();
-	Params.Version = Shot->ShotInfo.VersionNumber;
-
-	// If time dilation is in effect, RootFrameNumber and ShotFrameNumber will contain duplicates and the files will overwrite each other, 
-	// so we force them into relative mode and then warn users we did that (as their numbers will jump from say 1001 -> 0000).
-	bool bForceRelativeFrameNumbers = true; // TODO: Use relative frame numbers until we track Root vs. Shot frame numbers. (Previously false);
-	//if (FileNameFormatString.Contains(TEXT("{frame")) Payload->TraversalContext.Time.IsTimeDilated() && !FileNameFormatString.Contains(TEXT("_rel}")))
-	//{
-	//	UE_LOG(LogMovieRenderPipeline, Warning, TEXT("Time Dilation was used but output format does not use relative time, forcing relative numbers. Change {frame_number} to {frame_number_rel} (or shot version) to remove this message."));
-	//	bForceRelativeFrameNumbers = true;
-	//}
-	Params.bForceRelativeFrameNumbers = bForceRelativeFrameNumbers;
-	Params.bEnsureAbsolutePath = true;
-	Params.FileNameFormatOverrides = AdditionalFormatArgs;
-	Params.InitializationTime = InPipeline->GetInitializationTime();
-	Params.Shot = Shot;
-	Params.Job = InPipeline->GetCurrentJob();
-	Params.EvaluatedConfig = InRawFrameData->EvaluatedConfig.Get();
+	UE::MovieGraph::FMovieGraphSampleState* Payload = InRenderData.Value->GetPayload<UE::MovieGraph::FMovieGraphSampleState>();
+
+	FMovieGraphFilenameResolveParams Params = UE::MovieGraph::Private::MakeResolveParams(
+		InRenderData.Key,
+		InPipeline,
+		InRawFrameData->EvaluatedConfig.Get(),
+		Payload->TraversalContext,
+		AdditionalFormatArgs);
 
 	// Take our string path from the Output Setting and resolve it.
 	return UMovieGraphBlueprintLibrary::ResolveFilenameFormatArguments(FileNameFormatString, Params, OutMergedFormatArgs);
@@ -273,9 +331,17 @@ void UMovieGraphImageSequenceOutputNode::OnReceiveImageDataImpl(UMovieGraphPipel
 #if WITH_EDITOR
 		if (ParentNode->OCIOConfiguration.bIsEnabled && Payload->bAllowOCIO)
 		{
+			TMap<FString, FString> ResolvedOCIOContext = UE::MovieGraph::Private::ResolveOpenColorIOContext(
+				ParentNode->OCIOContext,
+				RenderData.Key,
+				InPipeline,
+				InRawFrameData->EvaluatedConfig.Get(),
+				Payload->TraversalContext
+			);
+
 			FPixelPreProcessor OCIOPixelPreProcessor = UE::MovieGraph::Private::CreateOpenColorIOPixelPreProcessor(
 				ParentNode->OCIOConfiguration.ColorConfiguration,
-				ParentNode->OCIOContext
+				ResolvedOCIOContext
 			);
 			if (OCIOPixelPreProcessor)
 			{
@@ -371,7 +437,8 @@ void UMovieGraphImageSequenceOutputNode_EXR::UpdateTaskPerLayer(
 	const UMovieGraphImageSequenceOutputNode* InParentNode,
 	TUniquePtr<FImagePixelData> InImageData,
 	int32 InLayerIndex,
-	const FString& InLayerName) const
+	const FString& InLayerName,
+	const TMap<FString, FString>& InResolvedOCIOContext) const
 {
 	const UE::MovieGraph::FMovieGraphSampleState* Payload = InImageData->GetPayload<UE::MovieGraph::FMovieGraphSampleState>();
 
@@ -381,7 +448,7 @@ void UMovieGraphImageSequenceOutputNode_EXR::UpdateTaskPerLayer(
 	{
 		FPixelPreProcessor OCIOPixelPreProcessor = UE::MovieGraph::Private::CreateOpenColorIOPixelPreProcessor(
 			InParentNode->OCIOConfiguration.ColorConfiguration,
-			InParentNode->OCIOContext
+			InResolvedOCIOContext
 		);
 		if (OCIOPixelPreProcessor)
 		{
@@ -476,7 +543,19 @@ void UMovieGraphImageSequenceOutputNode_EXR::OnReceiveImageDataImpl(UMovieGraphP
 		{
 			PixelData = RenderData.Value->MoveImageDataToNew();
 		}
-		UpdateTaskPerLayer(*ImageWriteTask, ParentNode, MoveTemp(PixelData), LayerIndex);
+
+		TMap<FString, FString> ResolvedOCIOContext = {};
+#if WITH_EDITOR
+		ResolvedOCIOContext = UE::MovieGraph::Private::ResolveOpenColorIOContext(
+			ParentNode->OCIOContext,
+			RenderData.Key,
+			InPipeline,
+			InRawFrameData->EvaluatedConfig.Get(),
+			Payload->TraversalContext
+		);
+#endif // WITH_EDITOR
+
+		UpdateTaskPerLayer(*ImageWriteTask, ParentNode, MoveTemp(PixelData), LayerIndex, FString(), ResolvedOCIOContext);
 
 		// Perform compositing if any composited passes were found earlier
 		for (TPair<FMovieGraphRenderDataIdentifier, TUniquePtr<FImagePixelData>>& CompositedPass : CompositedPasses)
@@ -563,7 +642,18 @@ void UMovieGraphImageSequenceOutputNode_MultiLayerEXR::OnReceiveImageDataImpl(UM
 			{
 				PixelData = ImageData->MoveImageDataToNew();
 			}
-			UpdateTaskPerLayer(*MultiLayerImageTask, ParentNode, MoveTemp(PixelData), LayerIndex, LayerName);
+
+			TMap<FString, FString> ResolvedOCIOContext = {};
+#if WITH_EDITOR
+			ResolvedOCIOContext = UE::MovieGraph::Private::ResolveOpenColorIOContext(
+				ParentNode->OCIOContext,
+				RenderID,
+				InPipeline,
+				InRawFrameData->EvaluatedConfig.Get(),
+				Payload->TraversalContext
+			);
+#endif // WITH_EDITOR
+			UpdateTaskPerLayer(*MultiLayerImageTask, ParentNode, MoveTemp(PixelData), LayerIndex, LayerName, ResolvedOCIOContext);
 
 			LayerIndex++;
 		}
@@ -695,20 +785,8 @@ FString UMovieGraphImageSequenceOutputNode_MultiLayerEXR::ResolveOutputFilename(
 	FMovieGraphRenderDataIdentifier TempRenderDataIdentifier;
 	TempRenderDataIdentifier.RootBranchName = InBranchName;
 
-	// This resolves the filename format and gathers metadata from the settings at the same time.
-	FMovieGraphFilenameResolveParams Params = FMovieGraphFilenameResolveParams();
-	Params.RenderDataIdentifier = TempRenderDataIdentifier;
-	Params.RootFrameNumberRel = InRawFrameData->TraversalContext.Time.OutputFrameNumber;
-	Params.ZeroPadFrameNumberCount = OutputSettings->ZeroPadFrameNumbers;
-	Params.FrameNumberOffset = OutputSettings->FrameNumberOffset;
-	Params.bForceRelativeFrameNumbers = true;	// TODO: Use the right value
-	Params.bEnsureAbsolutePath = true;
-	Params.FileNameFormatOverrides = FormatOverrides;
-	Params.InitializationTime = InPipeline->GetInitializationTime();
-	Params.Job = InPipeline->GetCurrentJob();
-	Params.Shot = InPipeline->GetActiveShotList()[InRawFrameData->TraversalContext.ShotIndex];
-	Params.EvaluatedConfig = InRawFrameData->EvaluatedConfig.Get();
-	Params.Version = Params.Shot->ShotInfo.VersionNumber;
+	FMovieGraphFilenameResolveParams Params = UE::MovieGraph::Private::MakeResolveParams(
+		TempRenderDataIdentifier, InPipeline, InRawFrameData->EvaluatedConfig.Get(), InRawFrameData->TraversalContext, FormatOverrides);
 	
 	FString FinalFilePath = UMovieGraphBlueprintLibrary::ResolveFilenameFormatArguments(FilePathFormatString, Params, OutResolveArgs);
 
diff --git a/Engine/Plugins/MovieScene/MovieRenderPipeline/Source/MovieRenderPipelineRenderPasses/Public/MovieGraphImageSequenceOutputNode.h b/Engine/Plugins/MovieScene/MovieRenderPipeline/Source/MovieRenderPipelineRenderPasses/Public/MovieGraphImageSequenceOutputNode.h
--- a/Engine/Plugins/MovieScene/MovieRenderPipeline/Source/MovieRenderPipelineRenderPasses/Public/MovieGraphImageSequenceOutputNode.h
+++ b/Engine/Plugins/MovieScene/MovieRenderPipeline/Source/MovieRenderPipelineRenderPasses/Public/MovieGraphImageSequenceOutputNode.h
@@ -160,7 +160,8 @@ protected:
 		const UMovieGraphImageSequenceOutputNode* InParentNode,
 		TUniquePtr<FImagePixelData> InImageData,
 		int32 InLayerIndex,
-		const FString& InLayerName = {}
+		const FString& InLayerName = {},
+		const TMap<FString, FString>& InResolvedOCIOContext = {}
 	) const;
 };
 
