MoviePipeline: Fixed a bug which caused the "Actor" condition group query to not work correctly (PIE vs. editor actor pointer mismatch). Also fixed an oversight that caused some lists in condition group queries to have items that are non-deletable.

#rnx

diff --git a/Engine/Plugins/MovieScene/MovieRenderPipeline/Source/MovieRenderPipelineCore/MovieRenderPipelineCore.Build.cs b/Engine/Plugins/MovieScene/MovieRenderPipeline/Source/MovieRenderPipelineCore/MovieRenderPipelineCore.Build.cs
--- a/Engine/Plugins/MovieScene/MovieRenderPipeline/Source/MovieRenderPipelineCore/MovieRenderPipelineCore.Build.cs
+++ b/Engine/Plugins/MovieScene/MovieRenderPipeline/Source/MovieRenderPipelineCore/MovieRenderPipelineCore.Build.cs
@@ -47,6 +47,7 @@ public class MovieRenderPipelineCore : ModuleRules
 		        {
 			        "ClassViewer",
 			        "SceneOutliner",
+			        "UnrealEd",
 		        });
 	        
 			PublicDependencyModuleNames.Add("MovieSceneTools");
diff --git a/Engine/Plugins/MovieScene/MovieRenderPipeline/Source/MovieRenderPipelineCore/Private/Graph/MoviePipelineRenderLayerSubsystem.cpp b/Engine/Plugins/MovieScene/MovieRenderPipeline/Source/MovieRenderPipelineCore/Private/Graph/MoviePipelineRenderLayerSubsystem.cpp
--- a/Engine/Plugins/MovieScene/MovieRenderPipeline/Source/MovieRenderPipelineCore/Private/Graph/MoviePipelineRenderLayerSubsystem.cpp
+++ b/Engine/Plugins/MovieScene/MovieRenderPipeline/Source/MovieRenderPipelineCore/Private/Graph/MoviePipelineRenderLayerSubsystem.cpp
@@ -20,6 +20,7 @@
 #include "ActorTreeItem.h"
 #include "ClassViewerFilter.h"
 #include "ClassViewerModule.h"
+#include "Editor.h"
 #include "Graph/MovieGraphSharedWidgets.h"
 #include "ISceneOutliner.h"
 #include "SceneOutlinerModule.h"
@@ -477,9 +478,39 @@ void UMovieGraphConditionGroupQuery_Actor::Evaluate(const TArray<AActor*>& InAct
 {
 	TRACE_CPUPROFILER_EVENT_SCOPE(UMovieGraphConditionGroupQuery_Actor::Evaluate);
 
+	// Convert the actors in the query to PIE equivalents once, rather than constantly in the loop.
+	TArray<AActor*> ActorsToMatch_Pie;
+	ActorsToMatch_Pie.Reserve(ActorsToMatch.Num());
+	for (const TSoftObjectPtr<AActor>& SoftActorToMatch : ActorsToMatch)
+	{
+		if (!SoftActorToMatch.IsValid())
+		{
+			continue;
+		}
+		
+		AActor* ActorToMatch = SoftActorToMatch.Get();
+
+		// Only do editor -> PIE actor conversion when the actor is from an editor world
+		const UWorld* ActorWorld = ActorToMatch->GetWorld();
+		if (ActorWorld && ActorWorld->IsEditorWorld())
+		{
+#if WITH_EDITOR
+			if (AActor* PieActor = EditorUtilities::GetSimWorldCounterpartActor(ActorToMatch))
+			{
+				ActorsToMatch_Pie.Add(PieActor);
+			}
+#endif
+		}
+		else
+		{
+			// Just use ActorToMatch as-is if it's not from an editor actor
+			ActorsToMatch_Pie.Add(ActorToMatch);
+		}
+	}
+	
 	for (AActor* Actor : InActorsToQuery)
 	{
-		if (ActorsToMatch.Contains(Actor))
+		if (ActorsToMatch_Pie.Contains(Actor))
 		{
 			OutMatchingActors.Add(Actor);
 		}
@@ -511,12 +542,18 @@ TArray<TSharedRef<SWidget>> UMovieGraphConditionGroupQuery_Actor::GetWidgets()
 	}
 
 	Widgets.Add(
-		SNew(SMovieGraphSimpleList<TSharedPtr<TSoftObjectPtr<AActor>>>)
+		SAssignNew(ActorsList, SMovieGraphSimpleList<TSharedPtr<TSoftObjectPtr<AActor>>>)
 			.DataSource(&ListDataSource)
 			.DataType(FText::FromString("Actor"))
 			.DataTypePlural(FText::FromString("Actors"))
 			.OnGetRowText_Static(&GetRowText)
 			.OnGetRowIcon_Static(&GetRowIcon)
+			.OnDelete_Lambda([this](const TSharedPtr<TSoftObjectPtr<AActor>> InActor)
+			{
+				ListDataSource.Remove(InActor);
+				ActorsToMatch.Remove(*InActor.Get());
+				ActorsList->Refresh();
+			})
 	);
 
 	return Widgets;
@@ -773,12 +810,17 @@ TArray<TSharedRef<SWidget>> UMovieGraphConditionGroupQuery_ActorType::GetWidgets
 	TArray<TSharedRef<SWidget>> Widgets;
 
 	Widgets.Add(
-		SNew(SMovieGraphSimpleList<UClass*>)
+		SAssignNew(ActorTypesList, SMovieGraphSimpleList<UClass*>)
 			.DataSource(&ActorTypes)
 			.DataType(FText::FromString("Actor Type"))
 			.DataTypePlural(FText::FromString("Actor Types"))
 			.OnGetRowText_Static(&GetRowText)
 			.OnGetRowIcon_Static(&GetRowIcon)
+			.OnDelete_Lambda([this](UClass* InActorClass)
+			{
+				ActorTypes.Remove(InActorClass);
+				ActorTypesList->Refresh();
+			})
 	);
 
 	return Widgets;
@@ -964,12 +1006,17 @@ TArray<TSharedRef<SWidget>> UMovieGraphConditionGroupQuery_ComponentType::GetWid
 	TArray<TSharedRef<SWidget>> Widgets;
 
 	Widgets.Add(
-		SNew(SMovieGraphSimpleList<UClass*>)
+		SAssignNew(ComponentTypesList, SMovieGraphSimpleList<UClass*>)
 			.DataSource(&ComponentTypes)
 			.DataType(FText::FromString("Component Type"))
 			.DataTypePlural(FText::FromString("Component Types"))
 			.OnGetRowText_Static(&GetRowText)
 			.OnGetRowIcon_Static(&GetRowIcon)
+			.OnDelete_Lambda([this](UClass* InComponentType)
+			{
+				ComponentTypes.Remove(InComponentType);
+				ComponentTypesList->Refresh();
+			})			
 	);
 
 	return Widgets;
diff --git a/Engine/Plugins/MovieScene/MovieRenderPipeline/Source/MovieRenderPipelineCore/Public/Graph/MovieGraphSharedWidgets.h b/Engine/Plugins/MovieScene/MovieRenderPipeline/Source/MovieRenderPipelineCore/Public/Graph/MovieGraphSharedWidgets.h
--- a/Engine/Plugins/MovieScene/MovieRenderPipeline/Source/MovieRenderPipelineCore/Public/Graph/MovieGraphSharedWidgets.h
+++ b/Engine/Plugins/MovieScene/MovieRenderPipeline/Source/MovieRenderPipelineCore/Public/Graph/MovieGraphSharedWidgets.h
@@ -65,6 +65,7 @@ public:
 			[
 				SAssignNew(ListView, SListView<ListType>)
 				.ListItemsSource(DataSource)
+				.SelectionMode(ESelectionMode::Single)
 				.OnKeyDownHandler(this, &SMovieGraphSimpleList<ListType>::HandleDelete)
 				.OnGenerateRow(this, &SMovieGraphSimpleList<ListType>::GenerateRow)
 			]
diff --git a/Engine/Plugins/MovieScene/MovieRenderPipeline/Source/MovieRenderPipelineCore/Public/Graph/MoviePipelineRenderLayerSubsystem.h b/Engine/Plugins/MovieScene/MovieRenderPipeline/Source/MovieRenderPipelineCore/Public/Graph/MoviePipelineRenderLayerSubsystem.h
--- a/Engine/Plugins/MovieScene/MovieRenderPipeline/Source/MovieRenderPipelineCore/Public/Graph/MoviePipelineRenderLayerSubsystem.h
+++ b/Engine/Plugins/MovieScene/MovieRenderPipeline/Source/MovieRenderPipelineCore/Public/Graph/MoviePipelineRenderLayerSubsystem.h
@@ -3,6 +3,7 @@
 #pragma once
 
 #include "CoreMinimal.h"
+#include "Graph/MovieGraphSharedWidgets.h"
 #include "Subsystems/WorldSubsystem.h"
 #include "Textures/SlateIcon.h"
 
@@ -176,7 +177,7 @@ private:
 #endif
 
 public:
-	/** The query must match one of the actors in order to be a match. */
+	/** The query must match one of the actors in order to be a match. If these are editor actors, they will be converted to PIE actors automatically. */
 	UPROPERTY(EditAnywhere, Category="General")
 	TArray<TSoftObjectPtr<AActor>> ActorsToMatch;
 
@@ -184,6 +185,9 @@ private:
 #if WITH_EDITOR
 	TSharedPtr<class ISceneOutliner> ActorPickerWidget;
 
+	/** Displays the actors which have been chosen. */
+	TSharedPtr<SMovieGraphSimpleList<TSharedPtr<TSoftObjectPtr<AActor>>>> ActorsList;
+
 	// Not ideal to store a duplicate of ActorsToMatch, but SListView requires TSharedPtr<...> as the data source, and UPROPERTY does not
 	// support TSharedPtr<...>
 	TArray<TSharedPtr<TSoftObjectPtr<AActor>>> ListDataSource;
@@ -282,6 +286,9 @@ private:
 #if WITH_EDITOR
 	static const FSlateBrush* GetRowIcon(UClass* InActorType);
 	static FText GetRowText(UClass* InActorType);
+
+	/** Displays the actor types which have been chosen. */
+	TSharedPtr<SMovieGraphSimpleList<UClass*>> ActorTypesList;
 #endif
 	
 	/** The class viewer widget to show in the Add menu. */
@@ -341,6 +348,9 @@ private:
 #if WITH_EDITOR
 	static const FSlateBrush* GetRowIcon(UClass* InComponentType);
 	static FText GetRowText(UClass* InComponentType);
+
+	/** Displays the component types which have been chosen. */
+	TSharedPtr<SMovieGraphSimpleList<UClass*>> ComponentTypesList;
 #endif
 	
 	/** The class viewer widget to show in the Add menu. */
