MoviePipeline: Removed support for transient graph configs -- they will not be available in the first version of the graph.

Also added a (temporary?) warning in the graph editor that notifies the user if they are editing the default graph. This is important because the default graph that is used when "Replace with Graph" is executed is the graph that MRG ships with, and it should not be modified.

#rnx

diff --git a/Engine/Plugins/MovieScene/MovieRenderPipeline/Source/MovieRenderPipelineCore/Private/Graph/MovieGraphPipeline.cpp b/Engine/Plugins/MovieScene/MovieRenderPipeline/Source/MovieRenderPipelineCore/Private/Graph/MovieGraphPipeline.cpp
--- a/Engine/Plugins/MovieScene/MovieRenderPipeline/Source/MovieRenderPipelineCore/Private/Graph/MovieGraphPipeline.cpp
+++ b/Engine/Plugins/MovieScene/MovieRenderPipeline/Source/MovieRenderPipelineCore/Private/Graph/MovieGraphPipeline.cpp
@@ -47,7 +47,7 @@ void UMovieGraphPipeline::Initialize(UMoviePipelineExecutorJob* InJob, const FMo
 		return;
 	}
 
-	if (!ensureAlwaysMsgf(InJob->GetGraphConfig(), TEXT("MoviePipeline cannot be initialized with a null job configuration. Make sure you've created a Graph Config for this job (or use the regular UMoviePipeline instead of UMovieGraphPipeline). Aborting.")))
+	if (!ensureAlwaysMsgf(InJob->GetGraphPreset(), TEXT("MoviePipeline cannot be initialized with a null job configuration. Make sure you've created a Graph Config for this job (or use the regular UMoviePipeline instead of UMovieGraphPipeline). Aborting.")))
 	{
 		//Shutdown(true);
 		return;
@@ -1308,10 +1308,6 @@ UMovieGraphConfig* UMovieGraphPipeline::GetRootGraphForShot(UMoviePipelineExecut
 				{
 					return Shot->GetGraphPreset();
 				}
-				else if (Shot->GetGraphConfig())
-				{
-					return Shot->GetGraphConfig();
-				}
 			}
 
 		}
@@ -1321,10 +1317,6 @@ UMovieGraphConfig* UMovieGraphPipeline::GetRootGraphForShot(UMoviePipelineExecut
 		{
 			return GetCurrentJob()->GetGraphPreset();
 		}
-		else if (GetCurrentJob()->GetGraphConfig())
-		{
-			return GetCurrentJob()->GetGraphConfig();
-		}
 	}
 
 	return nullptr;
diff --git a/Engine/Plugins/MovieScene/MovieRenderPipeline/Source/MovieRenderPipelineCore/Private/MoviePipelineInProcessExecutor.cpp b/Engine/Plugins/MovieScene/MovieRenderPipeline/Source/MovieRenderPipelineCore/Private/MoviePipelineInProcessExecutor.cpp
--- a/Engine/Plugins/MovieScene/MovieRenderPipeline/Source/MovieRenderPipelineCore/Private/MoviePipelineInProcessExecutor.cpp
+++ b/Engine/Plugins/MovieScene/MovieRenderPipeline/Source/MovieRenderPipelineCore/Private/MoviePipelineInProcessExecutor.cpp
@@ -117,7 +117,7 @@ void UMoviePipelineInProcessExecutor::OnMapLoadFinished(UWorld* NewWorld)
 	}
 
 	// Temporary
-	if (CurrentJob->GetGraphConfig())
+	if (CurrentJob->IsUsingGraphConfiguration())
 	{
 		MoviePipelineClass = UMovieGraphPipeline::StaticClass();
 	}
diff --git a/Engine/Plugins/MovieScene/MovieRenderPipeline/Source/MovieRenderPipelineCore/Public/MoviePipelineQueue.h b/Engine/Plugins/MovieScene/MovieRenderPipeline/Source/MovieRenderPipelineCore/Public/MoviePipelineQueue.h
--- a/Engine/Plugins/MovieScene/MovieRenderPipeline/Source/MovieRenderPipelineCore/Public/MoviePipelineQueue.h
+++ b/Engine/Plugins/MovieScene/MovieRenderPipeline/Source/MovieRenderPipelineCore/Public/MoviePipelineQueue.h
@@ -142,7 +142,6 @@ public:
 
 	/**
 	 * Gets the graph-style preset that this job is using. If the job is not using a graph-style preset, returns nullptr.
-	 * @see GetGraphConfig()
 	 */
 	UFUNCTION(BlueprintPure, Category = "Movie Render Pipeline")
 	UMovieGraphConfig* GetGraphPreset() const
@@ -150,20 +149,9 @@ public:
 		return GraphPreset.LoadSynchronous();
 	}
 
-	/**
-	 * Gets the graph-style config that this job is using. If the job is not using a graph-style config, returns nullptr.
-	 * @see GetGraphPreset()
-	 */
-	UFUNCTION(BlueprintPure, Category = "Movie Render Pipeline")
-	UMovieGraphConfig* GetGraphConfig() const
-	{
-		return GraphConfig;
-	}
-
 	/**
 	 * Sets the graph-style preset that this job will use. Note that this will cause the graph to switch over to using
 	 * graph-style configuration if it is not already using it.
-	 * @see SetGraphConfig()
 	 */
 	UFUNCTION(BlueprintCallable, Category = "Movie Render Pipeline")
 	void SetGraphPreset(const UMovieGraphConfig* InGraphPreset)
@@ -176,17 +164,6 @@ public:
 #endif
 	}
 
-	/**
-	 * Sets the graph-style config that this job will use. Note that this will cause the graph to switch over to using
-	 * graph-style configuration if it is not already using it.
-	 * @see SetGraphPreset()
-	 */
-	UFUNCTION(BlueprintCallable, Category = "Movie Render Pipeline")
-	void SetGraphConfig(UMovieGraphConfig* InGraphConfig)
-	{
-		GraphConfig = InGraphConfig;
-	}
-
 	/** Returns whether this should should be rendered */
 	UFUNCTION(BlueprintPure, Category = "Movie Render Pipeline")
 	bool ShouldRender() const
@@ -256,10 +233,6 @@ private:
 	UPROPERTY()
 	TSoftObjectPtr<UMoviePipelineShotConfig> ShotOverridePresetOrigin;
 
-	/** The graph-based configuration that this shot is using; Can be nullptr. */
-	UPROPERTY()
-	TObjectPtr<UMovieGraphConfig> GraphConfig;
-
 	/** The graph-based configuration preset that this shot is using. Can be nullptr. */
 	UPROPERTY()
 	TSoftObjectPtr<UMovieGraphConfig> GraphPreset;
@@ -446,7 +419,7 @@ public:
 	UFUNCTION(BlueprintPure, Category = "Movie Render Pipeline")
 	bool IsUsingGraphConfiguration() const
 	{
-		return GraphPreset.IsValid() || GraphConfig != nullptr;
+		return GraphPreset.IsValid();
 	}
 
 	/**
@@ -459,20 +432,9 @@ public:
 		return GraphPreset.LoadSynchronous();
 	}
 
-	/**
-	 * Gets the graph-style config that this job is using. If the job is not using a graph-style config, returns nullptr.
-	 * @see GetGraphPreset()
-	 */
-	UFUNCTION(BlueprintPure, Category = "Movie Render Pipeline")
-	UMovieGraphConfig* GetGraphConfig() const
-	{
-		return GraphConfig;
-	}
-
 	/**
 	 * Sets the graph-style preset that this job will use. Note that this will cause the graph to switch over to using
 	 * graph-style configuration if it is not already using it.
-	 * @see SetGraphConfig()
 	 */
 	UFUNCTION(BlueprintCallable, Category = "Movie Render Pipeline")
 	void SetGraphPreset(const UMovieGraphConfig* InGraphPreset)
@@ -485,17 +447,6 @@ public:
 #endif
 	}
 
-	/**
-	 * Sets the graph-style config that this job will use. Note that this will cause the graph to switch over to using
-	 * graph-style configuration if it is not already using it.
-	 * @see SetGraphPreset()
-	 */
-	UFUNCTION(BlueprintCallable, Category = "Movie Render Pipeline")
-	void SetGraphConfig(UMovieGraphConfig* InGraphConfig)
-	{
-		GraphConfig = InGraphConfig;
-	}
-
 	UFUNCTION(BlueprintSetter, Category = "Movie Render Pipeline")
 	void SetSequence(FSoftObjectPath InSequence);
 
@@ -588,10 +539,6 @@ private:
 	UPROPERTY()
 	bool bEnabled;
 
-	/** The graph-based configuration that this job is using; this graph has not been saved as an asset. Can be nullptr.  */
-	UPROPERTY()
-	TObjectPtr<UMovieGraphConfig> GraphConfig;
-
 	/** The graph-based configuration preset that this job is using. Can be nullptr. */
 	UPROPERTY()
 	TSoftObjectPtr<UMovieGraphConfig> GraphPreset;
diff --git a/Engine/Plugins/MovieScene/MovieRenderPipeline/Source/MovieRenderPipelineEditor/Private/Graph/MovieGraphAssetToolkit.cpp b/Engine/Plugins/MovieScene/MovieRenderPipeline/Source/MovieRenderPipelineEditor/Private/Graph/MovieGraphAssetToolkit.cpp
--- a/Engine/Plugins/MovieScene/MovieRenderPipeline/Source/MovieRenderPipelineEditor/Private/Graph/MovieGraphAssetToolkit.cpp
+++ b/Engine/Plugins/MovieScene/MovieRenderPipeline/Source/MovieRenderPipelineEditor/Private/Graph/MovieGraphAssetToolkit.cpp
@@ -19,6 +19,7 @@
 #include "Graph/Nodes/MovieGraphSelectNode.h"
 #include "MovieEdGraphNode.h"
 #include "MovieGraphSchema.h"
+#include "MovieRenderPipelineSettings.h"
 #include "SMovieGraphActiveRenderSettingsTabContent.h"
 #include "SMovieGraphMembersTabContent.h"
 
@@ -206,6 +207,12 @@ TSharedRef<SDockTab> FMovieGraphAssetToolkit::SpawnTab_RenderGraphEditor(const F
 		.Label(LOCTEXT("GraphTab_Title", "Graph"))
 		[
 			SNew(SVerticalBox)
+			+ SVerticalBox::Slot()
+			.AutoHeight()
+			[
+				GetDefaultGraphEditWarning()
+			]
+			
 			+ SVerticalBox::Slot()
 			[
 				SAssignNew(MovieGraphWidget, SMoviePipelineGraphPanel)
@@ -387,6 +394,46 @@ void FMovieGraphAssetToolkit::PersistEditorOnlyNodes() const
 	}
 }
 
+TSharedRef<SWidget> FMovieGraphAssetToolkit::GetDefaultGraphEditWarning() const
+{
+	// Determine if the default graph is being edited
+	bool bIsDefaultGraphBeingEdited = false;
+	const UMovieRenderPipelineProjectSettings* ProjectSettings = GetDefault<UMovieRenderPipelineProjectSettings>();
+	const TSoftObjectPtr<UMovieGraphConfig> ProjectDefaultGraph = ProjectSettings->DefaultGraph;
+	if (const UMovieGraphConfig* DefaultGraph = ProjectDefaultGraph.LoadSynchronous())
+	{
+		bIsDefaultGraphBeingEdited = (DefaultGraph == InitialGraph.Get());
+	}
+	
+	return SNew(SBorder)
+		.BorderImage(FAppStyle::GetBrush("Brushes.Warning"))
+		.BorderBackgroundColor_Lambda([]()
+		{
+			FLinearColor WarningColor = FAppStyle::GetSlateColor("Colors.Warning").GetSpecifiedColor();
+			WarningColor.A = 0.3;
+			return WarningColor;
+		})
+		.Padding(5.f)
+		.Visibility(bIsDefaultGraphBeingEdited ? EVisibility::Visible : EVisibility::Collapsed)
+		[
+			SNew(SHorizontalBox)
+			+ SHorizontalBox::Slot()
+			.Padding(0, 0, 5.f, 0)
+			.AutoWidth()
+			[
+				SNew(SImage)
+				.Image(FAppStyle::Get().GetBrush("Icons.WarningWithColor"))
+			]
+
+			+ SHorizontalBox::Slot()
+			[
+				SNew(STextBlock)
+				.Text(LOCTEXT("GraphTab_EditingDefaultGraphWarning", "The default graph asset is being edited. 'Save As' to save a new graph asset."))
+				.Font(FCoreStyle::GetDefaultFontStyle("Bold", 10))
+			]
+		];
+}
+
 FName FMovieGraphAssetToolkit::GetToolkitFName() const
 {
 	return FName("MovieGraphEditor");
diff --git a/Engine/Plugins/MovieScene/MovieRenderPipeline/Source/MovieRenderPipelineEditor/Private/Graph/MovieGraphAssetToolkit.h b/Engine/Plugins/MovieScene/MovieRenderPipeline/Source/MovieRenderPipelineEditor/Private/Graph/MovieGraphAssetToolkit.h
--- a/Engine/Plugins/MovieScene/MovieRenderPipeline/Source/MovieRenderPipelineEditor/Private/Graph/MovieGraphAssetToolkit.h
+++ b/Engine/Plugins/MovieScene/MovieRenderPipeline/Source/MovieRenderPipelineEditor/Private/Graph/MovieGraphAssetToolkit.h
@@ -45,6 +45,9 @@ private:
 	/** Copies editor-only nodes to the underlying runtime graph so they can be restored later. */
 	void PersistEditorOnlyNodes() const;
 
+	/** Generate the widget that displays a warning banner if the default graph is being edited. */
+	TSharedRef<SWidget> GetDefaultGraphEditWarning() const;
+
 private:
 	/** The details panel for the selected object(s) in the graph */
 	TSharedPtr<IDetailsView> SelectedGraphObjectsDetailsWidget;
diff --git a/Engine/Plugins/MovieScene/MovieRenderPipeline/Source/MovieRenderPipelineEditor/Private/MoviePipelinePIEExecutor.cpp b/Engine/Plugins/MovieScene/MovieRenderPipeline/Source/MovieRenderPipelineEditor/Private/MoviePipelinePIEExecutor.cpp
--- a/Engine/Plugins/MovieScene/MovieRenderPipeline/Source/MovieRenderPipelineEditor/Private/MoviePipelinePIEExecutor.cpp
+++ b/Engine/Plugins/MovieScene/MovieRenderPipeline/Source/MovieRenderPipelineEditor/Private/MoviePipelinePIEExecutor.cpp
@@ -223,8 +223,8 @@ void UMoviePipelinePIEExecutor::OnPIEStartupFinished(bool)
 		return;
 	}
 
-	UMoviePipelineExecutorJob* CurrentJob = Queue->GetJobs()[CurrentPipelineIndex];
-	if (CurrentJob->GetGraphConfig())
+	const UMoviePipelineExecutorJob* CurrentJob = Queue->GetJobs()[CurrentPipelineIndex];
+	if (CurrentJob->IsUsingGraphConfiguration())
 	{
 		PipelineClass = UMovieGraphPipeline::StaticClass();
 	}
diff --git a/Engine/Plugins/MovieScene/MovieRenderPipeline/Source/MovieRenderPipelineEditor/Private/Widgets/SMoviePipelineQueueEditor.cpp b/Engine/Plugins/MovieScene/MovieRenderPipeline/Source/MovieRenderPipelineEditor/Private/Widgets/SMoviePipelineQueueEditor.cpp
--- a/Engine/Plugins/MovieScene/MovieRenderPipeline/Source/MovieRenderPipelineEditor/Private/Widgets/SMoviePipelineQueueEditor.cpp
+++ b/Engine/Plugins/MovieScene/MovieRenderPipeline/Source/MovieRenderPipelineEditor/Private/Widgets/SMoviePipelineQueueEditor.cpp
@@ -319,8 +319,13 @@ public:
 		}
 
 		// Note: Setting the graph preset will transition the job to use a graph-based configuration
-		UMovieGraphConfig* NewGraph = NewObject<UMovieGraphConfig>(Job);
-		Job->SetGraphConfig(NewGraph);
+		// Use the default graph specified in Project Settings.
+		const UMovieRenderPipelineProjectSettings* ProjectSettings = GetDefault<UMovieRenderPipelineProjectSettings>();
+		const TSoftObjectPtr<UMovieGraphConfig> ProjectDefaultGraph = ProjectSettings->DefaultGraph;
+		if (const UMovieGraphConfig* DefaultGraph = ProjectDefaultGraph.LoadSynchronous())
+		{
+			Job->SetGraphPreset(DefaultGraph);
+		}
 	}
 
 	void OnCreateNewGraphAndAssign() const
@@ -1004,8 +1009,13 @@ struct FMoviePipelineShotItem : IMoviePipelineQueueTreeItem
 		}
 
 		// Note: Setting the graph preset will transition the job to use a graph-based configuration
-		UMovieGraphConfig* NewGraph = NewObject<UMovieGraphConfig>(Shot);
-		Shot->SetGraphConfig(NewGraph);
+		// Use the default graph specified in Project Settings.
+		const UMovieRenderPipelineProjectSettings* ProjectSettings = GetDefault<UMovieRenderPipelineProjectSettings>();
+		const TSoftObjectPtr<UMovieGraphConfig> ProjectDefaultGraph = ProjectSettings->DefaultGraph;
+		if (const UMovieGraphConfig* DefaultGraph = ProjectDefaultGraph.LoadSynchronous())
+		{
+			Shot->SetGraphPreset(DefaultGraph);
+		}
 	}
 
 	void OnCreateNewGraphAndAssign() const
diff --git a/Engine/Plugins/MovieScene/MovieRenderPipeline/Source/MovieRenderPipelineEditor/Private/Widgets/SMoviePipelineQueuePanel.cpp b/Engine/Plugins/MovieScene/MovieRenderPipeline/Source/MovieRenderPipelineEditor/Private/Widgets/SMoviePipelineQueuePanel.cpp
--- a/Engine/Plugins/MovieScene/MovieRenderPipeline/Source/MovieRenderPipelineEditor/Private/Widgets/SMoviePipelineQueuePanel.cpp
+++ b/Engine/Plugins/MovieScene/MovieRenderPipeline/Source/MovieRenderPipelineEditor/Private/Widgets/SMoviePipelineQueuePanel.cpp
@@ -387,7 +387,7 @@ void SMoviePipelineQueuePanel::OnEditJobConfigRequested(TWeakObjectPtr<UMoviePip
 	}
 	else if (InJob.IsValid() && InJob->IsUsingGraphConfiguration())
 	{
-		GraphToEdit = (InJob->GetGraphPreset() != nullptr) ? InJob->GetGraphPreset() : InJob->GetGraphConfig();
+		GraphToEdit = InJob->GetGraphPreset();
 	}
 
 	if (GraphToEdit)
