MoviePipeline: Blocked cross-graph copy/paste of variable nodes, which is currently not supported.

Also fixed an issue which caused graph assets to not be pickable from shot-based jobs in the queue.

#rnx

diff --git a/Engine/Plugins/MovieScene/MovieRenderPipeline/Source/MovieRenderPipelineEditor/Private/Graph/MovieEdGraphInputNode.h b/Engine/Plugins/MovieScene/MovieRenderPipeline/Source/MovieRenderPipelineEditor/Private/Graph/MovieEdGraphInputNode.h
--- a/Engine/Plugins/MovieScene/MovieRenderPipeline/Source/MovieRenderPipelineEditor/Private/Graph/MovieEdGraphInputNode.h
+++ b/Engine/Plugins/MovieScene/MovieRenderPipeline/Source/MovieRenderPipelineEditor/Private/Graph/MovieEdGraphInputNode.h
@@ -19,6 +19,6 @@ public:
 	// This is set to true so that the node is deserialized on paste, required for CanPasteHere to be called.
 	virtual bool CanDuplicateNode() const override { return true; } 
 	// And this is set to false so the paste doesn't actually go through, and instead we get the editor prompt notifying user.
-	virtual bool CanPasteHere(const UEdGraph* TargetGraph) const { return false; }
+	virtual bool CanPasteHere(const UEdGraph* TargetGraph) const override { return false; }
 	// ~End UEdGraphNode interface
 };
diff --git a/Engine/Plugins/MovieScene/MovieRenderPipeline/Source/MovieRenderPipelineEditor/Private/Graph/MovieEdGraphNode.cpp b/Engine/Plugins/MovieScene/MovieRenderPipeline/Source/MovieRenderPipelineEditor/Private/Graph/MovieEdGraphNode.cpp
--- a/Engine/Plugins/MovieScene/MovieRenderPipeline/Source/MovieRenderPipelineEditor/Private/Graph/MovieEdGraphNode.cpp
+++ b/Engine/Plugins/MovieScene/MovieRenderPipeline/Source/MovieRenderPipelineEditor/Private/Graph/MovieEdGraphNode.cpp
@@ -504,6 +504,12 @@ void UMoviePipelineEdGraphNodeBase::PrepareForCopying()
 		// This is restored in PostCopy
 		RuntimeNode->Rename(nullptr, this, REN_DontCreateRedirectors | REN_DoNotDirty);
 	}
+
+	const UMoviePipelineEdGraph* MovieGraphEditorGraph = CastChecked<UMoviePipelineEdGraph>(GetGraph());
+	const UMovieGraphConfig* RuntimeGraph = MovieGraphEditorGraph->GetPipelineGraph();
+
+	// Track where this node came from for copy/paste purposes
+	OriginGraph = RuntimeGraph->GetPathName();
 }
 
 void UMoviePipelineEdGraphNodeBase::PostCopy()
diff --git a/Engine/Plugins/MovieScene/MovieRenderPipeline/Source/MovieRenderPipelineEditor/Private/Graph/MovieEdGraphNode.h b/Engine/Plugins/MovieScene/MovieRenderPipeline/Source/MovieRenderPipelineEditor/Private/Graph/MovieEdGraphNode.h
--- a/Engine/Plugins/MovieScene/MovieRenderPipeline/Source/MovieRenderPipelineEditor/Private/Graph/MovieEdGraphNode.h
+++ b/Engine/Plugins/MovieScene/MovieRenderPipeline/Source/MovieRenderPipelineEditor/Private/Graph/MovieEdGraphNode.h
@@ -75,6 +75,10 @@ protected:
 
 	/** Should we early out in ReconstructNode(), skipping restoring the connections to other nodes. Set during Copy/Paste. */
 	bool bDisableReconstructNode = false;
+
+	/** The graph that this node originated from during a copy/paste operation. This value should never be used outside of copy/paste logic. */
+	UPROPERTY(SkipSerialization)
+	FSoftObjectPath OriginGraph;
 };
 
 UCLASS()
diff --git a/Engine/Plugins/MovieScene/MovieRenderPipeline/Source/MovieRenderPipelineEditor/Private/Graph/MovieEdGraphOutputNode.h b/Engine/Plugins/MovieScene/MovieRenderPipeline/Source/MovieRenderPipelineEditor/Private/Graph/MovieEdGraphOutputNode.h
--- a/Engine/Plugins/MovieScene/MovieRenderPipeline/Source/MovieRenderPipelineEditor/Private/Graph/MovieEdGraphOutputNode.h
+++ b/Engine/Plugins/MovieScene/MovieRenderPipeline/Source/MovieRenderPipelineEditor/Private/Graph/MovieEdGraphOutputNode.h
@@ -18,6 +18,6 @@ public:
 	// This is set to true so that the node is deserialized on paste, required for CanPasteHere to be called
 	virtual bool CanDuplicateNode() const override { return true; }
 	// And this is set to false so the paste doesn't actually go through, and instead we get the editor prompt notifying user.
-	virtual bool CanPasteHere(const UEdGraph* TargetGraph) const { return false; }
+	virtual bool CanPasteHere(const UEdGraph* TargetGraph) const override { return false; }
 	// ~End UEdGraphNode interface
 };
\ No newline at end of file
diff --git a/Engine/Plugins/MovieScene/MovieRenderPipeline/Source/MovieRenderPipelineEditor/Private/Graph/MovieEdGraphVariableNode.cpp b/Engine/Plugins/MovieScene/MovieRenderPipeline/Source/MovieRenderPipelineEditor/Private/Graph/MovieEdGraphVariableNode.cpp
--- a/Engine/Plugins/MovieScene/MovieRenderPipeline/Source/MovieRenderPipelineEditor/Private/Graph/MovieEdGraphVariableNode.cpp
+++ b/Engine/Plugins/MovieScene/MovieRenderPipeline/Source/MovieRenderPipelineEditor/Private/Graph/MovieEdGraphVariableNode.cpp
@@ -1,10 +1,13 @@
-﻿// Copyright Epic Games, Inc. All Rights Reserved.
+// Copyright Epic Games, Inc. All Rights Reserved.
 
 #include "Graph/MovieEdGraphVariableNode.h"
 
+#include "MovieEdGraph.h"
 #include "EdGraph/EdGraphPin.h"
+#include "Framework/Notifications/NotificationManager.h"
 #include "Graph/MovieGraphConfig.h"
 #include "Graph/Nodes/MovieGraphVariableNode.h"
+#include "Widgets/Notifications/SNotificationList.h"
 
 #define LOCTEXT_NAMESPACE "MoviePipelineGraph"
 
@@ -30,4 +33,37 @@ void UMoviePipelineEdGraphVariableNode::AllocateDefaultPins()
 	}
 }
 
+bool UMoviePipelineEdGraphVariableNode::CanPasteHere(const UEdGraph* TargetGraph) const
+{
+	const UMoviePipelineEdGraph* MovieEdGraph = Cast<UMoviePipelineEdGraph>(TargetGraph);
+	if (!MovieEdGraph)
+	{
+		return false;
+	}
+	
+	const UMovieGraphConfig* DestinationMovieGraph = MovieEdGraph->GetPipelineGraph();
+	if (!DestinationMovieGraph)
+	{
+		return false;
+	}
+	
+	// Only allow pasting the variable node into the graph that it originated from
+	if (DestinationMovieGraph->GetPathName() != OriginGraph.ToString())
+	{
+		FNotificationInfo Info(LOCTEXT("VariableNodePasteWarning", "Variable nodes cannot be pasted between graphs."));
+		Info.ExpireDuration = 3.0f;
+		Info.bUseLargeFont = false;
+		Info.Image = FCoreStyle::Get().GetBrush(TEXT("MessageLog.Warning"));
+		const TSharedPtr<SNotificationItem> Notification = FSlateNotificationManager::Get().AddNotification(Info);
+		if (Notification.IsValid())
+		{
+			Notification->SetCompletionState(SNotificationItem::CS_None);
+		}
+		
+		return false;
+	}
+
+	return true;
+}
+
 #undef LOCTEXT_NAMESPACE
\ No newline at end of file
diff --git a/Engine/Plugins/MovieScene/MovieRenderPipeline/Source/MovieRenderPipelineEditor/Private/Graph/MovieEdGraphVariableNode.h b/Engine/Plugins/MovieScene/MovieRenderPipeline/Source/MovieRenderPipelineEditor/Private/Graph/MovieEdGraphVariableNode.h
--- a/Engine/Plugins/MovieScene/MovieRenderPipeline/Source/MovieRenderPipelineEditor/Private/Graph/MovieEdGraphVariableNode.h
+++ b/Engine/Plugins/MovieScene/MovieRenderPipeline/Source/MovieRenderPipelineEditor/Private/Graph/MovieEdGraphVariableNode.h
@@ -16,5 +16,6 @@ public:
 	// ~Begin UEdGraphNode interface
 	virtual FText GetNodeTitle(ENodeTitleType::Type TitleType) const override;
 	virtual void AllocateDefaultPins() override;
+	virtual bool CanPasteHere(const UEdGraph* TargetGraph) const override;
 	// ~End UEdGraphNode interface
 };
\ No newline at end of file
diff --git a/Engine/Plugins/MovieScene/MovieRenderPipeline/Source/MovieRenderPipelineEditor/Private/Widgets/SMoviePipelineQueueEditor.cpp b/Engine/Plugins/MovieScene/MovieRenderPipeline/Source/MovieRenderPipelineEditor/Private/Widgets/SMoviePipelineQueueEditor.cpp
--- a/Engine/Plugins/MovieScene/MovieRenderPipeline/Source/MovieRenderPipelineEditor/Private/Widgets/SMoviePipelineQueueEditor.cpp
+++ b/Engine/Plugins/MovieScene/MovieRenderPipeline/Source/MovieRenderPipelineEditor/Private/Widgets/SMoviePipelineQueueEditor.cpp
@@ -977,10 +977,19 @@ struct FMoviePipelineShotItem : IMoviePipelineQueueTreeItem
 		// Close the dropdown menu that showed them the assets to pick from.
 		FSlateApplication::Get().DismissAllMenus();
 
-		UMoviePipelineExecutorShot* Shot = WeakShot.Get();
-		if (Shot)
+		if (UMoviePipelineExecutorShot* Shot = WeakShot.Get())
 		{
-			Shot->SetShotOverridePresetOrigin(CastChecked<UMoviePipelineShotConfig>(AssetData.GetAsset()));
+			FScopedTransaction Transaction(LOCTEXT("PickShotPresetAsset_Transaction", "Set Shot Configuration Asset"));
+			Shot->Modify();
+
+			if (Shot->IsUsingGraphConfiguration())
+			{
+				Shot->SetGraphPreset(CastChecked<UMovieGraphConfig>(AssetData.GetAsset()));
+			}
+			else
+			{
+				Shot->SetShotOverridePresetOrigin(CastChecked<UMoviePipelineShotConfig>(AssetData.GetAsset()));
+			}
 		}
 
 		OnChosePresetCallback.ExecuteIfBound(WeakJob, WeakShot);
@@ -1048,8 +1057,11 @@ struct FMoviePipelineShotItem : IMoviePipelineQueueTreeItem
 
 	TSharedRef<SWidget> OnGenerateShotConfigPresetPickerMenu()
 	{
+		const UMoviePipelineExecutorShot* Shot = WeakShot.Get();
+		UClass* ConfigType = Shot && Shot->IsUsingGraphConfiguration() ? UMovieGraphConfig::StaticClass() : UMoviePipelinePrimaryConfig::StaticClass();
+		
 		return FMoviePipelineQueueJobTreeItem::OnGenerateConfigPresetPickerMenuFromClass(
-			UMoviePipelineShotConfig::StaticClass(),
+			ConfigType,
 			WeakJob,
 			FOnAssetSelected::CreateRaw(this, &FMoviePipelineShotItem::OnPickShotPresetFromAsset),
 			FExecuteAction::CreateRaw(this, &FMoviePipelineShotItem::OnPickNewShotPreset),
