MoviePipeline: Fixed issue where the collection node does not update its text when the collection is renamed. Recent changes to collections broke this.

#rnx

diff --git a/Engine/Plugins/MovieScene/MovieRenderPipeline/Source/MovieRenderPipelineCore/Private/Graph/MoviePipelineRenderLayerSubsystem.cpp b/Engine/Plugins/MovieScene/MovieRenderPipeline/Source/MovieRenderPipelineCore/Private/Graph/MoviePipelineRenderLayerSubsystem.cpp
--- a/Engine/Plugins/MovieScene/MovieRenderPipeline/Source/MovieRenderPipelineCore/Private/Graph/MoviePipelineRenderLayerSubsystem.cpp
+++ b/Engine/Plugins/MovieScene/MovieRenderPipeline/Source/MovieRenderPipelineCore/Private/Graph/MoviePipelineRenderLayerSubsystem.cpp
@@ -1324,6 +1324,19 @@ bool UMovieGraphCollection::MoveConditionGroupToIndex(UMovieGraphConditionGroup*
 	return true;
 }
 
+#if WITH_EDITOR
+void UMovieGraphCollection::PostEditChangeProperty(FPropertyChangedEvent& PropertyChangedEvent)
+{
+	Super::PostEditChangeProperty(PropertyChangedEvent);
+
+	// Name change delegate is broadcast here so it catches both SetCollectionName() and a direct change of the property via the details panel
+	if (PropertyChangedEvent.GetPropertyName() == GET_MEMBER_NAME_CHECKED(UMovieGraphCollection, CollectionName))
+	{
+		OnCollectionNameChangedDelegate.Broadcast(this);
+	}
+}
+#endif	// WITH_EDITOR
+
 void UMovieGraphCollection::SetCollectionName(const FString& InName)
 {
 	CollectionName = InName;
diff --git a/Engine/Plugins/MovieScene/MovieRenderPipeline/Source/MovieRenderPipelineCore/Private/Graph/Nodes/MovieGraphCollectionNode.cpp b/Engine/Plugins/MovieScene/MovieRenderPipeline/Source/MovieRenderPipelineCore/Private/Graph/Nodes/MovieGraphCollectionNode.cpp
--- a/Engine/Plugins/MovieScene/MovieRenderPipeline/Source/MovieRenderPipelineCore/Private/Graph/Nodes/MovieGraphCollectionNode.cpp
+++ b/Engine/Plugins/MovieScene/MovieRenderPipeline/Source/MovieRenderPipelineCore/Private/Graph/Nodes/MovieGraphCollectionNode.cpp
@@ -11,6 +11,13 @@
 UMovieGraphCollectionNode::UMovieGraphCollectionNode()
 {
 	Collection = CreateDefaultSubobject<UMovieGraphCollection>(TEXT("Collection"));
+
+#if WITH_EDITOR
+	Collection->OnCollectionNameChangedDelegate.AddWeakLambda(this, [this](UMovieGraphCollection* ChangedCollection)
+	{
+		OnNodeChangedDelegate.Broadcast(this);
+	});
+#endif
 }
 
 #if WITH_EDITOR
@@ -45,19 +52,6 @@ FSlateIcon UMovieGraphCollectionNode::GetIconAndTint(FLinearColor& OutColor) con
 	OutColor = FLinearColor::White;
 	return CollectionIcon;
 }
-
-void UMovieGraphCollectionNode::PostEditChangeProperty(FPropertyChangedEvent& PropertyChangedEvent)
-{
-	Super::PostEditChangeProperty(PropertyChangedEvent);
-
-	// TODO: Ideally this only fires when the collection name changes
-	// Broadcast a node-changed delegate so that the node title's UI gets updated.
-	if (PropertyChangedEvent.GetPropertyName() == GET_MEMBER_NAME_CHECKED(UMovieGraphCollectionNode, Collection) ||
-		PropertyChangedEvent.GetPropertyName() == GET_MEMBER_NAME_CHECKED(UMovieGraphCollectionNode, bOverride_Collection))
-	{
-		OnNodeChangedDelegate.Broadcast(this);
-	}
-}
 #endif // WITH_EDITOR
 
 FString UMovieGraphCollectionNode::GetNodeInstanceName() const
diff --git a/Engine/Plugins/MovieScene/MovieRenderPipeline/Source/MovieRenderPipelineCore/Public/Graph/MoviePipelineRenderLayerSubsystem.h b/Engine/Plugins/MovieScene/MovieRenderPipeline/Source/MovieRenderPipelineCore/Public/Graph/MoviePipelineRenderLayerSubsystem.h
--- a/Engine/Plugins/MovieScene/MovieRenderPipeline/Source/MovieRenderPipelineCore/Public/Graph/MoviePipelineRenderLayerSubsystem.h
+++ b/Engine/Plugins/MovieScene/MovieRenderPipeline/Source/MovieRenderPipelineCore/Public/Graph/MoviePipelineRenderLayerSubsystem.h
@@ -438,7 +438,16 @@ class MOVIERENDERPIPELINECORE_API UMovieGraphCollection : public UObject
 	GENERATED_BODY()
 
 public:
+	/** Delegate which is called when the collection name changes. */
+	DECLARE_MULTICAST_DELEGATE_OneParam(FMovieGraphCollectionNameChanged, UMovieGraphCollection*)
+	
 	UMovieGraphCollection() = default;
+	
+#if WITH_EDITOR
+	//~ Begin UObject interface
+	virtual void PostEditChangeProperty(FPropertyChangedEvent& PropertyChangedEvent) override;
+	//~ End UObject interface
+#endif
 
 	/** Sets the name of the collection as seen in the UI. */
 	UFUNCTION(BlueprintCallable, Category = "Settings")
@@ -480,6 +489,12 @@ public:
 	UFUNCTION(BlueprintCallable, Category = "Settings")
 	bool MoveConditionGroupToIndex(UMovieGraphConditionGroup* InConditionGroup, const int32 NewIndex);
 
+public:
+#if WITH_EDITOR
+	/** Called when the collection name changes. */
+	FMovieGraphCollectionNameChanged OnCollectionNameChangedDelegate;
+#endif
+
 private:
 	/** The display name of the collection, shown in the UI. Does not need to be unique across collections. */
 	UPROPERTY(EditAnywhere, Category="Collection")
diff --git a/Engine/Plugins/MovieScene/MovieRenderPipeline/Source/MovieRenderPipelineCore/Public/Graph/Nodes/MovieGraphCollectionNode.h b/Engine/Plugins/MovieScene/MovieRenderPipeline/Source/MovieRenderPipelineCore/Public/Graph/Nodes/MovieGraphCollectionNode.h
--- a/Engine/Plugins/MovieScene/MovieRenderPipeline/Source/MovieRenderPipelineCore/Public/Graph/Nodes/MovieGraphCollectionNode.h
+++ b/Engine/Plugins/MovieScene/MovieRenderPipeline/Source/MovieRenderPipelineCore/Public/Graph/Nodes/MovieGraphCollectionNode.h
@@ -26,10 +26,6 @@ public:
 	virtual FText GetMenuCategory() const override;
 	virtual FLinearColor GetNodeTitleColor() const override;
 	virtual FSlateIcon GetIconAndTint(FLinearColor& OutColor) const override;
-
-	//~ Begin UObject interface
-	virtual void PostEditChangeProperty(FPropertyChangedEvent& PropertyChangedEvent) override;
-	//~ End UObject interface
 #endif
 
 	virtual FString GetNodeInstanceName() const override;
