MoviePipeline: Enable details search for graph assets.

#rnx

diff --git a/Engine/Plugins/MovieScene/MovieRenderPipeline/Source/MovieRenderPipelineEditor/Private/Customizations/Graph/MovieGraphShowFlagsCustomization.h b/Engine/Plugins/MovieScene/MovieRenderPipeline/Source/MovieRenderPipelineEditor/Private/Customizations/Graph/MovieGraphShowFlagsCustomization.h
--- a/Engine/Plugins/MovieScene/MovieRenderPipeline/Source/MovieRenderPipelineEditor/Private/Customizations/Graph/MovieGraphShowFlagsCustomization.h
+++ b/Engine/Plugins/MovieScene/MovieRenderPipeline/Source/MovieRenderPipelineEditor/Private/Customizations/Graph/MovieGraphShowFlagsCustomization.h
@@ -78,6 +78,7 @@ protected:
 				const uint32 ShowFlagIndex = ShowFlag.EngineShowFlagIndex;
 				
 				FlagGroup.AddWidgetRow()
+				.FilterString(ShowFlag.DisplayName)
 				.NameContent()
 				[
 					SNew(SHorizontalBox)
diff --git a/Engine/Plugins/MovieScene/MovieRenderPipeline/Source/MovieRenderPipelineEditor/Private/Graph/MovieGraphAssetToolkit.cpp b/Engine/Plugins/MovieScene/MovieRenderPipeline/Source/MovieRenderPipelineEditor/Private/Graph/MovieGraphAssetToolkit.cpp
--- a/Engine/Plugins/MovieScene/MovieRenderPipeline/Source/MovieRenderPipelineEditor/Private/Graph/MovieGraphAssetToolkit.cpp
+++ b/Engine/Plugins/MovieScene/MovieRenderPipeline/Source/MovieRenderPipelineEditor/Private/Graph/MovieGraphAssetToolkit.cpp
@@ -289,7 +289,7 @@ TSharedRef<SDockTab> FMovieGraphAssetToolkit::SpawnTab_RenderGraphDetails(const
 {
 	FPropertyEditorModule& PropertyEditorModule = FModuleManager::Get().LoadModuleChecked<FPropertyEditorModule>("PropertyEditor");
 	FDetailsViewArgs DetailsViewArgs;
-	DetailsViewArgs.bAllowSearch = false;
+	DetailsViewArgs.bShowPropertyMatrixButton = false;
 	DetailsViewArgs.NameAreaSettings = FDetailsViewArgs::HideNameArea;
 	DetailsViewArgs.bHideSelectionTip = true;
 	DetailsViewArgs.bAllowMultipleTopLevelObjects = true;
