Movie Graph:

-New: Separate UMovieGraphOutputSettingNode into Globals and RenderLayer specific nodes

diff --git a/Engine/Plugins/MovieScene/MovieRenderPipeline/Config/BaseMovieRenderPipeline.ini b/Engine/Plugins/MovieScene/MovieRenderPipeline/Config/BaseMovieRenderPipeline.ini
--- a/Engine/Plugins/MovieScene/MovieRenderPipeline/Config/BaseMovieRenderPipeline.ini
+++ b/Engine/Plugins/MovieScene/MovieRenderPipeline/Config/BaseMovieRenderPipeline.ini
@@ -1,2 +1,5 @@
 [CoreRedirects]
 +PropertyRedirects=(OldName="MoviePipelineDeferredPassBase.StencilLayers", NewName="ActorLayers")
++ClassRedirects=(OldName="/Script/MovieRenderPipelineCore.MovieGraphOutputSettingNode",NewName="/Script/MovieRenderPipelineCore.MovieGraphGlobalOutputSettingNode")
++PropertyRedirects=(OldName="/Script/MovieRenderPipelineCore.MovieGraphCommandLineEncoderNode.FileNameFormatOverride",NewName="/Script/MovieRenderPipelineCore.MovieGraphCommandLineEncoderNode.FileNameFormat")
++PropertyRedirects=(OldName="/Script/MovieRenderPipelineCore.MovieGraphCommandLineEncoderNode.bOverride_FileNameFormatOverride",NewName="/Script/MovieRenderPipelineCore.MovieGraphCommandLineEncoderNode.bOverride_FileNameFormat")
diff --git a/Engine/Plugins/MovieScene/MovieRenderPipeline/Source/MovieRenderPipelineCore/Private/Graph/MovieGraphBlueprintLibrary.cpp b/Engine/Plugins/MovieScene/MovieRenderPipeline/Source/MovieRenderPipelineCore/Private/Graph/MovieGraphBlueprintLibrary.cpp
--- a/Engine/Plugins/MovieScene/MovieRenderPipeline/Source/MovieRenderPipelineCore/Private/Graph/MovieGraphBlueprintLibrary.cpp
+++ b/Engine/Plugins/MovieScene/MovieRenderPipeline/Source/MovieRenderPipelineCore/Private/Graph/MovieGraphBlueprintLibrary.cpp
@@ -4,15 +4,18 @@
 
 #include "Graph/MovieGraphPipeline.h"
 #include "Graph/MovieGraphProjectSettings.h"
-#include "Graph/Nodes/MovieGraphOutputSettingNode.h"
-#include "Graph/Nodes/MovieGraphRenderLayerNode.h"
 #include "Graph/Nodes/MovieGraphCameraNode.h"
-#include "HAL/FileManager.h"
-#include "Internationalization/Regex.h"
+#include "Graph/Nodes/MovieGraphCommandLineEncoderNode.h"
+#include "Graph/Nodes/MovieGraphFileOutputNode.h"
+#include "Graph/Nodes/MovieGraphGlobalOutputSettingNode.h"
+#include "Graph/Nodes/MovieGraphRenderLayerNode.h"
 #include "MoviePipelineBlueprintLibrary.h"
 #include "MoviePipelineUtils.h"
 
-FFrameRate UMovieGraphBlueprintLibrary::GetEffectiveFrameRate(UMovieGraphOutputSettingNode* InNode, const FFrameRate& InDefaultRate)
+#include "HAL/FileManager.h"
+#include "Internationalization/Regex.h"
+
+FFrameRate UMovieGraphBlueprintLibrary::GetEffectiveFrameRate(UMovieGraphGlobalOutputSettingNode* InNode, const FFrameRate& InDefaultRate)
 {
 	if (InNode && InNode->bOverride_OutputFrameRate)
 	{
@@ -68,7 +71,7 @@ FString UMovieGraphBlueprintLibrary::ResolveFilenameFormatArguments(const FStrin
 				RenderLayerName = RenderLayerNode->GetRenderLayerName();
 			}
 
-			TObjectPtr<UMovieGraphOutputSettingNode> OutputSettingNode = InParams.EvaluatedConfig->GetSettingForBranch<UMovieGraphOutputSettingNode>(InParams.RenderDataIdentifier.RootBranchName, bIncludeCDOs);
+			TObjectPtr<UMovieGraphGlobalOutputSettingNode> OutputSettingNode = InParams.EvaluatedConfig->GetSettingForBranch<UMovieGraphGlobalOutputSettingNode>(UMovieGraphNode::GlobalsPinName, bIncludeCDOs);
 			if(OutputSettingNode)
 			{
 				bOverwriteExisting = OutputSettingNode->bOverwriteExistingOutput;
@@ -96,7 +99,7 @@ FString UMovieGraphBlueprintLibrary::ResolveFilenameFormatArguments(const FStrin
 		if (InParams.EvaluatedConfig)
 		{
 			const bool bIncludeCDOs = false;
-			UMovieGraphOutputSettingNode* OutputSettingNode = InParams.EvaluatedConfig->GetSettingForBranch<UMovieGraphOutputSettingNode>(InParams.RenderDataIdentifier.RootBranchName, bIncludeCDOs);
+			UMovieGraphGlobalOutputSettingNode* OutputSettingNode = InParams.EvaluatedConfig->GetSettingForBranch<UMovieGraphGlobalOutputSettingNode>(InParams.RenderDataIdentifier.RootBranchName, bIncludeCDOs);
 			FrameRate = GetEffectiveFrameRate(OutputSettingNode, InParams.DefaultFrameRate).AsDecimal();
 		}
 
@@ -228,74 +231,100 @@ int32 UMovieGraphBlueprintLibrary::ResolveVersionNumber(FMovieGraphFilenameResol
 
 	constexpr bool bIncludeCDOs = true;
 	constexpr bool bExactMatch = true;
-	const UMovieGraphOutputSettingNode* OutputSettingNode = InParams.EvaluatedConfig->GetSettingForBranch<UMovieGraphOutputSettingNode>(UMovieGraphNode::GlobalsPinName, bIncludeCDOs, bExactMatch);
+	const UMovieGraphGlobalOutputSettingNode* OutputSettingNode =
+		InParams.EvaluatedConfig->GetSettingForBranch<UMovieGraphGlobalOutputSettingNode>(UMovieGraphNode::GlobalsPinName, bIncludeCDOs, bExactMatch);
 	if (!OutputSettingNode->VersioningSettings.bAutoVersioning)
 	{
 		return OutputSettingNode->VersioningSettings.VersionNumber;
 	}
 
-	// Calculate a version number by looking at the output path and then scanning for a version token.
-	const FString FileNameFormatString = InParams.FileNameOverride.Len() > 0
-		? InParams.FileNameOverride
-		: OutputSettingNode->OutputDirectory.Path / OutputSettingNode->FileNameFormat;
-
 	// Force the Version string to stay as {version} so we can substring based on it later.
 	InParams.FileNameFormatOverrides.Add(TEXT("version"), TEXT("{version}"));
 
-	FMovieGraphResolveArgs FinalFormatArgs;
-	FString FinalPath = ResolveFilenameFormatArguments(FileNameFormatString, InParams, FinalFormatArgs);
-	FinalPath = FPaths::ConvertRelativePathToFull(FinalPath);
-	FPaths::NormalizeFilename(FinalPath);
+	// Get output nodes from the evaluated graph
+	TArray<UMovieGraphSettingNode*> ResultNodes = InParams.EvaluatedConfig->GetSettingsForBranch(
+		UMovieGraphCommandLineEncoderNode::StaticClass(), InParams.RenderDataIdentifier.RootBranchName, bIncludeCDOs, bExactMatch);
+	ResultNodes.Append(InParams.EvaluatedConfig->GetSettingsForBranch(
+		UMovieGraphFileOutputNode::StaticClass(), InParams.RenderDataIdentifier.RootBranchName, bIncludeCDOs, bExactMatch));
+	
+	int32 HighestVersion = -1;
 
-	// Can't resolve a version if it's not clear from the path where the version number will be used.
-	if (!FinalPath.Contains(TEXT("{version}")))
+	auto ExtrapolateHighestVersionFromResultNode = [&InParams, OutputSettingNode, &HighestVersion](const FString& FileNameFormat)
 	{
-		return -1;
-	}
-	
-	int32 HighestVersion = 0;
-
-	// FinalPath can have {version} either in a folder name or in a file name. We need to find the 'parent' of either the
-	// file or folder that contains it. We can do this by looking for {version} and then finding the last "/" character,
-	// which will be the containing folder.
-	const int32 VersionStringIndex = FinalPath.Find(TEXT("{version}"), ESearchCase::Type::IgnoreCase, ESearchDir::Type::FromStart);
-	if (VersionStringIndex >= 0)
-	{
-		const int32 LastParentFolder = FinalPath.Find(TEXT("/"), ESearchCase::Type::IgnoreCase, ESearchDir::Type::FromEnd, VersionStringIndex);
-		FinalPath.LeftInline(LastParentFolder + 1);
-
-		// Now that we have the parent folder of either the folder with the version token, or the file with the version
-		// token, we will look through all immediate children and scan for version tokens so we can find the highest one.
-		const FRegexPattern VersionSearchPattern(TEXT("v([0-9]{3})"));
-		constexpr bool bFindFiles = true;
-		constexpr bool bFindDirectories = true;
-		const FString SearchString = FinalPath / TEXT("*.*");
-		TArray<FString> FoundFilesAndFoldersInDirectory;
-		IFileManager& FileManager = IFileManager::Get();
-		FileManager.FindFiles(FoundFilesAndFoldersInDirectory, *SearchString, bFindFiles, bFindDirectories);
-
-		for (const FString& Path : FoundFilesAndFoldersInDirectory)
+		// Calculate a version number by looking at the output path and then scanning for a version token.
+		const FString FileNameFormatString = InParams.FileNameOverride.Len() > 0
+			? InParams.FileNameOverride
+			: OutputSettingNode->OutputDirectory.Path / FileNameFormat;
+
+		FMovieGraphResolveArgs FinalFormatArgs;
+		FString FinalPath = ResolveFilenameFormatArguments(FileNameFormatString, InParams, FinalFormatArgs);
+		FinalPath = FPaths::ConvertRelativePathToFull(FinalPath);
+		FPaths::NormalizeFilename(FinalPath);
+
+		// Can't resolve a version if it's not clear from the path where the version number will be used.
+		if (!FinalPath.Contains(TEXT("{version}")))
 		{
-			FRegexMatcher Regex(VersionSearchPattern, *Path);
-			if (Regex.FindNext())
-			{
-				FString Result = Regex.GetCaptureGroup(0);
-				if (Result.Len() > 0)
-				{
-					// Strip the "v" token off, expected pattern is vXXX
-					Result.RightChopInline(1);
-				}
+			return;
+		}
 
-				int32 VersionNumber = 0;
-				LexFromString(VersionNumber, *Result);
-				if (VersionNumber > HighestVersion)
+		// FinalPath can have {version} either in a folder name or in a file name. We need to find the 'parent' of either the
+		// file or folder that contains it. We can do this by looking for {version} and then finding the last "/" character,
+		// which will be the containing folder.
+		const int32 VersionStringIndex = FinalPath.Find(TEXT("{version}"), ESearchCase::Type::IgnoreCase, ESearchDir::Type::FromStart);
+		if (VersionStringIndex >= 0)
+		{
+			const int32 LastParentFolder = FinalPath.Find(TEXT("/"), ESearchCase::Type::IgnoreCase, ESearchDir::Type::FromEnd, VersionStringIndex);
+			FinalPath.LeftInline(LastParentFolder + 1);
+
+			// Now that we have the parent folder of either the folder with the version token, or the file with the version
+			// token, we will look through all immediate children and scan for version tokens so we can find the highest one.
+			const FRegexPattern VersionSearchPattern(TEXT("v([0-9]{3})"));
+			constexpr bool bFindFiles = true;
+			constexpr bool bFindDirectories = true;
+			const FString SearchString = FinalPath / TEXT("*.*");
+			TArray<FString> FoundFilesAndFoldersInDirectory;
+			IFileManager& FileManager = IFileManager::Get();
+			FileManager.FindFiles(FoundFilesAndFoldersInDirectory, *SearchString, bFindFiles, bFindDirectories);
+
+			for (const FString& Path : FoundFilesAndFoldersInDirectory)
+			{
+				FRegexMatcher Regex(VersionSearchPattern, *Path);
+				if (Regex.FindNext())
 				{
-					HighestVersion = VersionNumber;
+					FString Result = Regex.GetCaptureGroup(0);
+					if (Result.Len() > 0)
+					{
+						// Strip the "v" token off, expected pattern is vXXX
+						Result.RightChopInline(1);
+					}
+
+					int32 VersionNumber = 0;
+					LexFromString(VersionNumber, *Result);
+					if (VersionNumber > HighestVersion)
+					{
+						HighestVersion = VersionNumber;
+					}
 				}
 			}
 		}
-	}
+	};
+	
+	for (const UMovieGraphSettingNode* ResultNode : ResultNodes)
+	{
+		FString FileNameFormat = FString();
 
+		if (const UMovieGraphFileOutputNode* FileOutputNode = Cast<UMovieGraphFileOutputNode>(ResultNode))
+		{
+			FileNameFormat = FileOutputNode->FileNameFormat;
+		}
+		else if (const UMovieGraphCommandLineEncoderNode* CommandLineEncoderNode = Cast<UMovieGraphCommandLineEncoderNode>(ResultNode))
+		{
+			FileNameFormat = CommandLineEncoderNode->FileNameFormat;
+		}
+
+		ExtrapolateHighestVersionFromResultNode(FileNameFormat);
+	}
+	
 	return HighestVersion + (bGetNextVersion ? 1 : 0);
 }
 
@@ -308,7 +337,7 @@ FIntPoint UMovieGraphBlueprintLibrary::GetEffectiveOutputResolution(UMovieGraphE
 	}
 	
 	constexpr bool bIncludeCDOs = true;
-	const UMovieGraphOutputSettingNode* OutputSetting = InEvaluatedGraph->GetSettingForBranch<UMovieGraphOutputSettingNode>(InBranchName, bIncludeCDOs);
+	const UMovieGraphGlobalOutputSettingNode* OutputSetting = InEvaluatedGraph->GetSettingForBranch<UMovieGraphGlobalOutputSettingNode>(InBranchName, bIncludeCDOs);
 	if (!ensure(OutputSetting))
 	{
 		return FIntPoint();
diff --git a/Engine/Plugins/MovieScene/MovieRenderPipeline/Source/MovieRenderPipelineCore/Private/Graph/MovieGraphConfig.cpp b/Engine/Plugins/MovieScene/MovieRenderPipeline/Source/MovieRenderPipelineCore/Private/Graph/MovieGraphConfig.cpp
--- a/Engine/Plugins/MovieScene/MovieRenderPipeline/Source/MovieRenderPipelineCore/Private/Graph/MovieGraphConfig.cpp
+++ b/Engine/Plugins/MovieScene/MovieRenderPipeline/Source/MovieRenderPipelineCore/Private/Graph/MovieGraphConfig.cpp
@@ -9,7 +9,7 @@
 #include "Graph/MovieGraphPipeline.h"
 #include "Graph/Nodes/MovieGraphInputNode.h"
 #include "Graph/Nodes/MovieGraphOutputNode.h"
-#include "Graph/Nodes/MovieGraphOutputSettingNode.h"
+#include "Graph/Nodes/MovieGraphGlobalOutputSettingNode.h"
 #include "Graph/Nodes/MovieGraphRemoveRenderSettingNode.h"
 #include "Graph/Nodes/MovieGraphSubgraphNode.h"
 #include "Graph/Nodes/MovieGraphVariableNode.h"
@@ -1066,7 +1066,7 @@ void UMovieGraphConfig::RecurseUpGlobalsBranchToFindOutputDirectory(const UMovie
 			}
 		}
 	}
-	else if (const UMovieGraphOutputSettingNode* SettingsNode = Cast<UMovieGraphOutputSettingNode>(InNode))
+	else if (const UMovieGraphGlobalOutputSettingNode* SettingsNode = Cast<UMovieGraphGlobalOutputSettingNode>(InNode))
 	{
 		if (OutOutputDirectory.IsEmpty() && SettingsNode->bOverride_OutputDirectory)
 		{
@@ -1098,7 +1098,7 @@ void UMovieGraphConfig::GetOutputDirectory(FString& OutOutputDirectory) const
 		if (OutOutputDirectory.IsEmpty())
 		{
 			// If we didn't find any overrides, use the CDO values
-			UMovieGraphOutputSettingNode* CDO = Cast<UMovieGraphOutputSettingNode>(UMovieGraphOutputSettingNode::StaticClass()->ClassDefaultObject);
+			UMovieGraphGlobalOutputSettingNode* CDO = Cast<UMovieGraphGlobalOutputSettingNode>(UMovieGraphGlobalOutputSettingNode::StaticClass()->ClassDefaultObject);
 			check(CDO);
 			
 			OutOutputDirectory = CDO->OutputDirectory.Path;
diff --git a/Engine/Plugins/MovieScene/MovieRenderPipeline/Source/MovieRenderPipelineCore/Private/Graph/MovieGraphCoreTimeStep.cpp b/Engine/Plugins/MovieScene/MovieRenderPipeline/Source/MovieRenderPipelineCore/Private/Graph/MovieGraphCoreTimeStep.cpp
--- a/Engine/Plugins/MovieScene/MovieRenderPipeline/Source/MovieRenderPipelineCore/Private/Graph/MovieGraphCoreTimeStep.cpp
+++ b/Engine/Plugins/MovieScene/MovieRenderPipeline/Source/MovieRenderPipelineCore/Private/Graph/MovieGraphCoreTimeStep.cpp
@@ -4,7 +4,7 @@
 
 #include "Graph/MovieGraphPipeline.h"
 #include "Graph/MovieGraphBlueprintLibrary.h"
-#include "Graph/Nodes/MovieGraphOutputSettingNode.h"
+#include "Graph/Nodes/MovieGraphGlobalOutputSettingNode.h"
 #include "Graph/Nodes/MovieGraphCameraNode.h"
 #include "MoviePipelineQueue.h"
 #include "MovieRenderPipelineCoreModule.h"
@@ -378,7 +378,7 @@ void UMovieGraphCoreTimeStep::TickProducingFrames()
 	// Calculate frame numbers and timecodes for the current sequence (root) and shot
 	{
 		constexpr bool bIncludeCDOs = true;
-		UMovieGraphOutputSettingNode* OutputSetting = CurrentFrameData.EvaluatedConfig->GetSettingForBranch<UMovieGraphOutputSettingNode>(UMovieGraphNode::GlobalsPinName, bIncludeCDOs);
+		UMovieGraphGlobalOutputSettingNode* OutputSetting = CurrentFrameData.EvaluatedConfig->GetSettingForBranch<UMovieGraphGlobalOutputSettingNode>(UMovieGraphNode::GlobalsPinName, bIncludeCDOs);
 
 		// "Closest" isn't straightforward when using temporal sub-sampling, ie: A large enough shutter angle pushes a sample over the half way point and it rounds to
 		// the wrong one. Because temporal sub-sampling isn't centered around a frame (the centering is done via the final eval time) we can just subtract TSI*TPS to get our centered value.
diff --git a/Engine/Plugins/MovieScene/MovieRenderPipeline/Source/MovieRenderPipelineCore/Private/Graph/MovieGraphDefaultAudioRenderer.cpp b/Engine/Plugins/MovieScene/MovieRenderPipeline/Source/MovieRenderPipelineCore/Private/Graph/MovieGraphDefaultAudioRenderer.cpp
--- a/Engine/Plugins/MovieScene/MovieRenderPipeline/Source/MovieRenderPipelineCore/Private/Graph/MovieGraphDefaultAudioRenderer.cpp
+++ b/Engine/Plugins/MovieScene/MovieRenderPipeline/Source/MovieRenderPipelineCore/Private/Graph/MovieGraphDefaultAudioRenderer.cpp
@@ -10,7 +10,7 @@
 #include "Engine/Engine.h"
 #include "Graph/MovieGraphBlueprintLibrary.h"
 #include "Graph/MovieGraphUtils.h"
-#include "Graph/Nodes/MovieGraphOutputSettingNode.h"
+#include "Graph/Nodes/MovieGraphGlobalOutputSettingNode.h"
 #include "HAL/IConsoleManager.h"
 #include "MoviePipelineQueue.h"
 #include "MovieRenderPipelineDataTypes.h"
@@ -112,7 +112,7 @@ void UMovieGraphDefaultAudioRenderer::ProcessAudioTick()
 		constexpr bool bIncludeCDOs = true;
 		constexpr bool bExactMatch = true;
 		const TObjectPtr<UMovieGraphEvaluatedConfig> EvaluatedGraph = TimeStepInstance->GetCalculatedTimeData().EvaluatedConfig;
-		UMovieGraphOutputSettingNode* OutputSettingNode = EvaluatedGraph->GetSettingForBranch<UMovieGraphOutputSettingNode>(UMovieGraphNode::GlobalsPinName, bIncludeCDOs, bExactMatch);
+		UMovieGraphGlobalOutputSettingNode* OutputSettingNode = EvaluatedGraph->GetSettingForBranch<UMovieGraphGlobalOutputSettingNode>(UMovieGraphNode::GlobalsPinName, bIncludeCDOs, bExactMatch);
 
 		const FFrameRate SourceFrameRate = DataSourceInstance->GetDisplayRate();
 		const FFrameRate EffectiveFrameRate = UMovieGraphBlueprintLibrary::GetEffectiveFrameRate(OutputSettingNode, SourceFrameRate);
diff --git a/Engine/Plugins/MovieScene/MovieRenderPipeline/Source/MovieRenderPipelineCore/Private/Graph/MovieGraphPipeline.cpp b/Engine/Plugins/MovieScene/MovieRenderPipeline/Source/MovieRenderPipelineCore/Private/Graph/MovieGraphPipeline.cpp
--- a/Engine/Plugins/MovieScene/MovieRenderPipeline/Source/MovieRenderPipelineCore/Private/Graph/MovieGraphPipeline.cpp
+++ b/Engine/Plugins/MovieScene/MovieRenderPipeline/Source/MovieRenderPipelineCore/Private/Graph/MovieGraphPipeline.cpp
@@ -13,7 +13,7 @@
 #include "Graph/Nodes/MovieGraphModifierNode.h"
 #include "Graph/Nodes/MovieGraphRenderLayerNode.h"
 #include "Graph/Nodes/MovieGraphSamplingMethodNode.h"
-#include "Graph/Nodes/MovieGraphOutputSettingNode.h"
+#include "Graph/Nodes/MovieGraphGlobalOutputSettingNode.h"
 #include "Graph/Nodes/MovieGraphWarmUpSettingNode.h"
 #include "Graph/MovieGraphBlueprintLibrary.h"
 #include "MovieRenderPipelineCoreModule.h"
@@ -323,7 +323,7 @@ void UMovieGraphPipeline::BuildShotListFromDataSource()
 				*Shot->OuterName, *SamplingMethodNode->SamplingMethodClass.GetAssetName());
 		}
 		
-		UMovieGraphOutputSettingNode* OutputNode = EvaluatedConfig->GetSettingForBranch<UMovieGraphOutputSettingNode>(UMovieGraphSettingNode::GlobalsPinName);
+		UMovieGraphGlobalOutputSettingNode* OutputNode = EvaluatedConfig->GetSettingForBranch<UMovieGraphGlobalOutputSettingNode>(UMovieGraphSettingNode::GlobalsPinName);
 
 		const FFrameRate SourceFrameRate = GetDataSourceInstance()->GetDisplayRate();
 		const FFrameRate FinalFrameRate = UMovieGraphBlueprintLibrary::GetEffectiveFrameRate(OutputNode, SourceFrameRate);
@@ -681,7 +681,7 @@ void UMovieGraphPipeline::SetSoloShot(const TObjectPtr<UMoviePipelineExecutorSho
 
 		const FMovieGraphTimeStepData& TimeStepData = GetTimeStepInstance()->GetCalculatedTimeData();
 		TObjectPtr<UMovieGraphEvaluatedConfig> Config = TimeStepData.EvaluatedConfig;
-		UMovieGraphOutputSettingNode* OutputNode = Config->GetSettingForBranch<UMovieGraphOutputSettingNode>(UMovieGraphSettingNode::GlobalsPinName);
+		UMovieGraphGlobalOutputSettingNode* OutputNode = Config->GetSettingForBranch<UMovieGraphGlobalOutputSettingNode>(UMovieGraphSettingNode::GlobalsPinName);
 
 		const FFrameRate SourceFrameRate = GetDataSourceInstance()->GetDisplayRate();
 		const FFrameRate FinalFrameRate = UMovieGraphBlueprintLibrary::GetEffectiveFrameRate(OutputNode, SourceFrameRate);
@@ -837,32 +837,34 @@ int32 UMovieGraphPipeline::ResolveVersionForShot(const TObjectPtr<UMoviePipeline
 
 	const FMovieGraphTimeStepData& ShotTimeStepInstance = GraphTimeStepInstances.Last()->GetCalculatedTimeData();
 
+	constexpr bool bIncludeCDOs = true;
+	constexpr bool bExactMatch = true;
+	const UMovieGraphGlobalOutputSettingNode* BranchOutputSettingNode =
+		EvaluatedConfig->GetSettingForBranch<UMovieGraphGlobalOutputSettingNode>(UMovieGraphNode::GlobalsPinName, bIncludeCDOs, bExactMatch);
+	
+	// TODO: The RenderDataIdentifier needs to be fully populated in order to resolve the version reliably. For
+	// example, {renderer_name} may be used as a directory, and discovering files correctly depends on that token
+	// being resolved. We can loop over all renderer and file output nodes here since the branches are already being
+	// iterated, and resolve a version for each (still using the highest version found as the final resolved
+	// version). Note that to do this correctly, we need a way of asking these nodes for their renderer name and all
+	// possible sub-resource names.
+	ResolveParams.InitializationTime = GetInitializationTime();
+	ResolveParams.DefaultFrameRate = ShotTimeStepInstance.FrameRate;
+	ResolveParams.FrameNumberOffset = BranchOutputSettingNode->FrameNumberOffset;
+	ResolveParams.RenderDataIdentifier.CameraName = Shot->InnerName;
+	// ResolveParams.RenderDataIdentifier.RendererName = ?
+	// ResolveParams.RenderDataIdentifier.SubResourceName = ?
+	ResolveParams.RootFrameNumber = ShotTimeStepInstance.RootFrameNumber.Value;
+	ResolveParams.ShotFrameNumber = ShotTimeStepInstance.ShotFrameNumber.Value;
+	ResolveParams.bForceRelativeFrameNumbers = false;	// TODO: This should not be hardcoded
+	// ResolveParams.FileNameFormatOverrides = ?
+	ResolveParams.RootFrameNumberRel = ShotTimeStepInstance.OutputFrameNumber;
+	// ResolveParams.ShotFrameNumberRel = ?
+	ResolveParams.ZeroPadFrameNumberCount = BranchOutputSettingNode->ZeroPadFrameNumbers;
+
 	for (const FName& BranchName : EvaluatedConfig->GetBranchNames())
 	{
-		constexpr bool bIncludeCDOs = true;
-		constexpr bool bExactMatch = true;
-		const UMovieGraphOutputSettingNode* BranchOutputSettingNode = EvaluatedConfig->GetSettingForBranch<UMovieGraphOutputSettingNode>(BranchName, bIncludeCDOs, bExactMatch);
-
-		// TODO: The RenderDataIdentifier needs to be fully populated in order to resolve the version reliably. For
-		// example, {renderer_name} may be used as a directory, and discovering files correctly depends on that token
-		// being resolved. We can loop over all renderer and file output nodes here since the branches are already being
-		// iterated, and resolve a version for each (still using the highest version found as the final resolved
-		// version). Note that to do this correctly, we need a way of asking these nodes for their renderer name and all
-		// possible sub-resource names.
-		ResolveParams.InitializationTime = GetInitializationTime();
-		ResolveParams.DefaultFrameRate = ShotTimeStepInstance.FrameRate;
-		ResolveParams.FrameNumberOffset = BranchOutputSettingNode->FrameNumberOffset;
 		ResolveParams.RenderDataIdentifier.RootBranchName = BranchName;
-		ResolveParams.RenderDataIdentifier.CameraName = Shot->InnerName;
-		// ResolveParams.RenderDataIdentifier.RendererName = ?
-		// ResolveParams.RenderDataIdentifier.SubResourceName = ?
-		ResolveParams.RootFrameNumber = ShotTimeStepInstance.RootFrameNumber.Value;
-		ResolveParams.ShotFrameNumber = ShotTimeStepInstance.ShotFrameNumber.Value;
-		ResolveParams.bForceRelativeFrameNumbers = false;	// TODO: This should not be hardcoded
-		// ResolveParams.FileNameFormatOverrides = ?
-		ResolveParams.RootFrameNumberRel = ShotTimeStepInstance.OutputFrameNumber;
-		// ResolveParams.ShotFrameNumberRel = ?
-		ResolveParams.ZeroPadFrameNumberCount = BranchOutputSettingNode->ZeroPadFrameNumbers;
 
 		const int32 BranchVersion = UMovieGraphBlueprintLibrary::ResolveVersionNumber(ResolveParams);
 		if (BranchVersion > HighestVersionFound)
diff --git a/Engine/Plugins/MovieScene/MovieRenderPipeline/Source/MovieRenderPipelineCore/Private/Graph/Nodes/MovieGraphAudioOutputNode.cpp b/Engine/Plugins/MovieScene/MovieRenderPipeline/Source/MovieRenderPipelineCore/Private/Graph/Nodes/MovieGraphAudioOutputNode.cpp
--- a/Engine/Plugins/MovieScene/MovieRenderPipeline/Source/MovieRenderPipelineCore/Private/Graph/Nodes/MovieGraphAudioOutputNode.cpp
+++ b/Engine/Plugins/MovieScene/MovieRenderPipeline/Source/MovieRenderPipelineCore/Private/Graph/Nodes/MovieGraphAudioOutputNode.cpp
@@ -6,7 +6,7 @@
 #include "Graph/MovieGraphBlueprintLibrary.h"
 #include "Graph/MovieGraphPipeline.h"
 #include "Graph/MovieGraphUtils.h"
-#include "Graph/Nodes/MovieGraphOutputSettingNode.h"
+#include "Graph/Nodes/MovieGraphGlobalOutputSettingNode.h"
 #include "MoviePipelineUtils.h"
 #include "MovieRenderPipelineCoreModule.h"
 #include "Styling/AppStyle.h"
@@ -131,9 +131,8 @@ FString UMovieGraphAudioOutputNode::GenerateOutputPath(const FMovieGraphRenderDa
 	constexpr bool bIncludeCDOs = true;
 	constexpr bool bExactMatch = true;
 	const TObjectPtr<UMovieGraphEvaluatedConfig> EvaluatedConfig = CachedPipeline->GetTimeStepInstance()->GetCalculatedTimeData().EvaluatedConfig;
-	const UMovieGraphOutputSettingNode* OutputNode = EvaluatedConfig->GetSettingForBranch<UMovieGraphOutputSettingNode>(GlobalsPinName, bIncludeCDOs, bExactMatch);
-	const FString OutputFilename = FileNameFormatOverride.Len() > 0 ? FileNameFormatOverride : OutputNode->FileNameFormat;
-	FString FileNameFormatString = OutputNode->OutputDirectory.Path / OutputFilename;
+	const UMovieGraphGlobalOutputSettingNode* OutputNode = EvaluatedConfig->GetSettingForBranch<UMovieGraphGlobalOutputSettingNode>(GlobalsPinName, bIncludeCDOs, bExactMatch);
+	FString FileNameFormatString = OutputNode->OutputDirectory.Path / FileNameFormat;
 
 	constexpr bool bIncludeRenderPass = false;
 	constexpr bool bTestFrameNumber = false;
diff --git a/Engine/Plugins/MovieScene/MovieRenderPipeline/Source/MovieRenderPipelineCore/Private/Graph/Nodes/MovieGraphCommandLineEncoderNode.cpp b/Engine/Plugins/MovieScene/MovieRenderPipeline/Source/MovieRenderPipelineCore/Private/Graph/Nodes/MovieGraphCommandLineEncoderNode.cpp
--- a/Engine/Plugins/MovieScene/MovieRenderPipeline/Source/MovieRenderPipelineCore/Private/Graph/Nodes/MovieGraphCommandLineEncoderNode.cpp
+++ b/Engine/Plugins/MovieScene/MovieRenderPipeline/Source/MovieRenderPipelineCore/Private/Graph/Nodes/MovieGraphCommandLineEncoderNode.cpp
@@ -5,7 +5,7 @@
 #include "Graph/MovieGraphBlueprintLibrary.h"
 #include "Graph/MovieGraphPipeline.h"
 #include "Graph/Nodes/MovieGraphAudioOutputNode.h"
-#include "Graph/Nodes/MovieGraphOutputSettingNode.h"
+#include "Graph/Nodes/MovieGraphGlobalOutputSettingNode.h"
 #include "HAL/FileManager.h"
 #include "HAL/PlatformFileManager.h"
 #include "Misc/CoreDelegates.h"
@@ -84,9 +84,8 @@ void UMovieGraphCommandLineEncoderNode::StartEncodingProcess(TArray<FMovieGraphR
 	}
 
 	// Get the base output file path; it still needs to be validated and resolved at this point
-	const UMovieGraphOutputSettingNode* OutputSettingNode = GetSettingOnBranch<UMovieGraphOutputSettingNode>();
-	const FString OutputFilename = FileNameFormatOverride.Len() > 0 ? FileNameFormatOverride : OutputSettingNode->FileNameFormat;
-	FString FilePathFormatString = OutputSettingNode->OutputDirectory.Path / OutputFilename;
+	const UMovieGraphGlobalOutputSettingNode* OutputSettingNode = GetSettingOnBranch<UMovieGraphGlobalOutputSettingNode>();
+	FString FilePathFormatString = OutputSettingNode->OutputDirectory.Path / FileNameFormat;
 	
 	constexpr bool bTestRenderPass = false;
 	constexpr bool bTestFrameNumber = false;
@@ -238,10 +237,9 @@ void UMovieGraphCommandLineEncoderNode::OnTick()
 
 bool UMovieGraphCommandLineEncoderNode::NeedsPerShotFlushing() const
 {
-	const UMovieGraphOutputSettingNode* OutputSettingNode = GetSettingOnBranch<UMovieGraphOutputSettingNode>();
+	const UMovieGraphGlobalOutputSettingNode* OutputSettingNode = GetSettingOnBranch<UMovieGraphGlobalOutputSettingNode>();
 	
-	const FString OutputFilename = FileNameFormatOverride.Len() > 0 ? FileNameFormatOverride : OutputSettingNode->FileNameFormat;
-	const FString FullPath = OutputSettingNode->OutputDirectory.Path / OutputFilename;
+	const FString FullPath = OutputSettingNode->OutputDirectory.Path / FileNameFormat;
 	
 	if (FullPath.Contains(TEXT("{shot_name}")) || FullPath.Contains(TEXT("{camera_name}")))
 	{
@@ -262,7 +260,7 @@ TMap<FMovieGraphRenderDataIdentifier, UMovieGraphCommandLineEncoderNode::FEncode
 	FString ExecutablePathNoQuotes = EncoderSettings->ExecutablePath.Replace(TEXT("\""), TEXT(""));
 	FPaths::NormalizeFilename(ExecutablePathNoQuotes);
 
-	UMovieGraphOutputSettingNode* OutputSettingNode = GetSettingOnBranch<UMovieGraphOutputSettingNode>();
+	UMovieGraphGlobalOutputSettingNode* OutputSettingNode = GetSettingOnBranch<UMovieGraphGlobalOutputSettingNode>();
 	const FFrameRate SourceFrameRate = CachedPipeline->GetDataSourceInstance()->GetDisplayRate();
 	const FFrameRate EffectiveFrameRate = UMovieGraphBlueprintLibrary::GetEffectiveFrameRate(OutputSettingNode, SourceFrameRate);
 
@@ -359,6 +357,7 @@ FString UMovieGraphCommandLineEncoderNode::GetResolvedOutputFilename(const FMovi
 	ResolveParams.Job = CachedPipeline->GetCurrentJob();
 	ResolveParams.Shot = Shot.Get();
 	ResolveParams.FileNameFormatOverrides = FormatOverrides;
+	ResolveParams.FileNameOverride = FileNameFormat;
 	ResolveParams.EvaluatedConfig = CachedPipeline->GetTimeStepInstance()->GetCalculatedTimeData().EvaluatedConfig;
 	ResolveParams.RenderDataIdentifier = RenderIdentifier;
 	ResolveParams.Version = Shot.IsValid() ? Shot->ShotInfo.VersionNumber : UMovieGraphBlueprintLibrary::ResolveVersionNumber(ResolveParams);
@@ -388,7 +387,7 @@ FString UMovieGraphCommandLineEncoderNode::GetResolvedOutputFilename(const FMovi
 
 void UMovieGraphCommandLineEncoderNode::GenerateTemporaryEncoderInputFiles(const FEncoderParams& InParams, TArray<FString>& OutVideoInputFilePaths, TArray<FString>& OutAudioInputFilePaths) const
 {
-	UMovieGraphOutputSettingNode* OutputSettingNode = GetSettingOnBranch<UMovieGraphOutputSettingNode>();
+	UMovieGraphGlobalOutputSettingNode* OutputSettingNode = GetSettingOnBranch<UMovieGraphGlobalOutputSettingNode>();
 
 	// Generate a text file for each input type which lists the files for that input type. We generate a FGuid in case there are
 	// multiple encode jobs going at once.
diff --git a/Engine/Plugins/MovieScene/MovieRenderPipeline/Source/MovieRenderPipelineCore/Private/Graph/Nodes/MovieGraphGlobalOutputSettingNode.cpp b/Engine/Plugins/MovieScene/MovieRenderPipeline/Source/MovieRenderPipelineCore/Private/Graph/Nodes/MovieGraphGlobalOutputSettingNode.cpp
new file mode 100644
--- /dev/null
+++ b/Engine/Plugins/MovieScene/MovieRenderPipeline/Source/MovieRenderPipelineCore/Private/Graph/Nodes/MovieGraphGlobalOutputSettingNode.cpp
@@ -0,0 +1,109 @@
+// Copyright Epic Games, Inc. All Rights Reserved.
+
+#include "Graph/Nodes/MovieGraphGlobalOutputSettingNode.h"
+#include "Graph/MovieGraphProjectSettings.h"
+#include "Graph/MovieGraphBlueprintLibrary.h"
+#include "Styling/AppStyle.h"
+#include "Algo/Find.h"
+
+UMovieGraphGlobalOutputSettingNode::UMovieGraphGlobalOutputSettingNode()
+	: OutputFrameRate(FFrameRate(24, 1))
+	, bOverwriteExistingOutput(true)
+	, ZeroPadFrameNumbers(4)
+	, FrameNumberOffset(0)
+{
+	// We prefer a 1080p resolution by default, but users may not have a preset that matches that. So we'll 
+	// look for a matching resolution if we can find one, otherwise we go to Custom set to 1920x1080.
+	const FMovieGraphNamedResolution* FoundResolution = nullptr;
+	if (const UMovieGraphProjectSettings* MovieGraphProjectSettings =
+		GetDefault<UMovieGraphProjectSettings>())
+	{
+		FoundResolution = Algo::FindByPredicate(
+			MovieGraphProjectSettings->DefaultNamedResolutions,
+			[](const FMovieGraphNamedResolution& Other)
+			{
+				return Other.Resolution == FIntPoint(1920, 1080);
+			});
+	}
+
+	// If we found one that was 1080p, regardless of name, use it.
+	if (FoundResolution)
+	{
+		OutputResolution = *FoundResolution;
+	}
+	else
+	{
+		// We didn't find one that was 1080p, just force a custom resolution.
+		OutputResolution = UMovieGraphBlueprintLibrary::NamedResolutionFromSize(1920, 1080);
+	}
+	
+	OutputDirectory.Path = TEXT("{project_dir}/Saved/MovieRenders/");
+}
+
+void UMovieGraphGlobalOutputSettingNode::GetFormatResolveArgs(FMovieGraphResolveArgs& OutMergedFormatArgs, const FMovieGraphRenderDataIdentifier& InRenderDataIdentifier) const
+{
+	const FString ResolvedProjectDir = FPaths::ConvertRelativePathToFull(FPaths::ProjectDir());
+	OutMergedFormatArgs.FilenameArguments.Add(TEXT("project_dir"), ResolvedProjectDir);
+	OutMergedFormatArgs.FileMetadata.Add(TEXT("unreal/project_dir"), ResolvedProjectDir);
+
+	FIntPoint OutputResolutionAsIntPoint = GetSyncedOutputResolution();
+	
+	// Resolution Arguments
+	{
+		FString Resolution = FString::Printf(TEXT("%d_%d"), OutputResolutionAsIntPoint.X, OutputResolutionAsIntPoint.Y);
+		OutMergedFormatArgs.FilenameArguments.Add(TEXT("output_resolution"), Resolution);
+		OutMergedFormatArgs.FilenameArguments.Add(TEXT("output_width"), FString::FromInt(OutputResolutionAsIntPoint.X));
+		OutMergedFormatArgs.FilenameArguments.Add(TEXT("output_height"), FString::FromInt(OutputResolutionAsIntPoint.Y));
+	}
+
+	// We don't resolve the version here because that's handled on a per-file/shot basis
+}
+
+#if WITH_EDITOR
+FText UMovieGraphGlobalOutputSettingNode::GetNodeTitle(const bool bGetDescriptive) const
+{
+	static const FText OutputSettingsNodeName = NSLOCTEXT("MoviePipelineGraph", "NodeName_GlobalOutputSettings", "Global Output Settings");
+	return OutputSettingsNodeName;
+}
+
+FText UMovieGraphGlobalOutputSettingNode::GetMenuCategory() const 
+{
+	return NSLOCTEXT("MoviePipelineGraph", "Settings_Category", "Settings");
+}
+
+FLinearColor UMovieGraphGlobalOutputSettingNode::GetNodeTitleColor() const 
+{
+	static const FLinearColor OutputSettingsColor = FLinearColor(0.854f, 0.509f, 0.039f);
+	return OutputSettingsColor;
+}
+
+FSlateIcon UMovieGraphGlobalOutputSettingNode::GetIconAndTint(FLinearColor& OutColor) const 
+{
+	static const FSlateIcon SettingsIcon = FSlateIcon(FAppStyle::GetAppStyleSetName(), "Icons.Settings");
+
+	OutColor = FLinearColor::White;
+	return SettingsIcon;
+}
+
+EMovieGraphBranchRestriction UMovieGraphGlobalOutputSettingNode::GetBranchRestriction() const
+{
+	return EMovieGraphBranchRestriction::Globals;
+}
+#endif // WITH_EDITOR
+
+FIntPoint UMovieGraphGlobalOutputSettingNode::GetSyncedOutputResolution() const
+{
+	// Try to find a matching entry from Project Settings to stay in sync
+	const UMovieGraphProjectSettings* MovieGraphProjectSettings = GetDefault<UMovieGraphProjectSettings>();
+	if (ensureAlwaysMsgf(MovieGraphProjectSettings, TEXT("%hs: Failed to find UMovieGraphProjectSettings!"), __FUNCTION__))
+	{
+		const FMovieGraphNamedResolution* FoundResolution = MovieGraphProjectSettings->FindNamedResolutionForOption(OutputResolution.ProfileName);
+		if (FoundResolution)
+		{
+			return FoundResolution->Resolution;
+		}
+	}
+
+	// Otherwise return what we have saved locally
+	return OutputResolution.Resolution;
+}
diff --git a/Engine/Plugins/MovieScene/MovieRenderPipeline/Source/MovieRenderPipelineCore/Private/Graph/Nodes/MovieGraphOutputSettingNode.cpp b/Engine/Plugins/MovieScene/MovieRenderPipeline/Source/MovieRenderPipelineCore/Private/Graph/Nodes/MovieGraphOutputSettingNode.cpp
deleted file mode 100644
diff --git a/Engine/Plugins/MovieScene/MovieRenderPipeline/Source/MovieRenderPipelineCore/Private/Graph/Renderers/MovieGraphDeferredPass.cpp b/Engine/Plugins/MovieScene/MovieRenderPipeline/Source/MovieRenderPipelineCore/Private/Graph/Renderers/MovieGraphDeferredPass.cpp
--- a/Engine/Plugins/MovieScene/MovieRenderPipeline/Source/MovieRenderPipelineCore/Private/Graph/Renderers/MovieGraphDeferredPass.cpp
+++ b/Engine/Plugins/MovieScene/MovieRenderPipeline/Source/MovieRenderPipelineCore/Private/Graph/Renderers/MovieGraphDeferredPass.cpp
@@ -2,7 +2,6 @@
 
 #include "Graph/Renderers/MovieGraphDeferredPass.h"
 #include "Graph/Nodes/MovieGraphDeferredPassNode.h"
-#include "Graph/Nodes/MovieGraphOutputSettingNode.h"
 #include "Graph/Nodes/MovieGraphCameraNode.h"
 #include "Graph/MovieGraphDataTypes.h"
 #include "Graph/MovieGraphDefaultRenderer.h"
diff --git a/Engine/Plugins/MovieScene/MovieRenderPipeline/Source/MovieRenderPipelineCore/Public/Graph/MovieGraphBlueprintLibrary.h b/Engine/Plugins/MovieScene/MovieRenderPipeline/Source/MovieRenderPipelineCore/Public/Graph/MovieGraphBlueprintLibrary.h
--- a/Engine/Plugins/MovieScene/MovieRenderPipeline/Source/MovieRenderPipelineCore/Public/Graph/MovieGraphBlueprintLibrary.h
+++ b/Engine/Plugins/MovieScene/MovieRenderPipeline/Source/MovieRenderPipelineCore/Public/Graph/MovieGraphBlueprintLibrary.h
@@ -10,7 +10,7 @@
 #include "MovieGraphBlueprintLibrary.generated.h"
 
 // Forward Declare
-class UMovieGraphOutputSettingNode;
+class UMovieGraphGlobalOutputSettingNode;
 class UMovieGraphPipeline;
 
 UCLASS(meta = (ScriptName = "MovieGraphLibrary"))
@@ -28,7 +28,7 @@ public:
 	* @return					- The effective frame rate (taking into account the node's desire to override it). 
 	*/
 	UFUNCTION(BlueprintCallable, Category = "Movie Graph")
-	static FFrameRate GetEffectiveFrameRate(UMovieGraphOutputSettingNode* InNode, const FFrameRate& InDefaultRate);
+	static FFrameRate GetEffectiveFrameRate(UMovieGraphGlobalOutputSettingNode* InNode, const FFrameRate& InDefaultRate);
 
 	/**
 	* Takes a Movie Graph format string (in the form of {token}), a list of parameters (which normally come from the running UMovieGraphPipeline) and
diff --git a/Engine/Plugins/MovieScene/MovieRenderPipeline/Source/MovieRenderPipelineCore/Public/Graph/Nodes/MovieGraphAudioOutputNode.h b/Engine/Plugins/MovieScene/MovieRenderPipeline/Source/MovieRenderPipelineCore/Public/Graph/Nodes/MovieGraphAudioOutputNode.h
--- a/Engine/Plugins/MovieScene/MovieRenderPipeline/Source/MovieRenderPipelineCore/Public/Graph/Nodes/MovieGraphAudioOutputNode.h
+++ b/Engine/Plugins/MovieScene/MovieRenderPipeline/Source/MovieRenderPipelineCore/Public/Graph/Nodes/MovieGraphAudioOutputNode.h
@@ -1,4 +1,4 @@
-﻿// Copyright Epic Games, Inc. All Rights Reserved.
+// Copyright Epic Games, Inc. All Rights Reserved.
 
 #pragma once
 
@@ -42,13 +42,6 @@ public:
 	/** The extension of the audio files generated by this node. */
 	inline static const FString OutputExtension = TEXT("wav");
 
-public:
-	UPROPERTY(EditAnywhere, BlueprintReadWrite, Category = Overrides, meta = (InlineEditConditionToggle))
-	uint8 bOverride_FileNameFormatOverride : 1;
-	
-	/* File name format string override. If specified it will override the FileNameFormat from the Output Setting node. */
-	UPROPERTY(EditAnywhere, BlueprintReadWrite, Category = "File Output", meta=(EditCondition="bOverride_FileNameFormatOverride"))
-	FString FileNameFormatOverride;
 
 protected:
 	// UMovieGraphFileOutputNode Interface
diff --git a/Engine/Plugins/MovieScene/MovieRenderPipeline/Source/MovieRenderPipelineCore/Public/Graph/Nodes/MovieGraphCommandLineEncoderNode.h b/Engine/Plugins/MovieScene/MovieRenderPipeline/Source/MovieRenderPipelineCore/Public/Graph/Nodes/MovieGraphCommandLineEncoderNode.h
--- a/Engine/Plugins/MovieScene/MovieRenderPipeline/Source/MovieRenderPipelineCore/Public/Graph/Nodes/MovieGraphCommandLineEncoderNode.h
+++ b/Engine/Plugins/MovieScene/MovieRenderPipeline/Source/MovieRenderPipelineCore/Public/Graph/Nodes/MovieGraphCommandLineEncoderNode.h
@@ -41,7 +41,7 @@ public:
 
 public:
 	UPROPERTY(EditAnywhere, BlueprintReadWrite, Category = Overrides, meta = (InlineEditConditionToggle))
-	uint8 bOverride_FileNameFormatOverride : 1;
+	uint8 bOverride_FileNameFormat : 1;
 	
 	UPROPERTY(EditAnywhere, BlueprintReadWrite, Category = Overrides, meta = (InlineEditConditionToggle))
 	uint8 bOverride_Quality : 1;
@@ -82,8 +82,8 @@ public:
 	* Can be different from the main one in the Output setting so you can render out frames to individual
 	* shot folders but encode to one file.
 	*/
-	UPROPERTY(EditAnywhere, BlueprintReadWrite, Category = "General", meta=(EditCondition="bOverride_FileNameFormatOverride"))
-	FString FileNameFormatOverride;
+	UPROPERTY(EditAnywhere, BlueprintReadWrite, Category = "General", meta=(EditCondition="bOverride_FileNameFormat"))
+	FString FileNameFormat;
 
 	/** Whether the source files should be deleted on disk after encoding. */
 	UPROPERTY(EditAnywhere, BlueprintReadWrite, Category = "General", meta=(EditCondition="bOverride_bDeleteSourceFiles"))
diff --git a/Engine/Plugins/MovieScene/MovieRenderPipeline/Source/MovieRenderPipelineCore/Public/Graph/Nodes/MovieGraphFileOutputNode.h b/Engine/Plugins/MovieScene/MovieRenderPipeline/Source/MovieRenderPipelineCore/Public/Graph/Nodes/MovieGraphFileOutputNode.h
--- a/Engine/Plugins/MovieScene/MovieRenderPipeline/Source/MovieRenderPipelineCore/Public/Graph/Nodes/MovieGraphFileOutputNode.h
+++ b/Engine/Plugins/MovieScene/MovieRenderPipeline/Source/MovieRenderPipelineCore/Public/Graph/Nodes/MovieGraphFileOutputNode.h
@@ -51,6 +51,13 @@ public:
 		return NSLOCTEXT("MovieGraphNodes", "FileOutputGraphNode_Category", "Output Type");
 	}
 #endif
+	
+	UPROPERTY(EditAnywhere, BlueprintReadWrite, Category = Overrides, meta = (InlineEditConditionToggle))
+	uint8 bOverride_FileNameFormat : 1;
+
+	/** What format string should the final files use? Can include folder prefixes, and format string ({shot_name}, etc.) */
+	UPROPERTY(EditAnywhere, BlueprintReadWrite, Category = "File Output", meta=(EditCondition="bOverride_FileNameFormat"))
+	FString FileNameFormat;
 
 protected:
 	virtual void OnReceiveImageDataImpl(UMovieGraphPipeline* InPipeline, UE::MovieGraph::FMovieGraphOutputMergerFrame* InRawFrameData, const TSet<FMovieGraphRenderDataIdentifier>& InMask) {}
diff --git a/Engine/Plugins/MovieScene/MovieRenderPipeline/Source/MovieRenderPipelineCore/Public/Graph/Nodes/MovieGraphGlobalOutputSettingNode.h b/Engine/Plugins/MovieScene/MovieRenderPipeline/Source/MovieRenderPipelineCore/Public/Graph/Nodes/MovieGraphGlobalOutputSettingNode.h
new file mode 100644
--- /dev/null
+++ b/Engine/Plugins/MovieScene/MovieRenderPipeline/Source/MovieRenderPipelineCore/Public/Graph/Nodes/MovieGraphGlobalOutputSettingNode.h
@@ -0,0 +1,119 @@
+// Copyright Epic Games, Inc. All Rights Reserved.
+
+#pragma once
+
+#include "Graph/MovieGraphNamedResolution.h"
+#include "Graph/MovieGraphNode.h"
+
+#include "Misc/FrameRate.h"
+
+#include "MovieGraphGlobalOutputSettingNode.generated.h"
+
+USTRUCT(BlueprintType)
+struct FMovieGraphVersioningSettings
+{
+	GENERATED_BODY()
+
+	/**
+	 * If true, {version} tokens specified in the Output Directory and File Name Format properties will automatically
+	 * be incremented with each local render. If false, the version specified in Version Number will be used instead.
+	 *
+	 * Auto-versioning will search across all render branches and use the highest version found as the basis for the
+	 * next version used.
+	 */
+	UPROPERTY(EditAnywhere, BlueprintReadWrite, Category = "Versioning")
+	bool bAutoVersioning = true;
+	
+	/** The value to use for the version token if versions are not automatically incremented (Auto Version is off). */
+	UPROPERTY(EditAnywhere, BlueprintReadWrite, Category = "Versioning", meta = (UIMin = 1, UIMax = 50, ClampMin = 1))
+	int32 VersionNumber = 1;
+};
+
+UCLASS()
+class MOVIERENDERPIPELINECORE_API UMovieGraphGlobalOutputSettingNode : public UMovieGraphSettingNode
+{
+	GENERATED_BODY()
+public:
+	UMovieGraphGlobalOutputSettingNode();
+
+	// UMovieGraphSettingNode Interface
+	virtual void GetFormatResolveArgs(FMovieGraphResolveArgs& OutMergedFormatArgs, const FMovieGraphRenderDataIdentifier& InRenderDataIdentifier) const override;
+
+#if WITH_EDITOR
+	virtual FText GetNodeTitle(const bool bGetDescriptive) const override;
+	virtual FText GetMenuCategory() const override;
+	virtual FLinearColor GetNodeTitleColor() const override;
+	virtual FSlateIcon GetIconAndTint(FLinearColor& OutColor) const override;
+	virtual EMovieGraphBranchRestriction GetBranchRestriction() const override;
+#endif
+	// ~UMovieGraphSettingNode Interface
+
+	FIntPoint GetSyncedOutputResolution() const;
+
+public:
+	UPROPERTY(EditAnywhere, BlueprintReadWrite, Category = Overrides, meta = (InlineEditConditionToggle))
+	uint8 bOverride_OutputDirectory : 1;
+
+	UPROPERTY(EditAnywhere, BlueprintReadWrite, Category = Overrides, meta = (InlineEditConditionToggle))
+	uint8 bOverride_OutputResolution : 1;
+	
+	UPROPERTY(EditAnywhere, BlueprintReadWrite, Category = Overrides, meta = (InlineEditConditionToggle))
+	uint8 bOverride_OutputFrameRate : 1;
+	
+	UPROPERTY(EditAnywhere, BlueprintReadWrite, Category = Overrides, meta = (InlineEditConditionToggle))
+	uint8 bOverride_bOverwriteExistingOutput : 1;
+
+	UPROPERTY(EditAnywhere, BlueprintReadWrite, Category = Overrides, meta = (InlineEditConditionToggle))
+	uint8 bOverride_ZeroPadFrameNumbers : 1;
+
+	UPROPERTY(EditAnywhere, BlueprintReadWrite, Category = Overrides, meta = (InlineEditConditionToggle))
+	uint8 bOverride_FrameNumberOffset : 1;
+
+	UPROPERTY(EditAnywhere, BlueprintReadWrite, Category = Overrides, meta = (InlineEditConditionToggle))
+	uint8 bOverride_HandleFrameCount : 1;
+
+	UPROPERTY(EditAnywhere, BlueprintReadWrite, Category = Overrides, meta = (InlineEditConditionToggle))
+	uint8 bOverride_VersioningSettings : 1;
+
+	/** What directory should all of our output files be relative to. */
+	UPROPERTY(EditAnywhere, BlueprintReadWrite, Category = "File Output", meta=(EditCondition="bOverride_OutputDirectory"))
+	FDirectoryPath OutputDirectory;
+
+	/** What resolution should our output files be exported at? */
+	UPROPERTY(EditAnywhere, BlueprintReadWrite, Category = "File Output", meta = (EditCondition = "bOverride_OutputResolution"))
+	FMovieGraphNamedResolution OutputResolution;
+	
+	/** What frame rate should the output files be exported at? This overrides the Display Rate of the target sequence. If not overwritten, uses the default Sequence Display Rate. */
+	UPROPERTY(EditAnywhere, BlueprintReadWrite, Category = "File Output", meta=(EditCondition="bOverride_OutputFrameRate"))
+	FFrameRate OutputFrameRate;
+	
+	/** If true, output containers should attempt to override any existing files with the same name. */
+	UPROPERTY(EditAnywhere, BlueprintReadWrite, Category = "File Output", meta = (EditCondition = "bOverride_bOverwriteExistingOutput"))
+	bool bOverwriteExistingOutput;
+
+	/** How many digits should all output frame numbers be padded to? MySequence_1.png -> MySequence_0001.png. Useful for software that struggles to recognize frame ranges when non-padded. */
+	UPROPERTY(EditAnywhere, BlueprintReadWrite, Category = "File Output", meta = (EditCondition = "bOverride_ZeroPadFrameNumbers", UIMin = "1", UIMax = "5", ClampMin = "1"))
+	int32 ZeroPadFrameNumbers;
+
+	/**
+	* How many frames should we offset the output frame number by? This is useful when using handle frames on Sequences that start at frame 0,
+	* as the output would start in negative numbers. This can be used to offset by a fixed amount to ensure there's no negative numbers.
+	*/
+	UPROPERTY(EditAnywhere, BlueprintReadWrite,  Category = "File Output", meta = (EditCondition = "bOverride_FrameNumberOffset"))
+	int32 FrameNumberOffset;
+
+	/**
+	* Top level shot track sections will automatically be expanded by this many frames in both directions, and the resulting
+	* additional time will be rendered as part of that shot. The inner sequence should have sections long enough to cover
+	* this expanded range, otherwise these tracks will not evaluate during handle frames and may produce unexpected results.
+	* This can be used to generate handle frames for traditional non linear editing tools.
+	*/
+	UPROPERTY(EditAnywhere, BlueprintReadWrite, Category = "Frames", meta = (UIMin = 0, ClampMin = 0, EditCondition = "bOverride_HandleFrameCount"))
+	int32 HandleFrameCount;
+
+	/**
+	 * Determines how versioning should be handled (Auto Version, Version Number, etc.).
+	 */
+	UPROPERTY(EditAnywhere, BlueprintReadWrite, Category = "Versioning", meta = (EditCondition = "bOverride_VersioningSettings"))
+	FMovieGraphVersioningSettings VersioningSettings;
+};
diff --git a/Engine/Plugins/MovieScene/MovieRenderPipeline/Source/MovieRenderPipelineCore/Public/Graph/Nodes/MovieGraphOutputSettingNode.h b/Engine/Plugins/MovieScene/MovieRenderPipeline/Source/MovieRenderPipelineCore/Public/Graph/Nodes/MovieGraphOutputSettingNode.h
deleted file mode 100644
diff --git a/Engine/Plugins/MovieScene/MovieRenderPipeline/Source/MovieRenderPipelineEditor/Private/Customizations/Graph/MovieGraphVersioningSettingsCustomization.h b/Engine/Plugins/MovieScene/MovieRenderPipeline/Source/MovieRenderPipelineEditor/Private/Customizations/Graph/MovieGraphVersioningSettingsCustomization.h
--- a/Engine/Plugins/MovieScene/MovieRenderPipeline/Source/MovieRenderPipelineEditor/Private/Customizations/Graph/MovieGraphVersioningSettingsCustomization.h
+++ b/Engine/Plugins/MovieScene/MovieRenderPipeline/Source/MovieRenderPipelineEditor/Private/Customizations/Graph/MovieGraphVersioningSettingsCustomization.h
@@ -1,7 +1,7 @@
 ﻿// Copyright Epic Games, Inc. All Rights Reserved.
 #pragma once
 
-#include "Graph/Nodes/MovieGraphOutputSettingNode.h"
+#include "Graph/Nodes/MovieGraphGlobalOutputSettingNode.h"
 
 #include "DetailLayoutBuilder.h"
 #include "DetailWidgetRow.h"
diff --git a/Engine/Plugins/MovieScene/MovieRenderPipeline/Source/MovieRenderPipelineEditor/Private/Graph/MovieGraphAssetToolkit.cpp b/Engine/Plugins/MovieScene/MovieRenderPipeline/Source/MovieRenderPipelineEditor/Private/Graph/MovieGraphAssetToolkit.cpp
--- a/Engine/Plugins/MovieScene/MovieRenderPipeline/Source/MovieRenderPipelineEditor/Private/Graph/MovieGraphAssetToolkit.cpp
+++ b/Engine/Plugins/MovieScene/MovieRenderPipeline/Source/MovieRenderPipelineEditor/Private/Graph/MovieGraphAssetToolkit.cpp
@@ -15,7 +15,6 @@
 #include "Graph/MovieGraphConfig.h"
 #include "Graph/Nodes/MovieGraphCollectionNode.h"
 #include "Graph/Nodes/MovieGraphModifierNode.h"
-#include "Graph/Nodes/MovieGraphOutputSettingNode.h"
 #include "Graph/Nodes/MovieGraphSelectNode.h"
 #include "MovieEdGraphNode.h"
 #include "MovieGraphSchema.h"
diff --git a/Engine/Plugins/MovieScene/MovieRenderPipeline/Source/MovieRenderPipelineRenderPasses/Private/MovieGraphImageSequenceOutputNode.cpp b/Engine/Plugins/MovieScene/MovieRenderPipeline/Source/MovieRenderPipelineRenderPasses/Private/MovieGraphImageSequenceOutputNode.cpp
--- a/Engine/Plugins/MovieScene/MovieRenderPipeline/Source/MovieRenderPipelineRenderPasses/Private/MovieGraphImageSequenceOutputNode.cpp
+++ b/Engine/Plugins/MovieScene/MovieRenderPipeline/Source/MovieRenderPipelineRenderPasses/Private/MovieGraphImageSequenceOutputNode.cpp
@@ -2,7 +2,7 @@
 
 #include "MovieGraphImageSequenceOutputNode.h"
 
-#include "Graph/Nodes/MovieGraphOutputSettingNode.h"
+#include "Graph/Nodes/MovieGraphGlobalOutputSettingNode.h"
 #include "Graph/Nodes/MovieGraphRenderLayerNode.h"
 #include "Graph/MovieGraphDataTypes.h"
 #include "Graph/MovieGraphPipeline.h"
@@ -43,7 +43,7 @@ namespace UE::MovieGraph::Private
 		Params.RootFrameNumberRel = InTraversalContext.Time.OutputFrameNumber;
 		//Params.ShotFrameNumberRel = InTraversalContext.Time.ShotFrameNumberRel
 		//Params.FileMetadata = ToDo: Track File Metadata
-		const UMovieGraphOutputSettingNode* OutputSettingNode = InEvaluatedConfig->GetSettingForBranch<UMovieGraphOutputSettingNode>(InRenderId.RootBranchName);
+		const UMovieGraphGlobalOutputSettingNode* OutputSettingNode = InEvaluatedConfig->GetSettingForBranch<UMovieGraphGlobalOutputSettingNode>(UMovieGraphNode::GlobalsPinName);
 		if (IsValid(OutputSettingNode))
 		{
 			Params.ZeroPadFrameNumberCount = OutputSettingNode->ZeroPadFrameNumbers;
@@ -227,14 +227,14 @@ FString UMovieGraphImageSequenceOutputNode::CreateFileName(
 	case EImageFormat::EXR: Extension = TEXT("exr"); break;
 	}
 
-	UMovieGraphOutputSettingNode* OutputSettingNode = InRawFrameData->EvaluatedConfig->GetSettingForBranch<UMovieGraphOutputSettingNode>(InRenderData.Key.RootBranchName);
+	UMovieGraphGlobalOutputSettingNode* OutputSettingNode = InRawFrameData->EvaluatedConfig->GetSettingForBranch<UMovieGraphGlobalOutputSettingNode>(GlobalsPinName);
 	if (!OutputSettingNode)
 	{
 		return FString();
 	}
 
 	// Generate one string that puts the directory combined with the filename format.
-	FString FileNameFormatString = OutputSettingNode->OutputDirectory.Path / OutputSettingNode->FileNameFormat;
+	FString FileNameFormatString = OutputSettingNode->OutputDirectory.Path / FileNameFormat;
 
 	constexpr bool bIncludeRenderPass = false;
 	constexpr bool bTestFrameNumber = true;
@@ -742,14 +742,14 @@ FString UMovieGraphImageSequenceOutputNode_MultiLayerEXR::ResolveOutputFilename(
 	const TCHAR* Extension = TEXT("exr");
 
 	constexpr bool bIncludeCDOs = true;
-	const UMovieGraphOutputSettingNode* OutputSettings = InRawFrameData->EvaluatedConfig->GetSettingForBranch<UMovieGraphOutputSettingNode>(InBranchName, bIncludeCDOs);
+	const UMovieGraphGlobalOutputSettingNode* OutputSettings = InRawFrameData->EvaluatedConfig->GetSettingForBranch<UMovieGraphGlobalOutputSettingNode>(InBranchName, bIncludeCDOs);
 	if (!ensure(OutputSettings))
 	{
 		return FString();
 	}
 	
 	// If we have more than one resolution we'll store it as "_Add" / "_Add(1)" etc via {ExtraTag}.
-	FString FileNameFormatString = OutputSettings->FileNameFormat + "{ExtraTag}";
+	FString FileNameFormatString = FileNameFormat + "{ExtraTag}";
 
 	const FString FilePathFormatString = OutputSettings->OutputDirectory.Path / FileNameFormatString;
 	
diff --git a/Engine/Plugins/MovieScene/MovieRenderPipeline/Source/MovieRenderPipelineRenderPasses/Public/MovieGraphImageSequenceOutputNode.h b/Engine/Plugins/MovieScene/MovieRenderPipeline/Source/MovieRenderPipelineRenderPasses/Public/MovieGraphImageSequenceOutputNode.h
--- a/Engine/Plugins/MovieScene/MovieRenderPipeline/Source/MovieRenderPipelineRenderPasses/Public/MovieGraphImageSequenceOutputNode.h
+++ b/Engine/Plugins/MovieScene/MovieRenderPipeline/Source/MovieRenderPipelineRenderPasses/Public/MovieGraphImageSequenceOutputNode.h
@@ -97,6 +97,7 @@ public:
 	{
 		OutputFormat = EImageFormat::EXR;
 		bQuantizeTo8Bit = false;
+		FileNameFormat = TEXT("{sequence_name}.{frame_number}");
 		Compression = EEXRCompressionFormat::PIZ;
 	}
 
