MovieGraph:

-Fix: Double clicking a Sub-Graph node doesn't open the sub-graph asset

diff --git a/Engine/Plugins/MovieScene/MovieRenderPipeline/Source/MovieRenderPipelineEditor/Private/Graph/MovieGraphSchema.cpp b/Engine/Plugins/MovieScene/MovieRenderPipeline/Source/MovieRenderPipelineEditor/Private/Graph/MovieGraphSchema.cpp
--- a/Engine/Plugins/MovieScene/MovieRenderPipeline/Source/MovieRenderPipelineEditor/Private/Graph/MovieGraphSchema.cpp
+++ b/Engine/Plugins/MovieScene/MovieRenderPipeline/Source/MovieRenderPipelineEditor/Private/Graph/MovieGraphSchema.cpp
@@ -17,6 +17,7 @@
 #include "ScopedTransaction.h"
 #include "GraphEditor.h"
 #include "MovieEdGraphVariableNode.h"
+#include "MoviePipelineEdGraphSubgraphNode.h"
 
 TArray<UClass*> UMovieGraphSchema::MoviePipelineNodeClasses;
 
@@ -419,7 +420,18 @@ UEdGraphNode* FMovieGraphSchemaAction_NewNode::PerformAction(UEdGraph* ParentGra
 
 	// Now create the editor graph node
 	FGraphNodeCreator<UMoviePipelineEdGraphNode> NodeCreator(*ParentGraph);
-	UMoviePipelineEdGraphNode* GraphNode = NodeCreator.CreateUserInvokedNode(bSelectNewNode);
+
+	// Define the ed graph node type here if it differs from UMoviePipelineEdGraphNode
+	// If other ed node class types are needed here,
+	// we should let ed nodes declare their equivalent runtime node,
+	// and use that mapping to determine the applicable ed node type rather than hard-coding.
+	TSubclassOf<UMoviePipelineEdGraphNode> InvokableEdGraphNodeClass = UMoviePipelineEdGraphNode::StaticClass();
+	if (RuntimeNode->IsA(UMovieGraphSubgraphNode::StaticClass()))
+	{
+		InvokableEdGraphNodeClass = UMoviePipelineEdGraphSubgraphNode::StaticClass();
+	}
+	
+	UMoviePipelineEdGraphNode* GraphNode = NodeCreator.CreateUserInvokedNode(bSelectNewNode, InvokableEdGraphNodeClass);
 	GraphNode->Construct(RuntimeNode);
 	GraphNode->NodePosX = Location.X;
 	GraphNode->NodePosY = Location.Y;
