MoviePipeline: Updated component-based collection queries to look inside of blueprints.

#rnx

diff --git a/Engine/Plugins/MovieScene/MovieRenderPipeline/Source/MovieRenderPipelineCore/Private/Graph/MoviePipelineRenderLayerSubsystem.cpp b/Engine/Plugins/MovieScene/MovieRenderPipeline/Source/MovieRenderPipelineCore/Private/Graph/MoviePipelineRenderLayerSubsystem.cpp
--- a/Engine/Plugins/MovieScene/MovieRenderPipeline/Source/MovieRenderPipelineCore/Private/Graph/MoviePipelineRenderLayerSubsystem.cpp
+++ b/Engine/Plugins/MovieScene/MovieRenderPipeline/Source/MovieRenderPipelineCore/Private/Graph/MoviePipelineRenderLayerSubsystem.cpp
@@ -896,7 +896,9 @@ void UMovieGraphConditionGroupQuery_ComponentTagName::Evaluate(const TArray<AAct
 
 	for (AActor* Actor : InActorsToQuery)
 	{
-		Actor->GetComponents<UActorComponent*>(ActorComponents);
+		// Include child components so components inside of Blueprints can be found
+		constexpr bool bIncludeFromChildActors = false;
+		Actor->GetComponents<UActorComponent*>(ActorComponents, bIncludeFromChildActors);
 		
 		for (const UActorComponent* Component : ActorComponents)
 		{
@@ -974,7 +976,9 @@ void UMovieGraphConditionGroupQuery_ComponentType::Evaluate(const TArray<AActor*
 
 	for (AActor* Actor : InActorsToQuery)
 	{
-		Actor->GetComponents<UActorComponent*>(ActorComponents);
+		// Include child components so components inside of Blueprints can be found
+		constexpr bool bIncludeFromChildActors = false;
+		Actor->GetComponents<UActorComponent*>(ActorComponents, bIncludeFromChildActors);
 		
 		for (const UActorComponent* Component : ActorComponents)
 		{
