MRG: Removing editor-only constraints on linking the OpenColorIO library through the wrapper module.

Also removing the recently added transform getter, which is now unused.

diff --git a/Engine/Plugins/Compositing/OpenColorIO/Source/OpenColorIO/Private/OpenColorIOColorTransform.cpp b/Engine/Plugins/Compositing/OpenColorIO/Source/OpenColorIO/Private/OpenColorIOColorTransform.cpp
--- a/Engine/Plugins/Compositing/OpenColorIO/Source/OpenColorIO/Private/OpenColorIOColorTransform.cpp
+++ b/Engine/Plugins/Compositing/OpenColorIO/Source/OpenColorIO/Private/OpenColorIOColorTransform.cpp
@@ -566,11 +566,6 @@ bool UOpenColorIOColorTransform::IsTransform(const FString& InSourceColorSpace,
 }
 
 bool UOpenColorIOColorTransform::GetTransformProcessor(FOpenColorIOWrapperProcessor& OutProcessor) const
-{
-	return GetTransformProcessor(GetContextKeyValues(), OutProcessor);
-}
-
-bool UOpenColorIOColorTransform::GetTransformProcessor(const TMap<FString, FString>& InLocalContext, FOpenColorIOWrapperProcessor& OutProcessor) const
 {
 #if WITH_OCIO
 	UOpenColorIOConfiguration* ConfigurationOwner = Cast<UOpenColorIOConfiguration>(GetOuter());
@@ -586,11 +581,11 @@ bool UOpenColorIOColorTransform::GetTransformProcessor(const TMap<FString, FStri
 	EOpenColorIOViewTransformDirection CurrentDisplayViewDirection;
 	if (GetDisplayViewDirection(CurrentDisplayViewDirection))
 	{
-		OutProcessor = FOpenColorIOWrapperProcessor(ConfigWrapper, SourceColorSpace, Display, View, CurrentDisplayViewDirection == EOpenColorIOViewTransformDirection::Inverse, InLocalContext);
+		OutProcessor = FOpenColorIOWrapperProcessor(ConfigWrapper, SourceColorSpace, Display, View, CurrentDisplayViewDirection == EOpenColorIOViewTransformDirection::Inverse, GetContextKeyValues());
 	}
 	else
 	{
-		OutProcessor = FOpenColorIOWrapperProcessor(ConfigWrapper, SourceColorSpace, DestinationColorSpace, InLocalContext);
+		OutProcessor = FOpenColorIOWrapperProcessor(ConfigWrapper, SourceColorSpace, DestinationColorSpace, GetContextKeyValues());
 	}
 
 	return OutProcessor.IsValid();
diff --git a/Engine/Plugins/Compositing/OpenColorIO/Source/OpenColorIO/Public/OpenColorIOColorTransform.h b/Engine/Plugins/Compositing/OpenColorIO/Source/OpenColorIO/Public/OpenColorIOColorTransform.h
--- a/Engine/Plugins/Compositing/OpenColorIO/Source/OpenColorIO/Public/OpenColorIOColorTransform.h
+++ b/Engine/Plugins/Compositing/OpenColorIO/Source/OpenColorIO/Public/OpenColorIOColorTransform.h
@@ -93,15 +93,6 @@ public:
 	 */
 	bool GetTransformProcessor(FOpenColorIOWrapperProcessor& OutProcessor) const;
 
-	/**
-	 * Get the transform processor.
-	 *
-	 * @param InLocalContext Context key-value string pairs.
-	 * @param OutProcessor Processor wrapper for the current transform, with a specified context.
-	 * @return True if the processor was created and is valid.
-	 */
-	bool GetTransformProcessor(const TMap<FString, FString>& InLocalContext, FOpenColorIOWrapperProcessor& OutProcessor) const;
-
 	/** Apply the color transform in-place to the specified color. */
 	bool TransformColor(FLinearColor& InOutColor) const;
 		
diff --git a/Engine/Plugins/MovieScene/MovieRenderPipeline/Source/MovieRenderPipelineRenderPasses/MovieRenderPipelineRenderPasses.Build.cs b/Engine/Plugins/MovieScene/MovieRenderPipeline/Source/MovieRenderPipelineRenderPasses/MovieRenderPipelineRenderPasses.Build.cs
--- a/Engine/Plugins/MovieScene/MovieRenderPipeline/Source/MovieRenderPipelineRenderPasses/MovieRenderPipelineRenderPasses.Build.cs
+++ b/Engine/Plugins/MovieScene/MovieRenderPipeline/Source/MovieRenderPipelineRenderPasses/MovieRenderPipelineRenderPasses.Build.cs
@@ -24,6 +24,8 @@ public class MovieRenderPipelineRenderPasses : ModuleRules
 				"MovieRenderPipelineSettings", // For settings
 				"ColorManagement",
 				"SlateCore",
+				"ImageCore",
+				"OpenColorIOWrapper",
 			}
 		);
 
@@ -40,15 +42,5 @@ public class MovieRenderPipelineRenderPasses : ModuleRules
 
 		// Required for UEOpenExr
 		AddEngineThirdPartyPrivateStaticDependencies(Target, "zlib");
-
-		if (Target.bBuildEditor == true)
-		{
-			PrivateDependencyModuleNames.AddRange(
-				new string[] {
-					"ImageCore",
-					"OpenColorIOWrapper"
-				}
-			);
-		}
 	}
 }
diff --git a/Engine/Plugins/MovieScene/MovieRenderPipeline/Source/MovieRenderPipelineRenderPasses/Private/MovieGraphImageSequenceOutputNode.cpp b/Engine/Plugins/MovieScene/MovieRenderPipeline/Source/MovieRenderPipelineRenderPasses/Private/MovieGraphImageSequenceOutputNode.cpp
--- a/Engine/Plugins/MovieScene/MovieRenderPipeline/Source/MovieRenderPipelineRenderPasses/Private/MovieGraphImageSequenceOutputNode.cpp
+++ b/Engine/Plugins/MovieScene/MovieRenderPipeline/Source/MovieRenderPipelineRenderPasses/Private/MovieGraphImageSequenceOutputNode.cpp
@@ -17,7 +17,7 @@
 #include "Misc/Paths.h"
 #include "Async/TaskGraphInterfaces.h"
 
-#if WITH_EDITOR
+#if WITH_OCIO
 #include "ImageCore.h" // For GetImageView()
 #include "OpenColorIOConfiguration.h"
 #include "OpenColorIOColorTransform.h"
@@ -70,7 +70,7 @@ namespace UE::MovieGraph::Private
 		return Params;
 	}
 	
-#if WITH_EDITOR
+#if WITH_OCIO
 	struct FOpenColorIOPixelPreProcessor
 	{
 		FOpenColorIOPixelPreProcessor(FOpenColorIOWrapperProcessor&& InProcessor)
@@ -123,7 +123,8 @@ namespace UE::MovieGraph::Private
 	}
 
 	/**
-	 * Convenience function to create an OpenColorIO CPU processor based on the specified conversion settings. Editor-only.
+	 * Convenience function to create an OpenColorIO CPU processor based on the specified conversion settings.
+	 * We use the OpenColorIO processor wrapper directly to avoid concurrency issues with the uobjects lifetime.
 	 *
 	 * @return The pixel preprocessor if successful, nullptr otherwise.
 	*/
@@ -132,26 +133,43 @@ namespace UE::MovieGraph::Private
 		const TObjectPtr<UOpenColorIOConfiguration>& ConfigurationSource = InConversionSettings.ConfigurationSource;
 		if (IsValid(ConfigurationSource))
 		{
+			const FOpenColorIOWrapperConfig* ConfigWrapper = ConfigurationSource->GetOrCreateConfigWrapper();
 			TObjectPtr<const UOpenColorIOColorTransform> ColorTransform = ConfigurationSource->FindTransform(InConversionSettings);
 			if (IsValid(ColorTransform))
 			{
-				//TODO: We need further processing of the context values to support dynamic tokens such as {shot_name}.
 				FOpenColorIOWrapperProcessor Processor;
-				if (ColorTransform->GetTransformProcessor(InContext, Processor))
+				EOpenColorIOViewTransformDirection CurrentDisplayViewDirection;
+
+				if (ColorTransform->GetDisplayViewDirection(CurrentDisplayViewDirection))
+				{
+					Processor = FOpenColorIOWrapperProcessor(
+							ConfigWrapper,
+							ColorTransform->SourceColorSpace,
+							ColorTransform->Display,
+							ColorTransform->View,
+							static_cast<bool>(CurrentDisplayViewDirection),
+							InContext
+						);
+				}
+				else
+				{
+					Processor = FOpenColorIOWrapperProcessor(
+							ConfigWrapper,
+							ColorTransform->SourceColorSpace,
+							ColorTransform->DestinationColorSpace,
+							InContext
+						);
+				}
+
+				if (Processor.IsValid())
 				{
 					return FOpenColorIOPixelPreProcessor(MoveTemp(Processor));
 				}
 			}
-			else
-			{
-				UE_LOG(LogMovieRenderPipeline, Warning, TEXT("Invalid conversion settings, bypassing OpenColorIO transform."));
-			}
-		}
-		else
-		{
-			UE_LOG(LogMovieRenderPipeline, Warning, TEXT("Invalid configuration source, bypassing OpenColorIO transform."));
 		}
 
+		UE_LOG(LogMovieRenderPipeline, Warning, TEXT("Invalid configuration source or conversion settings, bypassing OpenColorIO transform."));
+
 		return {};
 	}
 #endif
@@ -328,7 +346,7 @@ void UMovieGraphImageSequenceOutputNode::OnReceiveImageDataImpl(UMovieGraphPipel
 		UE::MovieGraph::FMovieGraphSampleState* Payload = RenderData.Value->GetPayload<UE::MovieGraph::FMovieGraphSampleState>();
 		
 		bool bQuantizationEncodeSRGB = true;
-#if WITH_EDITOR
+#if WITH_OCIO
 		if (ParentNode->OCIOConfiguration.bIsEnabled && Payload->bAllowOCIO)
 		{
 			TMap<FString, FString> ResolvedOCIOContext = UE::MovieGraph::Private::ResolveOpenColorIOContext(
@@ -443,7 +461,7 @@ void UMovieGraphImageSequenceOutputNode_EXR::UpdateTaskPerLayer(
 	const UE::MovieGraph::FMovieGraphSampleState* Payload = InImageData->GetPayload<UE::MovieGraph::FMovieGraphSampleState>();
 
 	bool bEnabledOCIO = false;
-#if WITH_EDITOR
+#if WITH_OCIO
 	if (InParentNode->OCIOConfiguration.bIsEnabled && Payload->bAllowOCIO)
 	{
 		FPixelPreProcessor OCIOPixelPreProcessor = UE::MovieGraph::Private::CreateOpenColorIOPixelPreProcessor(
@@ -470,7 +488,7 @@ void UMovieGraphImageSequenceOutputNode_EXR::UpdateTaskPerLayer(
 		InOutImageTask.Height = Resolution.Y;
 
 		InOutImageTask.OverscanPercentage = Payload->OverscanFraction;
-#if WITH_EDITOR
+#if WITH_OCIO
 		if (bEnabledOCIO)
 		{
 			UE::MoviePipeline::UpdateColorSpaceMetadata(InParentNode->OCIOConfiguration.ColorConfiguration, InOutImageTask);
@@ -545,7 +563,7 @@ void UMovieGraphImageSequenceOutputNode_EXR::OnReceiveImageDataImpl(UMovieGraphP
 		}
 
 		TMap<FString, FString> ResolvedOCIOContext = {};
-#if WITH_EDITOR
+#if WITH_OCIO
 		ResolvedOCIOContext = UE::MovieGraph::Private::ResolveOpenColorIOContext(
 			ParentNode->OCIOContext,
 			RenderData.Key,
@@ -644,7 +662,7 @@ void UMovieGraphImageSequenceOutputNode_MultiLayerEXR::OnReceiveImageDataImpl(UM
 			}
 
 			TMap<FString, FString> ResolvedOCIOContext = {};
-#if WITH_EDITOR
+#if WITH_OCIO
 			ResolvedOCIOContext = UE::MovieGraph::Private::ResolveOpenColorIOContext(
 				ParentNode->OCIOContext,
 				RenderID,
diff --git a/Engine/Plugins/MovieScene/MovieRenderPipeline/Source/MovieRenderPipelineRenderPasses/Public/MovieGraphImageSequenceOutputNode.h b/Engine/Plugins/MovieScene/MovieRenderPipeline/Source/MovieRenderPipelineRenderPasses/Public/MovieGraphImageSequenceOutputNode.h
--- a/Engine/Plugins/MovieScene/MovieRenderPipeline/Source/MovieRenderPipelineRenderPasses/Public/MovieGraphImageSequenceOutputNode.h
+++ b/Engine/Plugins/MovieScene/MovieRenderPipeline/Source/MovieRenderPipelineRenderPasses/Public/MovieGraphImageSequenceOutputNode.h
@@ -33,7 +33,6 @@ public:
 	virtual bool IsFinishedWritingToDiskImpl() const override;
 	// ~UMovieGraphFileOutputNode Interface
 
-#if WITH_EDITORONLY_DATA
 	UPROPERTY(EditAnywhere, BlueprintReadWrite, Category = Overrides, meta = (InlineEditConditionToggle))
 	uint8 bOverride_OCIOConfiguration : 1;
 
@@ -44,7 +43,7 @@ public:
 	* OCIO configuration/transform settings.
 	*
 	* Note: There are differences from the previous implementation in MRQ given that we are now doing CPU-side processing.
-	* 1) This feature only works in editor-mode, since the OpenColorIO library is currently unavailable in game builds.
+	* 1) This feature only works on desktop platforms when the OpenColorIO library is available.
 	* 2) Users are now responsible for setting the renderer output space to Final Color (HDR) in Linear Working Color Space (SCS_FinalColorHDR).
 	*/
 	UPROPERTY(EditAnywhere, BlueprintReadWrite, Category = "OCIO", DisplayName="OCIO Configuration", meta = (EditCondition = "bOverride_OCIOConfiguration"))
@@ -55,7 +54,6 @@ public:
 	*/
 	UPROPERTY(EditAnywhere, BlueprintReadWrite, Category = "OCIO", DisplayName = "OCIO Context", meta = (EditCondition = "bOverride_OCIOContext"))
 	TMap<FString, FString> OCIOContext;
-#endif
 
 protected:
 	/** Convenience function to get the list of active composite passes from render data. */
