MoviePipeline: A set of fixes that prevent and/or detect recursive issues within the graph.
* In a common case, interactively prevent the user from creating a circular subgraph reference.
* In more complex cases, detect circular subgraph references and warn the user about them in the UI (Active Render Settings) and via errors in the log (which will be displayed if the user does a render).
* Slightly updated the circular node reference logic to be graph-specific so it doesn't conflict with circular subgraph detection logic.
* Updated the output directory detection logic to not crash on circular subgraph references.

#rnx

diff --git a/Engine/Plugins/MovieScene/MovieRenderPipeline/Source/MovieRenderPipelineCore/Private/Graph/MovieGraphConfig.cpp b/Engine/Plugins/MovieScene/MovieRenderPipeline/Source/MovieRenderPipelineCore/Private/Graph/MovieGraphConfig.cpp
--- a/Engine/Plugins/MovieScene/MovieRenderPipeline/Source/MovieRenderPipelineCore/Private/Graph/MovieGraphConfig.cpp
+++ b/Engine/Plugins/MovieScene/MovieRenderPipeline/Source/MovieRenderPipelineCore/Private/Graph/MovieGraphConfig.cpp
@@ -1035,7 +1035,7 @@ void UMovieGraphConfig::GetAllContainedSubgraphs(TSet<UMovieGraphConfig*>& OutSu
 	}
 }
 
-void UMovieGraphConfig::RecurseUpGlobalsBranchToFindOutputDirectory(const UMovieGraphNode* InNode, FString& OutOutputDirectory) const
+void UMovieGraphConfig::RecurseUpGlobalsBranchToFindOutputDirectory(const UMovieGraphNode* InNode, FString& OutOutputDirectory, TArray<const UMovieGraphConfig*>& VisitedGraphStack) const
 {
 	// If there's no Node, no upstream pin or no downstream pin for whatever reason,
 	// there is no way to continue so we early out
@@ -1051,10 +1051,19 @@ void UMovieGraphConfig::RecurseUpGlobalsBranchToFindOutputDirectory(const UMovie
 
 	UMovieGraphNode* ConnectedNode = UpstreamGlobalsPin->Node;
 
+	VisitedGraphStack.Push(this);
+
 	// Overrides can be set within Subgraphs
 	if (const UMovieGraphSubgraphNode* SubgraphNode = Cast<UMovieGraphSubgraphNode>(ConnectedNode))
 	{
 		const UMovieGraphConfig* SubgraphConfig = SubgraphNode->GetSubgraphAsset();
+
+		// Stop recursing if circular references are found
+		if (VisitedGraphStack.Contains(SubgraphConfig))
+		{
+			return;
+		}
+		
 		if (SubgraphConfig && OutOutputDirectory.IsEmpty())
 		{
 			const UMovieGraphPin* SubgraphGlobalsPin =
@@ -1062,9 +1071,11 @@ void UMovieGraphConfig::RecurseUpGlobalsBranchToFindOutputDirectory(const UMovie
 
 			if (SubgraphGlobalsPin && SubgraphGlobalsPin->IsConnected())
 			{
-				SubgraphConfig->RecurseUpGlobalsBranchToFindOutputDirectory(SubgraphGlobalsPin->Node, OutOutputDirectory);
+				SubgraphConfig->RecurseUpGlobalsBranchToFindOutputDirectory(SubgraphGlobalsPin->Node, OutOutputDirectory, VisitedGraphStack);
 			}
 		}
+
+		VisitedGraphStack.Pop();
 	}
 	else if (const UMovieGraphGlobalOutputSettingNode* SettingsNode = Cast<UMovieGraphGlobalOutputSettingNode>(InNode))
 	{
@@ -1077,7 +1088,7 @@ void UMovieGraphConfig::RecurseUpGlobalsBranchToFindOutputDirectory(const UMovie
 	// Keep looking upstream if we haven't found any overrides
 	if (OutOutputDirectory.IsEmpty())
 	{
-		RecurseUpGlobalsBranchToFindOutputDirectory(UpstreamGlobalsPin->Node, OutOutputDirectory);
+		RecurseUpGlobalsBranchToFindOutputDirectory(UpstreamGlobalsPin->Node, OutOutputDirectory, VisitedGraphStack);
 	}
 };
 
@@ -1093,7 +1104,8 @@ void UMovieGraphConfig::GetOutputDirectory(FString& OutOutputDirectory) const
 
 	if (GlobalsPin && GlobalsPin->IsConnected())
 	{
-		RecurseUpGlobalsBranchToFindOutputDirectory(OutputNode, OutOutputDirectory);
+		TArray<const UMovieGraphConfig*> VisitedGraphStack;
+		RecurseUpGlobalsBranchToFindOutputDirectory(OutputNode, OutOutputDirectory, VisitedGraphStack);
 
 		if (OutOutputDirectory.IsEmpty())
 		{
@@ -1271,30 +1283,37 @@ void UMovieGraphConfig::CopyOverriddenProperties(UMovieGraphNode* FromNode, UMov
 	}
 }
 
-void UMovieGraphConfig::CreateFlattenedGraph_Recursive(UMovieGraphEvaluatedConfig* InOwningConfig, FMovieGraphEvaluatedBranchConfig& OutBranchConfig,
+bool UMovieGraphConfig::CreateFlattenedGraph_Recursive(UMovieGraphEvaluatedConfig* InOwningConfig, FMovieGraphEvaluatedBranchConfig& OutBranchConfig,
 	FMovieGraphEvaluationContext& InEvaluationContext, UMovieGraphPin* InPinToFollow)
 {
 	if (!InPinToFollow)
 	{
-		return;
+		InEvaluationContext.TraversalError = LOCTEXT("TraversalError_InvalidPin", "Found an invalid pin during graph traversal.");
+		return false;
 	}
 
 	UMovieGraphNode* Node = InPinToFollow->Node;
 	if (!Node)
 	{
-		return;
+		InEvaluationContext.TraversalError = LOCTEXT("TraversalError_InvalidNode", "Found an invalid node during graph traversal.");
+		return false;
 	}
 
 	// We only follow execution pins during traversal.
 	if (!ensureMsgf(InPinToFollow->Properties.bIsBranch, TEXT("Only Branch pins should be contained by InPinToFollow!")))
 	{
-		return;
+		InEvaluationContext.TraversalError = LOCTEXT("TraversalError_NonBranchPin", "Attempting to follow a non-branch pin during graph traversal.");
+		return false;
 	}
 
 	InEvaluationContext.PinBeingFollowed = InPinToFollow;
 
-	// Check to see if our flattened evaluation graph already has a copy of this node.
-	InEvaluationContext.VisitedNodes.Add(Node);
+	// Add this node to the set of visited nodes so it can be checked for cycles later. Get the graph from GetTypedOuter() rather than "this" because
+	// this method will be called recursively, potentially on pins within subgraphs.
+	const UMovieGraphConfig* OwningGraph = InPinToFollow->GetTypedOuter<UMovieGraphConfig>();
+	ensure(OwningGraph);
+	TSet<TObjectPtr<UMovieGraphNode>>& VisitedNodeSet = InEvaluationContext.VisitedNodesByOwningGraph.FindOrAdd(OwningGraph).VisitedNodes;
+	VisitedNodeSet.Add(Node);
 	
 	const bool bShouldIncludeNode =
 		Node->IsA<UMovieGraphSettingNode>() &&
@@ -1337,26 +1356,60 @@ void UMovieGraphConfig::CreateFlattenedGraph_Recursive(UMovieGraphEvaluatedConfi
 	// Now that we've potentially resolved the values on this node, continue to travel up-stream along any execution pins,
 	// potentially following re-route nodes, sub-graph nodes, through branches, etc.
 	TArray<UMovieGraphPin*> NewPinsToFollow = Node->EvaluatePinsToFollow(InEvaluationContext);
+
+	// Immediately stop traversal if a circular subgraph reference was found. This is done after EvaluatePinsToFollow() because
+	// subgraph nodes will set bCircularGraphReferenceFound in EvaluatePinsToFollow().
+	if (InEvaluationContext.bCircularGraphReferenceFound)
+	{
+		// Generate a string illustrating the problematic subgraph stack
+		FString GraphCycleTraversalPath;
+		for (const TObjectPtr<const UMovieGraphSubgraphNode>& SubgraphNode : InEvaluationContext.SubgraphStack)
+		{
+			if (const UMovieGraphConfig* SubgraphAsset = SubgraphNode->GetSubgraphAsset())
+			{
+				GraphCycleTraversalPath += FString::Printf(TEXT("\n%s -> "), *SubgraphAsset->GetName());
+			}
+		}
+
+		InEvaluationContext.TraversalError = FText::Format(
+			LOCTEXT("TraversalError_CircularGraphReference", "Circular subgraph reference found during traversal.{0}"), FText::FromString(GraphCycleTraversalPath));
+		
+		return false;
+	}
 	
 	for (UMovieGraphPin* Pin : NewPinsToFollow)
 	{
 		for (UMovieGraphEdge* Edge : Pin->Edges)
 		{
-			if (UMovieGraphPin* OtherPin = Edge->GetOtherPin(Pin))
+			UMovieGraphPin* OtherPin = Edge->GetOtherPin(Pin);
+			if (!OtherPin)
 			{
-				UMovieGraphNode* OtherNode = OtherPin->Node;
+				continue;
+			}
+			
+			UMovieGraphNode* OtherNode = OtherPin->Node;
 
-				if (InEvaluationContext.VisitedNodes.Contains(OtherNode))
+			// Detect cycles within node connections
+			if (VisitedNodeSet.Contains(OtherNode))
+			{
+				// Generate a string illustrating the problematic node connections
+				FString NodeCycleTraversalPath;
+				for (const TObjectPtr<UMovieGraphNode>& VisitedNode : VisitedNodeSet)
 				{
-					// ToDo: This won't work long term if you have two different branches visiting the same node
-					// also we need to reset this every time we go start from the root.
-					FFrame::KismetExecutionMessage(
-						*FString::Printf(TEXT("%hs: Circular graph?"), __FUNCTION__),
-						ELogVerbosity::Warning);
-					continue;
+					NodeCycleTraversalPath += FString::Printf(TEXT("\n%s -> "), *VisitedNode->GetName());
 				}
+				
+				InEvaluationContext.TraversalError = FText::Format(
+					LOCTEXT("TraversalError_CircularNodeReference", "Node connection cycle found during traversal.{0}"), FText::FromString(NodeCycleTraversalPath));
+				
+				return false;
+			}
 
-				CreateFlattenedGraph_Recursive(InOwningConfig, OutBranchConfig, InEvaluationContext, OtherPin);
+			// If no cycle detected, continue following the pin
+			const bool bSuccess = CreateFlattenedGraph_Recursive(InOwningConfig, OutBranchConfig, InEvaluationContext, OtherPin);
+			if (!bSuccess)
+			{
+				return false;
 			}
 		}
 	}
@@ -1366,6 +1419,8 @@ void UMovieGraphConfig::CreateFlattenedGraph_Recursive(UMovieGraphEvaluatedConfi
 	{
 		InEvaluationContext.NodeTypesToRemoveStack.Pop();
 	}
+
+	return true;
 }
 
 void UMovieGraphConfig::VisitUpstreamNodes_Recursive(UMovieGraphNode* FromNode,	const FVisitNodesCallback& VisitCallback, TSet<UMovieGraphNode*>& VisitedNodes) const
@@ -1440,11 +1495,13 @@ void UMovieGraphConfig::VisitDownstreamNodes_Recursive(UMovieGraphNode* FromNode
 	}
 }
 
-UMovieGraphEvaluatedConfig* UMovieGraphConfig::CreateFlattenedGraph(const FMovieGraphTraversalContext& InContext)
+UMovieGraphEvaluatedConfig* UMovieGraphConfig::CreateFlattenedGraph(const FMovieGraphTraversalContext& InContext, FString& OutError)
 {
 	TRACE_CPUPROFILER_EVENT_SCOPE(MRQ_CreateFlattenedGraph);
 	LLM_SCOPE_BYNAME(TEXT("MovieGraph/CreateFlattenedGraph"));
 
+	OutError.Empty();
+
 	UMovieGraphEvaluatedConfig* NewContext = NewObject<UMovieGraphEvaluatedConfig>(this);
 
 	if (OutputNode)
@@ -1475,7 +1532,13 @@ UMovieGraphEvaluatedConfig* UMovieGraphConfig::CreateFlattenedGraph(const FMovie
 				UMovieGraphPin* OtherPin = Edge->GetOtherPin(InputPin);
 				if (OtherPin)
 				{
-					CreateFlattenedGraph_Recursive(NewContext, /*InOut*/ BranchConfig, StackContext, OtherPin);
+					const bool bTraversalSuccessful = CreateFlattenedGraph_Recursive(NewContext, /*InOut*/ BranchConfig, StackContext, OtherPin);
+					if (!bTraversalSuccessful)
+					{
+						UE_LOG(LogMovieRenderPipeline, Error, TEXT("%s"), *StackContext.TraversalError.ToString());
+						OutError = StackContext.TraversalError.ToString();
+						return nullptr;
+					}
 				}
 			}
 
@@ -1493,7 +1556,13 @@ UMovieGraphEvaluatedConfig* UMovieGraphConfig::CreateFlattenedGraph(const FMovie
 					UMovieGraphPin* OtherPin = Edge->GetOtherPin(GlobalsPin);
 					if (OtherPin)
 					{
-						CreateFlattenedGraph_Recursive(NewContext, /*InOut*/ BranchConfig, GlobalStackContext, OtherPin);
+						const bool bTraversalSuccessful = CreateFlattenedGraph_Recursive(NewContext, /*InOut*/ BranchConfig, GlobalStackContext, OtherPin);
+						if (!bTraversalSuccessful)
+						{
+							UE_LOG(LogMovieRenderPipeline, Error, TEXT("%s"), *StackContext.TraversalError.ToString());
+							OutError = StackContext.TraversalError.ToString();
+                            return nullptr;
+						}
 					}
 				}
 			}
diff --git a/Engine/Plugins/MovieScene/MovieRenderPipeline/Source/MovieRenderPipelineCore/Private/Graph/MovieGraphCoreTimeStep.cpp b/Engine/Plugins/MovieScene/MovieRenderPipeline/Source/MovieRenderPipelineCore/Private/Graph/MovieGraphCoreTimeStep.cpp
--- a/Engine/Plugins/MovieScene/MovieRenderPipeline/Source/MovieRenderPipelineCore/Private/Graph/MovieGraphCoreTimeStep.cpp
+++ b/Engine/Plugins/MovieScene/MovieRenderPipeline/Source/MovieRenderPipelineCore/Private/Graph/MovieGraphCoreTimeStep.cpp
@@ -67,7 +67,8 @@ void UMovieGraphCoreTimeStep::TickProducingFrames()
 		// Update global variables before evaluating the graph
 		Context.RootGraph->UpdateGlobalVariableValues(GetOwningGraph());
 
-		CurrentFrameData.EvaluatedConfig = TStrongObjectPtr<UMovieGraphEvaluatedConfig>(Context.RootGraph->CreateFlattenedGraph(Context));
+		FString OutError;
+		CurrentFrameData.EvaluatedConfig = TStrongObjectPtr<UMovieGraphEvaluatedConfig>(Context.RootGraph->CreateFlattenedGraph(Context, OutError));
 		CurrentFrameData.TemporalSampleCount = GetTemporalSampleCount();
 		UpdateFrameMetrics();
 
@@ -182,8 +183,9 @@ void UMovieGraphCoreTimeStep::TickProducingFrames()
 		FMovieGraphTraversalContext Context = GetOwningGraph()->GetCurrentTraversalContext();
 
 		// Update global variables before evaluating the graph
+		FString OutError;
 		Context.RootGraph->UpdateGlobalVariableValues(GetOwningGraph());
-		CurrentFrameData.EvaluatedConfig = TStrongObjectPtr<UMovieGraphEvaluatedConfig>(Context.RootGraph->CreateFlattenedGraph(Context));
+		CurrentFrameData.EvaluatedConfig = TStrongObjectPtr<UMovieGraphEvaluatedConfig>(Context.RootGraph->CreateFlattenedGraph(Context, OutError));
 
 		// The temporal sample count can change every frame due to graph evaluations, so when we're on our first temporal
 		// sub-sample of the new frame, we need to re-fetch the value. 
diff --git a/Engine/Plugins/MovieScene/MovieRenderPipeline/Source/MovieRenderPipelineCore/Private/Graph/MovieGraphPipeline.cpp b/Engine/Plugins/MovieScene/MovieRenderPipeline/Source/MovieRenderPipelineCore/Private/Graph/MovieGraphPipeline.cpp
--- a/Engine/Plugins/MovieScene/MovieRenderPipeline/Source/MovieRenderPipelineCore/Private/Graph/MovieGraphPipeline.cpp
+++ b/Engine/Plugins/MovieScene/MovieRenderPipeline/Source/MovieRenderPipelineCore/Private/Graph/MovieGraphPipeline.cpp
@@ -303,7 +303,16 @@ void UMovieGraphPipeline::BuildShotListFromDataSource()
 		FMovieGraphTimeStepData TimeContext;
 		CurrentContext.Time = TimeContext;
 
-		TObjectPtr<UMovieGraphEvaluatedConfig> EvaluatedConfig = CurrentContext.RootGraph->CreateFlattenedGraph(CurrentContext);
+		FString OutError;
+		TObjectPtr<UMovieGraphEvaluatedConfig> EvaluatedConfig = CurrentContext.RootGraph->CreateFlattenedGraph(CurrentContext, OutError);
+
+		// Shut down if there was an error when generating the evaluated graph
+		if (!OutError.IsEmpty())
+		{
+			constexpr bool bIsError = true;
+			Shutdown(bIsError);
+			return;
+		}
 		
 		// Create the time step instance for this shot. The time step method cannot vary per branch or per frame, so it
 		// is fetched from the Globals branch. This is the earliest point in the pipeline where the evaluated graph is
diff --git a/Engine/Plugins/MovieScene/MovieRenderPipeline/Source/MovieRenderPipelineCore/Private/Graph/Nodes/MovieGraphGlobalGameOverrides.cpp b/Engine/Plugins/MovieScene/MovieRenderPipeline/Source/MovieRenderPipelineCore/Private/Graph/Nodes/MovieGraphGlobalGameOverrides.cpp
--- a/Engine/Plugins/MovieScene/MovieRenderPipeline/Source/MovieRenderPipelineCore/Private/Graph/Nodes/MovieGraphGlobalGameOverrides.cpp
+++ b/Engine/Plugins/MovieScene/MovieRenderPipeline/Source/MovieRenderPipelineCore/Private/Graph/Nodes/MovieGraphGlobalGameOverrides.cpp
@@ -190,7 +190,8 @@ TSubclassOf<AGameModeBase> UMovieGraphGlobalGameOverridesNode::GetGameModeOverri
 		FMovieGraphTraversalContext TraversalContext;
 		TraversalContext.Job = const_cast<UMoviePipelineExecutorJob*>(InJob);
 		
-		if (const UMovieGraphEvaluatedConfig* EvaluatedGraph = GraphConfig->CreateFlattenedGraph(TraversalContext))
+		FString OutTraversalError;
+		if (const UMovieGraphEvaluatedConfig* EvaluatedGraph = GraphConfig->CreateFlattenedGraph(TraversalContext, OutTraversalError))
 		{
 			constexpr bool bIncludeCDOs = true;
 			constexpr bool bExactMatch = true;
diff --git a/Engine/Plugins/MovieScene/MovieRenderPipeline/Source/MovieRenderPipelineCore/Private/Graph/Nodes/MovieGraphSubgraphNode.cpp b/Engine/Plugins/MovieScene/MovieRenderPipeline/Source/MovieRenderPipelineCore/Private/Graph/Nodes/MovieGraphSubgraphNode.cpp
--- a/Engine/Plugins/MovieScene/MovieRenderPipeline/Source/MovieRenderPipelineCore/Private/Graph/Nodes/MovieGraphSubgraphNode.cpp
+++ b/Engine/Plugins/MovieScene/MovieRenderPipeline/Source/MovieRenderPipelineCore/Private/Graph/Nodes/MovieGraphSubgraphNode.cpp
@@ -2,9 +2,15 @@
 
 #include "Graph/Nodes/MovieGraphSubgraphNode.h"
 
+#include "MovieRenderPipelineCoreModule.h"
 #include "UObject/ObjectSaveContext.h"
 #include "UObject/Package.h"
 
+#if WITH_EDITOR
+#include "Framework/Notifications/NotificationManager.h"
+#include "Widgets/Notifications/SNotificationList.h"
+#endif
+
 #define LOCTEXT_NAMESPACE "MovieGraphNode"
 
 namespace UE::MovieGraph::Private
@@ -91,6 +97,18 @@ TArray<UMovieGraphPin*> UMovieGraphSubgraphNode::EvaluatePinsToFollow(FMovieGrap
 		return PinsToFollow;
 	}
 
+	// If this subgraph has already been visited, throw a circular graph reference error
+	if (InContext.SubgraphStack.Contains(this))
+	{
+		InContext.bCircularGraphReferenceFound = true;
+
+		// Normally there shouldn't be two identical graphs in the stack (cycle!), but for the purposes of error reporting, include this graph
+		// as the last in the stack so the cycle is clear
+		InContext.SubgraphStack.Add(this);
+		
+		return PinsToFollow;
+	}
+
 	// Find the input pin on the subgraph asset's outputs node that corresponds to the pin being followed
 	//
 	// Subgraph
@@ -126,6 +144,27 @@ void UMovieGraphSubgraphNode::PostEditChangeProperty(FPropertyChangedEvent& Prop
 {
 	Super::PostEditChangeProperty(PropertyChangedEvent);
 
+	// Prevent circular subgraph references. This won't catch everything, but it will at least catch the common case interactively. Other cases will
+	// be caught when the graph is fully evaluated.
+	if (GetTypedOuter<UMovieGraphConfig>() == SubgraphAsset)
+	{
+		SubgraphAsset = nullptr;
+
+#if WITH_EDITOR
+		FNotificationInfo Info(LOCTEXT("CircularGraphAssignmentWarning", "Assigning the subgraph to the provided asset would cause a circular reference."));
+		Info.ExpireDuration = 3.0f;
+		Info.bUseLargeFont = false;
+		Info.Image = FCoreStyle::Get().GetBrush(TEXT("MessageLog.Warning"));
+		const TSharedPtr<SNotificationItem> Notification = FSlateNotificationManager::Get().AddNotification(Info);
+		if (Notification.IsValid())
+		{
+			Notification->SetCompletionState(SNotificationItem::CS_None);
+		}
+#endif
+
+		UE_LOG(LogMovieRenderPipeline, Warning, TEXT("Found a circular reference when assigning the asset to the subgraph. Reverting."));
+	}
+
 	// Update pins when the subgraph asset is changed
 	if (PropertyChangedEvent.GetPropertyName() == GET_MEMBER_NAME_CHECKED(UMovieGraphSubgraphNode, SubgraphAsset))
 	{
diff --git a/Engine/Plugins/MovieScene/MovieRenderPipeline/Source/MovieRenderPipelineCore/Public/Graph/MovieGraphConfig.h b/Engine/Plugins/MovieScene/MovieRenderPipeline/Source/MovieRenderPipelineCore/Public/Graph/MovieGraphConfig.h
--- a/Engine/Plugins/MovieScene/MovieRenderPipeline/Source/MovieRenderPipelineCore/Public/Graph/MovieGraphConfig.h
+++ b/Engine/Plugins/MovieScene/MovieRenderPipeline/Source/MovieRenderPipelineCore/Public/Graph/MovieGraphConfig.h
@@ -354,6 +354,18 @@ private:
 	TMap<FString, FMovieGraphEvaluatedSettingsStack> NamedNodes;
 };
 
+// Note: This struct exists purely as a workaround for UHT throwing an error when putting a TSet in a TArray.
+/** Information on visited nodes found during traversal. */
+USTRUCT()
+struct FMovieGraphEvaluationContext_VisitedNodeInfo
+{
+	GENERATED_BODY()
+	
+	/** The nodes that were visited during traversal. */
+	UPROPERTY()
+	TSet<TObjectPtr<UMovieGraphNode>> VisitedNodes;
+};
+
 /**
 * This stores short-term information needed during traversal of the graph
 * such as disabled nodes, already visited nodes, etc. This information is
@@ -373,10 +385,10 @@ public:
 	FMovieGraphTraversalContext UserContext;
 
 	/**
-	* A list of nodes that have been visited. Used for cycle detection right now.
+	* A list of nodes that have been visited, where the key is the graph where the node was found. Used for cycle detection right now.
 	*/
 	UPROPERTY()
-	TSet<TObjectPtr<UMovieGraphNode>> VisitedNodes;
+	TMap<const UMovieGraphConfig*, FMovieGraphEvaluationContext_VisitedNodeInfo> VisitedNodesByOwningGraph;
 
 	/**
 	* The pin that is currently being followed in the traversal process.
@@ -391,6 +403,18 @@ public:
 	UPROPERTY()
 	TArray<TObjectPtr<const UMovieGraphSubgraphNode>> SubgraphStack;
 
+	/**
+	 * Whether a circular graph reference was found during traversal.
+	 */
+	UPROPERTY()
+	bool bCircularGraphReferenceFound = false;
+
+	/*
+	 * The error that was generated during traversal. A non-empty string implies that the traversal did not complete successfully.
+	 */
+	UPROPERTY()
+	FText TraversalError;
+
 	/**
 	 * The stack of node types (exact match) that should be removed from the graph while it is being traversed. Each node
 	 * which specifies a type to removed adds to the stack.
@@ -603,9 +627,12 @@ public:
 	void SetEditorOnlyNodes(const TArray<TObjectPtr<const UObject>>& InNodes);
 #endif
 
-	/** Given a user-defined evaluation context, evaluate the graph and build a "flattened" list of settings for each branch discovered. */
+	/**
+	 * Given a user-defined evaluation context, evaluate the graph and build a "flattened" list of settings for each branch discovered.
+	 * If there was an error while evaluating the graph, nullptr will be returned and OutError will be populated with a description of the problem.
+	 */
 	UFUNCTION(BlueprintCallable, Category="Experimental")
-	UMovieGraphEvaluatedConfig* CreateFlattenedGraph(const FMovieGraphTraversalContext& InContext);
+	UMovieGraphEvaluatedConfig* CreateFlattenedGraph(const FMovieGraphTraversalContext& InContext, FString& OutError);
 
 	/** Given a class and FProperty that belongs to that class, search for a FBoolProperty that matches the name "bOverride_<name of InRealProperty>. */
 	static FBoolProperty* FindOverridePropertyForRealProperty(UClass* InClass, const FProperty* InRealProperty);
@@ -644,9 +671,8 @@ public:
 	void GetOutputDirectory(FString& OutOutputDirectory) const;
 
 protected:
-
 	/** Look for the output directory in the UMovieGraphOutputSettings nodes found upstream of InNode. */
-	void RecurseUpGlobalsBranchToFindOutputDirectory(const UMovieGraphNode* InNode, FString& OutOutputDirectory) const;
+	void RecurseUpGlobalsBranchToFindOutputDirectory(const UMovieGraphNode* InNode, FString& OutOutputDirectory, TArray<const UMovieGraphConfig*>& VisitedGraphStack) const;
 	
 	/** Copies properties in FromNode that are marked for override into ToNode, but only if ToNode doesn't already override that value. */
 	void CopyOverriddenProperties(UMovieGraphNode* FromNode, UMovieGraphNode* ToNode, const FMovieGraphTraversalContext* InContext);
@@ -654,8 +680,11 @@ protected:
 	/** Find all "Overrideable" marked properties, then find their edit condition properties, then set those to false. */
 	void InitializeFlattenedNode(UMovieGraphNode* InNode);
 
-	/** Traverse the graph, generating a combined "flatten" graph as it goes. */
-	void CreateFlattenedGraph_Recursive(UMovieGraphEvaluatedConfig* InOwningConfig, FMovieGraphEvaluatedBranchConfig& OutBranchConfig,
+	/**
+	 * Traverse the graph, generating a combined "flatten" graph as it goes. Returns false if there was an issue (and the evaluation context will be
+	 * updated with more details regarding the failure).
+	 */
+	bool CreateFlattenedGraph_Recursive(UMovieGraphEvaluatedConfig* InOwningConfig, FMovieGraphEvaluatedBranchConfig& OutBranchConfig,
 		FMovieGraphEvaluationContext& InEvaluationContext, UMovieGraphPin* InPinToFollow);
 
 	/** Recursive helper for VisitUpstreamNodes(). */
diff --git a/Engine/Plugins/MovieScene/MovieRenderPipeline/Source/MovieRenderPipelineEditor/Private/Graph/SMovieGraphActiveRenderSettingsTabContent.cpp b/Engine/Plugins/MovieScene/MovieRenderPipeline/Source/MovieRenderPipelineEditor/Private/Graph/SMovieGraphActiveRenderSettingsTabContent.cpp
--- a/Engine/Plugins/MovieScene/MovieRenderPipeline/Source/MovieRenderPipelineEditor/Private/Graph/SMovieGraphActiveRenderSettingsTabContent.cpp
+++ b/Engine/Plugins/MovieScene/MovieRenderPipeline/Source/MovieRenderPipelineEditor/Private/Graph/SMovieGraphActiveRenderSettingsTabContent.cpp
@@ -360,9 +360,44 @@ void SMovieGraphActiveRenderSettingsTabContent::Construct(const FArguments& InAr
 			]
 		]
 
+		+ SVerticalBox::Slot()
+		.Padding(0, 10, 10, 0)
+		.AutoHeight()
+		[
+			SNew(SHorizontalBox)
+			.Visibility_Lambda([this]()
+			{
+				return TraversalError.IsEmpty() ? EVisibility::Collapsed : EVisibility::Visible;
+			})
+
+			+ SHorizontalBox::Slot()
+			.VAlign(VAlign_Center)
+			.Padding(10)
+			.AutoWidth()
+			[
+				SNew(SImage)
+				.Image(FAppStyle::GetBrush("Icons.ErrorWithColor"))
+			]
+
+			+ SHorizontalBox::Slot()
+			.VAlign(VAlign_Center)
+			[
+				SNew(STextBlock)
+				.Font(FAppStyle::Get().GetFontStyle("NormalFontBold"))
+				.Text_Lambda([this]()
+				{
+					return FText::FromString(TraversalError);
+				})
+			]
+		]
+
 		+ SVerticalBox::Slot()
 		[
 			SAssignNew(TreeView, STreeView<TSharedPtr<FActiveRenderSettingsTreeElement>>)
+			.Visibility_Lambda([this]()
+			{
+				return TraversalError.IsEmpty() ? EVisibility::Visible : EVisibility::Collapsed;
+			})
 			.ItemHeight(28)
 			.TreeItemsSource(&RootElements)
 			.SelectionMode(ESelectionMode::Single)
@@ -397,6 +432,12 @@ void SMovieGraphActiveRenderSettingsTabContent::Construct(const FArguments& InAr
 			.Font(FCoreStyle::GetDefaultFontStyle("Italic", 10))
 			.Visibility_Lambda([this]()
 			{
+				// If there's a traversal error, only show the error and not this warning
+				if (!TraversalError.IsEmpty())
+				{
+					return EVisibility::Collapsed;					
+				}
+				
 				return FlattenedGraph.IsValid() ? EVisibility::Collapsed : EVisibility::Visible;
 			})
 		]
@@ -420,7 +461,13 @@ void SMovieGraphActiveRenderSettingsTabContent::TraverseGraph()
 	Context.Job = TraversalJob.Get();
 
 	// Traverse the graph, and update the root elements
-	FlattenedGraph = TStrongObjectPtr(CurrentGraph->CreateFlattenedGraph(Context));
+	FlattenedGraph = TStrongObjectPtr(CurrentGraph->CreateFlattenedGraph(Context, TraversalError));
+	if (!FlattenedGraph.IsValid())
+	{
+		// TraversalError was set, which will be picked up by the UI and displayed instead of the tree
+		return;
+	}
+	
 	for (const TSharedPtr<FActiveRenderSettingsTreeElement>& RootElement : RootElements)
 	{
 		RootElement->FlattenedGraph = MakeWeakObjectPtr(FlattenedGraph.Get());
diff --git a/Engine/Plugins/MovieScene/MovieRenderPipeline/Source/MovieRenderPipelineEditor/Private/Graph/SMovieGraphActiveRenderSettingsTabContent.h b/Engine/Plugins/MovieScene/MovieRenderPipeline/Source/MovieRenderPipelineEditor/Private/Graph/SMovieGraphActiveRenderSettingsTabContent.h
--- a/Engine/Plugins/MovieScene/MovieRenderPipeline/Source/MovieRenderPipelineEditor/Private/Graph/SMovieGraphActiveRenderSettingsTabContent.h
+++ b/Engine/Plugins/MovieScene/MovieRenderPipeline/Source/MovieRenderPipelineEditor/Private/Graph/SMovieGraphActiveRenderSettingsTabContent.h
@@ -181,4 +181,7 @@ private:
 
 	/** The jobs that are available to be selected in the evaluation context settings. */
 	TArray<TSharedPtr<FString>> AvailableJobsInQueue;
+
+	/** An error that was set during traversal. This will show up in the UI and prevent the tree from being shown. */
+	FString TraversalError;
 };
diff --git a/Engine/Plugins/MovieScene/MovieRenderPipeline/Source/MovieRenderPipelineEditor/Private/MoviePipelineEditorBlueprintLibrary.cpp b/Engine/Plugins/MovieScene/MovieRenderPipeline/Source/MovieRenderPipelineEditor/Private/MoviePipelineEditorBlueprintLibrary.cpp
--- a/Engine/Plugins/MovieScene/MovieRenderPipeline/Source/MovieRenderPipelineEditor/Private/MoviePipelineEditorBlueprintLibrary.cpp
+++ b/Engine/Plugins/MovieScene/MovieRenderPipeline/Source/MovieRenderPipelineEditor/Private/MoviePipelineEditorBlueprintLibrary.cpp
@@ -262,9 +262,10 @@ FString UMoviePipelineEditorBlueprintLibrary::ResolveOutputDirectoryFromJob(UMov
 		Params.Version = UMovieGraphBlueprintLibrary::ResolveVersionNumber(Params, bGetNextVersion);
 		Params.RenderDataIdentifier.RootBranchName = UMovieGraphNode::GlobalsPinName;
 		
+		FString OutTraversalError;
 		FMovieGraphTraversalContext TraversalContext;
 		TraversalContext.Job = InJob;
-		Params.EvaluatedConfig = InJob->GetGraphPreset()->CreateFlattenedGraph(TraversalContext);
+		Params.EvaluatedConfig = InJob->GetGraphPreset()->CreateFlattenedGraph(TraversalContext, OutTraversalError);
 
 		FMovieGraphResolveArgs Dummy;
 		OutResolvedPath = UMovieGraphBlueprintLibrary::ResolveFilenameFormatArguments(FormatString, Params, Dummy);
diff --git a/Engine/Plugins/MovieScene/MovieRenderPipeline/Source/MovieRenderPipelineEditor/Private/MoviePipelineNewProcessExecutor.cpp b/Engine/Plugins/MovieScene/MovieRenderPipeline/Source/MovieRenderPipelineEditor/Private/MoviePipelineNewProcessExecutor.cpp
--- a/Engine/Plugins/MovieScene/MovieRenderPipeline/Source/MovieRenderPipelineEditor/Private/MoviePipelineNewProcessExecutor.cpp
+++ b/Engine/Plugins/MovieScene/MovieRenderPipeline/Source/MovieRenderPipelineEditor/Private/MoviePipelineNewProcessExecutor.cpp
@@ -106,7 +106,12 @@ void UMoviePipelineNewProcessExecutor::Execute_Implementation(UMoviePipelineQueu
 			FMovieGraphTraversalContext TraversalContext;
 			TraversalContext.Job = Job;
 
-			UMovieGraphEvaluatedConfig* EvaluatedGraph = GraphConfig->CreateFlattenedGraph(TraversalContext);
+			FString OutTraversalError;
+			UMovieGraphEvaluatedConfig* EvaluatedGraph = GraphConfig->CreateFlattenedGraph(TraversalContext, OutTraversalError);
+			if (!OutTraversalError.IsEmpty())
+			{
+				return;
+			}
 
 			// Only nodes on the Globals branch can apply command line args. Get CDOs here as well, so all nodes which define command line args
 			// will be found, even if they were not instantiated in the graph.
