Movie Graph:

-Fix: Output nodes (JPG, EXR, etc) do not have a default file name format anymore. For now, the default should be {sequence_name}.{render_layer}.{frame_number}
-Fix: ImageOutput nodes getting FileNameFormat from CDO instead of the instanced node

diff --git a/Engine/Plugins/MovieScene/MovieRenderPipeline/Source/MovieRenderPipelineCore/Private/Graph/MovieGraphBlueprintLibrary.cpp b/Engine/Plugins/MovieScene/MovieRenderPipeline/Source/MovieRenderPipelineCore/Private/Graph/MovieGraphBlueprintLibrary.cpp
--- a/Engine/Plugins/MovieScene/MovieRenderPipeline/Source/MovieRenderPipelineCore/Private/Graph/MovieGraphBlueprintLibrary.cpp
+++ b/Engine/Plugins/MovieScene/MovieRenderPipeline/Source/MovieRenderPipelineCore/Private/Graph/MovieGraphBlueprintLibrary.cpp
@@ -229,8 +229,8 @@ int32 UMovieGraphBlueprintLibrary::ResolveVersionNumber(FMovieGraphFilenameResol
 		return -1;
 	}
 
-	constexpr bool bIncludeCDOs = true;
-	constexpr bool bExactMatch = true;
+	bool bIncludeCDOs = true;
+	bool bExactMatch = true;
 	const UMovieGraphGlobalOutputSettingNode* OutputSettingNode =
 		InParams.EvaluatedConfig->GetSettingForBranch<UMovieGraphGlobalOutputSettingNode>(UMovieGraphNode::GlobalsPinName, bIncludeCDOs, bExactMatch);
 	if (!OutputSettingNode->VersioningSettings.bAutoVersioning)
@@ -242,6 +242,8 @@ int32 UMovieGraphBlueprintLibrary::ResolveVersionNumber(FMovieGraphFilenameResol
 	InParams.FileNameFormatOverrides.Add(TEXT("version"), TEXT("{version}"));
 
 	// Get output nodes from the evaluated graph
+	bIncludeCDOs = false;
+	bExactMatch = false;
 	TArray<UMovieGraphSettingNode*> ResultNodes = InParams.EvaluatedConfig->GetSettingsForBranch(
 		UMovieGraphCommandLineEncoderNode::StaticClass(), InParams.RenderDataIdentifier.RootBranchName, bIncludeCDOs, bExactMatch);
 	ResultNodes.Append(InParams.EvaluatedConfig->GetSettingsForBranch(
diff --git a/Engine/Plugins/MovieScene/MovieRenderPipeline/Source/MovieRenderPipelineCore/Private/Graph/Nodes/MovieGraphCommandLineEncoderNode.cpp b/Engine/Plugins/MovieScene/MovieRenderPipeline/Source/MovieRenderPipelineCore/Private/Graph/Nodes/MovieGraphCommandLineEncoderNode.cpp
--- a/Engine/Plugins/MovieScene/MovieRenderPipeline/Source/MovieRenderPipelineCore/Private/Graph/Nodes/MovieGraphCommandLineEncoderNode.cpp
+++ b/Engine/Plugins/MovieScene/MovieRenderPipeline/Source/MovieRenderPipelineCore/Private/Graph/Nodes/MovieGraphCommandLineEncoderNode.cpp
@@ -22,6 +22,8 @@ UMovieGraphCommandLineEncoderNode::UMovieGraphCommandLineEncoderNode()
 	, bSkipEncodeOnRenderCanceled(true)
 	, bRetainInputTextFiles(false)
 {
+	FileNameFormat = TEXT("{sequence_name}.{render_layer}");
+		
 	// Sensible defaults for cmdline arguments
 	CommandLineFormat = TEXT("-hide_banner -y -loglevel error {VideoInputs} {AudioInputs} -acodec {AudioCodec} -vcodec {VideoCodec} {Quality} \"{OutputPath}\"");
 	VideoInputStringFormat = TEXT("-f concat -safe 0 -i \"{InputFile}\" -r {FrameRate}");
diff --git a/Engine/Plugins/MovieScene/MovieRenderPipeline/Source/MovieRenderPipelineCore/Public/Graph/Nodes/MovieGraphFileOutputNode.h b/Engine/Plugins/MovieScene/MovieRenderPipeline/Source/MovieRenderPipelineCore/Public/Graph/Nodes/MovieGraphFileOutputNode.h
--- a/Engine/Plugins/MovieScene/MovieRenderPipeline/Source/MovieRenderPipelineCore/Public/Graph/Nodes/MovieGraphFileOutputNode.h
+++ b/Engine/Plugins/MovieScene/MovieRenderPipeline/Source/MovieRenderPipelineCore/Public/Graph/Nodes/MovieGraphFileOutputNode.h
@@ -57,7 +57,7 @@ public:
 
 	/** What format string should the final files use? Can include folder prefixes, and format string ({shot_name}, etc.) */
 	UPROPERTY(EditAnywhere, BlueprintReadWrite, Category = "File Output", meta=(EditCondition="bOverride_FileNameFormat"))
-	FString FileNameFormat;
+	FString FileNameFormat = "{sequence_name}.{render_layer}.{frame_number}";
 
 protected:
 	virtual void OnReceiveImageDataImpl(UMovieGraphPipeline* InPipeline, UE::MovieGraph::FMovieGraphOutputMergerFrame* InRawFrameData, const TSet<FMovieGraphRenderDataIdentifier>& InMask) {}
diff --git a/Engine/Plugins/MovieScene/MovieRenderPipeline/Source/MovieRenderPipelineRenderPasses/Private/MovieGraphImageSequenceOutputNode.cpp b/Engine/Plugins/MovieScene/MovieRenderPipeline/Source/MovieRenderPipelineRenderPasses/Private/MovieGraphImageSequenceOutputNode.cpp
--- a/Engine/Plugins/MovieScene/MovieRenderPipeline/Source/MovieRenderPipelineRenderPasses/Private/MovieGraphImageSequenceOutputNode.cpp
+++ b/Engine/Plugins/MovieScene/MovieRenderPipeline/Source/MovieRenderPipelineRenderPasses/Private/MovieGraphImageSequenceOutputNode.cpp
@@ -9,10 +9,12 @@
 #include "Graph/MovieGraphConfig.h"
 #include "Graph/MovieGraphFilenameResolveParams.h"
 #include "Graph/MovieGraphBlueprintLibrary.h"
-#include "Modules/ModuleManager.h"
 #include "MoviePipelineUtils.h"
 #include "MoviePipelineImageSequenceOutput.h" // for FAsyncImageQuantization
 #include "MovieRenderPipelineCoreModule.h"
+
+#include "Algo/Find.h"
+#include "Modules/ModuleManager.h"
 #include "ImageWriteQueue.h"
 #include "Misc/Paths.h"
 #include "Async/TaskGraphInterfaces.h"
@@ -251,8 +253,22 @@ FString UMovieGraphImageSequenceOutputNode::CreateFileName(
 		return FString();
 	}
 
+	constexpr bool bIncludeCDOs = false;
+	constexpr bool bExactMatch = false;
+	const TArray<UMovieGraphImageSequenceOutputNode*> ImageSequenceOutputNodes = InRawFrameData->EvaluatedConfig->GetSettingsForBranch<UMovieGraphImageSequenceOutputNode>(
+			InRenderData.Key.RootBranchName, bIncludeCDOs, bExactMatch);
+
+	// Find the instance of this node from the evaluated graph which matches our class
+	UMovieGraphImageSequenceOutputNode* const* InstanceNode = Algo::FindByPredicate(ImageSequenceOutputNodes,
+		[this](const UMovieGraphImageSequenceOutputNode* Node)
+	{
+		return Node->GetClass() == GetClass();
+	});
+
+	ensureAlwaysMsgf(InstanceNode, TEXT("Failed to find ImageSequenceOutputNode instance for class: %s; using the CDO value of FileNameFormat instead."), *GetClass()->GetName());
+	
 	// Generate one string that puts the directory combined with the filename format.
-	FString FileNameFormatString = OutputSettingNode->OutputDirectory.Path / FileNameFormat;
+	FString FileNameFormatString = OutputSettingNode->OutputDirectory.Path / (InstanceNode ? (*InstanceNode)->FileNameFormat : FileNameFormat);
 
 	constexpr bool bIncludeRenderPass = false;
 	constexpr bool bTestFrameNumber = true;
diff --git a/Engine/Plugins/MovieScene/MovieRenderPipeline/Source/MovieRenderPipelineRenderPasses/Public/MovieGraphImageSequenceOutputNode.h b/Engine/Plugins/MovieScene/MovieRenderPipeline/Source/MovieRenderPipelineRenderPasses/Public/MovieGraphImageSequenceOutputNode.h
--- a/Engine/Plugins/MovieScene/MovieRenderPipeline/Source/MovieRenderPipelineRenderPasses/Public/MovieGraphImageSequenceOutputNode.h
+++ b/Engine/Plugins/MovieScene/MovieRenderPipeline/Source/MovieRenderPipelineRenderPasses/Public/MovieGraphImageSequenceOutputNode.h
@@ -95,7 +95,6 @@ public:
 	{
 		OutputFormat = EImageFormat::EXR;
 		bQuantizeTo8Bit = false;
-		FileNameFormat = TEXT("{sequence_name}.{frame_number}");
 		Compression = EEXRCompressionFormat::PIZ;
 	}
 
