MovieGraph:

-Fix: Exposing a NamedResolution as a Job Variable doesn't show customization in MRQ Details Panel
-Fix: Editing a resolution set as a variable results in a custom resolution of 0x0

diff --git a/Engine/Plugins/MovieScene/MovieRenderPipeline/Source/MovieRenderPipelineEditor/Private/Customizations/Graph/MovieGraphNamedResolutionCustomization.h b/Engine/Plugins/MovieScene/MovieRenderPipeline/Source/MovieRenderPipelineEditor/Private/Customizations/Graph/MovieGraphNamedResolutionCustomization.h
--- a/Engine/Plugins/MovieScene/MovieRenderPipeline/Source/MovieRenderPipelineEditor/Private/Customizations/Graph/MovieGraphNamedResolutionCustomization.h
+++ b/Engine/Plugins/MovieScene/MovieRenderPipeline/Source/MovieRenderPipelineEditor/Private/Customizations/Graph/MovieGraphNamedResolutionCustomization.h
@@ -12,6 +12,7 @@
 #include "IPropertyTypeCustomization.h"
 #include "Widgets/Input/SButton.h"
 #include "Widgets/Input/SComboBox.h"
+#include "Widgets/Input/SNumericEntryBox.h"
 #include "Widgets/SBoxPanel.h"
 
 #define LOCTEXT_NAMESPACE "FMovieGraphNamedResolutionCustomization"
@@ -268,7 +269,7 @@ protected:
 		];
 	}
 
-	void UpdateLockImage()
+	void UpdateAspectRatioLockImage()
 	{
 		LockImage->SetImage(GetLockBrush());
 		LockImage->SetToolTipText(GetLockTooltipText());
@@ -279,7 +280,7 @@ protected:
 		bLockedAspectRatio = !bLockedAspectRatio;
 		if (bLockedAspectRatio) { CacheAspectRatio(); }
 
-		UpdateLockImage();
+		UpdateAspectRatioLockImage();
 	}
 
 	TSharedRef<SWidget> MakeLockExtensionWidget()
@@ -299,85 +300,17 @@ protected:
 		];
 	}
 
-	virtual void CustomizeChildren(
-		TSharedRef<IPropertyHandle> InStructPropertyHandle,
-		IDetailChildrenBuilder& StructBuilder, IPropertyTypeCustomizationUtils& StructCustomizationUtils) override
-	{
-		const TSharedPtr<IPropertyHandle> ResolutionPropertyHandle =
-			InStructPropertyHandle->GetChildHandle(GET_MEMBER_NAME_CHECKED(FMovieGraphNamedResolution, Resolution));
-
-		ResolutionPropertyHandle->SetOnChildPropertyValuePreChange(
-			FSimpleDelegate::CreateSP(this, &FMovieGraphNamedResolutionCustomization::OnCustomResolutionPreManualChange));
-		ResolutionPropertyHandle->SetOnChildPropertyValueChangedWithData(TDelegate<void(const FPropertyChangedEvent&)>::CreateSP(
-			this, &FMovieGraphNamedResolutionCustomization::OnCustomResolutionManualChange));
-
-		if (const TSharedPtr<IPropertyHandle> PropertyXHandle = ResolutionPropertyHandle->GetChildHandle(0))
-		{			
-			IDetailPropertyRow& PropertyXRow = StructBuilder.AddProperty(PropertyXHandle.ToSharedRef());
-			ResolutionXPropertyHandle = PropertyXHandle;
-
-			// Prevent showing Reset to Default
-			PropertyXRow.OverrideResetToDefault(FResetToDefaultOverride::Hide());
-
-			PropertyXRow.CustomWidget()
-			.NameContent()
-			[
-				PropertyXHandle->CreatePropertyNameWidget()
-			]
-			.ValueContent()
-			.HAlign(HAlign_Fill)
-			[
-				PropertyXHandle->CreatePropertyValueWidget()
-			]
-			.ExtensionContent()
-			[
-				MakeLockExtensionWidget()
-			];
-
-			UpdateLockImage();
-		}
-
-		if (const TSharedPtr<IPropertyHandle> PropertyYHandle = ResolutionPropertyHandle->GetChildHandle(1))
-		{			
-			IDetailPropertyRow& PropertyYRow = StructBuilder.AddProperty(PropertyYHandle.ToSharedRef());
-			ResolutionYPropertyHandle = PropertyYHandle;
-
-			// Prevent showing Reset to Default
-			PropertyYRow.OverrideResetToDefault(FResetToDefaultOverride::Hide());
-
-			// Custom widget to enforce length and padding on value widget
-			PropertyYRow.CustomWidget()
-			.NameContent()
-			[
-				PropertyYHandle->CreatePropertyNameWidget()
-			]
-			.ValueContent()
-			.HAlign(HAlign_Fill)
-			[
-				SNew(SBox)
-				.HAlign(HAlign_Fill)
-				.Padding(FMargin(0, 0, 24, 0))
-				[
-					PropertyYHandle->CreatePropertyValueWidget()
-				]
-			];
-		}
-
-		CacheAspectRatio();
-	}
-	//~ End IPropertyTypeCustomization interface
-	
-	void OnProjectSettingsChanged(UObject*, FPropertyChangedEvent& PropertyChangedEvent)
+	void CacheAspectRatio()
 	{
-		if (PropertyChangedEvent.MemberProperty->GetFName().IsEqual(GET_MEMBER_NAME_CHECKED(UMovieGraphProjectSettings, DefaultNamedResolutions)))
+		float X = (float)CustomEntry.Resolution.X;
+		float Y = (float)CustomEntry.Resolution.Y;
+		
+		if (const FMovieGraphNamedResolution* SelectedResolution = FindNamedResolutionForOption(GetCurrentOptionAssigningIfNeeded()))
 		{
-			RepopulateOptions();
+			X = (float)SelectedResolution->Resolution.X;
+			Y = (float)SelectedResolution->Resolution.Y;
 		}
-	}
-
-	void CacheAspectRatio()
-	{
-		CurrentAspectRatio = (double)CustomEntry.Resolution.X / CustomEntry.Resolution.Y;
+		CurrentAspectRatio = X / Y;
 	}
 
 	void UpdateCustomEntryValues()
@@ -406,7 +339,26 @@ protected:
 		CacheAspectRatio();
 	}
 
-	void OnCustomResolutionPreManualChange()
+	int32 RoundToNearestEvenNumber(double InNumberToRound) const
+	{
+		// Round to nearest even
+		const int32 Floored = FMath::FloorToInt(InNumberToRound);
+		const int32 Ceiled = FMath::CeilToInt(InNumberToRound);
+
+		return Floored % 2 == 0 ? Floored : Ceiled;
+	}
+	
+	int32 GetCurrentlySelectedResolutionByAxis(const EAxis::Type Axis)
+	{
+		if (const FMovieGraphNamedResolution* const CurrentNamedResolution = FindNamedResolutionForOption(GetCurrentOptionAssigningIfNeeded()))
+		{
+			return Axis == EAxis::Y ? CurrentNamedResolution->Resolution.Y : CurrentNamedResolution->Resolution.X;
+		}
+
+		return INDEX_NONE;
+	}
+
+	void OnCustomSliderValueChanged(uint32 NewValue, EAxis::Type Axis)
 	{
 		// If the custom resolution is changed, switch to 'Custom' option
 		if (!CurrentOption.IsEqual(FMovieGraphNamedResolution::CustomEntryName))
@@ -416,54 +368,117 @@ protected:
 			// Switch to 'Custom' option
 			SetCurrentOptionAndCacheTooltipText(FMovieGraphNamedResolution::CustomEntryName);
 		}
+
+		// Enforce aspect ratio if desired
+		if (Axis == EAxis::Y)
+		{
+			CustomEntry.Resolution.Y = NewValue;
+			
+			if (bLockedAspectRatio)
+			{
+				const double NewXResolution = (double)CustomEntry.Resolution.Y * CurrentAspectRatio;
+				CustomEntry.Resolution.X = RoundToNearestEvenNumber(NewXResolution);
+			}
+		}
+		else // X Changed
+		{
+			CustomEntry.Resolution.X = NewValue;
+			
+			if (bLockedAspectRatio)
+			{
+				const double NewYResolution = (double)CustomEntry.Resolution.X / CurrentAspectRatio;
+				CustomEntry.Resolution.Y = RoundToNearestEvenNumber(NewYResolution);
+			}
+		}
 	}
 
-	int32 RoundToNearestEvenNumber(double InNumberToRound) const
+	void OnCustomSliderValueCommitted(uint32 NewValue, ETextCommit::Type, EAxis::Type Axis)
 	{
-		// Round to nearest even
-		const int32 Floored = FMath::FloorToInt(InNumberToRound);
-		const int32 Ceiled = FMath::CeilToInt(InNumberToRound);
-
-		return Floored % 2 == 0 ? Floored : Ceiled;
+		OnCustomSliderValueChanged(NewValue, Axis);
+		
+		UpdateOwningStruct();
 	}
 
-	void OnCustomResolutionManualChange(const FPropertyChangedEvent& Event) const
+	void AddCustomRowForResolutionAxis(
+		EAxis::Type Axis, IDetailChildrenBuilder& StructBuilder, const FString& FilterString,
+		TSharedRef<SWidget> NameContentWidget, TSharedPtr<SWidget> AspectRatioLockExtensionWidget = nullptr)
 	{
-		if (!ensureMsgf(ResolutionXPropertyHandle, TEXT("%hs: `ResolutionXPropertyHandle` is null."), __FUNCTION__))
-		{
-			return;
-		}
-		if (!ensureMsgf(ResolutionYPropertyHandle, TEXT("%hs: `ResolutionYPropertyHandle` is null."), __FUNCTION__))
+		const FSlateFontInfo PropertyFontStyle = FAppStyle::GetFontStyle( TEXT("PropertyWindow.NormalFont") );
+		const TOptional<uint32> PropertyMaxValue = TOptional<uint32>();
+	
+		FDetailWidgetRow& CustomRow = StructBuilder.AddCustomRow(FText::FromString(FilterString))
+		.NameContent()
+		[
+			NameContentWidget
+		]
+		.ValueContent()
+		.HAlign(HAlign_Fill)
+		[
+			// We make our own value widget because we don't want the widget tied to the property handle
+			// as the property handle gets re-instanced every time a blueprint is edited.
+			SNew(SBox)
+			// Maintain spacing when there would be no extension widget
+			.Padding(0, 0, Axis == EAxis::Y ? 24 /* Assuming image is 16x16 + 8 padding */ : 0, 0) 
+			[
+				SNew(SNumericEntryBox<uint32>)
+				.Font(PropertyFontStyle)
+				.AllowSpin(true)
+				.MinSliderValue(2)
+				.MinValue(2)
+				.MaxValue(PropertyMaxValue)
+				.MaxSliderValue(PropertyMaxValue)
+				.Delta(2)
+				.Value_Lambda([this, Axis] (){ return GetCurrentlySelectedResolutionByAxis(Axis); }) // Lambda to bypass const requirement
+				.OnValueChanged(this, &FMovieGraphNamedResolutionCustomization::OnCustomSliderValueChanged, Axis)
+				.OnValueCommitted(this, &FMovieGraphNamedResolutionCustomization::OnCustomSliderValueCommitted, Axis)
+			]
+		]
+		.OverrideResetToDefault(FResetToDefaultOverride::Hide()); // Prevent showing Reset to Default
+
+		if (AspectRatioLockExtensionWidget.IsValid())
 		{
-			return;
+			CustomRow
+			.ExtensionContent()
+			[
+				AspectRatioLockExtensionWidget.ToSharedRef()
+			];
 		}
+	}
 
-		// Update CustomEntry values
-		int32 ResolutionX = 0;
-		ResolutionXPropertyHandle->GetValue(ResolutionX);
-
-		CustomEntry.Resolution.X = ResolutionX;
+	virtual void CustomizeChildren(
+		TSharedRef<IPropertyHandle> InStructPropertyHandle,
+		IDetailChildrenBuilder& StructBuilder, IPropertyTypeCustomizationUtils& StructCustomizationUtils) override
+	{
+		const TSharedPtr<IPropertyHandle> ResolutionPropertyHandle =
+			InStructPropertyHandle->GetChildHandle(GET_MEMBER_NAME_CHECKED(FMovieGraphNamedResolution, Resolution));
+		
+		if (const TSharedPtr<IPropertyHandle> PropertyXHandle = ResolutionPropertyHandle->GetChildHandle(0))
+		{
+			ResolutionXPropertyHandle = PropertyXHandle;
 
-		int32 ResolutionY = 0;
-		ResolutionYPropertyHandle->GetValue(ResolutionY);
+			AddCustomRowForResolutionAxis(
+				EAxis::X, StructBuilder, PropertyXHandle->GeneratePathToProperty(), PropertyXHandle->CreatePropertyNameWidget(), MakeLockExtensionWidget());
+			
+			UpdateAspectRatioLockImage();
+		}
 
-		CustomEntry.Resolution.Y = ResolutionY;
+		if (const TSharedPtr<IPropertyHandle> PropertyYHandle = ResolutionPropertyHandle->GetChildHandle(1))
+		{			
+			ResolutionYPropertyHandle = PropertyYHandle;
+			
+			AddCustomRowForResolutionAxis(
+				EAxis::Y, StructBuilder, PropertyYHandle->GeneratePathToProperty(), PropertyYHandle->CreatePropertyNameWidget());
+		}
 
-		// Enforce aspect ratio if desired
-		if (bLockedAspectRatio && Event.ChangeType != EPropertyChangeType::Interactive)
+		CacheAspectRatio();
+	}
+	//~ End IPropertyTypeCustomization interface
+	
+	void OnProjectSettingsChanged(UObject*, FPropertyChangedEvent& PropertyChangedEvent)
+	{
+		if (PropertyChangedEvent.MemberProperty->GetFName().IsEqual(GET_MEMBER_NAME_CHECKED(UMovieGraphProjectSettings, DefaultNamedResolutions)))
 		{
-			if (Event.Property == ResolutionXPropertyHandle->GetProperty()) // X Changed
-			{
-				const double NewYResolution = (double)CustomEntry.Resolution.X / CurrentAspectRatio;
-
-				// Specify an interactive change to avoid feedback loop
-				ResolutionYPropertyHandle->SetValue(RoundToNearestEvenNumber(NewYResolution), EPropertyValueSetFlags::InteractiveChange);
-			}
-			else // Y Changed
-			{
-				const double NewXResolution = (double)CustomEntry.Resolution.Y * CurrentAspectRatio;
-				ResolutionXPropertyHandle->SetValue(RoundToNearestEvenNumber(NewXResolution), EPropertyValueSetFlags::InteractiveChange);
-			}
+			RepopulateOptions();
 		}
 	}
 
diff --git a/Engine/Plugins/MovieScene/MovieRenderPipeline/Source/MovieRenderPipelineEditor/Private/Widgets/SMoviePipelineQueuePanel.cpp b/Engine/Plugins/MovieScene/MovieRenderPipeline/Source/MovieRenderPipelineEditor/Private/Widgets/SMoviePipelineQueuePanel.cpp
--- a/Engine/Plugins/MovieScene/MovieRenderPipeline/Source/MovieRenderPipelineEditor/Private/Widgets/SMoviePipelineQueuePanel.cpp
+++ b/Engine/Plugins/MovieScene/MovieRenderPipeline/Source/MovieRenderPipelineEditor/Private/Widgets/SMoviePipelineQueuePanel.cpp
@@ -1,7 +1,9 @@
 // Copyright Epic Games, Inc. All Rights Reserved.
 
 // Movie Pipeline Includes
+
 #include "Widgets/SMoviePipelineQueuePanel.h"
+#include "Customizations/Graph/MovieGraphNamedResolutionCustomization.h"
 #include "Customizations/JobCustomization.h"
 #include "Widgets/MoviePipelineWidgetConstants.h"
 #include "SMoviePipelineQueueEditor.h"
@@ -70,13 +72,16 @@ void SMoviePipelineQueuePanel::Construct(const FArguments& InArgs)
 		UMoviePipelineExecutorJob::StaticClass(),
 		FOnGetDetailCustomizationInstance::CreateStatic(&FJobDetailsCustomization::MakeInstance));
 
+	JobDetailsPanelWidget->RegisterInstancedCustomPropertyTypeLayout(
+		FMovieGraphNamedResolution::StaticStruct()->GetFName(),
+		FOnGetPropertyTypeCustomizationInstance::CreateStatic(&FMovieGraphNamedResolutionCustomization::MakeInstance));
+
 	// Create the child widgets that need to know about our pipeline
 	PipelineQueueEditorWidget = SNew(SMoviePipelineQueueEditor)
 		.OnEditConfigRequested(this, &SMoviePipelineQueuePanel::OnEditJobConfigRequested)
 		.OnPresetChosen(this, &SMoviePipelineQueuePanel::OnJobPresetChosen)
 		.OnJobSelectionChanged(this, &SMoviePipelineQueuePanel::OnSelectionChanged);
 
-
 	{
 		// Automatically select the first job in the queue
 		UMoviePipelineQueueSubsystem* Subsystem = GEditor->GetEditorSubsystem<UMoviePipelineQueueSubsystem>();
