MovieGraph: Fix file name format for multi-layer EXR.

diff --git a/Engine/Plugins/MovieScene/MovieRenderPipeline/Source/MovieRenderPipelineRenderPasses/Private/MovieGraphImageSequenceOutputNode.cpp b/Engine/Plugins/MovieScene/MovieRenderPipeline/Source/MovieRenderPipelineRenderPasses/Private/MovieGraphImageSequenceOutputNode.cpp
--- a/Engine/Plugins/MovieScene/MovieRenderPipeline/Source/MovieRenderPipelineRenderPasses/Private/MovieGraphImageSequenceOutputNode.cpp
+++ b/Engine/Plugins/MovieScene/MovieRenderPipeline/Source/MovieRenderPipelineRenderPasses/Private/MovieGraphImageSequenceOutputNode.cpp
@@ -13,7 +13,6 @@
 #include "MoviePipelineImageSequenceOutput.h" // for FAsyncImageQuantization
 #include "MovieRenderPipelineCoreModule.h"
 
-#include "Algo/Find.h"
 #include "Modules/ModuleManager.h"
 #include "ImageWriteQueue.h"
 #include "Misc/Paths.h"
@@ -236,6 +235,7 @@ FString UMovieGraphImageSequenceOutputNode::CreateFileName(
 	const UMovieGraphPipeline* InPipeline,
 	const TPair<FMovieGraphRenderDataIdentifier, TUniquePtr<FImagePixelData>>& InRenderData,
 	const EImageFormat InImageFormat,
+	const FString& InFileNameFormat,
 	FMovieGraphResolveArgs& OutMergedFormatArgs) const
 {
 	const TCHAR* Extension = TEXT("");
@@ -253,22 +253,8 @@ FString UMovieGraphImageSequenceOutputNode::CreateFileName(
 		return FString();
 	}
 
-	constexpr bool bIncludeCDOs = false;
-	constexpr bool bExactMatch = false;
-	const TArray<UMovieGraphImageSequenceOutputNode*> ImageSequenceOutputNodes = InRawFrameData->EvaluatedConfig->GetSettingsForBranch<UMovieGraphImageSequenceOutputNode>(
-			InRenderData.Key.RootBranchName, bIncludeCDOs, bExactMatch);
-
-	// Find the instance of this node from the evaluated graph which matches our class
-	UMovieGraphImageSequenceOutputNode* const* InstanceNode = Algo::FindByPredicate(ImageSequenceOutputNodes,
-		[this](const UMovieGraphImageSequenceOutputNode* Node)
-	{
-		return Node->GetClass() == GetClass();
-	});
-
-	ensureAlwaysMsgf(InstanceNode, TEXT("Failed to find ImageSequenceOutputNode instance for class: %s; using the CDO value of FileNameFormat instead."), *GetClass()->GetName());
-	
 	// Generate one string that puts the directory combined with the filename format.
-	FString FileNameFormatString = OutputSettingNode->OutputDirectory.Path / (InstanceNode ? (*InstanceNode)->FileNameFormat : FileNameFormat);
+	FString FileNameFormatString = OutputSettingNode->OutputDirectory.Path / InFileNameFormat;
 
 	constexpr bool bIncludeRenderPass = false;
 	constexpr bool bTestFrameNumber = true;
@@ -332,9 +318,15 @@ void UMovieGraphImageSequenceOutputNode::OnReceiveImageDataImpl(UMovieGraphPipel
 		// ToDo: Certain images may require transparency, at which point
 		// we write out a .png instead of a .jpeg.
 		EImageFormat PreferredOutputFormat = OutputFormat;
+
+		constexpr bool bIncludeCDOs = false;
+		constexpr bool bExactMatch = true;
+		const UMovieGraphImageSequenceOutputNode* ParentNode = Cast<UMovieGraphImageSequenceOutputNode>(
+			InRawFrameData->EvaluatedConfig->GetSettingForBranch(GetClass(), RenderData.Key.RootBranchName, bIncludeCDOs, bExactMatch));
+		checkf(ParentNode, TEXT("Image sequence output should not exist without a parent node in the graph."));
 		
 		FMovieGraphResolveArgs FinalResolvedKVPs;
-		FString FileName = CreateFileName(InRawFrameData, InPipeline, RenderData, PreferredOutputFormat, FinalResolvedKVPs);
+		FString FileName = CreateFileName(InRawFrameData, InPipeline, RenderData, PreferredOutputFormat, ParentNode->FileNameFormat, FinalResolvedKVPs);
 		if (!ensureMsgf(!FileName.IsEmpty(), TEXT("Unexpected empty file name, skipping frame.")))
 		{
 			continue;
@@ -355,10 +347,6 @@ void UMovieGraphImageSequenceOutputNode::OnReceiveImageDataImpl(UMovieGraphPipel
 			TileImageTask->PixelData = RenderData.Value->MoveImageDataToNew();
 		}
 
-		const UMovieGraphImageSequenceOutputNode* ParentNode = Cast<UMovieGraphImageSequenceOutputNode>(
-			InRawFrameData->EvaluatedConfig->GetSettingForBranch(GetClass(), RenderData.Key.RootBranchName, false, true));
-		checkf(ParentNode, TEXT("Image sequence output should not exist without a parent node in the graph."));
-
 		UE::MovieGraph::FMovieGraphSampleState* Payload = RenderData.Value->GetPayload<UE::MovieGraph::FMovieGraphSampleState>();
 		
 		bool bQuantizationEncodeSRGB = true;
@@ -549,18 +537,20 @@ void UMovieGraphImageSequenceOutputNode_EXR::OnReceiveImageDataImpl(UMovieGraphP
 
 		checkf(RenderData.Value.IsValid(), TEXT("Unexpected empty image data: incorrectly moved or its production failed?"));
 
+		constexpr bool bIncludeCDOs = false;
+		constexpr bool bExactMatch = true;
+		const UMovieGraphImageSequenceOutputNode_EXR* ParentNode = InRawFrameData->EvaluatedConfig->GetSettingForBranch<UMovieGraphImageSequenceOutputNode_EXR>(
+			RenderData.Key.RootBranchName, bIncludeCDOs, bExactMatch);
+		checkf(ParentNode, TEXT("Single-layer EXR should not exist without a parent node in the graph."));
+
 		FMovieGraphResolveArgs ResolvedFormatArgs;
-		FString FileName = CreateFileName(InRawFrameData, InPipeline, RenderData, OutputFormat, ResolvedFormatArgs);
+		FString FileName = CreateFileName(InRawFrameData, InPipeline, RenderData, OutputFormat, ParentNode->FileNameFormat, ResolvedFormatArgs);
 		if (!ensureMsgf(!FileName.IsEmpty(), TEXT("Unexpected empty file name, skipping frame.")))
 		{
 			continue;
 		}
 
 		UE::MovieGraph::FMovieGraphSampleState* Payload = RenderData.Value->GetPayload<UE::MovieGraph::FMovieGraphSampleState>();
-
-		const UMovieGraphImageSequenceOutputNode_EXR* ParentNode = InRawFrameData->EvaluatedConfig->GetSettingForBranch<UMovieGraphImageSequenceOutputNode_EXR>(
-			RenderData.Key.RootBranchName, false /*bIncludeCDOs*/, true /*bExactMatch*/);
-		checkf(ParentNode, TEXT("Single-layer EXR should not exist without a parent node in the graph."));
 		
 		TUniquePtr<FEXRImageWriteTask> ImageWriteTask = CreateImageWriteTask(FileName, ParentNode->Compression);
 		
@@ -618,8 +608,10 @@ void UMovieGraphImageSequenceOutputNode_MultiLayerEXR::OnReceiveImageDataImpl(UM
 {
 	check(InRawFrameData);
 
+	constexpr bool bIncludeCDOs = false;
+	constexpr bool bExactMatch = true;
 	const UMovieGraphImageSequenceOutputNode_MultiLayerEXR* ParentNode = InRawFrameData->EvaluatedConfig->GetSettingForBranch<UMovieGraphImageSequenceOutputNode_MultiLayerEXR>(
-		UMovieGraphNode::GlobalsPinName, false /*bIncludeCDOs*/, true /*bExactMatch*/);
+		UMovieGraphNode::GlobalsPinName, bIncludeCDOs, bExactMatch);
 	checkf(ParentNode, TEXT("Multi-Layer EXR should not exist without a parent node in the graph."));
 
 	// Generate a mapping of resolved filename -> RenderIDs, and filename -> resolve args. The generated EXRs can only
diff --git a/Engine/Plugins/MovieScene/MovieRenderPipeline/Source/MovieRenderPipelineRenderPasses/Public/MovieGraphImageSequenceOutputNode.h b/Engine/Plugins/MovieScene/MovieRenderPipeline/Source/MovieRenderPipelineRenderPasses/Public/MovieGraphImageSequenceOutputNode.h
--- a/Engine/Plugins/MovieScene/MovieRenderPipeline/Source/MovieRenderPipelineRenderPasses/Public/MovieGraphImageSequenceOutputNode.h
+++ b/Engine/Plugins/MovieScene/MovieRenderPipeline/Source/MovieRenderPipelineRenderPasses/Public/MovieGraphImageSequenceOutputNode.h
@@ -66,6 +66,7 @@ protected:
 		const UMovieGraphPipeline* InPipeline,
 		const TPair<FMovieGraphRenderDataIdentifier, TUniquePtr<FImagePixelData>>& InRenderData,
 		const EImageFormat InImageFormat,
+		const FString& InFileNameFormat,
 		FMovieGraphResolveArgs& OutMergedFormatArgs) const;
 
 protected:
@@ -173,6 +174,12 @@ class UMovieGraphImageSequenceOutputNode_MultiLayerEXR : public UMovieGraphImage
 	GENERATED_BODY()
 
 public:
+	UMovieGraphImageSequenceOutputNode_MultiLayerEXR()
+		: UMovieGraphImageSequenceOutputNode_EXR()
+	{
+		// Multi-layer default excludes {render_layer}.
+		FileNameFormat = TEXT("{sequence_name}.{frame_number}");
+	}
 
 	virtual EMovieGraphBranchRestriction GetBranchRestriction() const override
 	{
