MoviePipeline: Added support for the Collections UI in the details panel to select matching actors in the Outliner.

Also updated the "Actor" collection query class to skip trying to convert an editor actor reference to the PIE equivalent when not in PIE.

#rnx

diff --git a/Engine/Plugins/MovieScene/MovieRenderPipeline/Source/MovieRenderPipelineCore/Private/Graph/MoviePipelineRenderLayerSubsystem.cpp b/Engine/Plugins/MovieScene/MovieRenderPipeline/Source/MovieRenderPipelineCore/Private/Graph/MoviePipelineRenderLayerSubsystem.cpp
--- a/Engine/Plugins/MovieScene/MovieRenderPipeline/Source/MovieRenderPipelineCore/Private/Graph/MoviePipelineRenderLayerSubsystem.cpp
+++ b/Engine/Plugins/MovieScene/MovieRenderPipeline/Source/MovieRenderPipelineCore/Private/Graph/MoviePipelineRenderLayerSubsystem.cpp
@@ -478,9 +478,15 @@ void UMovieGraphConditionGroupQuery_Actor::Evaluate(const TArray<AActor*>& InAct
 {
 	TRACE_CPUPROFILER_EVENT_SCOPE(UMovieGraphConditionGroupQuery_Actor::Evaluate);
 
-	// Convert the actors in the query to PIE equivalents once, rather than constantly in the loop.
-	TArray<AActor*> ActorsToMatch_Pie;
-	ActorsToMatch_Pie.Reserve(ActorsToMatch.Num());
+#if WITH_EDITOR
+	const bool bIsPIE = GEditor->IsPlaySessionInProgress();
+#else
+	const bool bIsPIE = false;
+#endif
+
+	// Convert the actors in the query to PIE (or editor) equivalents once, rather than constantly in the loop.
+	TArray<AActor*> ActorsToMatch_Converted;
+	ActorsToMatch_Converted.Reserve(ActorsToMatch.Num());
 	for (const TSoftObjectPtr<AActor>& SoftActorToMatch : ActorsToMatch)
 	{
 		if (!SoftActorToMatch.IsValid())
@@ -489,28 +495,53 @@ void UMovieGraphConditionGroupQuery_Actor::Evaluate(const TArray<AActor*>& InAct
 		}
 		
 		AActor* ActorToMatch = SoftActorToMatch.Get();
-
-		// Only do editor -> PIE actor conversion when the actor is from an editor world
 		const UWorld* ActorWorld = ActorToMatch->GetWorld();
-		if (ActorWorld && ActorWorld->IsEditorWorld())
+		const bool bIsEditorActor = ActorWorld && ActorWorld->IsEditorWorld();
+
+		// If a PIE session is NOT in progress, make sure that the actor is the editor equivalent
+		if (!bIsPIE)
 		{
+			// Only do PIE -> editor actor conversion when the actor is NOT from the editor
+			if (!bIsEditorActor)
+			{
 #if WITH_EDITOR
-			if (AActor* PieActor = EditorUtilities::GetSimWorldCounterpartActor(ActorToMatch))
+				if (AActor* EditorActor = EditorUtilities::GetEditorWorldCounterpartActor(ActorToMatch))
+				{
+					ActorsToMatch_Converted.Add(EditorActor);
+				}
+#endif
+			}
+			else
 			{
-				ActorsToMatch_Pie.Add(PieActor);
+				// Just use ActorToMatch as-is if it's not from PIE
+				ActorsToMatch_Converted.Add(ActorToMatch);
 			}
-#endif
 		}
+
+		// If a PIE session IS active, try to get the PIE equivalent of the editor actor
 		else
 		{
-			// Just use ActorToMatch as-is if it's not from an editor actor
-			ActorsToMatch_Pie.Add(ActorToMatch);
+			// Only do editor -> PIE actor conversion when the actor is from an editor world
+			if (bIsEditorActor)
+			{
+#if WITH_EDITOR
+				if (AActor* PieActor = EditorUtilities::GetSimWorldCounterpartActor(ActorToMatch))
+				{
+					ActorsToMatch_Converted.Add(PieActor);
+				}
+#endif
+			}
+			else
+			{
+				// Just use ActorToMatch as-is if it's not from an editor actor
+				ActorsToMatch_Converted.Add(ActorToMatch);
+			}	
 		}
 	}
 	
 	for (AActor* Actor : InActorsToQuery)
 	{
-		if (ActorsToMatch_Pie.Contains(Actor))
+		if (ActorsToMatch_Converted.Contains(Actor))
 		{
 			OutMatchingActors.Add(Actor);
 		}
diff --git a/Engine/Plugins/MovieScene/MovieRenderPipeline/Source/MovieRenderPipelineEditor/Private/Graph/MovieGraphAssetToolkit.cpp b/Engine/Plugins/MovieScene/MovieRenderPipeline/Source/MovieRenderPipelineEditor/Private/Graph/MovieGraphAssetToolkit.cpp
--- a/Engine/Plugins/MovieScene/MovieRenderPipeline/Source/MovieRenderPipelineEditor/Private/Graph/MovieGraphAssetToolkit.cpp
+++ b/Engine/Plugins/MovieScene/MovieRenderPipeline/Source/MovieRenderPipelineEditor/Private/Graph/MovieGraphAssetToolkit.cpp
@@ -25,6 +25,7 @@
 #include "Framework/Commands/GenericCommands.h"
 #include "IDetailRootObjectCustomization.h"
 #include "PropertyEditorModule.h"
+#include "Selection.h"
 #include "ToolMenus.h"
 #include "Widgets/Docking/SDockTab.h"
 #include "Widgets/Graph/SMovieGraphConfigPanel.h"
@@ -37,6 +38,73 @@ const FName FMovieGraphAssetToolkit::DetailsTabId(TEXT("MovieGraphAssetToolkitDe
 const FName FMovieGraphAssetToolkit::MembersTabId(TEXT("MovieGraphAssetToolkitMembers"));
 const FName FMovieGraphAssetToolkit::ActiveRenderSettingsTabId(TEXT("MovieGraphAssetToolkitActiveRenderSettings"));
 
+void SMovieGraphSyncCollectionToOutlinerButton::Construct(const FArguments& InArgs)
+{
+	SelectedNodesAttribute = InArgs._SelectedNodes;
+
+	ChildSlot
+	[
+		SNew(SButton)
+		.ButtonStyle(&FAppStyle::Get().GetWidgetStyle<FButtonStyle>("SimpleButton"))
+		.ContentPadding(0.f)
+		.ToolTipText(LOCTEXT("PreviewCollectionButton_Tooltip", "Evaluate the collection and select the matched actors in the Outliner."))
+		.Visibility_Lambda([this]()
+		{
+			const TArray<TWeakObjectPtr<UObject>>& SelectedObjects = SelectedNodesAttribute.Get();
+			if ((SelectedObjects.Num() == 1) && SelectedObjects[0].IsValid() && SelectedObjects[0]->IsA<UMovieGraphCollectionNode>())
+			{
+				return EVisibility::Visible;
+			}
+
+			return EVisibility::Hidden;
+		})
+		.OnClicked_Lambda([this]()
+		{
+			EvaluateCollectionAndSelect();
+
+			return FReply::Handled();
+		})
+		[
+			SNew(SImage)
+			.ColorAndOpacity(FSlateColor::UseForeground())
+			.Image(FAppStyle::Get().GetBrush("FoliageEditMode.SelectAll"))
+		]
+	];
+}
+
+void SMovieGraphSyncCollectionToOutlinerButton::EvaluateCollectionAndSelect() const
+{
+	const TArray<TWeakObjectPtr<UObject>>& SelectedObjects = SelectedNodesAttribute.Get();
+	if ((SelectedObjects.Num() != 1) || !SelectedObjects[0].IsValid())
+	{
+		return;
+	}
+			
+	if (const UMovieGraphCollectionNode* CollectionNode = Cast<UMovieGraphCollectionNode>(SelectedObjects[0]))
+	{
+		// Evaluate the collection based on the current editor world
+		TSet<AActor*> EvaluatedActors = CollectionNode->Collection->Evaluate(GEditor->GetEditorWorldContext().World());
+
+		// Select all actors matched by the collection
+		{
+			GEditor->GetSelectedActors()->Modify();
+			GEditor->GetSelectedActors()->BeginBatchSelectOperation();
+			GEditor->GetSelectedActors()->DeselectAll();
+							
+			for (AActor* Actor : EvaluatedActors)
+			{
+				constexpr bool bShouldSelect = true;
+				constexpr bool bNotifyAfterSelect = false;
+				constexpr bool bSelectEvenIfHidden = true;
+				GEditor->SelectActor(Actor, bShouldSelect, bNotifyAfterSelect, bSelectEvenIfHidden);
+			}
+
+			constexpr bool bNotify = false;
+			GEditor->GetSelectedActors()->EndBatchSelectOperation(bNotify);
+		}
+	}
+}
+
 /** Header customization for when multiple objects are displayed in the details panel. */
 class FMovieGraphDetailsRootObjectCustomization final : public IDetailRootObjectCustomization
 {
@@ -289,10 +357,12 @@ TSharedRef<SDockTab> FMovieGraphAssetToolkit::SpawnTab_RenderGraphDetails(const
 	FPropertyEditorModule& PropertyEditorModule = FModuleManager::Get().LoadModuleChecked<FPropertyEditorModule>("PropertyEditor");
 	FDetailsViewArgs DetailsViewArgs;
 	DetailsViewArgs.bShowPropertyMatrixButton = false;
-	DetailsViewArgs.NameAreaSettings = FDetailsViewArgs::HideNameArea;
+	DetailsViewArgs.bCustomNameAreaLocation = true;
+	DetailsViewArgs.NameAreaSettings = FDetailsViewArgs::ObjectsUseNameArea;
 	DetailsViewArgs.bHideSelectionTip = true;
 	DetailsViewArgs.bAllowMultipleTopLevelObjects = true;
 	DetailsViewArgs.ViewIdentifier = "MovieGraphSettings";
+	DetailsViewArgs.bLockable = false;
 
 	SelectedGraphObjectsDetailsWidget = PropertyEditorModule.CreateDetailView(DetailsViewArgs);
 	SelectedGraphObjectsDetailsWidget->SetRootObjectCustomizationInstance(
@@ -330,11 +400,35 @@ TSharedRef<SDockTab> FMovieGraphAssetToolkit::SpawnTab_RenderGraphDetails(const
 		UMovieGraphModifierNode::StaticClass(),
 		FOnGetDetailCustomizationInstance::CreateStatic(&FMovieGraphModifiersCustomization::MakeInstance));
 	
+	TSharedRef<SWidget> CustomContent = SAssignNew(NameAreaCustomContent, SHorizontalBox)
+	+ SHorizontalBox::Slot()
+	[
+		SNew(SMovieGraphSyncCollectionToOutlinerButton)
+		.SelectedNodes_Lambda([this]()
+		{
+			return SelectedGraphObjectsDetailsWidget->GetSelectedObjects();
+		})
+	];
+	
+	SelectedGraphObjectsDetailsWidget->SetNameAreaCustomContent(CustomContent);
+
 	return SNew(SDockTab)
 		.TabColorScale(GetTabColorScale())
 		.Label(LOCTEXT("DetailsTab_Title", "Details"))
 		[
-			SelectedGraphObjectsDetailsWidget.ToSharedRef()
+			SNew(SVerticalBox)
+			
+			+ SVerticalBox::Slot()
+			.Padding(10.f, 4.f, 0.f, 0.f)
+			.AutoHeight()
+			[
+				SelectedGraphObjectsDetailsWidget->GetNameAreaWidget().ToSharedRef()
+			]
+
+			+ SVerticalBox::Slot()
+			[
+				SelectedGraphObjectsDetailsWidget.ToSharedRef()
+			]
 		];
 }
 
diff --git a/Engine/Plugins/MovieScene/MovieRenderPipeline/Source/MovieRenderPipelineEditor/Private/Graph/MovieGraphAssetToolkit.h b/Engine/Plugins/MovieScene/MovieRenderPipeline/Source/MovieRenderPipelineEditor/Private/Graph/MovieGraphAssetToolkit.h
--- a/Engine/Plugins/MovieScene/MovieRenderPipeline/Source/MovieRenderPipelineEditor/Private/Graph/MovieGraphAssetToolkit.h
+++ b/Engine/Plugins/MovieScene/MovieRenderPipeline/Source/MovieRenderPipelineEditor/Private/Graph/MovieGraphAssetToolkit.h
@@ -5,6 +5,26 @@
 #include "CoreMinimal.h"
 #include "Toolkits/AssetEditorToolkit.h"
 
+/** Evaluates a collection and syncs the matched actors to the Outliner selection. */
+class SMovieGraphSyncCollectionToOutlinerButton final : public SCompoundWidget
+{
+public:
+	SLATE_BEGIN_ARGS(SMovieGraphSyncCollectionToOutlinerButton) {}
+		/** The nodes selected in the graph. */
+		SLATE_ATTRIBUTE(TArray<TWeakObjectPtr<UObject>>, SelectedNodes)
+	SLATE_END_ARGS()
+
+	void Construct(const FArguments& InArgs);
+
+private:
+	/** Evaluates the collection and selects all matches in the Outliner. */
+	void EvaluateCollectionAndSelect() const;
+
+private:
+	/** The nodes which are currently selected in the graph. */
+	TAttribute<TArray<TWeakObjectPtr<UObject>>> SelectedNodesAttribute;
+};
+
 class FMovieGraphAssetToolkit :  public FAssetEditorToolkit
 {
 public:
@@ -52,6 +72,9 @@ private:
 	/** The details panel for the selected object(s) in the graph */
 	TSharedPtr<IDetailsView> SelectedGraphObjectsDetailsWidget;
 
+	/** Custom widgetry shown above the filter box in the details panel. */
+	TSharedPtr<SWidget> NameAreaCustomContent;
+
 	/** The widget that contains the main node graph */
 	TSharedPtr<class SMoviePipelineGraphPanel> MovieGraphWidget;
 
