MoviePipeline: Refactored the Collections UI in the details panel to use node builders rather than construct a monolothic STreeView. This has many advantages:
* Fixes an issue where the condition groups tree would randomly disappear.
* Expansion state of condition groups is remembered.
* Gripper handles are shown for condition groups and condition group queries. Wasn't obvious before that these were re-orderable.
* No more awkward indentation for the collections tree. The prior STreeView didn't fit in well, but using node builders makes the UI look much more "correct".
* Searchable via the details panel filter box.
* New condition groups have an Actor Name query by default.

#rnx

diff --git a/Engine/Plugins/MovieScene/MovieRenderPipeline/Source/MovieRenderPipelineCore/Private/Graph/MoviePipelineRenderLayerSubsystem.cpp b/Engine/Plugins/MovieScene/MovieRenderPipeline/Source/MovieRenderPipelineCore/Private/Graph/MoviePipelineRenderLayerSubsystem.cpp
--- a/Engine/Plugins/MovieScene/MovieRenderPipeline/Source/MovieRenderPipelineCore/Private/Graph/MoviePipelineRenderLayerSubsystem.cpp
+++ b/Engine/Plugins/MovieScene/MovieRenderPipeline/Source/MovieRenderPipelineCore/Private/Graph/MoviePipelineRenderLayerSubsystem.cpp
@@ -627,6 +627,7 @@ TSharedRef<SWidget> UMovieGraphConditionGroupQuery_Actor::GetAddMenuContents(con
 			if (ActorPickerWidget.IsValid())
 			{
 				ActorPickerWidget->FullRefresh();
+				ActorsList->Refresh();
 			}
 		}));
 	
@@ -884,6 +885,7 @@ TSharedRef<SWidget> UMovieGraphConditionGroupQuery_ActorType::GetAddMenuContents
 			if (ClassViewerWidget.IsValid())
 			{
 				ClassViewerWidget->Refresh();
+				ActorTypesList->Refresh();
 			}
 		}));
 
@@ -1085,6 +1087,7 @@ TSharedRef<SWidget> UMovieGraphConditionGroupQuery_ComponentType::GetAddMenuCont
 			if (ClassViewerWidget.IsValid())
 			{
 				ClassViewerWidget->Refresh();
+				ComponentTypesList->Refresh();
 			}
 		}));
 
@@ -1112,6 +1115,15 @@ FText UMovieGraphConditionGroupQuery_ComponentType::GetRowText(UClass* InCompone
 UMovieGraphConditionGroup::UMovieGraphConditionGroup()
 	: OpType(EMovieGraphConditionGroupOpType::Add)
 {
+	// The CDO will always have the default GUID
+	if (!HasAllFlags(RF_ClassDefaultObject))
+	{
+		Id = FGuid::NewGuid();
+	}
+	else
+	{
+		Id = FGuid();
+	}
 }
 
 void UMovieGraphConditionGroup::SetOperationType(const EMovieGraphConditionGroupOpType OperationType)
@@ -1268,6 +1280,11 @@ bool UMovieGraphConditionGroup::MoveQueryToIndex(UMovieGraphConditionGroupQueryB
 	return true;
 }
 
+const FGuid& UMovieGraphConditionGroup::GetId() const
+{
+	return Id;
+}
+
 TSet<AActor*> UMovieGraphCollection::Evaluate(const UWorld* InWorld) const
 {
 	TRACE_CPUPROFILER_EVENT_SCOPE(UMovieGraphCollection::Evaluate);
diff --git a/Engine/Plugins/MovieScene/MovieRenderPipeline/Source/MovieRenderPipelineCore/Public/Graph/MoviePipelineRenderLayerSubsystem.h b/Engine/Plugins/MovieScene/MovieRenderPipeline/Source/MovieRenderPipelineCore/Public/Graph/MoviePipelineRenderLayerSubsystem.h
--- a/Engine/Plugins/MovieScene/MovieRenderPipeline/Source/MovieRenderPipelineCore/Public/Graph/MoviePipelineRenderLayerSubsystem.h
+++ b/Engine/Plugins/MovieScene/MovieRenderPipeline/Source/MovieRenderPipelineCore/Public/Graph/MoviePipelineRenderLayerSubsystem.h
@@ -413,7 +413,14 @@ public:
 	UFUNCTION(BlueprintCallable, Category = "Settings")
 	bool MoveQueryToIndex(UMovieGraphConditionGroupQueryBase* InQuery, const int32 NewIndex);
 
+	/** Gets a persistent unique identifier for this condition group. */
+	const FGuid& GetId() const;
+
 private:
+	/** A unique identifier for this condition group. Needed in some cases because condition groups do not have names. */
+	UPROPERTY()
+	FGuid Id;
+	
 	/** The operation type that the condition group is using. */
 	UPROPERTY(EditAnywhere, Category="General")
 	EMovieGraphConditionGroupOpType OpType;
diff --git a/Engine/Plugins/MovieScene/MovieRenderPipeline/Source/MovieRenderPipelineEditor/Private/Customizations/Graph/MovieGraphCollectionsCustomization.cpp b/Engine/Plugins/MovieScene/MovieRenderPipeline/Source/MovieRenderPipelineEditor/Private/Customizations/Graph/MovieGraphCollectionsCustomization.cpp
--- a/Engine/Plugins/MovieScene/MovieRenderPipeline/Source/MovieRenderPipelineEditor/Private/Customizations/Graph/MovieGraphCollectionsCustomization.cpp
+++ b/Engine/Plugins/MovieScene/MovieRenderPipeline/Source/MovieRenderPipelineEditor/Private/Customizations/Graph/MovieGraphCollectionsCustomization.cpp
@@ -1,4 +1,4 @@
-﻿// Copyright Epic Games, Inc. All Rights Reserved.
+// Copyright Epic Games, Inc. All Rights Reserved.
 
 #include "MovieGraphCollectionsCustomization.h"
 
@@ -9,6 +9,7 @@
 #include "Framework/MultiBox/MultiBoxBuilder.h"
 #include "Graph/MoviePipelineRenderLayerSubsystem.h"
 #include "Graph/Nodes/MovieGraphCollectionNode.h"
+#include "IDetailChildrenBuilder.h"
 #include "MovieRenderPipelineStyle.h"
 #include "PropertyHandle.h"
 #include "ScopedTransaction.h"
@@ -24,222 +25,583 @@
 
 #define LOCTEXT_NAMESPACE "MoviePipelineEditorCollectionCustomization"
 
-void SMovieGraphCollectionTreeViewIndent::Construct(const FArguments& InArgs)
+FMovieGraphCollectionDragDropOp::FMovieGraphCollectionDragDropOp(const bool bIsConditionGroup, const int32 InitialIndex, TWeakObjectPtr<UMovieGraphConditionGroup> InWeakOwningConditionGroup)
+	: bIsConditionGroup(bIsConditionGroup)
+	, InitialIndex(InitialIndex)
+	, OwningConditionGroup(InWeakOwningConditionGroup)
 {
-	const int NumIndents = InArgs._NumIndents.Get();
-	const TSharedRef<SHorizontalBox> IndentContainer = SNew(SHorizontalBox);
-	
-	ChildSlot
-	.VAlign(VAlign_Fill)
-	[
-		// A single indent doesn't change the background color; more indents get a darker background
-		SNew(SBorder)
-		.Padding(0)
-		.VAlign(VAlign_Fill)
-		.BorderImage((NumIndents <= 1)
-			? FAppStyle::Get().GetBrush("NoBorder")
-			: FMovieRenderPipelineStyle::Get().GetBrush("MovieRenderGraph.CollectionsTree.IndentedRowBackground"))
-		[
-			IndentContainer
-		]
-	];
+	MouseCursor = EMouseCursor::GrabHandClosed;
+}
+
+void FMovieGraphCollectionDragDropOp::Init()
+{
+	SetValidTarget(false);
+	SetupDefaults();
+	Construct();
+}
 
-	// Add a new indent for each level of indent requested
-	for (int32 Index = 0; Index < NumIndents; ++Index)
+void FMovieGraphCollectionDragDropOp::SetValidTarget(const bool bIsValidTarget)
+{
+	if (bIsValidTarget)
 	{
-		IndentContainer->AddSlot()
-		.AutoWidth()
-		.Padding(0)
-		.VAlign(VAlign_Fill)
-		[
-			SNew(SBorder)
-			.BorderImage(FAppStyle::Get().GetBrush("DetailsView.ArrayDropShadow"))
-			[
-				SNew(SBox)
-				.WidthOverride(16.f)
-			]
-		];
+		if (bIsConditionGroup)
+		{
+			CurrentHoverText = LOCTEXT("DragConditionGroup", "Move 1 condition group here");
+		}
+		else
+		{
+			CurrentHoverText = LOCTEXT("DragConditionGroupQuery", "Move 1 condition group query here");
+		}
+		
+		CurrentIconBrush = FAppStyle::GetBrush("Graph.ConnectorFeedback.OK");
+	}
+	else
+	{
+		if (bIsConditionGroup)
+		{
+			CurrentHoverText = LOCTEXT("DragConditionGroup_Invalid", "Cannot move condition group here");
+		}
+		else
+		{
+			CurrentHoverText = LOCTEXT("DragConditionGroupQuery_Invalid", "Cannot move condition group query here");
+		}
+		
+		CurrentIconBrush = FAppStyle::GetBrush("Graph.ConnectorFeedback.Error");
 	}
 }
 
-TSharedRef<FMovieGraphCollectionTreeDragOp> FMovieGraphCollectionTreeDragOp::New(const TSharedPtr<FCollectionCustomizationTreeElement>& InElement)
+int32 FMovieGraphCollectionDragDropOp::GetIndex() const
 {
-	TSharedRef<FMovieGraphCollectionTreeDragOp> Operation = MakeShared<FMovieGraphCollectionTreeDragOp>();
-	Operation->Element = InElement;
-	Operation->Construct();
-	
-	return Operation;
+	return InitialIndex;
 }
 
-void SMovieGraphCollectionTreeView::Construct(const FArguments& InArgs)
+bool FMovieGraphCollectionDragDropOp::IsConditionGroup() const
 {
-	STreeView<TSharedPtr<FCollectionCustomizationTreeElement>>::FArguments SuperArgs;
-
-	DetailBuilder = InArgs._DetailBuilder;
-	Collection = InArgs._Collection;
+	return bIsConditionGroup;
+}
 
-	SuperArgs
-	.ItemHeight(28)
-	.TreeItemsSource(&RootElements)
-	.SelectionMode(ESelectionMode::Single)
-	.OnGenerateRow_Lambda([this](const TSharedPtr<FCollectionCustomizationTreeElement>& TreeItem, const TSharedPtr<STableViewBase>& OwnerTable)
-	{
-		return GenerateTreeRow(TreeItem, OwnerTable.ToSharedRef());
-	})
-	.OnGetChildren(this, &SMovieGraphCollectionTreeView::GetChildrenForTree)
-	.OnExpansionChanged(this, &SMovieGraphCollectionTreeView::OnExpansionChanged)
-	.HeaderRow
-	(
-		SNew(SHeaderRow)
-		.Visibility(EVisibility::Collapsed)
-
-		+ SHeaderRow::Column(SMovieGraphCollectionTreeItem::ColumnID_Name)
-		.DefaultLabel(LOCTEXT("NameColumnLabel", "Name"))
-		.SortMode(EColumnSortMode::None)
-		.FillWidth(0.5f)
-
-		+ SHeaderRow::Column(SMovieGraphCollectionTreeItem::ColumnID_Value)
-		.DefaultLabel(LOCTEXT("ValueColumnLabel", "Value"))
-		.SortMode(EColumnSortMode::None)
-		.FillWidth(0.5f)
-	);
-	
-	STreeView<TSharedPtr<FCollectionCustomizationTreeElement>>::Construct(SuperArgs);
+TWeakObjectPtr<UMovieGraphConditionGroup> FMovieGraphCollectionDragDropOp::GetOwningConditionGroup() const
+{
+	return OwningConditionGroup;
 }
 
-void SMovieGraphCollectionTreeView::OnExpansionChanged(const TSharedPtr<FCollectionCustomizationTreeElement> InElement, const bool bInExpanded)
+
+FMovieGraphCollectionDragDropHandler::FMovieGraphCollectionDragDropHandler(const bool bIsConditionGroup, IDetailLayoutBuilder* InDetailBuilder, UMovieGraphCollection* InCollection, const int32 ConditionGroupIndex, TWeakObjectPtr<UMovieGraphConditionGroup> InWeakConditionGroup, const int32 ConditionGroupQueryIndex)
+	: bIsConditionGroup(bIsConditionGroup)
+	, DetailLayoutBuilder(InDetailBuilder)
+	, WeakCollection(InCollection)
+	, ConditionGroupIndex(ConditionGroupIndex)
+	, WeakConditionGroup(InWeakConditionGroup)
+	, ConditionGroupQueryIndex(ConditionGroupQueryIndex)
 {
-	if (bInExpanded)
-	{
-		ExpandedElements.Add(InElement->GetHash());
-	}
-	else
-	{
-		ExpandedElements.Remove(InElement->GetHash());
-	}
 }
 
-void SMovieGraphCollectionTreeView::GetChildrenForTree(const TSharedPtr<FCollectionCustomizationTreeElement> InItem, TArray<TSharedPtr<FCollectionCustomizationTreeElement>>& OutChildren)
+TSharedPtr<FDragDropOperation> FMovieGraphCollectionDragDropHandler::CreateDragDropOperation() const
 {
-	OutChildren.Append(InItem->GetChildren());
+	TSharedPtr<FMovieGraphCollectionDragDropOp> DragOp = MakeShared<FMovieGraphCollectionDragDropOp>(bIsConditionGroup, bIsConditionGroup ? ConditionGroupIndex : ConditionGroupQueryIndex, WeakConditionGroup);
+	DragOp->Init();
+	return DragOp;
 }
 
-void SMovieGraphCollectionTreeView::RestoreExpansionStateRecursive(const TSharedPtr<FCollectionCustomizationTreeElement>& InElement)
+bool FMovieGraphCollectionDragDropHandler::AcceptDrop(const FDragDropEvent& DragDropEvent, EItemDropZone DropZone) const
 {
-	SetItemExpansion(InElement, ExpandedElements.Contains(InElement->GetHash()));
+	const TSharedPtr<FMovieGraphCollectionDragDropOp> DragOp = DragDropEvent.GetOperationAs<FMovieGraphCollectionDragDropOp>();
+	if (!DragOp.IsValid() || DropZone == EItemDropZone::OntoItem)
+	{
+		return false;
+	}
+
+	// The index of the condition group or query being dragged
+	const int32 SourceIndex = DragOp->GetIndex();
 	
-	for (const TSharedPtr<FCollectionCustomizationTreeElement>& ChildElement : InElement->GetChildren())
+	if (bIsConditionGroup && WeakCollection.IsValid())
 	{
-		RestoreExpansionStateRecursive(ChildElement);
+		const TArray<UMovieGraphConditionGroup*>& ConditionGroups = WeakCollection->GetConditionGroups();
+		if (!ConditionGroups.IsValidIndex(SourceIndex))
+		{
+			return false;
+		}
+
+		UMovieGraphConditionGroup* TargetConditionGroup = ConditionGroups[SourceIndex]; 
+
+		FScopedTransaction Transaction(LOCTEXT("ReorderConditionGroup", "Reorder Condition Group"));
+		const int32 NewIndex = DropZone == EItemDropZone::AboveItem ? ConditionGroupIndex : ConditionGroupIndex + 1;
+		WeakCollection->MoveConditionGroupToIndex(TargetConditionGroup, NewIndex);
+	}
+	else if (!bIsConditionGroup && WeakConditionGroup.IsValid())
+	{
+		const TArray<UMovieGraphConditionGroupQueryBase*>& ConditionGroupQueries = WeakConditionGroup->GetQueries();
+		if (!ConditionGroupQueries.IsValidIndex(SourceIndex))
+		{
+			return false;
+		}
+
+		UMovieGraphConditionGroupQueryBase* TargetQuery = ConditionGroupQueries[SourceIndex]; 
+
+		FScopedTransaction Transaction(LOCTEXT("ReorderConditionGroupQuery", "Reorder Condition Group Query"));
+		const int32 NewIndex = DropZone == EItemDropZone::AboveItem ? ConditionGroupQueryIndex : ConditionGroupQueryIndex + 1;
+		WeakConditionGroup->MoveQueryToIndex(TargetQuery, NewIndex);
 	}
+
+	DetailLayoutBuilder->ForceRefreshDetails();
+	return true;
 }
 
-TSharedRef<ITableRow> SMovieGraphCollectionTreeView::GenerateTreeRow(TSharedPtr<FCollectionCustomizationTreeElement> InTreeElement, const TSharedRef<STableViewBase>& OwnerTable)
+TOptional<EItemDropZone> FMovieGraphCollectionDragDropHandler::CanAcceptDrop(const FDragDropEvent& DragDropEvent, EItemDropZone DropZone) const
 {
-	if (InTreeElement->Type == FCollectionCustomizationTreeElement::EElementType::QueryProperty)
+	TOptional<EItemDropZone> ReturnedDropZone;
+	const TSharedPtr<FMovieGraphCollectionDragDropOp> DragDropOp = DragDropEvent.GetOperationAs<FMovieGraphCollectionDragDropOp>();
+
+	// Can't drag an element onto another element -- only changing order is allowed (above/below)
+	if (!DragDropOp.IsValid() || (DropZone == EItemDropZone::OntoItem))
+	{
+		return ReturnedDropZone;
+	}
+
+	if (bIsConditionGroup)
 	{
-		if (InTreeElement->Query && !InTreeElement->Query->GetWidgets().IsEmpty())
+		// Can't drop a condition group into a query
+		const bool bIsTargetValid = DragDropOp->IsConditionGroup();
+		
+		if (bIsTargetValid && WeakCollection.IsValid() && WeakCollection->GetConditionGroups().IsValidIndex(ConditionGroupIndex))
 		{
-			return SNew(SMovieGraphCollectionTreeItem_Stretch, SharedThis(this), InTreeElement);
+			ReturnedDropZone = DropZone;
+			DragDropOp->SetValidTarget(true);
 		}
 	}
-
-	if ((InTreeElement->Type == FCollectionCustomizationTreeElement::EElementType::EmptyCollection) ||
-		(InTreeElement->Type == FCollectionCustomizationTreeElement::EElementType::EmptyConditionGroup))
+	else
 	{
-		return SNew(SMovieGraphCollectionTreeItem_Stretch, SharedThis(this), InTreeElement);
+		// Can't drop a query into a condition group, and can't drop a query into a condition group that didn't originally own it
+		const bool bIsTargetValid = !DragDropOp->IsConditionGroup() && (WeakConditionGroup == DragDropOp->GetOwningConditionGroup());
+		
+		if (bIsTargetValid && WeakConditionGroup.IsValid() && WeakConditionGroup->GetQueries().IsValidIndex(ConditionGroupQueryIndex))
+		{
+			ReturnedDropZone = DropZone;
+			DragDropOp->SetValidTarget(true);
+		}
 	}
-	
-	return SNew(SMovieGraphCollectionTreeItem, SharedThis(this), InTreeElement, DetailBuilder);
+
+	return ReturnedDropZone;
 }
 
-void SMovieGraphCollectionTreeView::AddRootElement()
+FMovieGraphConditionGroupQueryBuilder::FMovieGraphConditionGroupQueryBuilder(IDetailLayoutBuilder* InDetailBuilder, TSharedRef<IPropertyHandle> InConditionGroupQueryProperty, int32 ConditionGroupQueryIndex, const TWeakObjectPtr<UMovieGraphConditionGroup>& InWeakConditionGroup)
+	: DetailLayoutBuilder(InDetailBuilder)
+	, ConditionGroupQueryProperty(InConditionGroupQueryProperty)
+	, ConditionGroupQueryIndex(ConditionGroupQueryIndex)
+	, WeakConditionGroup(InWeakConditionGroup)
 {
-	const TSharedPtr<FCollectionCustomizationTreeElement> RootElement = MakeShared<FCollectionCustomizationTreeElement>(FCollectionCustomizationTreeElement::EElementType::Root);
-	RootElement->Collection = Collection.Get();
-	
-	RootElements.Add(RootElement);
-
-	// Expand the root element by default because it is never visible and the user cannot expand it themselves
-	constexpr bool bIsExpanded = true;
-	SetItemExpansion(RootElement, bIsExpanded);
-	ExpandedElements.Add(RootElement->GetHash());
 }
 
-void SMovieGraphCollectionTreeView::AddRootConditionGroupElements(const TObjectPtr<UMovieGraphCollection>& InCollection) const
+FName FMovieGraphConditionGroupQueryBuilder::GetName() const
 {
-	// Add the collection's condition groups. These will be the root elements of the tree, and each condition group will take
-	// over adding its own children to the tree.
-	const TArray<UMovieGraphConditionGroup*> ConditionGroups = InCollection->GetConditionGroups();
-	for (int32 Index = 0; Index < ConditionGroups.Num(); ++Index)
-	{
-		AddNewConditionGroupElement(ConditionGroups[Index]);
-	}
+	// Name doesn't matter here because expansion state doesn't need to be persisted (this node does not have any children nested under it)
+	return NAME_None;
 }
 
-void SMovieGraphCollectionTreeView::AddNewConditionGroupElement(UMovieGraphConditionGroup* InConditionGroup) const
+void FMovieGraphConditionGroupQueryBuilder::GenerateHeaderRowContent(FDetailWidgetRow& NodeRow)
 {
-	const TSharedPtr<FCollectionCustomizationTreeElement> NewElement = MakeShared<FCollectionCustomizationTreeElement>(FCollectionCustomizationTreeElement::EElementType::ConditionGroup);
-	NewElement->ConditionGroup = InConditionGroup;
-	NewElement->Collection = Collection.Get();
+
 }
 
-void SMovieGraphCollectionTreeView::AddConditionGroup()
+void FMovieGraphConditionGroupQueryBuilder::GenerateChildContent(IDetailChildrenBuilder& ChildrenBuilder)
 {
-	if (Collection.IsValid())
+	if (!WeakConditionGroup.IsValid())
 	{
-		UMovieGraphConditionGroup* NewConditionGroup = Collection.Get()->AddConditionGroup();
-		
-		AddNewConditionGroupElement(NewConditionGroup);
+		return;
+	}
+	
+	const TArray<UMovieGraphConditionGroupQueryBase*>& Queries = WeakConditionGroup->GetQueries();
+	
+	const TSharedPtr<IPropertyHandleArray> QueriesArrayProperty = ConditionGroupQueryProperty->AsArray();
+	check(QueriesArrayProperty);
+	
+	uint32 NumQueries = 0;
+	QueriesArrayProperty->GetNumElements(NumQueries);
+	
+	for (uint32 QueryIndex = 0; QueryIndex < NumQueries; ++QueryIndex)
+	{
+		const TSharedRef<IPropertyHandle> ElementHandle = QueriesArrayProperty->GetElement(QueryIndex);
+		check(ElementHandle->IsValidHandle());
+		check(Queries.IsValidIndex(QueryIndex));
+
+		UMovieGraphConditionGroupQueryBase* Query = Queries[QueryIndex];
+		TWeakObjectPtr<UMovieGraphConditionGroupQueryBase> WeakQuery = MakeWeakObjectPtr(Query);
+
+		// Drag-drop handler for condition group queries doesn't provide condition group information
+		TWeakObjectPtr<UMovieGraphConditionGroup> InvalidConditionGroup;
+		constexpr int32 InvalidConditionGroupIndex = 0;
+		constexpr UMovieGraphCollection* InvalidCollection = nullptr;
+		constexpr bool bIsConditionGroup = false;
+
+		ChildrenBuilder.AddProperty(ElementHandle)
+		.DragDropHandler(MakeShared<FMovieGraphCollectionDragDropHandler>(bIsConditionGroup, DetailLayoutBuilder, InvalidCollection, InvalidConditionGroupIndex, WeakConditionGroup, QueryIndex))
+		.ShowPropertyButtons(false)
+		.OverrideResetToDefault(FResetToDefaultOverride::Hide())
+		.CustomWidget()
+		.FilterString(Query->GetClass()->GetDisplayNameText())
+		.NameContent()
+		.HAlign(HAlign_Fill)
+		[
+			SNew(SHorizontalBox)
+			+ SHorizontalBox::Slot()
+			.AutoWidth()
+			.VAlign(VAlign_Center)
+			.Padding(0, 0, 4.f, 0)
+			[
+				SNew(SCheckBox)
+				.IsChecked_Lambda([WeakQuery]()
+				{
+					return WeakQuery.IsValid() && WeakQuery->IsEnabled() ? ECheckBoxState::Checked : ECheckBoxState::Unchecked;
+				})
+				.OnCheckStateChanged_Lambda([WeakQuery](ECheckBoxState NewCheckState)
+				{
+					if (WeakQuery.IsValid())
+					{
+						WeakQuery->SetEnabled(NewCheckState == ECheckBoxState::Checked);
+					}
+				})
+			]
 
-		if (!RootElements.IsEmpty())
-		{
-			// Needed so the "empty collection" notice is cleared out (if this is the first condition group being added to the collection).
-			RootElements[0]->ClearCachedChildren();
-		}
+			+ SHorizontalBox::Slot()
+			.VAlign(VAlign_Center)
+			.HAlign(HAlign_Fill)
+			.Padding(0, 0, 7.f, 0)
+			[
+				SNew(SMovieGraphCollectionTreeQueryTypeSelectorWidget)
+				.WeakConditionGroup(WeakConditionGroup)
+				.WeakQuery(WeakQuery)
+				.OnQueryTypeChanged_Lambda([this](TWeakObjectPtr<UMovieGraphConditionGroupQueryBase> NewQuery)
+				{
+					// The other widgets will be referring to an outdated query object (since changing the type generates a completely new
+					// query), so the UI needs to be refreshed
+					DetailLayoutBuilder->ForceRefreshDetails();
+				})
+			]
+		]
+		.ValueContent()
+		.HAlign(HAlign_Fill)
+		[
+			SNew(SHorizontalBox)
+			
+			+ SHorizontalBox::Slot()
+			.VAlign(VAlign_Center)
+			.HAlign(HAlign_Left)
+			.Padding(4.f, 0, 0, 0)
+			.AutoWidth()
+			[
+				SNew(SMovieGraphCollectionTreeAddQueryContentWidget)
+				.WeakQuery(WeakQuery)
+			]
+
+			+ SHorizontalBox::Slot()
+			.VAlign(VAlign_Center)
+			.HAlign(HAlign_Left)
+			.AutoWidth()
+			[
+				SNew(SComboButton)
+				.IsEnabled_Lambda([WeakQuery]()
+				{
+					return WeakQuery.IsValid() && WeakQuery->IsEnabled();
+				})
+				.ComboButtonStyle(FMovieRenderPipelineStyle::Get(), "MovieRenderGraph.CollectionsTree.SmallComboButton")
+				.ContentPadding(0)
+				.ForegroundColor(FSlateColor::UseForeground())
+				.HasDownArrow(true)
+				.OnGetMenuContent_Lambda([this, WeakQuery]()
+				{
+					FMenuBuilder ViewOptions(true, nullptr);
+
+					ViewOptions.AddMenuEntry(
+						LOCTEXT("DeleteConditionGroupQuery", "Delete"),
+						LOCTEXT("DeleteConditionGroupQueryTooltip", "Delete this condition group query."),
+						FSlateIcon(),
+						FUIAction(
+							FExecuteAction::CreateLambda([this, WeakQuery]()
+							{
+								if (WeakConditionGroup.IsValid() && WeakQuery.IsValid())
+								{
+									WeakConditionGroup->RemoveQuery(WeakQuery.Get());
+								}
+							}),
+							FCanExecuteAction()));
+
+					return ViewOptions.MakeWidget();
+				})
+			]
+
+			+ SHorizontalBox::Slot()
+			.VAlign(VAlign_Center)
+			.HAlign(HAlign_Right)
+			.Padding(2.f)
+			.FillWidth(1.f)
+			[
+				SNew(SMovieGraphCollectionTreeOpTypeWidget)
+				.IsConditionGroup(false)
+				.WeakConditionGroupQuery(WeakQuery)
+			]
+		];
 		
-		RefreshView();
+		// Add the widgetry that allows the query's properties to be updated
+		AddQueryTypeWidgets(WeakQuery, ChildrenBuilder);
 	}
+
+	// Refresh details when the number of queries changes
+	IDetailCategoryBuilder& ParentCategory = ChildrenBuilder.GetParentCategory();
+	QueriesArrayProperty->SetOnNumElementsChanged(FSimpleDelegate::CreateSPLambda(this, [&ParentCategory]()
+	{
+		ParentCategory.GetParentLayout().ForceRefreshDetails();
+	}));
 }
 
-void SMovieGraphCollectionTreeView::RemoveConditionGroup(const TWeakPtr<FCollectionCustomizationTreeElement>& InConditionGroupElement)
+void FMovieGraphConditionGroupQueryBuilder::AddQueryTypeWidgets(TWeakObjectPtr<UMovieGraphConditionGroupQueryBase> WeakQuery, IDetailChildrenBuilder& ChildrenBuilder) const
 {
-	if (!InConditionGroupElement.IsValid())
+	if (!WeakQuery.IsValid())
 	{
 		return;
 	}
-	
-	const TSharedPtr<FCollectionCustomizationTreeElement> TreeElementPin = InConditionGroupElement.Pin();
-	if (!TreeElementPin->Collection)
+
+	// If the query defines widgets to display, use those
+	const TArray<TSharedRef<SWidget>> QueryWidgets = WeakQuery->GetWidgets();
+	if (!QueryWidgets.IsEmpty())
 	{
+		TSharedPtr<SVerticalBox> WidgetContainer;
+
+		ChildrenBuilder.AddCustomRow(FText())
+		.WholeRowContent()
+		[
+			SAssignNew(WidgetContainer, SVerticalBox)
+		];
+
+		// Add all widgets that are generated by the column
+		for (const TSharedRef<SWidget>& Widget : WeakQuery->GetWidgets())
+		{
+			WidgetContainer->AddSlot()
+				.HAlign(HAlign_Fill)
+				.VAlign(VAlign_Center)
+				[
+					SNew(SBox)
+					.IsEnabled_Lambda([WeakQuery]()
+					{
+						return WeakQuery.IsValid() && WeakQuery->IsEnabled();
+					})
+					[
+						Widget
+					]
+				];
+		}
+
 		return;
 	}
 
-	TreeElementPin->Collection->RemoveConditionGroup(TreeElementPin->ConditionGroup);
+	// If the query does NOT define widgets to display, just display the properties that the query exposes
+	for (TFieldIterator<FProperty> PropIt(WeakQuery->GetClass()); PropIt; ++PropIt)
+	{
+		const FProperty* QueryProperty = *PropIt;
+		if (!QueryProperty->HasAllPropertyFlags(CPF_Edit))
+		{
+			continue;
+		}
+		
+		// Get the property handle for the property within the query. This is needed in order to get the name/value widgets for it.
+		TSharedPtr<IPropertyHandle> QueryPropertyHandle =
+			DetailLayoutBuilder->AddObjectPropertyData({QueryProperty->GetOwnerUObject()}, QueryProperty->GetFName());
+
+		if (!QueryPropertyHandle || !QueryPropertyHandle.IsValid())
+		{
+			continue;
+		}
+
+		ChildrenBuilder.AddProperty(QueryPropertyHandle.ToSharedRef())
+		.CustomWidget()
+		.NameContent()
+		[
+			SNew(SBox)
+			.Padding(0)
+			.IsEnabled_Lambda([WeakQuery]()
+			{
+				return WeakQuery.IsValid() && WeakQuery->IsEnabled();
+			})
+			[
+				QueryPropertyHandle->CreatePropertyNameWidget()
+			]
+		]
+		.ValueContent()
+		[
+			SNew(SBox)
+			.Padding(5.f, 2.f)
+			.IsEnabled_Lambda([WeakQuery]()
+			{
+				return WeakQuery.IsValid() && WeakQuery->IsEnabled();
+			})
+			[
+				QueryPropertyHandle->CreatePropertyValueWidget()
+			]
+		];
+	}
+}
+
+FMovieGraphConditionGroupBuilder::FMovieGraphConditionGroupBuilder(IDetailLayoutBuilder* InDetailBuilder, TSharedRef<IPropertyHandle> InConditionGroupProperty, const uint32 InConditionGroupIndex, TWeakObjectPtr<UMovieGraphCollection>& InWeakCollection)
+	: DetailLayoutBuilder(InDetailBuilder)
+	, ConditionGroupProperty(InConditionGroupProperty)
+	, ConditionGroupIndex(InConditionGroupIndex)
+	, WeakCollection(InWeakCollection)
+{
+}
 
-	if (TreeElementPin->ParentElement.IsValid())
+FName FMovieGraphConditionGroupBuilder::GetName() const
+{
+	if (WeakCollection.IsValid())
 	{
-		TreeElementPin->ParentElement.Pin()->ClearCachedChildren();
+		const TArray<UMovieGraphConditionGroup*>& ConditionGroups = WeakCollection->GetConditionGroups();
+		if (ConditionGroups.IsValidIndex(ConditionGroupIndex))
+		{
+			return FName(ConditionGroups[ConditionGroupIndex]->GetId().ToString());
+		}
 	}
 
-	RefreshView();
+	return NAME_None;
 }
 
-void SMovieGraphCollectionTreeView::RefreshView()
+void FMovieGraphConditionGroupBuilder::GenerateHeaderRowContent(FDetailWidgetRow& NodeRow)
 {
-	RequestTreeRefresh();
-    
-    // Restore tree expansion state after the refresh
-    for (const TSharedPtr<FCollectionCustomizationTreeElement>& RootElement : RootElements)
-    {
-    	RestoreExpansionStateRecursive(RootElement);
-    }
+	// Generate the condition group name
+	const FString GroupName = FString::Printf(TEXT("Condition Group %i"), (ConditionGroupIndex + 1));
+
+	// Drag-drop handler for condition groups doesn't provide query information
+	TWeakObjectPtr<UMovieGraphConditionGroup> InvalidQuery;
+	constexpr int32 InvalidQueryIndex = 0;
+	constexpr bool bIsConditionGroup = true;
+	
+	NodeRow
+	.DragDropHandler(MakeShared<FMovieGraphCollectionDragDropHandler>(bIsConditionGroup, DetailLayoutBuilder, WeakCollection.Get(), ConditionGroupIndex, InvalidQuery, InvalidQueryIndex))
+	.OverrideResetToDefault(FResetToDefaultOverride::Hide())
+	.FilterString(FText::FromString(GroupName))
+	.RowTag(GetName())
+	.PropertyHandleList({ConditionGroupProperty})
+	.NameContent()
+	[
+		SNew(SHorizontalBox)
+
+		+ SHorizontalBox::Slot()
+		.AutoWidth()
+		.VAlign(VAlign_Center)
+		.Padding(0, 0, 2.f, 0)
+		[
+			SNew(SImage)
+			.Image(FAppStyle::GetBrush("ClassIcon.GroupActor"))
+			.ColorAndOpacity(FSlateColor::UseForeground())
+		]
+
+		+ SHorizontalBox::Slot()
+		.AutoWidth()
+		.VAlign(VAlign_Center)
+		.Padding(0, 0, 2.f, 0)
+		[
+			SNew(STextBlock)
+			.Text(FText::FromString(GroupName))
+			.Font(FAppStyle::GetFontStyle(TEXT("PropertyWindow.BoldFont")))
+		]
+	]
+	.ValueContent()
+	.HAlign(HAlign_Fill)
+	.VAlign(VAlign_Center)
+	[
+		SNew(SHorizontalBox)
+		
+		+ SHorizontalBox::Slot()
+		.VAlign(VAlign_Center)
+		.HAlign(HAlign_Left)
+		.AutoWidth()
+		[
+			SNew(SButton)
+			.HAlign(HAlign_Center)
+			.ToolTip(SNew(SToolTip).Text(LOCTEXT("AddConditionGroupQuery", "Add a condition to this condition group.")))
+			.ButtonStyle(FAppStyle::Get(), "SimpleButton")
+			.OnClicked_Lambda([this]()
+			{
+				if (WeakCollection.IsValid())
+				{
+					const TArray<UMovieGraphConditionGroup*>& ConditionGroups = WeakCollection->GetConditionGroups();
+					if (ConditionGroups.IsValidIndex(ConditionGroupIndex))
+					{
+						// Use the Actor Name query as the default
+						ConditionGroups[ConditionGroupIndex]->AddQuery(UMovieGraphConditionGroupQuery_ActorName::StaticClass());
+					}
+				}
+
+				return FReply::Handled();
+			})
+			.Content()
+			[
+				SNew(SImage)
+				.Image(FAppStyle::GetBrush("Icons.PlusCircle"))
+				.ColorAndOpacity(FSlateColor::UseForeground())
+			]
+		]
+
+		+ SHorizontalBox::Slot()
+		.VAlign(VAlign_Center)
+		.HAlign(HAlign_Left)
+		.AutoWidth()
+		[
+			SNew(SComboButton)
+			.ComboButtonStyle(FMovieRenderPipelineStyle::Get(), "MovieRenderGraph.CollectionsTree.SmallComboButton")
+			.ContentPadding(0)
+			.ForegroundColor(FSlateColor::UseForeground())
+			.HasDownArrow(true)
+			.OnGetMenuContent_Lambda([this]()
+			{
+				FMenuBuilder ViewOptions(true, nullptr);
+
+				ViewOptions.AddMenuEntry(
+					LOCTEXT("DeleteConditionGroup", "Delete"),
+					LOCTEXT("DeleteConditionGroupTooltip", "Delete this condition group and the queries in it."),
+					FSlateIcon(),
+					FUIAction(
+						FExecuteAction::CreateLambda([this]()
+						{
+							if (WeakCollection.IsValid())
+							{
+								const TArray<UMovieGraphConditionGroup*>& ConditionGroups = WeakCollection->GetConditionGroups();
+								if (ConditionGroups.IsValidIndex(ConditionGroupIndex))
+								{
+									WeakCollection->RemoveConditionGroup(ConditionGroups[ConditionGroupIndex]);
+								}
+							}
+                        }),
+						FCanExecuteAction()));
+
+				return ViewOptions.MakeWidget();
+			})
+		]
+
+		+ SHorizontalBox::Slot()
+		.VAlign(VAlign_Center)
+		.HAlign(HAlign_Right)
+		.Padding(2.f)
+		.FillWidth(1.f)
+		[
+			SNew(SMovieGraphCollectionTreeOpTypeWidget)
+			.IsConditionGroup(true)
+			.WeakConditionGroup(WeakCollection->GetConditionGroups()[ConditionGroupIndex])
+		]
+	];
+}
+
+void FMovieGraphConditionGroupBuilder::GenerateChildContent(IDetailChildrenBuilder& ChildrenBuilder)
+{
+	UMovieGraphConditionGroup* ConditionGroup = WeakCollection.Get()->GetConditionGroups()[ConditionGroupIndex];
+	const TSharedPtr<IPropertyHandle> QueriesProperty = DetailLayoutBuilder->AddObjectPropertyData({ConditionGroup}, FName(TEXT("Queries")));
+	
+	ChildrenBuilder.AddCustomBuilder(
+		MakeShareable(new FMovieGraphConditionGroupQueryBuilder(DetailLayoutBuilder, QueriesProperty.ToSharedRef(), ConditionGroupIndex, MakeWeakObjectPtr(ConditionGroup))));
 }
 
 void SMovieGraphCollectionTreeOpTypeWidget::Construct(const FArguments& InArgs)
 {
 	bIsConditionGroup = InArgs._IsConditionGroup.Get();
-	WeakTreeElement = InArgs._WeakTreeElement.Get();
+	WeakConditionGroup = InArgs._WeakConditionGroup.Get();
+	WeakConditionGroupQuery = InArgs._WeakConditionGroupQuery.Get();
 	
 	ChildSlot
 	[
@@ -324,24 +686,16 @@ TSharedRef<SWidget> SMovieGraphCollectionTreeOpTypeWidget::GetOpTypeContents(con
 
 bool SMovieGraphCollectionTreeOpTypeWidget::IsWidgetEnabled() const
 {
-	if (WeakTreeElement.IsValid())
+	if (bIsConditionGroup && WeakConditionGroup.IsValid())
 	{
-		if (bIsConditionGroup)
-		{
-			if (const TObjectPtr<UMovieGraphConditionGroup> ConditionGroup = WeakTreeElement.Pin()->ConditionGroup)
-			{
-				// Disable the first condition group
-				return !ConditionGroup->IsFirstConditionGroup();
-			}
-		}
-		else
-		{
-			if (const TObjectPtr<UMovieGraphConditionGroupQueryBase> ConditionGroupQuery = WeakTreeElement.Pin()->Query)
-			{
-				// Disable the first query, or disable it if it was flagged as disabled
-				return !ConditionGroupQuery->IsFirstConditionGroupQuery() && ConditionGroupQuery->IsEnabled();
-			}
-		}
+		// Disable the first condition group
+		return !WeakConditionGroup->IsFirstConditionGroup();
+	}
+	
+	if (WeakConditionGroupQuery.IsValid())
+	{
+		// Disable the first query, or disable it if it was flagged as disabled
+		return !WeakConditionGroupQuery->IsFirstConditionGroupQuery() && WeakConditionGroupQuery->IsEnabled();
 	}
 	
 	return true;
@@ -349,41 +703,27 @@ bool SMovieGraphCollectionTreeOpTypeWidget::IsWidgetEnabled() const
 
 void SMovieGraphCollectionTreeOpTypeWidget::SetOpType(const FName InNewOpType, ESelectInfo::Type SelectInfo) const
 {
-	if (!WeakTreeElement.IsValid())
-	{
-		return;
-	}
-	
-	if (bIsConditionGroup)
+	if (bIsConditionGroup && WeakConditionGroup.IsValid())
 	{
-		if (const TObjectPtr<UMovieGraphConditionGroup>& ConditionGroup = WeakTreeElement.Pin()->ConditionGroup)
-		{
-			const UEnum* OpTypeEnum = StaticEnum<EMovieGraphConditionGroupOpType>();
-			ConditionGroup->SetOperationType(static_cast<EMovieGraphConditionGroupOpType>(OpTypeEnum->GetValueByName(InNewOpType)));
-		}
+		const UEnum* OpTypeEnum = StaticEnum<EMovieGraphConditionGroupOpType>();
+		WeakConditionGroup->SetOperationType(static_cast<EMovieGraphConditionGroupOpType>(OpTypeEnum->GetValueByName(InNewOpType)));
 	}
-	else
+	else if (WeakConditionGroupQuery.IsValid())
 	{
-		if (const TObjectPtr<UMovieGraphConditionGroupQueryBase>& Query = WeakTreeElement.Pin()->Query)
-		{
-			const UEnum* OpTypeEnum = StaticEnum<EMovieGraphConditionGroupQueryOpType>();
-			Query->SetOperationType(static_cast<EMovieGraphConditionGroupQueryOpType>(OpTypeEnum->GetValueByName(InNewOpType)));
-		}
+		const UEnum* OpTypeEnum = StaticEnum<EMovieGraphConditionGroupQueryOpType>();
+		WeakConditionGroupQuery->SetOperationType(static_cast<EMovieGraphConditionGroupQueryOpType>(OpTypeEnum->GetValueByName(InNewOpType)));
 	}
 }
 
 FName SMovieGraphCollectionTreeOpTypeWidget::GetCurrentConditionGroupOpType() const
 {
-	if (WeakTreeElement.IsValid())
+	if (WeakConditionGroup.IsValid())
 	{
-		if (const TObjectPtr<UMovieGraphConditionGroup>& ConditionGroup = WeakTreeElement.Pin()->ConditionGroup)
-		{
-			const UEnum* OpTypeEnum = StaticEnum<EMovieGraphConditionGroupOpType>();
-			const EMovieGraphConditionGroupOpType OpType = ConditionGroup->GetOperationType();
-			const FText DisplayNameText = OpTypeEnum->GetDisplayNameTextByValue(static_cast<__underlying_type(EMovieGraphConditionGroupOpType)>(OpType));
-			
-			return FName(DisplayNameText.ToString());
-		}
+		const UEnum* OpTypeEnum = StaticEnum<EMovieGraphConditionGroupOpType>();
+		const EMovieGraphConditionGroupOpType OpType = WeakConditionGroup->GetOperationType();
+		const FText DisplayNameText = OpTypeEnum->GetDisplayNameTextByValue(static_cast<__underlying_type(EMovieGraphConditionGroupOpType)>(OpType));
+		
+		return FName(DisplayNameText.ToString());
 	}
 		
 	return FName();
@@ -391,16 +731,13 @@ FName SMovieGraphCollectionTreeOpTypeWidget::GetCurrentConditionGroupOpType() co
 
 FName SMovieGraphCollectionTreeOpTypeWidget::GetCurrentConditionGroupQueryOpType() const
 {
-	if (WeakTreeElement.IsValid())
+	if (WeakConditionGroupQuery.IsValid())
 	{
-		if (const TObjectPtr<UMovieGraphConditionGroupQueryBase>& Query = WeakTreeElement.Pin()->Query)
-		{
-			const UEnum* OpTypeEnum = StaticEnum<EMovieGraphConditionGroupQueryOpType>();
-			const EMovieGraphConditionGroupQueryOpType OpType = Query->GetOperationType();
-			const FText DisplayNameText = OpTypeEnum->GetDisplayNameTextByValue(static_cast<__underlying_type(EMovieGraphConditionGroupQueryOpType)>(OpType));
-			
-			return FName(DisplayNameText.ToString());
-		}
+		const UEnum* OpTypeEnum = StaticEnum<EMovieGraphConditionGroupQueryOpType>();
+		const EMovieGraphConditionGroupQueryOpType OpType = WeakConditionGroupQuery->GetOperationType();
+		const FText DisplayNameText = OpTypeEnum->GetDisplayNameTextByValue(static_cast<__underlying_type(EMovieGraphConditionGroupQueryOpType)>(OpType));
+		
+		return FName(DisplayNameText.ToString());
 	}
 		
 	return FName();
@@ -408,29 +745,19 @@ FName SMovieGraphCollectionTreeOpTypeWidget::GetCurrentConditionGroupQueryOpType
 
 void SMovieGraphCollectionTreeAddQueryContentWidget::Construct(const FArguments& InArgs)
 {
-	TWeakPtr<FCollectionCustomizationTreeElement> WeakTreeElement = InArgs._WeakTreeElement.Get();
-	TWeakPtr<SMovieGraphCollectionTreeView> WeakTreeView = InArgs._WeakTreeView.Get();
+	TWeakObjectPtr<UMovieGraphConditionGroupQueryBase> WeakQuery = InArgs._WeakQuery.Get();
 
-	if (!WeakTreeElement.IsValid() || !WeakTreeView.IsValid())
+	if (!WeakQuery.IsValid())
 	{
 		return;
 	}
 
-	// Call the underlying query to get the contents of the menu. The query will call back to the widget if something was added so the tree view has
+	// Call the underlying query to get the contents of the menu. The query will call back to the widget if something was added so the UI has
 	// an opportunity to refresh itself.
-	TSharedRef<SWidget> AddMenuContents = WeakTreeElement.Pin()->Query->GetAddMenuContents(
-		UMovieGraphConditionGroupQueryBase::FMovieGraphConditionGroupQueryContentsChanged::CreateLambda([WeakTreeElement, WeakTreeView]()
+	TSharedRef<SWidget> AddMenuContents = WeakQuery->GetAddMenuContents(
+		UMovieGraphConditionGroupQueryBase::FMovieGraphConditionGroupQueryContentsChanged::CreateLambda([]()
     {
-    	if (WeakTreeView.IsValid())
-    	{
-    		if (WeakTreeElement.IsValid())
-    		{
-    			WeakTreeElement.Pin()->ClearCachedChildren();
-    		}
-    		
-    		WeakTreeView.Pin()->RefreshView();
-    		FSlateApplication::Get().DismissAllMenus();
-    	}
+    	FSlateApplication::Get().DismissAllMenus();
     }));
 
     // Generate a (multi-layered) icon for the "Add" menu
@@ -447,9 +774,9 @@ void SMovieGraphCollectionTreeAddQueryContentWidget::Construct(const FArguments&
 		.ComboButtonStyle(FAppStyle::Get(), "SimpleComboButton")
 		.ContentPadding(0)
 		.HasDownArrow(false)
-		.IsEnabled_Lambda([WeakTreeElement]()
+		.IsEnabled_Lambda([WeakQuery]()
 		{
-			return WeakTreeElement.IsValid() && WeakTreeElement.Pin()->IsQueryEnabled();
+			return WeakQuery.IsValid() && WeakQuery->IsEnabled();
 		})
 		.OnGetMenuContent_Lambda([AddMenuContents]()
 		{
@@ -462,22 +789,28 @@ void SMovieGraphCollectionTreeAddQueryContentWidget::Construct(const FArguments&
 		})
 		.ButtonContent()
 		[
-			AddIcon
+			SNew(SBox)
+			.WidthOverride(16)
+			.HeightOverride(16)
+			[
+				AddIcon
+			]
 		]
 	];
 }
 
 void SMovieGraphCollectionTreeQueryTypeSelectorWidget::Construct(const FArguments& InArgs)
 {
-	WeakTreeElement = InArgs._WeakTreeElement.Get();
-	WeakTreeView = InArgs._WeakTreeView.Get();
+	WeakConditionGroup = InArgs._WeakConditionGroup.Get();
+	WeakQuery = InArgs._WeakQuery.Get();
+	OnQueryTypeChanged = InArgs._OnQueryTypeChanged;
 
 	ChildSlot
 	[
 		SNew(SComboBox<UClass*>)
 		.IsEnabled_Lambda([this]()
 		{
-			return WeakTreeElement.IsValid() && WeakTreeElement.Pin()->IsQueryEnabled();
+			return WeakQuery.IsValid() && WeakQuery->IsEnabled();
 		})
 		.OptionsSource(GetAvailableQueryTypes())
 		.OnSelectionChanged(this, &SMovieGraphCollectionTreeQueryTypeSelectorWidget::SetQueryType)
@@ -573,819 +906,45 @@ TSharedRef<SWidget> SMovieGraphCollectionTreeQueryTypeSelectorWidget::GetQueryTy
 
 FText SMovieGraphCollectionTreeQueryTypeSelectorWidget::GetCurrentQueryTypeDisplayName() const
 {
-	if (WeakTreeElement.IsValid())
+	if (const UClass* CurrentQueryType = GetCurrentQueryType())
 	{
-		if (const UClass* CurrentQueryType = GetCurrentQueryType())
-		{
-			check(CurrentQueryType->IsChildOf(UMovieGraphConditionGroupQueryBase::StaticClass()));
-			
-			return Cast<UMovieGraphConditionGroupQueryBase>(CurrentQueryType->GetDefaultObject())->GetDisplayName();
-		}
-	}
+		check(CurrentQueryType->IsChildOf(UMovieGraphConditionGroupQueryBase::StaticClass()));
 		
+		return Cast<UMovieGraphConditionGroupQueryBase>(CurrentQueryType->GetDefaultObject())->GetDisplayName();
+	}
+	
 	return FText();
 }
 
 UClass* SMovieGraphCollectionTreeQueryTypeSelectorWidget::GetCurrentQueryType() const
 {
-	if (WeakTreeElement.IsValid())
+	if (WeakQuery.IsValid())
 	{
-		return WeakTreeElement.Pin()->Query->GetClass();
+		return WeakQuery->GetClass();
 	}
 
 	return nullptr;
 }
 
-void SMovieGraphCollectionTreeQueryTypeSelectorWidget::SetQueryType(UClass* InNewQueryType, ESelectInfo::Type SelectInfo) const
+void SMovieGraphCollectionTreeQueryTypeSelectorWidget::SetQueryType(UClass* InNewQueryType, ESelectInfo::Type SelectInfo)
 {
-	if (!WeakTreeElement.IsValid())
+	if (!WeakConditionGroup.IsValid() || !WeakQuery.IsValid())
 	{
 		return;
 	}
 	
-	const TSharedPtr<FCollectionCustomizationTreeElement> TreeElementPin = WeakTreeElement.Pin();
-
-	// Swap the old query with a new one
-	if (const TObjectPtr<UMovieGraphConditionGroup>& ConditionGroup = TreeElementPin->ConditionGroup)
-	{
-		const auto& ExistingQuery = TreeElementPin->Query;
-		const int32 ExistingQueryIndex = ConditionGroup->GetQueries().Find(ExistingQuery);
-
-		// Add the new query at the index of the current query. If the query couldn't be found for some reason, default to the end of the array (-1)
-		ConditionGroup->AddQuery(InNewQueryType, ExistingQueryIndex != INDEX_NONE ? ExistingQueryIndex : -1);
-	
-		ConditionGroup->RemoveQuery(ExistingQuery);
-	}
-
-	// Changing the query type removes the current query and adds a new one to the parent condition group, which means the parent's children
-	// need to be refreshed
-	if (TreeElementPin->ParentElement.IsValid())
-	{
-		TreeElementPin->ParentElement.Pin()->ClearCachedChildren();			
-	}
-
-	if (WeakTreeView.IsValid())
-	{
-		WeakTreeView.Pin()->RefreshView();
-	}
-}
-
-FCollectionCustomizationTreeElement::FCollectionCustomizationTreeElement(const EElementType InType)
-	: Type(InType)
-{
-		
-}
-
-bool FCollectionCustomizationTreeElement::IsQueryEnabled() const
-{
-	return Query ? Query->IsEnabled() : false;
-}
-
-const TArray<TSharedPtr<FCollectionCustomizationTreeElement>>& FCollectionCustomizationTreeElement::GetChildren() const
-{
-	// Returned the cached children if they are available
-	if (!ChildrenCache.IsEmpty())
-	{
-		return ChildrenCache;
-	}
-
-	/** Creates a new child element of the specified type, with the associated pointers to the objects needed for that type. */
-	auto AddChildElement = [this](
-		const EElementType ElementType, UMovieGraphConditionGroup* InConditionGroup = nullptr,
-		UMovieGraphConditionGroupQueryBase* InQuery = nullptr, FProperty* InQueryProperty = nullptr)
-	{
-		const TSharedPtr<FCollectionCustomizationTreeElement> NewElement = MakeShared<FCollectionCustomizationTreeElement>(ElementType);
-		NewElement->Collection = Collection;
-		NewElement->ConditionGroup = InConditionGroup;
-		NewElement->Query = InQuery;
-		NewElement->QueryProperty = InQueryProperty;
-		NewElement->ParentElement = AsWeak();
-
-		ChildrenCache.Add(NewElement);
-	};
-
-	if (Type == EElementType::Root)
-	{
-		const TArray<UMovieGraphConditionGroup*>& ConditionGroups = Collection->GetConditionGroups();
-
-		if (ConditionGroups.IsEmpty())
-		{
-			AddChildElement(EElementType::EmptyCollection);
-		}
-		else
-		{
-			for (UMovieGraphConditionGroup* Group : ConditionGroups)
-			{
-				AddChildElement(EElementType::ConditionGroup, Group);
-			}
-		}
-	}
-	
-	if (Type == EElementType::ConditionGroup)
-	{
-		const TArray<UMovieGraphConditionGroupQueryBase*>& ConditionGroupQueries = ConditionGroup->GetQueries();
-		
-		if (ConditionGroupQueries.IsEmpty())
-		{
-			AddChildElement(EElementType::EmptyConditionGroup);
-		}
-		else
-		{
-			for (UMovieGraphConditionGroupQueryBase* GroupQuery : ConditionGroup->GetQueries())
-			{
-				AddChildElement(EElementType::Query, ConditionGroup, GroupQuery);
-			}
-		}
-	}
-
-	if (Type == EElementType::Query)
-	{
-		// If the query provides widgets, add a row for each widget. Otherwise, add a row for each EditAnywhere property within the query.
-		const TArray<TSharedRef<SWidget>> QueryWidgets = Query->GetWidgets();
-		if (!QueryWidgets.IsEmpty())
-		{
-			for (const TSharedRef<SWidget>& Widget : QueryWidgets)
-			{
-				FProperty* UnknownProperty = nullptr;
-				AddChildElement(EElementType::QueryProperty, ConditionGroup, Query, UnknownProperty);
-			}
-		}
-		else
-		{
-			for (TFieldIterator<FProperty> PropIt(Query->GetClass()); PropIt; ++PropIt)
-			{
-				FProperty* Prop = *PropIt;
-				if (!Prop->HasAllPropertyFlags(CPF_Edit))
-				{
-					continue;
-				}
-
-				AddChildElement(EElementType::QueryProperty, ConditionGroup, Query, Prop);
-			}
-		}
-	}
-
-	return ChildrenCache;
-}
-
-void FCollectionCustomizationTreeElement::ClearCachedChildren() const
-{
-	ChildrenCache.Empty();
-}
-
-uint32 FCollectionCustomizationTreeElement::GetHash() const
-{
-	if (ElementHash != 0)
-	{
-		return ElementHash;
-	}
-
-	switch (Type)
-	{
-	case EElementType::Root:
-	case EElementType::EmptyCollection:
-	case EElementType::EmptyConditionGroup:
-		ElementHash = GetTypeHash(true);	// There will only ever be one root, or one of the "empty" notices under a given parent
-		break;
-	case EElementType::ConditionGroup:
-		ElementHash = GetTypeHash(ConditionGroup);
-		break;
-	case EElementType::Query:
-		ElementHash = GetTypeHash(Query);
-		break;
-	case EElementType::QueryProperty:
-		ElementHash = GetTypeHash(QueryProperty);
-		break;
-	}
-		
-	if (ParentElement.IsValid())
-	{
-		ElementHash = HashCombine(ElementHash, ParentElement.Pin()->GetHash());
-	}
-
-	return ElementHash;
-}
-
-TArray<TSharedRef<SWidget>> SMovieGraphCollectionTreeItem_Stretch::GetColumnWidgets() const
-{
-	if (!WeakTreeElement.IsValid())
-	{
-		return { SNullWidget::NullWidget };
-	}
-
-	const TSharedPtr<FCollectionCustomizationTreeElement> TreeElement = WeakTreeElement.Pin();
-
-	if (TreeElement->Type == FCollectionCustomizationTreeElement::EElementType::QueryProperty)
-	{
-		// Display the widgets specified by the query
-		return TreeElement->Query->GetWidgets();
-	}
-
-	if (TreeElement->Type == FCollectionCustomizationTreeElement::EElementType::EmptyCollection ||
-		TreeElement->Type == FCollectionCustomizationTreeElement::EElementType::EmptyConditionGroup)
-	{
-		const FText NoticeText = (TreeElement->Type == FCollectionCustomizationTreeElement::EElementType::EmptyCollection)
-			? LOCTEXT("EmptyCollectionNotice", "Add a condition group to this collection.")
-			: LOCTEXT("EmptyConditionGroupNotice", "Add a condition to this condition group.");
-			
-		return {
-			SNew(SBox)
-			.MinDesiredHeight(50)
-			.Padding(10.f)
-			.HAlign(HAlign_Fill)
-			.VAlign(VAlign_Center)
-			.Content()
-			[
-				SNew(STextBlock)
-				.Text(NoticeText)
-			]
-		};
-	}
-	
-	return { SNullWidget::NullWidget };
-}
-
-void SMovieGraphCollectionTreeItem_Stretch::Construct(
-	const FArguments& InArgs, const TSharedPtr<SMovieGraphCollectionTreeView>& InOwnerTable, const TSharedPtr<FCollectionCustomizationTreeElement>& InTreeElement)
-{
-	check(InTreeElement);
-	WeakTreeElement = InTreeElement.ToWeakPtr();
-
-	TSharedPtr<SVerticalBox> WidgetContainer;
-
-	// Add a new indented row which has a designated place for widgets to be added (WidgetContainer)
-	this->ChildSlot
-	[
-		SNew(SHorizontalBox)
-		+ SHorizontalBox::Slot()
-		.AutoWidth()
-		.VAlign(VAlign_Fill)
-		[
-			SNew(SMovieGraphCollectionTreeViewIndent)
-			.NumIndents(2)
-		]
-
-		+ SHorizontalBox::Slot()
-		.VAlign(VAlign_Top)
-		.HAlign(HAlign_Fill)
-		[
-			SNew(SBorder)
-			.Padding(0, 0, 0, 1.f)
-			.HAlign(HAlign_Fill)
-			.BorderImage(FMovieRenderPipelineStyle::Get().GetBrush("MovieRenderGraph.CollectionsTree.BottomBorder"))
-			[
-				SAssignNew(WidgetContainer, SVerticalBox)
-			]
-		]
-	];
-
-	// Add all widgets that are generated by the column
-	for (const TSharedRef<SWidget>& Widget : GetColumnWidgets())
-	{
-		WidgetContainer->AddSlot()
-			.HAlign(HAlign_Fill)
-			[
-				SNew(SBox)
-				.IsEnabled_Lambda([InTreeElement]()
-				{
-					return InTreeElement->IsQueryEnabled();
-				})
-				[
-					Widget
-				]
-			];
-	}
-	
-	ConstructInternal(
-		FArguments()
-			.ShowSelection(true)
-			.Style(FMovieRenderPipelineStyle::Get(), "MovieRenderGraph.CollectionsTree.Row"),
-		InOwnerTable.ToSharedRef());
-}
-
-void SMovieGraphCollectionTreeItem::Construct(const FArguments& InArgs, const TSharedPtr<SMovieGraphCollectionTreeView>& InOwnerTable, const TSharedPtr<FCollectionCustomizationTreeElement>& InTreeElement, const TWeakPtr<IDetailLayoutBuilder>& InDetailBuilder)
-{
-	WeakTreeElement = InTreeElement;
-	WeakTreeView = InOwnerTable.ToWeakPtr();
-	DetailBuilder = InDetailBuilder;
-
-	SMultiColumnTableRow<TSharedPtr<FCollectionCustomizationTreeElement>>::Construct(
-		SMultiColumnTableRow<TSharedPtr<FCollectionCustomizationTreeElement>>::FArguments()
-			.Padding(0)
-			.Style(FMovieRenderPipelineStyle::Get(), "MovieRenderGraph.CollectionsTree.Row")
-			.OnDragDetected(this, &SMovieGraphCollectionTreeItem::HandleOnDragDetected)
-			.OnCanAcceptDrop(this, &SMovieGraphCollectionTreeItem::HandleOnCanAcceptDrop)
-			.OnAcceptDrop(this, &SMovieGraphCollectionTreeItem::HandleOnAcceptDrop)
-		, InOwnerTable.ToSharedRef());
-}
-
-FReply SMovieGraphCollectionTreeItem::HandleOnDragDetected(const FGeometry& InGeometry, const FPointerEvent& MouseEvent) const
-{
-	if (!WeakTreeView.IsValid())
-	{
-		return FReply::Unhandled();
-	}
-	
-	TArray<TSharedPtr<FCollectionCustomizationTreeElement>> SelectedElements = WeakTreeView.Pin()->GetSelectedItems();
-
-	if ((SelectedElements.Num() != 1) || !SelectedElements[0].IsValid())
-	{
-		return FReply::Unhandled();
-	}
-
-	if (MouseEvent.IsMouseButtonDown(EKeys::LeftMouseButton))
-	{
-		const FText DragDropText = SelectedElements[0]->Type == FCollectionCustomizationTreeElement::EElementType::ConditionGroup
-			? LOCTEXT("DragConditionGroup", "Move 1 Condition Group")
-			: LOCTEXT("DragConditionGroupQuery", "Move 1 Condition");
-					
-		const TSharedRef<FMovieGraphCollectionTreeDragOp> DragDropOp = FMovieGraphCollectionTreeDragOp::New(SelectedElements[0]);
-		DragDropOp->CurrentHoverText = DragDropText;
-					
-		return FReply::Handled().BeginDragDrop(DragDropOp);
-	}
-
-	return FReply::Unhandled();
-}
-
-TOptional<EItemDropZone> SMovieGraphCollectionTreeItem::HandleOnCanAcceptDrop(const FDragDropEvent& DragDropEvent, const EItemDropZone DropZone, const TSharedPtr<FCollectionCustomizationTreeElement> TargetItem)
-{
-	TOptional<EItemDropZone> ReturnedDropZone;
-	const TSharedPtr<FMovieGraphCollectionTreeDragOp> DragDropOp = DragDropEvent.GetOperationAs<FMovieGraphCollectionTreeDragOp>();
-
-	// Can't drag an element onto another element -- only changing order is allowed (above/below)
-	if (DropZone == EItemDropZone::OntoItem)
-	{
-		DragDropOp->CurrentIconBrush = FAppStyle::GetBrush(TEXT("Graph.ConnectorFeedback.Error"));
-		return ReturnedDropZone;
-	}
-				
-	if (DragDropOp.IsValid() && DragDropOp->Element.IsValid())
-	{
-		// All tree items can be moved *except* query properties
-		if (TargetItem.Get()->Type != FCollectionCustomizationTreeElement::EElementType::QueryProperty)
-		{
-			// For now, elements cannot be moved to a different parent
-			if (TargetItem.Get()->ParentElement == DragDropOp->Element->ParentElement)
-			{
-				DragDropOp->CurrentIconBrush = FAppStyle::GetBrush(TEXT("Graph.ConnectorFeedback.Ok"));
-				ReturnedDropZone = DropZone;
-			}
-			else
-			{
-				DragDropOp->CurrentIconBrush = FAppStyle::GetBrush(TEXT("Graph.ConnectorFeedback.Error"));
-			}
-		}
-	}
-
-	return ReturnedDropZone;
-}
-
-FReply SMovieGraphCollectionTreeItem::HandleOnAcceptDrop(const FDragDropEvent& DragDropEvent, const EItemDropZone DropZone, const TSharedPtr<FCollectionCustomizationTreeElement> TargetItem) const
-{
-	const TSharedPtr<FMovieGraphCollectionTreeDragOp> DragDropOp = DragDropEvent.GetOperationAs<FMovieGraphCollectionTreeDragOp>();
-	if (!DragDropOp.IsValid())
-	{
-		return FReply::Unhandled();
-	}
-
-	// Figure out the new order of elements
-	if (TargetItem->Type == FCollectionCustomizationTreeElement::EElementType::ConditionGroup)
-	{
-		if (TargetItem->Collection)
-		{
-			const TArray<UMovieGraphConditionGroup*>& ConditionGroups = TargetItem->Collection->GetConditionGroups();
-						
-			// Find the index of the target
-			const int32 TargetIndex = ConditionGroups.Find(TargetItem->ConditionGroup);
-			if (TargetIndex != INDEX_NONE)
-			{
-				FScopedTransaction Transaction(LOCTEXT("ReorderConditionGroup", "Reorder Condition Group"));
-				const int32 NewIndex = DropZone == EItemDropZone::AboveItem ? TargetIndex : TargetIndex + 1;
-				TargetItem->Collection->MoveConditionGroupToIndex(DragDropOp.Get()->Element->ConditionGroup, NewIndex);
-							
-				DragDropOp.Get()->Element->ParentElement.Pin()->ClearCachedChildren();
-				RefreshTreeView();
-			}
-		}
-	}
-	else if (TargetItem->Type == FCollectionCustomizationTreeElement::EElementType::Query)
-	{
-		if (TargetItem->ConditionGroup)
-		{
-			const TArray<UMovieGraphConditionGroupQueryBase*>& Queries = TargetItem->ConditionGroup->GetQueries();
-
-			// Find the index of the target
-			const int32 TargetIndex = Queries.Find(TargetItem->Query);
-			if (TargetIndex != INDEX_NONE)
-			{
-				FScopedTransaction Transaction(LOCTEXT("ReorderConditionGroupQuery", "Reorder Condition Group Query"));
-				const int32 NewIndex = DropZone == EItemDropZone::AboveItem ? TargetIndex : TargetIndex + 1;
-				TargetItem->ConditionGroup->MoveQueryToIndex(DragDropOp.Get()->Element->Query, NewIndex);
-							
-				DragDropOp.Get()->Element->ParentElement.Pin()->ClearCachedChildren();
-				RefreshTreeView();
-			}
-		}
-	}
-
-	return FReply::Handled();
-}
-
-TSharedRef<SWidget> SMovieGraphCollectionTreeItem::GenerateConditionGroupColumnWidget(const FName& ColumnName, const TSharedPtr<FCollectionCustomizationTreeElement>& TreeElement)
-{
-	if (ColumnName == ColumnID_Name)
-	{
-		// Generate the condition group name
-		const int32 GroupIndex = (TreeElement->ConditionGroup && TreeElement->Collection)
-			? TreeElement->Collection->GetConditionGroups().Find(TreeElement->ConditionGroup)
-			: 0;
-		const FString GroupName = FString::Printf(TEXT("Condition Group %i"), (GroupIndex + 1));
-		
-		return
-			SNew(SBorder)
-			.HAlign(HAlign_Fill)
-			.Padding(0)
-			.BorderImage(FMovieRenderPipelineStyle::Get().GetBrush("MovieRenderGraph.CollectionsTree.BottomBorder"))
-			[
-				SNew(SHorizontalBox)
-				+ SHorizontalBox::Slot()
-				.AutoWidth()
-				.VAlign(VAlign_Fill)
-				.Padding(0)
-				[
-					SNew(SMovieGraphCollectionTreeViewIndent)
-					.NumIndents(1)
-				]
-		
-				+ SHorizontalBox::Slot()
-				.AutoWidth()
-				.VAlign(VAlign_Center)
-				[
-					SNew(SExpanderArrow, SharedThis(this))
-					.IndentAmount(0)
-				]
-
-				+ SHorizontalBox::Slot()
-				.AutoWidth()
-				.VAlign(VAlign_Center)
-				.Padding(0, 0, 2.f, 0)
-				[
-					SNew(SImage)
-					.Image(FAppStyle::GetBrush("ClassIcon.GroupActor"))
-					.ColorAndOpacity(FSlateColor::UseForeground())
-				]
-
-				+ SHorizontalBox::Slot()
-				.AutoWidth()
-				.VAlign(VAlign_Center)
-				.Padding(0, 0, 2.f, 0)
-				[
-					SNew(STextBlock)
-					.Text(FText::FromString(GroupName))
-					.Font(FAppStyle::GetFontStyle(TEXT("PropertyWindow.BoldFont")))
-				]
-			];
-	}
-
-	if (ColumnName == ColumnID_Value)
-	{
-		return
-			SNew(SBorder)
-			.Padding(0)
-			.BorderImage(FMovieRenderPipelineStyle::Get().GetBrush("MovieRenderGraph.CollectionsTree.BottomLeftBorder"))
-			[
-				SNew(SHorizontalBox)
-				+ SHorizontalBox::Slot()
-				.VAlign(VAlign_Center)
-				.HAlign(HAlign_Left)
-				.AutoWidth()
-				[
-					SNew(SButton)
-					.HAlign(HAlign_Center)
-					.ToolTip(SNew(SToolTip).Text(LOCTEXT("AddConditionGroupQuery", "Add a condition to this condition group.")))
-					.ButtonStyle(FAppStyle::Get(), "SimpleButton")
-					.OnClicked_Lambda([this]()
-					{
-						if (WeakTreeElement.IsValid())
-						{
-							const TSharedPtr<FCollectionCustomizationTreeElement> TreeElementPin = WeakTreeElement.Pin();
-							
-							if (const TObjectPtr<UMovieGraphConditionGroup> ConditionGroup = TreeElementPin->ConditionGroup)
-							{
-								// Use the Actor Name query as the default
-								ConditionGroup->AddQuery(UMovieGraphConditionGroupQuery_ActorName::StaticClass());
-
-								if (TreeElementPin->ParentElement.IsValid())
-								{
-									TreeElementPin->ParentElement.Pin()->ClearCachedChildren();
-								}
-								
-								RefreshTreeView();
-							}
-						}
-
-						return FReply::Handled();
-					})
-					.Content()
-					[
-						SNew(SImage)
-						.Image(FAppStyle::GetBrush("Icons.PlusCircle"))
-						.ColorAndOpacity(FSlateColor::UseForeground())
-					]
-				]
-
-				+ SHorizontalBox::Slot()
-				.VAlign(VAlign_Center)
-				.HAlign(HAlign_Left)
-				.AutoWidth()
-				[
-					SNew(SComboButton)
-					.ComboButtonStyle(FAppStyle::Get(), "SimpleComboButton")
-					.ContentPadding(0)
-					.ForegroundColor(FSlateColor::UseForeground())
-					.HasDownArrow(true)
-					.OnGetMenuContent_Lambda([this]()
-					{
-						FMenuBuilder ViewOptions(true, nullptr);
-
-						ViewOptions.AddMenuEntry(
-							LOCTEXT("DeleteConditionGroup", "Delete"),
-							LOCTEXT("DeleteConditionGroupTooltip", "Delete this condition group and the queries in it."),
-							FSlateIcon(),
-							FUIAction(
-								FExecuteAction::CreateLambda([this]()
-								{
-									if (WeakTreeView.IsValid())
-									{
-										WeakTreeView.Pin()->RemoveConditionGroup(WeakTreeElement);
-									}
-                                }),
-								FCanExecuteAction()));
-
-						return ViewOptions.MakeWidget();
-					})
-				]
-
-				+ SHorizontalBox::Slot()
-				.VAlign(VAlign_Center)
-				.HAlign(HAlign_Right)
-				.Padding(2.f)
-				.FillWidth(1.f)
-				[
-					SNew(SMovieGraphCollectionTreeOpTypeWidget)
-					.IsConditionGroup(true)
-					.WeakTreeElement(WeakTreeElement)
-				]
-			];
-	}
-	
-	return SNullWidget::NullWidget;
-}
+	const int32 ExistingQueryIndex = WeakConditionGroup->GetQueries().Find(WeakQuery.Get());
 
-TSharedRef<SWidget> SMovieGraphCollectionTreeItem::GenerateConditionGroupQueryColumnWidget(const FName& ColumnName, const TSharedPtr<FCollectionCustomizationTreeElement>& TreeElement)
-{
-	if (ColumnName == ColumnID_Name)
-	{
-		return
-			SNew(SHorizontalBox)
-			+ SHorizontalBox::Slot()
-			.AutoWidth()
-			.VAlign(VAlign_Fill)
-			[
-				SNew(SMovieGraphCollectionTreeViewIndent)
-				.NumIndents(2)
-			]
+	// Swap the old query with a new one.
+	// Add the new query at the index of the current query. If the query couldn't be found for some reason, default to the end of the array (-1)
+	UMovieGraphConditionGroupQueryBase* NewQuery = WeakConditionGroup->AddQuery(InNewQueryType, ExistingQueryIndex != INDEX_NONE ? ExistingQueryIndex : -1);
+	WeakConditionGroup->RemoveQuery(WeakQuery.Get());
+	WeakQuery = MakeWeakObjectPtr(NewQuery);
 
-			+ SHorizontalBox::Slot()
-			.HAlign(HAlign_Fill)
-			[
-				SNew(SBorder)
-				.Padding(0)
-				.BorderImage(FMovieRenderPipelineStyle::Get().GetBrush("MovieRenderGraph.CollectionsTree.BottomBorder"))
-				[
-					SNew(SHorizontalBox)
-					+ SHorizontalBox::Slot()
-					.AutoWidth()
-					.VAlign(VAlign_Center)
-					[
-						SNew(SExpanderArrow, SharedThis(this))
-						.IndentAmount(0)
-					]
-
-					+ SHorizontalBox::Slot()
-					.AutoWidth()
-					.VAlign(VAlign_Center)
-					.Padding(0, 0, 4.f, 0)
-					[
-						SNew(SCheckBox)
-						.IsChecked_Lambda([TreeElement]()
-						{
-							return TreeElement->IsQueryEnabled() ? ECheckBoxState::Checked : ECheckBoxState::Unchecked;
-						})
-						.OnCheckStateChanged_Lambda([TreeElement](ECheckBoxState NewCheckState)
-						{
-							if (TreeElement.IsValid() && TreeElement->Query)
-							{
-								TreeElement->Query->SetEnabled(NewCheckState == ECheckBoxState::Checked);
-							}
-						})
-					]
-
-					+ SHorizontalBox::Slot()
-					.VAlign(VAlign_Center)
-					.HAlign(HAlign_Fill)
-					.Padding(0, 0, 7.f, 0)
-					[
-						SNew(SMovieGraphCollectionTreeQueryTypeSelectorWidget)
-						.WeakTreeElement(WeakTreeElement)
-						.WeakTreeView(WeakTreeView)
-					]
-				]
-			];
-	}
-
-	if (ColumnName == ColumnID_Value)
-	{			
-		return
-			SNew(SBorder)
-			.Padding(0)
-			.BorderImage(FMovieRenderPipelineStyle::Get().GetBrush("MovieRenderGraph.CollectionsTree.BottomLeftBorder"))
-			[
-				SNew(SHorizontalBox)
-				+ SHorizontalBox::Slot()
-				.VAlign(VAlign_Center)
-				.HAlign(HAlign_Left)
-				.Padding(4.f, 0, 0, 0)
-				.AutoWidth()
-				[
-					SNew(SMovieGraphCollectionTreeAddQueryContentWidget)
-					.WeakTreeElement(WeakTreeElement)
-					.WeakTreeView(WeakTreeView)
-				]
-
-				+ SHorizontalBox::Slot()
-				.VAlign(VAlign_Center)
-				.HAlign(HAlign_Left)
-				.AutoWidth()
-				[
-					SNew(SComboButton)
-					.IsEnabled_Lambda([this]()
-					{
-						return WeakTreeElement.IsValid() ? WeakTreeElement.Pin()->IsQueryEnabled() : false;
-					})
-					.ComboButtonStyle(FAppStyle::Get(), "SimpleComboButton")
-					.ContentPadding(0)
-					.ForegroundColor(FSlateColor::UseForeground())
-					.HasDownArrow(true)
-					.OnGetMenuContent_Lambda([this]()
-					{
-						FMenuBuilder ViewOptions(true, nullptr);
-
-						ViewOptions.AddMenuEntry(
-							LOCTEXT("DeleteConditionGroupQuery", "Delete"),
-							LOCTEXT("DeleteConditionGroupQueryTooltip", "Delete this condition group query."),
-							FSlateIcon(),
-							FUIAction(
-								FExecuteAction::CreateLambda([this]()
-								{
-									if (WeakTreeElement.IsValid())
-									{
-										const TSharedPtr<FCollectionCustomizationTreeElement> TreeElementPin = WeakTreeElement.Pin();
-										
-										TreeElementPin->ConditionGroup->RemoveQuery(TreeElementPin->Query);
-
-										if (TreeElementPin->ParentElement.IsValid())
-										{
-											TreeElementPin->ParentElement.Pin()->ClearCachedChildren();
-										}
-										
-										RefreshTreeView();
-									}
-								}),
-								FCanExecuteAction()));
-
-						return ViewOptions.MakeWidget();
-					})
-				]
-
-				+ SHorizontalBox::Slot()
-				.VAlign(VAlign_Center)
-				.HAlign(HAlign_Right)
-				.Padding(2.f)
-				.FillWidth(1.f)
-				[
-					SNew(SMovieGraphCollectionTreeOpTypeWidget)
-					.IsConditionGroup(false)
-					.WeakTreeElement(WeakTreeElement)
-				]
-			];
-	}
-	
-	return SNullWidget::NullWidget;
-}
-
-TSharedRef<SWidget> SMovieGraphCollectionTreeItem::GenerateQueryPropertyColumnWidget(const FName& ColumnName, const TSharedPtr<FCollectionCustomizationTreeElement>& TreeElement)
-{
-	if (!DetailBuilder.IsValid())
-	{
-		return SNullWidget::NullWidget;
-	}
-		
-	// Get the property handle for the property within the query. This is needed in order to get the name/value widgets for it.
-	const TSharedPtr<IPropertyHandle> QueryPropertyHandle =
-		DetailBuilder.Pin()->AddObjectPropertyData({TreeElement->QueryProperty->GetOwnerUObject()}, TreeElement->QueryProperty->GetFName());
-		
-	if (!QueryPropertyHandle || !QueryPropertyHandle.IsValid())
-	{
-		return SNullWidget::NullWidget;
-	}
-		
-	if (ColumnName == ColumnID_Name)
-	{
-		return
-			SNew(SBorder)
-			.Padding(0)
-			.IsEnabled_Lambda([this]()
-			{
-				return WeakTreeElement.IsValid() ? WeakTreeElement.Pin()->IsQueryEnabled() : false;
-			})
-			.BorderImage(FMovieRenderPipelineStyle::Get().GetBrush("MovieRenderGraph.CollectionsTree.BottomBorder"))
-			[
-				SNew(SHorizontalBox)
-				+ SHorizontalBox::Slot()
-				.AutoWidth()
-				.VAlign(VAlign_Fill)
-				[
-					SNew(SMovieGraphCollectionTreeViewIndent)
-					.NumIndents(2)
-				]
-
-				+ SHorizontalBox::Slot()
-				.AutoWidth()
-				[
-					SNew(SExpanderArrow, SharedThis(this))
-					.IndentAmount(0)
-				]
-
-				+ SHorizontalBox::Slot()
-				[
-					QueryPropertyHandle->CreatePropertyNameWidget()
-				]
-			];
-	}
-
-	if (ColumnName == ColumnID_Value)
-	{
-		return
-			SNew(SBorder)
-			.Padding(5.f, 2.f)
-			.IsEnabled_Lambda([this]()
-			{
-				return WeakTreeElement.IsValid() ? WeakTreeElement.Pin()->IsQueryEnabled() : false;
-			})
-			.BorderImage(FMovieRenderPipelineStyle::Get().GetBrush("MovieRenderGraph.CollectionsTree.BottomLeftBorder"))
-			[
-				QueryPropertyHandle->CreatePropertyValueWidget()
-			];
-	}
-		
-	return SNullWidget::NullWidget;
-}
-
-void SMovieGraphCollectionTreeItem::RefreshTreeView() const
-{
-	if (WeakTreeView.IsValid())
-	{
-		WeakTreeView.Pin()->RefreshView();
-	}
-}
-
-TSharedRef<SWidget> SMovieGraphCollectionTreeItem::GenerateWidgetForColumn(const FName& ColumnName)
-{
-	const TSharedPtr<FCollectionCustomizationTreeElement> TreeElement = WeakTreeElement.Pin();
-	if (!TreeElement.IsValid())
-	{
-		return SNullWidget::NullWidget;
-	}
-
-	if (TreeElement->Type == FCollectionCustomizationTreeElement::EElementType::Root)
-	{
-		// The root element isn't visible, but serves as the single entry point to the tree's hierarchy
-		return SNullWidget::NullWidget;
-	}
-
-	if (TreeElement->Type == FCollectionCustomizationTreeElement::EElementType::ConditionGroup)
-	{
-		return GenerateConditionGroupColumnWidget(ColumnName, TreeElement);
-	}
-
-	if (TreeElement->Type == FCollectionCustomizationTreeElement::EElementType::Query)
+	if (OnQueryTypeChanged.IsBound())
 	{
-		return GenerateConditionGroupQueryColumnWidget(ColumnName, TreeElement);
+		OnQueryTypeChanged.Execute(WeakQuery);
 	}
-		
-	return GenerateQueryPropertyColumnWidget(ColumnName, TreeElement);
 }
 
 TSharedRef<IDetailCustomization> FMovieGraphCollectionsCustomization::MakeInstance()
@@ -1410,8 +969,9 @@ void FMovieGraphCollectionsCustomization::CustomizeDetails(IDetailLayoutBuilder&
 	}
 
 	const TObjectPtr<UMovieGraphCollection> Collection = CollectionNode->Collection;
+	WeakCollection = Collection;
 
-	// The Collection property is being completely replaced by a custom STreeView
+	// The Collection property is being completely replaced by a custom builder
 	const TSharedRef<IPropertyHandle> CollectionProperty = InDetailBuilder.GetProperty("Collection");
 	CollectionProperty->MarkHiddenByCustomization();
 
@@ -1429,10 +989,10 @@ void FMovieGraphCollectionsCustomization::CustomizeDetails(IDetailLayoutBuilder&
 			.Text(LOCTEXT("AddNewConditionGroup", "Condition Group"))
 			.OnClicked_Lambda([this]
 			{
-				if (TreeView.IsValid())
+				if (WeakCollection.IsValid())
 				{
-					TreeView->AddConditionGroup();
-					TreeView->RefreshView();
+					UMovieGraphConditionGroup* NewConditionGroup = WeakCollection->AddConditionGroup();
+					NewConditionGroup->AddQuery(UMovieGraphConditionGroupQuery_ActorName::StaticClass());
 				}
 
 				return FReply::Handled();
@@ -1443,23 +1003,25 @@ void FMovieGraphCollectionsCustomization::CustomizeDetails(IDetailLayoutBuilder&
 	// Add the collection name property
 	CollectionCategory.AddExternalObjectProperty({Collection}, FName(TEXT("CollectionName")));
 
-	FDetailWidgetRow& ConditionsWidgetRow = InDetailBuilder.EditCategory("Conditions")
-		.AddCustomRow(FText::FromString(TEXT("Conditions")));
-	
-	ConditionsWidgetRow.WholeRowWidget
-	[
-		SNew(SBox)
-		.VAlign(VAlign_Fill)
-		.Padding(0)
-		[
-			SAssignNew(TreeView, SMovieGraphCollectionTreeView)
-			.DetailBuilder(DetailBuilder)
-			.Collection(MakeWeakObjectPtr(Collection))
-		]
-	];
+	IDetailCategoryBuilder& ConditionsCategory = InDetailBuilder.EditCategory("Conditions");
 
-	// Add the root elements to the tree; root elements will take care of adding children underneath them
-	TreeView->AddRootElement();
+	const TSharedPtr<IPropertyHandle> ConditionGroupsProperty = InDetailBuilder.AddObjectPropertyData({Collection}, FName(TEXT("ConditionGroups")));
+	const TSharedPtr<IPropertyHandleArray> ConditionGroupsArrayProperty = ConditionGroupsProperty->AsArray();
+	check(ConditionGroupsArrayProperty);
+
+	// Add builders for the condition group array
+	uint32 NumConditionGroups;
+	ConditionGroupsArrayProperty->GetNumElements(NumConditionGroups);
+	for (uint32 ConditionGroupIndex = 0; ConditionGroupIndex < NumConditionGroups; ++ConditionGroupIndex)
+	{
+		const TSharedRef<IPropertyHandle> ConditionGroupProperty = ConditionGroupsArrayProperty->GetElement(ConditionGroupIndex);
+		ConditionsCategory.AddCustomBuilder(MakeShareable(new FMovieGraphConditionGroupBuilder(DetailBuilder.Pin().Get(), ConditionGroupProperty, ConditionGroupIndex, WeakCollection)));
+	}
+
+	ConditionGroupsArrayProperty->SetOnNumElementsChanged(FSimpleDelegate::CreateSPLambda(this, [&ConditionsCategory]()
+	{
+		ConditionsCategory.GetParentLayout().ForceRefreshDetails();
+	}));
 }
 
 void FMovieGraphCollectionsCustomization::CustomizeDetails(const TSharedPtr<IDetailLayoutBuilder>& InDetailBuilder)
diff --git a/Engine/Plugins/MovieScene/MovieRenderPipeline/Source/MovieRenderPipelineEditor/Private/Customizations/Graph/MovieGraphCollectionsCustomization.h b/Engine/Plugins/MovieScene/MovieRenderPipeline/Source/MovieRenderPipelineEditor/Private/Customizations/Graph/MovieGraphCollectionsCustomization.h
--- a/Engine/Plugins/MovieScene/MovieRenderPipeline/Source/MovieRenderPipelineEditor/Private/Customizations/Graph/MovieGraphCollectionsCustomization.h
+++ b/Engine/Plugins/MovieScene/MovieRenderPipeline/Source/MovieRenderPipelineEditor/Private/Customizations/Graph/MovieGraphCollectionsCustomization.h
@@ -4,95 +4,140 @@
 
 #include "DragAndDrop/DecoratedDragDropOp.h"
 #include "IDetailCustomization.h"
+#include "IDetailCustomNodeBuilder.h"
+#include "IDetailDragDropHandler.h"
 #include "Widgets/Views/STreeView.h"
 
-class FCollectionCustomizationTreeElement;
 class UMovieGraphCollection;
 class UMovieGraphConditionGroup;
 class UMovieGraphConditionGroupQueryBase;
 
-/**
- * A widget which indents the specified number of indent levels, with the appropriate styling for the collection tree view.
- */
-class SMovieGraphCollectionTreeViewIndent final : public SCompoundWidget
+/** Drag/drop operation that contains data about a condition group or condition group query being dragged. */
+class FMovieGraphCollectionDragDropOp : public FDecoratedDragDropOp
 {
 public:
-	SLATE_BEGIN_ARGS(SMovieGraphCollectionTreeViewIndent) { }
-		SLATE_ATTRIBUTE(int, NumIndents)
-	SLATE_END_ARGS()
+	DRAG_DROP_OPERATOR_TYPE(FMovieGraphCollectionDragDropOp, FDecoratedDragDropOp);
 
-	void Construct(const FArguments& InArgs);
-};
+	FMovieGraphCollectionDragDropOp(const bool bIsConditionGroup, const int32 InitialIndex, TWeakObjectPtr<UMovieGraphConditionGroup> InWeakOwningConditionGroup = nullptr);
 
-/** The drag-drop operator for use within the collections tree view. */
-class FMovieGraphCollectionTreeDragOp final : public FDecoratedDragDropOp
-{
-public:
-	DRAG_DROP_OPERATOR_TYPE(SMovieGraphCollectionTreeDragOp, FDragDropOperation)
+	/** Initialize the operation. */
+	void Init();
 
-	static TSharedRef<FMovieGraphCollectionTreeDragOp> New(const TSharedPtr<FCollectionCustomizationTreeElement>& InElement);
+	/** Set whether the operation is currently valid or not. */
+	void SetValidTarget(const bool bIsValidTarget);
 
-public:
-	/** The tree element associated with this drag-drop operation. */
-	TSharedPtr<FCollectionCustomizationTreeElement> Element;
+	/** The index of the condition group or condition group query being dragged. */
+	int32 GetIndex() const;
+
+	/** Whether this operation represents a condition group (true) or condition group query (false). */
+	bool IsConditionGroup() const;
+
+	/** Get the condition group that owns the query being dragged (nullptr for dragging a condition group). */
+	TWeakObjectPtr<UMovieGraphConditionGroup> GetOwningConditionGroup() const;
 
-	/** Whether the drag-drop operation is currently valid. */
-	bool bIsValidOperation = false;
+private:
+	/** Whether this is a condition group (true) or condition group query (false) being dragged. */
+	bool bIsConditionGroup;
+	
+	/** The index of the condition group or condition group query being dragged. */
+	int32 InitialIndex;
+
+	/** The condition group that owns the query being dragged (this being valid implies !bIsConditionGroup). */
+	TWeakObjectPtr<UMovieGraphConditionGroup> OwningConditionGroup;
 };
 
-/** A tree view specifically for viewing the contents of a UMovieGraphCollection object. */
-class SMovieGraphCollectionTreeView final : public STreeView<TSharedPtr<FCollectionCustomizationTreeElement>>
+/** Handler for customizing the drag/drop functionality of condition groups and condition group queries. */
+class FMovieGraphCollectionDragDropHandler : public IDetailDragDropHandler
 {
 public:
-	SLATE_BEGIN_ARGS(SMovieGraphCollectionTreeView) { }
-		SLATE_ARGUMENT(TWeakPtr<IDetailLayoutBuilder>, DetailBuilder)
-		SLATE_ARGUMENT(TWeakObjectPtr<UMovieGraphCollection>, Collection)
-	SLATE_END_ARGS()
+	FMovieGraphCollectionDragDropHandler(const bool bIsConditionGroup, IDetailLayoutBuilder* InDetailBuilder, UMovieGraphCollection* InCollection, const int32 ConditionGroupIndex, TWeakObjectPtr<UMovieGraphConditionGroup> InWeakConditionGroup, const int32 ConditionGroupQueryIndex);
 
-	void Construct(const FArguments& InArgs);
+	//~ Begin IDetailDragDropHandler interface
+	virtual TSharedPtr<FDragDropOperation> CreateDragDropOperation() const override;
+	virtual bool AcceptDrop(const FDragDropEvent& DragDropEvent, EItemDropZone DropZone) const override;
+	virtual TOptional<EItemDropZone> CanAcceptDrop(const FDragDropEvent& DragDropEvent, EItemDropZone DropZone) const override;
+	//~ End IDetailDragDropHandler interface
 
-	/** Called when an row in the tree view expands/collapses. Caches the expansion state of the tree so it can be restored at a later time. */
-	void OnExpansionChanged(TSharedPtr<FCollectionCustomizationTreeElement> InElement, bool bInExpanded);
+private:
+	/** Whether the handler is targeting a condition group (true) or query (false). */
+	bool bIsConditionGroup;
+
+	/** The layout builder that the handler was created from. */
+	IDetailLayoutBuilder* DetailLayoutBuilder;
+
+	/** The collection that owns the condition group (or query). */
+	TWeakObjectPtr<UMovieGraphCollection> WeakCollection;
 
-	/** Gets the children of the specified tree element. */
-	void GetChildrenForTree(TSharedPtr<FCollectionCustomizationTreeElement> InItem, TArray<TSharedPtr<FCollectionCustomizationTreeElement>>& OutChildren);
+	/** The index of the condition group targeted by the handler. */
+	int32 ConditionGroupIndex;
 
-	/** Restores the cached expansion state recursively through all tree elements. */
-	void RestoreExpansionStateRecursive(const TSharedPtr<FCollectionCustomizationTreeElement>& InElement);
+	/** The condition group that is targeted by the handler. */
+	TWeakObjectPtr<UMovieGraphConditionGroup> WeakConditionGroup;
+
+	/** The index of the query targeted by the handler. */
+	int32 ConditionGroupQueryIndex;
+};
 
-	/** Generates a tree row suitable for the type of tree element provided. */
-	TSharedRef<ITableRow> GenerateTreeRow(TSharedPtr<FCollectionCustomizationTreeElement> InTreeElement, const TSharedRef<STableViewBase>& OwnerTable);
+/** Custom node builder for condition group queries. Displays each query as a child. */
+class FMovieGraphConditionGroupQueryBuilder : public IDetailCustomNodeBuilder, public TSharedFromThis<FMovieGraphConditionGroupQueryBuilder>
+{
+public:
+	FMovieGraphConditionGroupQueryBuilder(IDetailLayoutBuilder* InDetailBuilder, TSharedRef<IPropertyHandle> InConditionGroupQueryProperty, int32 ConditionGroupQueryIndex, const TWeakObjectPtr<UMovieGraphConditionGroup>& InWeakConditionGroup);
 
-	/** Adds the root tree element. The root is not visible, but serves as the single entry point into the tree. */
-	void AddRootElement();
+	//~ Begin IDetailCustomNodeBuilder interface
+	virtual FName GetName() const override;
+	virtual bool InitiallyCollapsed() const override { return false; }
+	virtual bool RequiresTick() const override { return false; }
+	virtual void Tick(float DeltaTime) override {}
+	virtual void GenerateHeaderRowContent(FDetailWidgetRow& NodeRow) override;
+	virtual void GenerateChildContent(IDetailChildrenBuilder& ChildrenBuilder) override;
+	//~ End IDetailCustomNodeBuilder interface
 
-	/** Adds condition group tree elements, one for each condition group in the provided collection. */
-	void AddRootConditionGroupElements(const TObjectPtr<UMovieGraphCollection>& InCollection) const;
+private:
+	void AddQueryTypeWidgets(TWeakObjectPtr<UMovieGraphConditionGroupQueryBase> WeakQuery, IDetailChildrenBuilder& ChildrenBuilder) const;
 
-	/** Adds a new condition group element associated with the provided condition group. */
-	void AddNewConditionGroupElement(UMovieGraphConditionGroup* InConditionGroup) const;
+private:
+	/** The layout builder that this node builder was created from. */
+	IDetailLayoutBuilder* DetailLayoutBuilder;
+	
+	/** The property for the condition group query being displayed. */
+	TSharedPtr<IPropertyHandle> ConditionGroupQueryProperty;
 
-	/** Adds a new, empty condition group element to the tree, and adds a corresponding empty condition group to the underlying collection. */
-	void AddConditionGroup();
+	/** The index of the query represented by this builder. */
+	int32 ConditionGroupQueryIndex;
 
-	/** Removes the condition group associated with the provided element. The condition group is also removed from the underlying collection. */
-	void RemoveConditionGroup(const TWeakPtr<FCollectionCustomizationTreeElement>& InConditionGroupElement);
+	/** The condition group that the query belongs to. */
+	TWeakObjectPtr<UMovieGraphConditionGroup> WeakConditionGroup;
+};
 
-	/** Refreshes the tree view while also taking care of restoring cached expansion state. Should always be called instead of RequestTreeRefresh(). */
-	void RefreshView();
+/** Custom node builder for condition groups. The condition group is added as the header, and children are added as FMovieGraphConditionGroupQueryBuilder rows. */
+class FMovieGraphConditionGroupBuilder : public IDetailCustomNodeBuilder, public TSharedFromThis<FMovieGraphConditionGroupBuilder>
+{
+public:
+	FMovieGraphConditionGroupBuilder(IDetailLayoutBuilder* InDetailBuilder, TSharedRef<IPropertyHandle> InConditionGroupProperty, const uint32 InConditionGroupIndex, TWeakObjectPtr<UMovieGraphCollection>& InWeakCollection);
+
+	//~ Begin IDetailCustomNodeBuilder interface
+	virtual FName GetName() const override;
+	virtual bool InitiallyCollapsed() const override { return false; }
+	virtual bool RequiresTick() const override { return false; }
+	virtual void Tick(float DeltaTime) override {}
+	virtual TSharedPtr<IPropertyHandle> GetPropertyHandle() const override { return ConditionGroupProperty; }
+	virtual void GenerateHeaderRowContent(FDetailWidgetRow& NodeRow) override;
+	virtual void GenerateChildContent(IDetailChildrenBuilder& ChildrenBuilder) override;
+	//~ End IDetailCustomNodeBuilder interface
 
 private:
-	/** The collection associated with this tree view. */
-	TWeakObjectPtr<UMovieGraphCollection> Collection = nullptr;
+	/** The layout builder that this node builder was created from. */
+	IDetailLayoutBuilder* DetailLayoutBuilder;
 	
-	/** The root-most elements in the tree. */
-	TArray<TSharedPtr<FCollectionCustomizationTreeElement>> RootElements;
+	/** The property for the condition group being displayed. */
+	TSharedPtr<IPropertyHandle> ConditionGroupProperty;
 	
-	/** The hashes of elements that are currently expanded in the tree. */
-	TSet<uint32> ExpandedElements;
+	/** The index of the condition group represented by this builder. */
+	uint32 ConditionGroupIndex;
 
-	/** The detail builder associated with this tree. Some items in the tree need to add widgets that can be fetched from the detail builder. */
-	TWeakPtr<IDetailLayoutBuilder> DetailBuilder;
+	/** The collection that the condition groups belong to. */
+	TWeakObjectPtr<UMovieGraphCollection> WeakCollection;
 };
 
 /**
@@ -103,7 +148,8 @@ class SMovieGraphCollectionTreeOpTypeWidget final : public SCompoundWidget
 public:
 	SLATE_BEGIN_ARGS(SMovieGraphCollectionTreeOpTypeWidget) { }
 		SLATE_ATTRIBUTE(bool, IsConditionGroup)
-		SLATE_ATTRIBUTE(TWeakPtr<FCollectionCustomizationTreeElement>, WeakTreeElement)
+		SLATE_ATTRIBUTE(TWeakObjectPtr<UMovieGraphConditionGroup>, WeakConditionGroup)
+		SLATE_ATTRIBUTE(TWeakObjectPtr<UMovieGraphConditionGroupQueryBase>, WeakConditionGroupQuery)
 	SLATE_END_ARGS()
 
 	void Construct(const FArguments& InArgs);
@@ -112,8 +158,11 @@ private:
 	/** Whether this widget is associated with a condition group (else a condition group query). */
 	bool bIsConditionGroup = false;
 
-	/** The tree element this widget is being displayed in. */
-	TWeakPtr<FCollectionCustomizationTreeElement> WeakTreeElement = nullptr;
+	/** The condition group associated with this widget (if any). */
+	TWeakObjectPtr<UMovieGraphConditionGroup> WeakConditionGroup;
+
+	/** The condition group query associated with this widget (if any). */
+	TWeakObjectPtr<UMovieGraphConditionGroupQueryBase> WeakConditionGroupQuery;
 
 private:
 	/** Gets the names of all op types available. */
@@ -144,8 +193,7 @@ class SMovieGraphCollectionTreeAddQueryContentWidget final : public SCompoundWid
 {
 public:
 	SLATE_BEGIN_ARGS(SMovieGraphCollectionTreeAddQueryContentWidget) { }
-		SLATE_ATTRIBUTE(TWeakPtr<FCollectionCustomizationTreeElement>, WeakTreeElement)
-		SLATE_ATTRIBUTE(TWeakPtr<SMovieGraphCollectionTreeView>, WeakTreeView)
+		SLATE_ATTRIBUTE(TWeakObjectPtr<UMovieGraphConditionGroupQueryBase>, WeakQuery)
 	SLATE_END_ARGS()
 
 	void Construct(const FArguments& InArgs);
@@ -157,9 +205,12 @@ public:
 class SMovieGraphCollectionTreeQueryTypeSelectorWidget final : public SCompoundWidget
 {
 public:
+	DECLARE_DELEGATE_OneParam(FOnQueryTypeChanged, TWeakObjectPtr<UMovieGraphConditionGroupQueryBase>);
+	
 	SLATE_BEGIN_ARGS(SMovieGraphCollectionTreeQueryTypeSelectorWidget) { }
-		SLATE_ATTRIBUTE(TWeakPtr<FCollectionCustomizationTreeElement>, WeakTreeElement)
-		SLATE_ATTRIBUTE(TWeakPtr<SMovieGraphCollectionTreeView>, WeakTreeView)
+		SLATE_ATTRIBUTE(TWeakObjectPtr<UMovieGraphConditionGroup>, WeakConditionGroup)
+		SLATE_ATTRIBUTE(TWeakObjectPtr<UMovieGraphConditionGroupQueryBase>, WeakQuery)
+		SLATE_EVENT(FOnQueryTypeChanged, OnQueryTypeChanged)
 	SLATE_END_ARGS()
 
 	void Construct(const FArguments& InArgs);
@@ -178,142 +229,17 @@ private:
 	UClass* GetCurrentQueryType() const;
 
 	/** Updates the query to the query type which was selected. */
-	void SetQueryType(UClass* InNewQueryType, ESelectInfo::Type SelectInfo) const;
-
-private:
-	/** The tree element this widget is being displayed in. */
-	TWeakPtr<FCollectionCustomizationTreeElement> WeakTreeElement = nullptr;
-
-	/** The tree view associated with the tree element. */
-	TWeakPtr<SMovieGraphCollectionTreeView> WeakTreeView = nullptr;
-};
-
-/**
- * Represents the data of one element displayed in the tree, which is one of the types in EElementType. Provides a way
- * of generating the children under the element (and caching the children for fast access later).
- */
-class FCollectionCustomizationTreeElement : public TSharedFromThis<FCollectionCustomizationTreeElement>
-{
-public:
-	/** The type of this element. */
-	enum class EElementType : uint8
-	{
-		Root,					///< The root element; does not represent any actual data in the collection
-		EmptyCollection,		///< Represents an "empty collection" notice within a collection
-		ConditionGroup,			///< Represents a condition group within the collection
-		EmptyConditionGroup,	///< Represents an "empty condition group" notice within a condition group
-		Query,					///< Represents a query within a condition group
-		QueryProperty			///< Represents a property within a query
-	};
-	
-	explicit FCollectionCustomizationTreeElement(const EElementType InType);
-
-	/** Whether the query associated with this element is currently enabled. */
-	bool IsQueryEnabled() const;
-
-	/**
-	 * Gets the child elements nested under this element. Returns a cached result if available, otherwise calculates the
-	 * children and returns the result. Call ClearCachedChildren() if the cache needs to be cleared (eg, when the tree
-	 * is being refreshed).
-	 */
-	const TArray<TSharedPtr<FCollectionCustomizationTreeElement>>& GetChildren() const;
-
-	/** Clears the cached result of GetChildren(). */
-	void ClearCachedChildren() const;
-
-	/** Gets the hash that uniquely identifies this element in the tree. */
-	uint32 GetHash() const;
-
-public:
-	/** The type of item in the tree this element represents. */
-	EElementType Type = EElementType::Root;
-
-	/** The collection associated with this element. */
-	TObjectPtr<UMovieGraphCollection> Collection = nullptr;
-
-	/** If this element represents a condition group, this points to the underlying condition group (else nullptr). */
-	TObjectPtr<UMovieGraphConditionGroup> ConditionGroup = nullptr;
-
-	/** If this element represents a condition group query, this points to the underlying query (else nullptr). */
-	TObjectPtr<UMovieGraphConditionGroupQueryBase> Query = nullptr;
-
-	/** If this element represents a query property, this points to the underlying property (else nullptr). */
-	FProperty* QueryProperty = nullptr;
-
-	/** This element's parent element in the tree.  */
-	TWeakPtr<const FCollectionCustomizationTreeElement> ParentElement = nullptr;
-
-private:
-	/** The (cached) child elements nested under this element. */
-	mutable TArray<TSharedPtr<FCollectionCustomizationTreeElement>> ChildrenCache;
-
-	/** The (cached) hash of this element. */
-	mutable uint32 ElementHash = 0;
-};
-
-/** Specialized row class that spans all columns within the tree view. */
-class SMovieGraphCollectionTreeItem_Stretch final : public STableRow<TSharedPtr<FCollectionCustomizationTreeElement>>
-{
-public:
-	/** Gets the widgets that should be displayed within the column (which stretches across all columns in the tree view). */
-	TArray<TSharedRef<SWidget>> GetColumnWidgets() const;
-	
-	void Construct(const FArguments& InArgs, const TSharedPtr<SMovieGraphCollectionTreeView>& InOwnerTable,
-		const TSharedPtr<FCollectionCustomizationTreeElement>& InTreeElement);
+	void SetQueryType(UClass* InNewQueryType, ESelectInfo::Type SelectInfo);
 
 private:
-	/** The tree element associated with this item. */
-	TWeakPtr<FCollectionCustomizationTreeElement> WeakTreeElement;
-};
-
-/**
- * The widget that is responsible for generating the per-column content for each element in the tree.
- */
-class SMovieGraphCollectionTreeItem final : public SMultiColumnTableRow<TSharedPtr<FCollectionCustomizationTreeElement>>
-{
-	SLATE_BEGIN_ARGS(SMovieGraphCollectionTreeItem) { }
-	SLATE_END_ARGS()
-
-	inline static const FName ColumnID_Name = TEXT("Name");
-	inline static const FName ColumnID_Value = TEXT("Value");
-
-	void Construct(const FArguments& InArgs, const TSharedPtr<SMovieGraphCollectionTreeView>& InOwnerTable,
-		const TSharedPtr<FCollectionCustomizationTreeElement>& InTreeElement, const TWeakPtr<IDetailLayoutBuilder>& InDetailBuilder);
-
-	/** DragDetected handler for this widget within the tree. */
-	FReply HandleOnDragDetected(const FGeometry& InGeometry, const FPointerEvent& MouseEvent) const;
-
-	/** CanAcceptDrop handler for this widget within the tree. */
-	TOptional<EItemDropZone> HandleOnCanAcceptDrop(const FDragDropEvent& DragDropEvent, const EItemDropZone DropZone, const TSharedPtr<FCollectionCustomizationTreeElement> TargetItem);
+	/** The condition group that the query belongs to. */
+	TWeakObjectPtr<UMovieGraphConditionGroup> WeakConditionGroup = nullptr;
 
-	/** AcceptDrop handler for this widget within the tree. */
-	FReply HandleOnAcceptDrop(const FDragDropEvent& DragDropEvent, EItemDropZone DropZone, TSharedPtr<FCollectionCustomizationTreeElement> TargetItem) const;
+	/** The condition group query associated with this widget. */
+	TWeakObjectPtr<UMovieGraphConditionGroupQueryBase> WeakQuery = nullptr;
 
-	/** Generates the widgets needed for condition group columns. */
-	TSharedRef<SWidget> GenerateConditionGroupColumnWidget(const FName& ColumnName, const TSharedPtr<FCollectionCustomizationTreeElement>& TreeElement);
-
-	/** Generates the widgets needed for condition group query columns. */
-	TSharedRef<SWidget> GenerateConditionGroupQueryColumnWidget(const FName& ColumnName, const TSharedPtr<FCollectionCustomizationTreeElement>& TreeElement);
-
-	/** Generates the widgets needed for condition group query property columns. */
-	TSharedRef<SWidget> GenerateQueryPropertyColumnWidget(const FName& ColumnName, const TSharedPtr<FCollectionCustomizationTreeElement>& TreeElement);
-
-	/** Convenience method for refreshing the tree view. */
-	void RefreshTreeView() const;
-
-	//~ Begin SMultiColumnTableRow overrides
-	virtual TSharedRef<SWidget> GenerateWidgetForColumn(const FName& ColumnName) override;
-	//~ End SMultiColumnTableRow overrides
-
-private:
-	/** The element that this tree item is displaying. */
-	TWeakPtr<FCollectionCustomizationTreeElement> WeakTreeElement;
-
-	/** The tree that owns this tree element. */
-	TWeakPtr<SMovieGraphCollectionTreeView> WeakTreeView;
-
-	/** The details builder associated with the customization. */
-	TWeakPtr<IDetailLayoutBuilder> DetailBuilder;
+	/** Called when the query type changes. */
+	FOnQueryTypeChanged OnQueryTypeChanged;
 };
 
 /** Customize how the Collection node appears in the details panel. */
@@ -329,8 +255,8 @@ protected:
 	//~ End IDetailCustomization interface
 
 private:
-	/** The tree view displaying condition groups and the queries contained in them. */
-	TSharedPtr<SMovieGraphCollectionTreeView> TreeView = nullptr;
+	/** The collection being displayed. */
+	TWeakObjectPtr<UMovieGraphCollection> WeakCollection;
 
 	/** The details builder associated with the customization. */
 	TWeakPtr<IDetailLayoutBuilder> DetailBuilder;
diff --git a/Engine/Plugins/MovieScene/MovieRenderPipeline/Source/MovieRenderPipelineEditor/Private/MovieRenderPipelineStyle.cpp b/Engine/Plugins/MovieScene/MovieRenderPipeline/Source/MovieRenderPipelineEditor/Private/MovieRenderPipelineStyle.cpp
--- a/Engine/Plugins/MovieScene/MovieRenderPipeline/Source/MovieRenderPipelineEditor/Private/MovieRenderPipelineStyle.cpp
+++ b/Engine/Plugins/MovieScene/MovieRenderPipeline/Source/MovieRenderPipelineEditor/Private/MovieRenderPipelineStyle.cpp
@@ -103,18 +103,8 @@ FMovieRenderPipelineStyle::FMovieRenderPipelineStyle()
 
 	Set("MovieRenderPipeline.TabIcon", new FSlateImageBrush(RootToContentDir(TEXT("TabIcon_24x.png")), Icon16x16));
 
-	Set("MovieRenderGraph.CollectionsTree.Row", FTableRowStyle(FAppStyle::Get().GetWidgetStyle<FTableRowStyle>("SimpleTableView.Row"))
-		.SetEvenRowBackgroundBrush(FSlateColorBrush(FStyleColors::Header))
-		.SetEvenRowBackgroundHoveredBrush(FSlateColorBrush(FStyleColors::Header))
-		.SetOddRowBackgroundBrush(FSlateColorBrush(FStyleColors::Header))
-		.SetOddRowBackgroundHoveredBrush(FSlateColorBrush(FStyleColors::Header))
-		.SetTextColor(FAppStyle::GetSlateColor("DefaultForeground"))
-		.SetSelectedTextColor(FAppStyle::GetSlateColor("DefaultForeground")));
-
-	Set("MovieRenderGraph.CollectionsTree.BottomBorder", new FSlateBorderBrush(NAME_None, FMargin(0, 0, 0, 1.0), FStyleColors::Recessed));
-	Set("MovieRenderGraph.CollectionsTree.TopBottomBorder", new FSlateBorderBrush(NAME_None, FMargin(0, 1.0, 0, 1.0), FStyleColors::Recessed));
-	Set("MovieRenderGraph.CollectionsTree.BottomLeftBorder", new FSlateBorderBrush(NAME_None, FMargin(1.0, 0, 0, 1.0), FStyleColors::Recessed));
-	Set("MovieRenderGraph.CollectionsTree.IndentedRowBackground", new FSlateColorBrush(FAppStyle::Get().GetSlateColor("Colors.Panel")));
+	Set("MovieRenderGraph.CollectionsTree.SmallComboButton", FComboButtonStyle(FAppStyle::Get().GetWidgetStyle<FComboButtonStyle>("SimpleComboButton"))
+		.SetDownArrowPadding(FMargin(0)));
 
 	FSlateStyleRegistry::RegisterSlateStyle(*this);
 }
