MovieGraph: Prevent multi-layer EXR node from using the CDO FileNameFormat.

#rnx

diff --git a/Engine/Plugins/MovieScene/MovieRenderPipeline/Source/MovieRenderPipelineRenderPasses/Private/MovieGraphImageSequenceOutputNode.cpp b/Engine/Plugins/MovieScene/MovieRenderPipeline/Source/MovieRenderPipelineRenderPasses/Private/MovieGraphImageSequenceOutputNode.cpp
--- a/Engine/Plugins/MovieScene/MovieRenderPipeline/Source/MovieRenderPipelineRenderPasses/Private/MovieGraphImageSequenceOutputNode.cpp
+++ b/Engine/Plugins/MovieScene/MovieRenderPipeline/Source/MovieRenderPipelineRenderPasses/Private/MovieGraphImageSequenceOutputNode.cpp
@@ -619,7 +619,7 @@ void UMovieGraphImageSequenceOutputNode_MultiLayerEXR::OnReceiveImageDataImpl(UM
 	// map to the same filename.
 	TMap<FString, TArray<FMovieGraphRenderDataIdentifier>> FilenameToRenderIDs;
 	TMap<FString, FMovieGraphResolveArgs> FilenameToResolveArgs;
-	GetFilenameToRenderIDMappings(InPipeline, InRawFrameData, FilenameToRenderIDs, FilenameToResolveArgs);
+	GetFilenameToRenderIDMappings(ParentNode, InPipeline, InRawFrameData, FilenameToRenderIDs, FilenameToResolveArgs);
 
 	// Write an EXR for each filename, which potentially contains multiple passes (render IDs).
 	for (const TPair<FString, TArray<FMovieGraphRenderDataIdentifier>>& RenderIDsForFilename : FilenameToRenderIDs)
@@ -694,6 +694,7 @@ void UMovieGraphImageSequenceOutputNode_MultiLayerEXR::OnReceiveImageDataImpl(UM
 }
 
 void UMovieGraphImageSequenceOutputNode_MultiLayerEXR::GetFilenameToRenderIDMappings(
+	const UMovieGraphImageSequenceOutputNode_MultiLayerEXR* InParentNode,
 	UMovieGraphPipeline* InPipeline, UE::MovieGraph::FMovieGraphOutputMergerFrame* InRawFrameData,
 	TMap<FString, TArray<FMovieGraphRenderDataIdentifier>>& OutFilenameToRenderIDs,
 	TMap<FString, FMovieGraphResolveArgs>& OutFilenameToResolveArgs) const
@@ -716,7 +717,7 @@ void UMovieGraphImageSequenceOutputNode_MultiLayerEXR::GetFilenameToRenderIDMapp
 	{
 		constexpr int32 ResolutionIndex = 0;
 		FMovieGraphResolveArgs ResolveArgs;
-		const FString PreliminaryFileName = ResolveOutputFilename(InPipeline, ResolutionIndex, InRawFrameData, RenderPassData.Key.RootBranchName, ResolveArgs);
+		const FString PreliminaryFileName = ResolveOutputFilename(InParentNode, InPipeline, ResolutionIndex, InRawFrameData, RenderPassData.Key.RootBranchName, ResolveArgs);
 		
 		TArray<FMovieGraphRenderDataIdentifier>& RenderIDs = OutFilenameToRenderIDs.FindOrAdd(PreliminaryFileName);
 		RenderIDs.Add(RenderPassData.Key);
@@ -750,7 +751,7 @@ void UMovieGraphImageSequenceOutputNode_MultiLayerEXR::GetFilenameToRenderIDMapp
 			// Re-resolve the filename, this time using the resolution index to generate a filename that will only contain
 			// passes with this particular resolution
 			FMovieGraphResolveArgs ResolveArgs;
-			const FString FinalFilename = ResolveOutputFilename(InPipeline, ResolutionIndex, InRawFrameData, RenderID.RootBranchName, ResolveArgs);
+			const FString FinalFilename = ResolveOutputFilename(InParentNode, InPipeline, ResolutionIndex, InRawFrameData, RenderID.RootBranchName, ResolveArgs);
 
 			TArray<FMovieGraphRenderDataIdentifier>& RenderIDs = OutFilenameToRenderIDs.FindOrAdd(FinalFilename);
 			RenderIDs.Add(RenderID);
@@ -761,6 +762,7 @@ void UMovieGraphImageSequenceOutputNode_MultiLayerEXR::GetFilenameToRenderIDMapp
 }
 
 FString UMovieGraphImageSequenceOutputNode_MultiLayerEXR::ResolveOutputFilename(
+	const UMovieGraphImageSequenceOutputNode_MultiLayerEXR* InParentNode,
 	const UMovieGraphPipeline* InPipeline,
 	const int32 ResolutionIndex, const UE::MovieGraph::FMovieGraphOutputMergerFrame* InRawFrameData,
 	const FName& InBranchName, FMovieGraphResolveArgs& OutResolveArgs) const
@@ -775,7 +777,7 @@ FString UMovieGraphImageSequenceOutputNode_MultiLayerEXR::ResolveOutputFilename(
 	}
 	
 	// If we have more than one resolution we'll store it as "_Add" / "_Add(1)" etc via {ExtraTag}.
-	FString FileNameFormatString = FileNameFormat + "{ExtraTag}";
+	FString FileNameFormatString = InParentNode->FileNameFormat + "{ExtraTag}";
 
 	const FString FilePathFormatString = OutputSettings->OutputDirectory.Path / FileNameFormatString;
 	
diff --git a/Engine/Plugins/MovieScene/MovieRenderPipeline/Source/MovieRenderPipelineRenderPasses/Public/MovieGraphImageSequenceOutputNode.h b/Engine/Plugins/MovieScene/MovieRenderPipeline/Source/MovieRenderPipelineRenderPasses/Public/MovieGraphImageSequenceOutputNode.h
--- a/Engine/Plugins/MovieScene/MovieRenderPipeline/Source/MovieRenderPipelineRenderPasses/Public/MovieGraphImageSequenceOutputNode.h
+++ b/Engine/Plugins/MovieScene/MovieRenderPipeline/Source/MovieRenderPipelineRenderPasses/Public/MovieGraphImageSequenceOutputNode.h
@@ -228,6 +228,7 @@ private:
 	 * with the same resolution.
 	 */
 	void GetFilenameToRenderIDMappings(
+		const UMovieGraphImageSequenceOutputNode_MultiLayerEXR* InParentNode,
 		UMovieGraphPipeline* InPipeline, UE::MovieGraph::FMovieGraphOutputMergerFrame* InRawFrameData,
 		TMap<FString, TArray<FMovieGraphRenderDataIdentifier>>& OutFilenameToRenderIDs,
 		TMap<FString, FMovieGraphResolveArgs>& OutFilenameToResolveArgs) const;
@@ -237,6 +238,7 @@ private:
 	 * Use GetFilenameToRenderIDMappings() to guarantee that the filename respects EXR limitations.
 	 */
 	FString ResolveOutputFilename(
+		const UMovieGraphImageSequenceOutputNode_MultiLayerEXR* InParentNode,
 		const UMovieGraphPipeline* InPipeline, const int32 ResolutionIndex, const UE::MovieGraph::FMovieGraphOutputMergerFrame* InRawFrameData,
 		const FName& InBranchName, FMovieGraphResolveArgs& OutResolveArgs) const;
 };
