Movie Graph:

-Fix: {frame_number} and {frame_number_shot} don't show correct frame numbers

diff --git a/Engine/Plugins/MovieScene/MovieRenderPipeline/Source/MovieRenderPipelineCore/Private/Graph/MovieGraphCoreTimeStep.cpp b/Engine/Plugins/MovieScene/MovieRenderPipeline/Source/MovieRenderPipelineCore/Private/Graph/MovieGraphCoreTimeStep.cpp
--- a/Engine/Plugins/MovieScene/MovieRenderPipeline/Source/MovieRenderPipelineCore/Private/Graph/MovieGraphCoreTimeStep.cpp
+++ b/Engine/Plugins/MovieScene/MovieRenderPipeline/Source/MovieRenderPipelineCore/Private/Graph/MovieGraphCoreTimeStep.cpp
@@ -74,6 +74,9 @@ void UMovieGraphCoreTimeStep::TickProducingFrames()
 
 		// Ensure we've set it in the CurrentTimeStepData so things can fetch from it below.
 		CurrentTimeStepData.EvaluatedConfig = TObjectPtr<UMovieGraphEvaluatedConfig>(CurrentFrameData.EvaluatedConfig.Get());
+		
+		// Get ready to begin tracking relative shot frame count
+		CurrentTimeStepData.ShotOutputFrameNumber = -1;
 
 		// Sets up the render state, etc.
 		GetOwningGraph()->SetupShot(CurrentCameraCut);
@@ -432,6 +435,7 @@ void UMovieGraphCoreTimeStep::TickProducingFrames()
 			if (CurrentCameraCut->ShotInfo.State == EMovieRenderShotState::Rendering)
 			{
 				GetOwningGraph()->GetCustomEngineTimeStep()->SharedTimeStepData.OutputFrameNumber++;
+				CurrentTimeStepData.ShotOutputFrameNumber++;
 			}
 				
 			// If we've rendered the last temporal sub-sample, we've started a new output frame
diff --git a/Engine/Plugins/MovieScene/MovieRenderPipeline/Source/MovieRenderPipelineCore/Public/Graph/MovieGraphTimeStepData.h b/Engine/Plugins/MovieScene/MovieRenderPipeline/Source/MovieRenderPipelineCore/Public/Graph/MovieGraphTimeStepData.h
--- a/Engine/Plugins/MovieScene/MovieRenderPipeline/Source/MovieRenderPipelineCore/Public/Graph/MovieGraphTimeStepData.h
+++ b/Engine/Plugins/MovieScene/MovieRenderPipeline/Source/MovieRenderPipelineCore/Public/Graph/MovieGraphTimeStepData.h
@@ -34,9 +34,18 @@ public:
 		, bRequiresAccumulator(false)
 	{}
 
+	/*
+	 * Relative to zero for the entire render.
+	 */
 	UPROPERTY(BlueprintReadWrite, EditAnywhere, Category = "Movie Graph")
 	int32 OutputFrameNumber;
 
+	/*
+	 * Relative to zero for the current shot.
+	 */
+	UPROPERTY(BlueprintReadWrite, EditAnywhere, Category = "Movie Graph")
+	int32 ShotOutputFrameNumber;
+
 	UPROPERTY(BlueprintReadWrite, EditAnywhere, Category = "Movie Graph")
 	int32 RenderedFrameNumber;
 
@@ -113,7 +122,10 @@ public:
 	UPROPERTY(BlueprintReadWrite, EditAnywhere, Category = "Movie Graph")
 	FTimecode RootTimeCode;
 
-	/** The current frame number at the root (sequence) level. */
+	/*
+	 * The current frame number at the root (sequence) level.
+	 * This is the same number as what is shown in sequencer when viewing the sequence.
+	 */
 	UPROPERTY(BlueprintReadWrite, EditAnywhere, Category = "Movie Graph")
 	FFrameNumber RootFrameNumber;
 
@@ -121,7 +133,10 @@ public:
 	UPROPERTY(BlueprintReadWrite, EditAnywhere, Category = "Movie Graph")
 	FTimecode ShotTimeCode;
 
-	/** The current frame number at the shot level. */
+	/*
+	 * The current frame number at the shot level.
+	 * This is the same number as what is shown in sequencer when viewing the shot.
+	 */
 	UPROPERTY(BlueprintReadWrite, EditAnywhere, Category = "Movie Graph")
 	FFrameNumber ShotFrameNumber;
 };
diff --git a/Engine/Plugins/MovieScene/MovieRenderPipeline/Source/MovieRenderPipelineRenderPasses/Private/MovieGraphImageSequenceOutputNode.cpp b/Engine/Plugins/MovieScene/MovieRenderPipeline/Source/MovieRenderPipelineRenderPasses/Private/MovieGraphImageSequenceOutputNode.cpp
--- a/Engine/Plugins/MovieScene/MovieRenderPipeline/Source/MovieRenderPipelineRenderPasses/Private/MovieGraphImageSequenceOutputNode.cpp
+++ b/Engine/Plugins/MovieScene/MovieRenderPipeline/Source/MovieRenderPipelineRenderPasses/Private/MovieGraphImageSequenceOutputNode.cpp
@@ -39,10 +39,10 @@ namespace UE::MovieGraph::Private
 
 		FMovieGraphFilenameResolveParams Params = FMovieGraphFilenameResolveParams();
 		Params.RenderDataIdentifier = InRenderId;
-		//Params.RootFrameNumber = InTraversalContext.Time.RootFrameNumber;
-		//Params.ShotFrameNumber = InTraversalContext.Time.ShotFrameNumber;
+		Params.RootFrameNumber = InTraversalContext.Time.RootFrameNumber.Value;
+		Params.ShotFrameNumber = InTraversalContext.Time.ShotFrameNumber.Value;
 		Params.RootFrameNumberRel = InTraversalContext.Time.OutputFrameNumber;
-		//Params.ShotFrameNumberRel = InTraversalContext.Time.ShotFrameNumberRel
+		Params.ShotFrameNumberRel = InTraversalContext.Time.ShotOutputFrameNumber;
 		//Params.FileMetadata = ToDo: Track File Metadata
 		const UMovieGraphGlobalOutputSettingNode* OutputSettingNode = InEvaluatedConfig->GetSettingForBranch<UMovieGraphGlobalOutputSettingNode>(UMovieGraphNode::GlobalsPinName);
 		if (IsValid(OutputSettingNode))
@@ -55,12 +55,12 @@ namespace UE::MovieGraph::Private
 
 		// If time dilation is in effect, RootFrameNumber and ShotFrameNumber will contain duplicates and the files will overwrite each other, 
 		// so we force them into relative mode and then warn users we did that (as their numbers will jump from say 1001 -> 0000).
-		bool bForceRelativeFrameNumbers = true; // TODO: Use relative frame numbers until we track Root vs. Shot frame numbers. (Previously false);
-		//if (FileNameFormatString.Contains(TEXT("{frame")) InTraversalContext.Time.IsTimeDilated() && !FileNameFormatString.Contains(TEXT("_rel}")))
-		//{
-		//	UE_LOG(LogMovieRenderPipeline, Warning, TEXT("Time Dilation was used but output format does not use relative time, forcing relative numbers. Change {frame_number} to {frame_number_rel} (or shot version) to remove this message."));
-		//	bForceRelativeFrameNumbers = true;
-		//}
+		bool bForceRelativeFrameNumbers = false;
+		// if (FileNameFormatString.Contains(TEXT("{frame")) && InTraversalContext.Time.IsTimeDilated() && !FileNameFormatString.Contains(TEXT("_rel}")))
+		// {
+		// 	UE_LOG(LogMovieRenderPipeline, Warning, TEXT("Time Dilation was used but output format does not use relative time, forcing relative numbers. Change {frame_number} to {frame_number_rel} (or shot version) to remove this message."));
+		// 	bForceRelativeFrameNumbers = true;
+		// }
 		Params.bForceRelativeFrameNumbers = bForceRelativeFrameNumbers;
 		Params.bEnsureAbsolutePath = true;
 		Params.FileNameFormatOverrides = InAdditionalFormatArgs;
