CIS: Fixed an uninitialized struct member variable in FMovieGraphTimeStepData

diff --git a/Engine/Plugins/MovieScene/MovieRenderPipeline/Source/MovieRenderPipelineCore/Public/Graph/MovieGraphTimeStepData.h b/Engine/Plugins/MovieScene/MovieRenderPipeline/Source/MovieRenderPipelineCore/Public/Graph/MovieGraphTimeStepData.h
--- a/Engine/Plugins/MovieScene/MovieRenderPipeline/Source/MovieRenderPipelineCore/Public/Graph/MovieGraphTimeStepData.h
+++ b/Engine/Plugins/MovieScene/MovieRenderPipeline/Source/MovieRenderPipelineCore/Public/Graph/MovieGraphTimeStepData.h
@@ -20,6 +20,7 @@ struct MOVIERENDERPIPELINECORE_API FMovieGraphTimeStepData
 public:
 	FMovieGraphTimeStepData()
 		: OutputFrameNumber(0)
+		, ShotOutputFrameNumber(0)
 		, RenderedFrameNumber(0)
 		, FrameDeltaTime(0.f)
 		, WorldTimeDilation(0.f)
