Movie Graph:

-Rename render_layer token to layer_name

diff --git a/Engine/Plugins/MovieScene/MovieRenderPipeline/Source/MovieRenderPipelineCore/Private/Graph/MovieGraphBlueprintLibrary.cpp b/Engine/Plugins/MovieScene/MovieRenderPipeline/Source/MovieRenderPipelineCore/Private/Graph/MovieGraphBlueprintLibrary.cpp
--- a/Engine/Plugins/MovieScene/MovieRenderPipeline/Source/MovieRenderPipelineCore/Private/Graph/MovieGraphBlueprintLibrary.cpp
+++ b/Engine/Plugins/MovieScene/MovieRenderPipeline/Source/MovieRenderPipelineCore/Private/Graph/MovieGraphBlueprintLibrary.cpp
@@ -82,7 +82,7 @@ FString UMovieGraphBlueprintLibrary::ResolveFilenameFormatArguments(const FStrin
 		OutMergedFormatArgs.FilenameArguments.Add(TEXT("camera_name"), InParams.RenderDataIdentifier.CameraName);
 		OutMergedFormatArgs.FilenameArguments.Add(TEXT("renderer_name"), InParams.RenderDataIdentifier.RendererName);
 		OutMergedFormatArgs.FilenameArguments.Add(TEXT("renderer_sub_name"), InParams.RenderDataIdentifier.SubResourceName);
-		OutMergedFormatArgs.FilenameArguments.Add(TEXT("render_layer"), RenderLayerName);
+		OutMergedFormatArgs.FilenameArguments.Add(TEXT("layer_name"), RenderLayerName);
 
 		// TODO: Some of these are per render layer and need to be stored that way. EXRs will have the metadata for all the layers/cameras/etc. in one file.
 		OutMergedFormatArgs.FileMetadata.Add(TEXT("unreal/cameraName"), InParams.RenderDataIdentifier.CameraName);
diff --git a/Engine/Plugins/MovieScene/MovieRenderPipeline/Source/MovieRenderPipelineCore/Private/Graph/Nodes/MovieGraphCommandLineEncoderNode.cpp b/Engine/Plugins/MovieScene/MovieRenderPipeline/Source/MovieRenderPipelineCore/Private/Graph/Nodes/MovieGraphCommandLineEncoderNode.cpp
--- a/Engine/Plugins/MovieScene/MovieRenderPipeline/Source/MovieRenderPipelineCore/Private/Graph/Nodes/MovieGraphCommandLineEncoderNode.cpp
+++ b/Engine/Plugins/MovieScene/MovieRenderPipeline/Source/MovieRenderPipelineCore/Private/Graph/Nodes/MovieGraphCommandLineEncoderNode.cpp
@@ -22,7 +22,7 @@ UMovieGraphCommandLineEncoderNode::UMovieGraphCommandLineEncoderNode()
 	, bSkipEncodeOnRenderCanceled(true)
 	, bRetainInputTextFiles(false)
 {
-	FileNameFormat = TEXT("{sequence_name}.{render_layer}");
+	FileNameFormat = TEXT("{sequence_name}.{layer_name}");
 		
 	// Sensible defaults for cmdline arguments
 	CommandLineFormat = TEXT("-hide_banner -y -loglevel error {VideoInputs} {AudioInputs} -acodec {AudioCodec} -vcodec {VideoCodec} {Quality} \"{OutputPath}\"");
diff --git a/Engine/Plugins/MovieScene/MovieRenderPipeline/Source/MovieRenderPipelineCore/Public/Graph/MovieGraphFilenameResolveParams.h b/Engine/Plugins/MovieScene/MovieRenderPipeline/Source/MovieRenderPipelineCore/Public/Graph/MovieGraphFilenameResolveParams.h
--- a/Engine/Plugins/MovieScene/MovieRenderPipeline/Source/MovieRenderPipelineCore/Public/Graph/MovieGraphFilenameResolveParams.h
+++ b/Engine/Plugins/MovieScene/MovieRenderPipeline/Source/MovieRenderPipelineCore/Public/Graph/MovieGraphFilenameResolveParams.h
@@ -66,7 +66,7 @@ struct FMovieGraphFilenameResolveParams
 	
 	/**
 	* This is used to fill out tokens related to the render data that the file represents.
-	* Used to fill out the {camera_name}, {render_layer}, {renderer_name}, {renderer_sub_name} tokens.
+	* Used to fill out the {camera_name}, {layer_name}, {renderer_name}, {renderer_sub_name} tokens.
 	*/
 	UPROPERTY(BlueprintReadWrite, EditAnywhere, Category = "Movie Graph")
 	FMovieGraphRenderDataIdentifier RenderDataIdentifier;
diff --git a/Engine/Plugins/MovieScene/MovieRenderPipeline/Source/MovieRenderPipelineCore/Public/Graph/Nodes/MovieGraphFileOutputNode.h b/Engine/Plugins/MovieScene/MovieRenderPipeline/Source/MovieRenderPipelineCore/Public/Graph/Nodes/MovieGraphFileOutputNode.h
--- a/Engine/Plugins/MovieScene/MovieRenderPipeline/Source/MovieRenderPipelineCore/Public/Graph/Nodes/MovieGraphFileOutputNode.h
+++ b/Engine/Plugins/MovieScene/MovieRenderPipeline/Source/MovieRenderPipelineCore/Public/Graph/Nodes/MovieGraphFileOutputNode.h
@@ -57,7 +57,7 @@ public:
 
 	/** What format string should the final files use? Can include folder prefixes, and format string ({shot_name}, etc.) */
 	UPROPERTY(EditAnywhere, BlueprintReadWrite, Category = "File Output", meta=(EditCondition="bOverride_FileNameFormat"))
-	FString FileNameFormat = "{sequence_name}.{render_layer}.{frame_number}";
+	FString FileNameFormat = "{sequence_name}.{layer_name}.{frame_number}";
 
 protected:
 	virtual void OnReceiveImageDataImpl(UMovieGraphPipeline* InPipeline, UE::MovieGraph::FMovieGraphOutputMergerFrame* InRawFrameData, const TSet<FMovieGraphRenderDataIdentifier>& InMask) {}
diff --git a/Engine/Plugins/MovieScene/MovieRenderPipeline/Source/MovieRenderPipelineRenderPasses/Private/MovieGraphImageSequenceOutputNode.cpp b/Engine/Plugins/MovieScene/MovieRenderPipeline/Source/MovieRenderPipelineRenderPasses/Private/MovieGraphImageSequenceOutputNode.cpp
--- a/Engine/Plugins/MovieScene/MovieRenderPipeline/Source/MovieRenderPipelineRenderPasses/Private/MovieGraphImageSequenceOutputNode.cpp
+++ b/Engine/Plugins/MovieScene/MovieRenderPipeline/Source/MovieRenderPipelineRenderPasses/Private/MovieGraphImageSequenceOutputNode.cpp
@@ -710,7 +710,7 @@ void UMovieGraphImageSequenceOutputNode_MultiLayerEXR::GetFilenameToRenderIDMapp
 	// This two-step process is necessary due to the flexibility in file naming, and the multi-layer nature of EXRs.
 	// For example, if the file name format is "{sequence_name}.{frame_number}", and the second of two branches in the
 	// graph has a differing resolution, only after resolving the output filenames for all outputs is a problem found;
-	// layers of differing resolutions will be written to the same file. Using "{render_layer}.{sequence_name}.{frame_number}"
+	// layers of differing resolutions will be written to the same file. Using "{layer_name}.{sequence_name}.{frame_number}"
 	// as the file name format would prevent the issue, but the two-step process is a generic way of approaching the
 	// problem.
 	for (const TPair<FMovieGraphRenderDataIdentifier, TUniquePtr<FImagePixelData>>& RenderPassData : InRawFrameData->ImageOutputData)
