MovieGraph:

-Fix: "Custom Playback Range" functionality is missing

diff --git a/Engine/Plugins/MovieScene/MovieRenderPipeline/Source/MovieRenderPipelineCore/Private/Graph/MovieGraphPipeline.cpp b/Engine/Plugins/MovieScene/MovieRenderPipeline/Source/MovieRenderPipelineCore/Private/Graph/MovieGraphPipeline.cpp
--- a/Engine/Plugins/MovieScene/MovieRenderPipeline/Source/MovieRenderPipelineCore/Private/Graph/MovieGraphPipeline.cpp
+++ b/Engine/Plugins/MovieScene/MovieRenderPipeline/Source/MovieRenderPipelineCore/Private/Graph/MovieGraphPipeline.cpp
@@ -1114,6 +1114,8 @@ void UMovieGraphPipeline::TransitionToState(const EMovieRenderPipelineState InNe
 
 			OutputNodesDataSentTo.Reset();
 
+			GraphDataSourceInstance->RestoreCachedDataPostJob();
+
 			OnMoviePipelineFinishedImpl();
 			UE_LOG(LogMovieRenderPipeline, Log, TEXT("Movie Graph Render completed. Duration: %s"), *(FDateTime::UtcNow() - GraphInitializationTime).ToString());
 		}
diff --git a/Engine/Plugins/MovieScene/MovieRenderPipeline/Source/MovieRenderPipelineCore/Private/Graph/MovieGraphSequenceDataSource.cpp b/Engine/Plugins/MovieScene/MovieRenderPipeline/Source/MovieRenderPipelineCore/Private/Graph/MovieGraphSequenceDataSource.cpp
--- a/Engine/Plugins/MovieScene/MovieRenderPipeline/Source/MovieRenderPipelineCore/Private/Graph/MovieGraphSequenceDataSource.cpp
+++ b/Engine/Plugins/MovieScene/MovieRenderPipeline/Source/MovieRenderPipelineCore/Private/Graph/MovieGraphSequenceDataSource.cpp
@@ -14,6 +14,7 @@
 #include "CoreGlobals.h"
 #include "EngineUtils.h"
 #include "Engine/GameViewportClient.h"
+#include "Graph/Nodes/MovieGraphGlobalOutputSettingNode.h"
 
 UMovieGraphSequenceDataSource::UMovieGraphSequenceDataSource()
 {
@@ -46,6 +47,9 @@ void UMovieGraphSequenceDataSource::RestoreCachedDataPostJob()
 	{
 		Viewport->bDisableWorldRendering = false;
 	}
+
+	ULevelSequence* RootSequence = Cast<ULevelSequence>(GetOwningGraph()->GetCurrentJob()->Sequence.TryLoad());
+	MoviePipeline::RestoreCompleteSequenceHierarchy(RootSequence, CachedSequenceHierarchyRoot);
 }
 
 void UMovieGraphSequenceDataSource::UpdateShotList()
@@ -58,6 +62,48 @@ void UMovieGraphSequenceDataSource::UpdateShotList()
 	}
 }
 
+void UMovieGraphSequenceDataSource::OverrideSequencePlaybackRangeFromGlobalOutputSettings(ULevelSequence* InSequence)
+{
+	FMovieGraphTraversalContext TraversalContext;
+	UMoviePipelineExecutorJob* CurrentJob = GetOwningGraph()->GetCurrentJob();
+	TraversalContext.Job = CurrentJob;
+	FString OutErrorMessage;
+	UMovieGraphEvaluatedConfig* EvaluatedGraph = CurrentJob->GetGraphPreset()->CreateFlattenedGraph(TraversalContext, OutErrorMessage);
+
+	constexpr bool bIncludeCDOs = true;
+	constexpr bool bExactMatch = false;
+	UMovieGraphGlobalOutputSettingNode* OutputSetting =
+		EvaluatedGraph->GetSettingForBranch<UMovieGraphGlobalOutputSettingNode>(UMovieGraphNode::GlobalsPinName, bIncludeCDOs, bExactMatch);
+
+	TRange<FFrameNumber> CurrentPlaybackRange = InSequence->GetMovieScene()->GetPlaybackRange();
+
+	FFrameNumber StartFrameTickResolution = CurrentPlaybackRange.GetLowerBound().GetValue();
+	FFrameNumber EndFrameTickResolution = CurrentPlaybackRange.GetUpperBound().GetValue();
+	
+	if (OutputSetting->bOverride_CustomPlaybackRangeStartFrame)
+	{
+		StartFrameTickResolution = FFrameRate::TransformTime(FFrameTime(FFrameNumber(OutputSetting->CustomPlaybackRangeStartFrame)), InSequence->GetMovieScene()->GetDisplayRate(), InSequence->GetMovieScene()->GetTickResolution()).FloorToFrame();
+	}
+	if (OutputSetting->bOverride_CustomPlaybackRangeEndFrame)
+	{
+		EndFrameTickResolution = FFrameRate::TransformTime(FFrameTime(FFrameNumber(OutputSetting->CustomPlaybackRangeEndFrame)), InSequence->GetMovieScene()->GetDisplayRate(), InSequence->GetMovieScene()->GetTickResolution()).CeilToFrame();
+	}
+
+	TRange<FFrameNumber> NewPlaybackRange = TRange<FFrameNumber>(StartFrameTickResolution, EndFrameTickResolution);
+	
+#if WITH_EDITOR
+	InSequence->GetMovieScene()->SetPlaybackRangeLocked(false);
+	InSequence->GetMovieScene()->SetReadOnly(false);
+#endif
+	InSequence->GetMovieScene()->SetPlaybackRange(NewPlaybackRange);
+
+	// Warn about zero length playback ranges, often happens because they set the Start/End frame to the same frame.
+	if (InSequence->GetMovieScene()->GetPlaybackRange().IsEmpty())
+	{
+		UE_LOG(LogMovieRenderPipeline, Error, TEXT("Playback Range was zero. End Frames are exclusive, did you mean [n, n+1]?"));
+	}
+}
+
 void UMovieGraphSequenceDataSource::CacheLevelSequenceData(ULevelSequence* InSequence)
 {
 	// There is a reasonable chance that there exists a Level Sequence Actor in the world already set up to play this sequence.
@@ -93,6 +139,13 @@ void UMovieGraphSequenceDataSource::CacheLevelSequenceData(ULevelSequence* InSeq
 	LevelSequenceActor->PlaybackSettings.bPauseAtEnd = true;
 	LevelSequenceActor->PlaybackSettings.FinishCompletionStateOverride = EMovieSceneCompletionModeOverride::ForceRestoreState;
 
+	// Cache the sequence data to be restored later
+	CachedSequenceHierarchyRoot = MakeShared<MoviePipeline::FCameraCutSubSectionHierarchyNode>();
+	MoviePipeline::CacheCompleteSequenceHierarchy(InSequence, CachedSequenceHierarchyRoot);
+	
+	// Override the frame range on the target sequence if needed first before anyone has a chance to modify it.
+	OverrideSequencePlaybackRangeFromGlobalOutputSettings(InSequence);
+
 	// Ensure the (possibly new) Level Sequence Actor uses our sequence
 	LevelSequenceActor->SetSequence(InSequence);
 
diff --git a/Engine/Plugins/MovieScene/MovieRenderPipeline/Source/MovieRenderPipelineCore/Public/Graph/MovieGraphSequenceDataSource.h b/Engine/Plugins/MovieScene/MovieRenderPipeline/Source/MovieRenderPipelineCore/Public/Graph/MovieGraphSequenceDataSource.h
--- a/Engine/Plugins/MovieScene/MovieRenderPipeline/Source/MovieRenderPipelineCore/Public/Graph/MovieGraphSequenceDataSource.h
+++ b/Engine/Plugins/MovieScene/MovieRenderPipeline/Source/MovieRenderPipelineCore/Public/Graph/MovieGraphSequenceDataSource.h
@@ -54,6 +54,7 @@ public:
 	virtual void ExpandShot(const TObjectPtr<UMoviePipelineExecutorShot>& InShot, const int32 InLeftDeltaFrames, const int32 InLeftDeltaFramesUserPoV,
 		const int32 InRightDeltaFrames, const bool bInPrepass) override;
 protected:
+	void OverrideSequencePlaybackRangeFromGlobalOutputSettings(ULevelSequence* InSequence);
 	void CacheLevelSequenceData(ULevelSequence* InSequence);
 	void OnSequenceEvaluated(const UMovieSceneSequencePlayer& Player, FFrameTime CurrentTime, FFrameTime PreviousTime);
 
@@ -63,4 +64,6 @@ protected:
 
 	/** Custom Time Controller for the Sequence Player, used to match Custom TimeStep without any floating point accumulation errors. */
 	TSharedPtr<UE::MovieGraph::FMovieGraphSequenceTimeController> CustomSequenceTimeController;
+
+	TSharedPtr<MoviePipeline::FCameraCutSubSectionHierarchyNode> CachedSequenceHierarchyRoot;
 };
\ No newline at end of file
diff --git a/Engine/Plugins/MovieScene/MovieRenderPipeline/Source/MovieRenderPipelineCore/Public/Graph/Nodes/MovieGraphGlobalOutputSettingNode.h b/Engine/Plugins/MovieScene/MovieRenderPipeline/Source/MovieRenderPipelineCore/Public/Graph/Nodes/MovieGraphGlobalOutputSettingNode.h
--- a/Engine/Plugins/MovieScene/MovieRenderPipeline/Source/MovieRenderPipelineCore/Public/Graph/Nodes/MovieGraphGlobalOutputSettingNode.h
+++ b/Engine/Plugins/MovieScene/MovieRenderPipeline/Source/MovieRenderPipelineCore/Public/Graph/Nodes/MovieGraphGlobalOutputSettingNode.h
@@ -72,6 +72,12 @@ public:
 	UPROPERTY(EditAnywhere, BlueprintReadWrite, Category = Overrides, meta = (InlineEditConditionToggle))
 	uint8 bOverride_HandleFrameCount : 1;
 
+	UPROPERTY(EditAnywhere, BlueprintReadWrite, Category = Overrides, meta = (InlineEditConditionToggle))
+	uint8 bOverride_CustomPlaybackRangeStartFrame : 1;
+
+	UPROPERTY(EditAnywhere, BlueprintReadWrite, Category = Overrides, meta = (InlineEditConditionToggle))
+	uint8 bOverride_CustomPlaybackRangeEndFrame : 1;
+
 	UPROPERTY(EditAnywhere, BlueprintReadWrite, Category = Overrides, meta = (InlineEditConditionToggle))
 	uint8 bOverride_VersioningSettings : 1;
 
@@ -111,6 +117,12 @@ public:
 	UPROPERTY(EditAnywhere, BlueprintReadWrite, Category = "Frames", meta = (UIMin = 0, ClampMin = 0, EditCondition = "bOverride_HandleFrameCount"))
 	int32 HandleFrameCount;
 
+	UPROPERTY(EditAnywhere, BlueprintReadWrite, Category = "Frames", meta = (EditCondition = "bOverride_CustomPlaybackRangeStartFrame"))
+	int32 CustomPlaybackRangeStartFrame;
+	
+	UPROPERTY(EditAnywhere, BlueprintReadWrite, Category = "Frames", meta = (EditCondition = "bOverride_CustomPlaybackRangeEndFrame"))
+	int32 CustomPlaybackRangeEndFrame;
+
 	/**
 	 * Determines how versioning should be handled (Auto Version, Version Number, etc.).
 	 */
