MoviePipeline: Updated the modifier execution order in the graph such that modifiers defined further downstream execute last.

#rnx

diff --git a/Engine/Plugins/MovieScene/MovieRenderPipeline/Source/MovieRenderPipelineCore/Private/Graph/MovieGraphPipeline.cpp b/Engine/Plugins/MovieScene/MovieRenderPipeline/Source/MovieRenderPipelineCore/Private/Graph/MovieGraphPipeline.cpp
--- a/Engine/Plugins/MovieScene/MovieRenderPipeline/Source/MovieRenderPipelineCore/Private/Graph/MovieGraphPipeline.cpp
+++ b/Engine/Plugins/MovieScene/MovieRenderPipeline/Source/MovieRenderPipelineCore/Private/Graph/MovieGraphPipeline.cpp
@@ -222,8 +222,20 @@ void UMovieGraphPipeline::UpdateLayerContentsInRenderLayerSubsystem(const UMovie
 		constexpr bool bExactMatch = true;
 		const TArray<UMovieGraphCollectionNode*> CollectionNodes =
 			EvaluatedConfig->GetSettingsForBranch<UMovieGraphCollectionNode>(LayerName, bIncludeCDOs, bExactMatch);
-		const TArray<UMovieGraphModifierNode*> ModifierNodes =
+		TArray<UMovieGraphModifierNode*> ModifierNodes =
 			EvaluatedConfig->GetSettingsForBranch<UMovieGraphModifierNode>(LayerName, bIncludeCDOs, bExactMatch);
+
+		// Graph evaluation discovers nodes working from the Outputs node, moving towards the Inputs node. Therefore modifiers (even if they are
+		// nodes which override an existing modifier) that are furthest downstream occur first in the array of modifiers returned by the evaluated
+		// graph. For example:
+		//
+		// A (definition) -> B (definition) -> C (definition) -> B (override) -> A (override) -> C (override)
+		//
+		// ... would be returned from the graph in the order of C, A, B
+		//
+		// However, users expect the *evaluation* order of the modifiers to be B, A, C, where the nodes furthest downstream execute last.
+		// Therefore, the modifier node array returned by the evaluated graph needs to be reversed.
+		Algo::Reverse(ModifierNodes);
 		
 		for (const UMovieGraphModifierNode* ModifierNode : ModifierNodes)
 		{
