MoviePipeline: Added layer name as an overlay to images that show up in the Render Preview window.

#rnx

#virtualized

diff --git a/Engine/Plugins/MovieScene/MovieRenderPipeline/Source/MovieRenderPipelineCore/Private/Graph/MovieGraphDefaultRenderer.cpp b/Engine/Plugins/MovieScene/MovieRenderPipeline/Source/MovieRenderPipelineCore/Private/Graph/MovieGraphDefaultRenderer.cpp
--- a/Engine/Plugins/MovieScene/MovieRenderPipeline/Source/MovieRenderPipelineCore/Private/Graph/MovieGraphDefaultRenderer.cpp
+++ b/Engine/Plugins/MovieScene/MovieRenderPipeline/Source/MovieRenderPipelineCore/Private/Graph/MovieGraphDefaultRenderer.cpp
@@ -433,13 +433,25 @@ TArray<FMovieGraphImagePreviewData> UMovieGraphDefaultRenderer::GetPreviewData()
 { 
 	TArray<FMovieGraphImagePreviewData> Results;
 
-	TArray<TObjectPtr<UTextureRenderTarget2D>> PooledTargets;
-	PooledViewRenderTargets.GenerateValueArray(PooledTargets);
+	// The evaluated config may not be available when the preview data is first requested
+	FString LayerName;
+	const UMovieGraphTimeStepBase* TimeStepInstance = GetOwningGraph()->GetTimeStepInstance();
+	const UMovieGraphEvaluatedConfig* EvaluatedConfig = TimeStepInstance ? TimeStepInstance->GetCalculatedTimeData().EvaluatedConfig : nullptr;
 
-	for (const TObjectPtr<UTextureRenderTarget2D>& Target : PooledTargets)
+	for (const TPair<UE::MovieGraph::DefaultRenderer::FMovieGraphImagePreviewDataPoolParams, TObjectPtr<UTextureRenderTarget2D>>& RenderTarget : PooledViewRenderTargets)
 	{
+		// Try to determine the layer name
+		if (EvaluatedConfig)
+		{
+			constexpr bool bIncludeCDOs = false;
+			const UMovieGraphRenderLayerNode* RenderLayerNode = EvaluatedConfig->GetSettingForBranch<UMovieGraphRenderLayerNode>(RenderTarget.Key.Identifier.RootBranchName, bIncludeCDOs);
+			LayerName = RenderLayerNode ? RenderLayerNode->GetRenderLayerName() : TEXT("");
+		}
+		
 		FMovieGraphImagePreviewData& Data = Results.AddDefaulted_GetRef();
-		Data.Texture = Target.Get();
+		Data.Identifier = RenderTarget.Key.Identifier;
+		Data.Texture = RenderTarget.Value.Get();
+		Data.LayerName = LayerName;
 	}
 
 	return Results;
diff --git a/Engine/Plugins/MovieScene/MovieRenderPipeline/Source/MovieRenderPipelineCore/Public/Graph/MovieGraphDataTypes.h b/Engine/Plugins/MovieScene/MovieRenderPipeline/Source/MovieRenderPipelineCore/Public/Graph/MovieGraphDataTypes.h
--- a/Engine/Plugins/MovieScene/MovieRenderPipeline/Source/MovieRenderPipelineCore/Public/Graph/MovieGraphDataTypes.h
+++ b/Engine/Plugins/MovieScene/MovieRenderPipeline/Source/MovieRenderPipelineCore/Public/Graph/MovieGraphDataTypes.h
@@ -43,9 +43,13 @@ struct FMovieGraphImagePreviewData
 	UPROPERTY(BlueprintReadOnly, Category = "Movie Graph")
 	class UTexture* Texture;
 	
-	/** The identifier for the image, containing the render layer name, renderer, etc. */
+	/** The identifier for the image, containing the branch name, renderer, etc. */
 	UPROPERTY(BlueprintReadOnly, Category = "Movie Graph")
 	FMovieGraphRenderDataIdentifier Identifier;
+	
+	/** The layer name associated with the render. */
+	UPROPERTY(BlueprintReadOnly, Category = "Movie Graph")
+	FString LayerName;
 };
 
 USTRUCT(BlueprintType)
