MoviePipeline: Added two new query types available for use in collections (editor folder and sublevel).

Also fixed a few misc minor bugs in other queries.

#rnx

diff --git a/Engine/Plugins/MovieScene/MovieRenderPipeline/Source/MovieRenderPipelineCore/Private/Graph/MoviePipelineRenderLayerSubsystem.cpp b/Engine/Plugins/MovieScene/MovieRenderPipeline/Source/MovieRenderPipelineCore/Private/Graph/MoviePipelineRenderLayerSubsystem.cpp
--- a/Engine/Plugins/MovieScene/MovieRenderPipeline/Source/MovieRenderPipelineCore/Private/Graph/MoviePipelineRenderLayerSubsystem.cpp
+++ b/Engine/Plugins/MovieScene/MovieRenderPipeline/Source/MovieRenderPipelineCore/Private/Graph/MoviePipelineRenderLayerSubsystem.cpp
@@ -17,11 +17,19 @@
 #include "Widgets/Input/SMultiLineEditableTextBox.h"
 
 #if WITH_EDITOR
+#include "ActorFolderPickingMode.h"
+#include "ActorFolderTreeItem.h"
+#include "ActorMode.h"
 #include "ActorTreeItem.h"
 #include "ClassViewerFilter.h"
 #include "ClassViewerModule.h"
+#include "ComponentTreeItem.h"
+#include "ContentBrowserDataSource.h"
+#include "ContentBrowserModule.h"
 #include "Editor.h"
+#include "EditorActorFolders.h"
 #include "Graph/MovieGraphSharedWidgets.h"
+#include "IContentBrowserSingleton.h"
 #include "ISceneOutliner.h"
 #include "SceneOutlinerModule.h"
 #include "SceneOutlinerPublicTypes.h"
@@ -411,7 +419,7 @@ EMovieGraphConditionGroupQueryOpType UMovieGraphConditionGroupQueryBase::GetOper
 	return OpType;
 }
 
-void UMovieGraphConditionGroupQueryBase::Evaluate(const TArray<AActor*>& InActorsToQuery, TSet<AActor*>& OutMatchingActors) const
+void UMovieGraphConditionGroupQueryBase::Evaluate(const TArray<AActor*>& InActorsToQuery, const UWorld* InWorld, TSet<AActor*>& OutMatchingActors) const
 {
 	// No implementation
 }
@@ -474,7 +482,7 @@ bool UMovieGraphConditionGroupQueryBase::IsFirstConditionGroupQuery() const
 	return false;
 }
 
-void UMovieGraphConditionGroupQuery_Actor::Evaluate(const TArray<AActor*>& InActorsToQuery, TSet<AActor*>& OutMatchingActors) const
+void UMovieGraphConditionGroupQuery_Actor::Evaluate(const TArray<AActor*>& InActorsToQuery, const UWorld* InWorld, TSet<AActor*>& OutMatchingActors) const
 {
 	TRACE_CPUPROFILER_EVENT_SCOPE(UMovieGraphConditionGroupQuery_Actor::Evaluate);
 
@@ -669,7 +677,7 @@ FText UMovieGraphConditionGroupQuery_Actor::GetRowText(TSharedPtr<TSoftObjectPtr
 }
 #endif	// WITH_EDITOR
 
-void UMovieGraphConditionGroupQuery_ActorTagName::Evaluate(const TArray<AActor*>& InActorsToQuery, TSet<AActor*>& OutMatchingActors) const
+void UMovieGraphConditionGroupQuery_ActorTagName::Evaluate(const TArray<AActor*>& InActorsToQuery, const UWorld* InWorld, TSet<AActor*>& OutMatchingActors) const
 {
 	TRACE_CPUPROFILER_EVENT_SCOPE(UMovieGraphConditionGroupQuery_ActorTag::Evaluate);
 	
@@ -743,7 +751,7 @@ TArray<TSharedRef<SWidget>> UMovieGraphConditionGroupQuery_ActorTagName::GetWidg
 }
 #endif
 
-void UMovieGraphConditionGroupQuery_ActorName::Evaluate(const TArray<AActor*>& InActorsToQuery, TSet<AActor*>& OutMatchingActors) const
+void UMovieGraphConditionGroupQuery_ActorName::Evaluate(const TArray<AActor*>& InActorsToQuery, const UWorld* InWorld, TSet<AActor*>& OutMatchingActors) const
 {
 	TRACE_CPUPROFILER_EVENT_SCOPE(UMovieGraphConditionGroupQuery_ActorName::Evaluate);
 	
@@ -811,7 +819,7 @@ bool UMovieGraphConditionGroupQuery_ActorName::IsEditorOnly() const
 	return true;
 }
 
-void UMovieGraphConditionGroupQuery_ActorType::Evaluate(const TArray<AActor*>& InActorsToQuery, TSet<AActor*>& OutMatchingActors) const
+void UMovieGraphConditionGroupQuery_ActorType::Evaluate(const TArray<AActor*>& InActorsToQuery, const UWorld* InWorld, TSet<AActor*>& OutMatchingActors) const
 {
 	TRACE_CPUPROFILER_EVENT_SCOPE(UMovieGraphConditionGroupQuery_ActorType::Evaluate);
 	
@@ -906,11 +914,16 @@ const FSlateBrush* UMovieGraphConditionGroupQuery_ActorType::GetRowIcon(UClass*
 
 FText UMovieGraphConditionGroupQuery_ActorType::GetRowText(UClass* InActorType)
 {
-	return InActorType->GetDisplayNameText();
+	if (InActorType)
+	{
+		return InActorType->GetDisplayNameText();
+	}
+
+	return LOCTEXT("MovieGraphActorTypeConditionGroupQuery_Invalid", "(invalid)");
 }
 #endif
 
-void UMovieGraphConditionGroupQuery_ComponentTagName::Evaluate(const TArray<AActor*>& InActorsToQuery, TSet<AActor*>& OutMatchingActors) const
+void UMovieGraphConditionGroupQuery_ComponentTagName::Evaluate(const TArray<AActor*>& InActorsToQuery, const UWorld* InWorld, TSet<AActor*>& OutMatchingActors) const
 {
 	TRACE_CPUPROFILER_EVENT_SCOPE(UMovieGraphConditionGroupQuery_ComponentTag::Evaluate);
 	
@@ -1001,7 +1014,7 @@ TArray<TSharedRef<SWidget>> UMovieGraphConditionGroupQuery_ComponentTagName::Get
 }
 #endif
 
-void UMovieGraphConditionGroupQuery_ComponentType::Evaluate(const TArray<AActor*>& InActorsToQuery, TSet<AActor*>& OutMatchingActors) const
+void UMovieGraphConditionGroupQuery_ComponentType::Evaluate(const TArray<AActor*>& InActorsToQuery, const UWorld* InWorld, TSet<AActor*>& OutMatchingActors) const
 {
 	TRACE_CPUPROFILER_EVENT_SCOPE(UMovieGraphConditionGroupQuery_ComponentType::Evaluate);
 	
@@ -1108,7 +1121,292 @@ const FSlateBrush* UMovieGraphConditionGroupQuery_ComponentType::GetRowIcon(UCla
 
 FText UMovieGraphConditionGroupQuery_ComponentType::GetRowText(UClass* InComponentType)
 {
-	return InComponentType->GetDisplayNameText();
+	if (InComponentType)
+	{
+		return InComponentType->GetDisplayNameText();
+	}
+	
+	return LOCTEXT("MovieGraphComponentTypeConditionGroupQuery_Invalid", "(invalid)");
+}
+#endif	// WITH_EDITOR
+
+void UMovieGraphConditionGroupQuery_EditorFolder::Evaluate(const TArray<AActor*>& InActorsToQuery, const UWorld* InWorld, TSet<AActor*>& OutMatchingActors) const
+{
+#if WITH_EDITOR
+	// This const cast is unfortunate, but should be harmless
+	TArray<AActor*> ActorsInFolders;
+	FActorFolders::GetActorsFromFolders(*const_cast<UWorld*>(InWorld), FolderPaths, ActorsInFolders);
+	
+	OutMatchingActors.Append(ActorsInFolders);
+#endif	// WITH_EDITOR
+}
+
+const FSlateIcon& UMovieGraphConditionGroupQuery_EditorFolder::GetIcon() const
+{
+	static const FSlateIcon EditorFolderIcon = FSlateIcon(FAppStyle::GetAppStyleSetName(), "Icons.FolderOpen");
+	return EditorFolderIcon;
+}
+
+const FText& UMovieGraphConditionGroupQuery_EditorFolder::GetDisplayName() const
+{
+	static const FText DisplayName = LOCTEXT("ConditionGroupQueryDisplayName_EditorFolder", "Editor Folder");
+	return DisplayName;
+}
+
+bool UMovieGraphConditionGroupQuery_EditorFolder::IsEditorOnlyQuery() const
+{
+	// This query is editor-only because folders do not exist outside of the editor
+	return true;
+}
+
+#if WITH_EDITOR
+TArray<TSharedRef<SWidget>> UMovieGraphConditionGroupQuery_EditorFolder::GetWidgets()
+{
+	TArray<TSharedRef<SWidget>> Widgets;
+    
+    Widgets.Add(
+    	SAssignNew(FolderPathsList, SMovieGraphSimpleList<FName>)
+    		.DataSource(&FolderPaths)
+    		.DataType(FText::FromString("Folder"))
+    		.DataTypePlural(FText::FromString("Folders"))
+    		.OnGetRowText_Static(&GetRowText)
+    		.OnGetRowIcon_Static(&GetRowIcon)
+    		.OnDelete_Lambda([this](FName InFolderPath)
+    		{
+    			FolderPaths.Remove(InFolderPath);
+    			FolderPathsList->Refresh();
+    			FolderPickerWidget->FullRefresh();
+    		})			
+    );
+
+    return Widgets;
+}
+
+TSharedRef<SWidget> UMovieGraphConditionGroupQuery_EditorFolder::GetAddMenuContents(const FMovieGraphConditionGroupQueryContentsChanged& OnAddFinished)
+{
+	auto OnItemPicked = FOnSceneOutlinerItemPicked::CreateLambda([this, OnAddFinished](TSharedRef<ISceneOutlinerTreeItem> Item)
+	{
+		if (const FActorFolderTreeItem* FolderItem = Item->CastTo<FActorFolderTreeItem>())
+		{
+			if (FolderItem->IsValid())
+			{
+				const FName& FolderPath = FolderItem->GetPath();
+					
+				// Don't allow duplicate folder paths
+				if (FolderPaths.Contains(FolderPath))
+				{
+					return;
+				}
+				
+				FolderPaths.AddUnique(FolderPath);
+				OnAddFinished.ExecuteIfBound();
+				
+				if (FolderPickerWidget.IsValid())
+				{
+					FolderPickerWidget->FullRefresh();
+					FolderPathsList->Refresh();
+				}
+			}
+		}
+	});
+
+	const FCreateSceneOutlinerMode ModeFactory = FCreateSceneOutlinerMode::CreateLambda([&OnItemPicked](SSceneOutliner* Outliner)
+	{
+		return new FActorFolderPickingMode(Outliner, OnItemPicked);
+	});
+	
+	FSceneOutlinerInitializationOptions SceneOutlinerInitOptions;
+	SceneOutlinerInitOptions.bShowCreateNewFolder = false;
+	SceneOutlinerInitOptions.bFocusSearchBoxWhenOpened = true;
+	SceneOutlinerInitOptions.ModeFactory = ModeFactory;
+
+	// Don't show folders which have already been picked
+	SceneOutlinerInitOptions.Filters->AddFilterPredicate<FActorFolderTreeItem>(
+		FActorFolderTreeItem::FFilterPredicate::CreateLambda([this](const FFolder& InFolder)
+		{
+			return !FolderPaths.Contains(InFolder.GetPath());
+		}));
+
+	// Only show the name/label column, that's the only column relevant to folders
+	SceneOutlinerInitOptions.ColumnMap.Add(
+		FSceneOutlinerBuiltInColumnTypes::Label(),
+		FSceneOutlinerColumnInfo(ESceneOutlinerColumnVisibility::Visible, 0, FCreateSceneOutlinerColumn(), false, TOptional<float>(), FSceneOutlinerBuiltInColumnTypes::Label_Localized()));
+
+	FolderPickerWidget = SNew(SSceneOutliner, SceneOutlinerInitOptions)
+		.IsEnabled(FSlateApplication::Get().GetNormalExecutionAttribute());
+	
+	return
+		SNew(SBox)
+		.WidthOverride(400.f)
+		.HeightOverride(300.f)
+		[
+			FolderPickerWidget.ToSharedRef()
+		];
+}
+
+const FSlateBrush* UMovieGraphConditionGroupQuery_EditorFolder::GetRowIcon(FName InFolderPath)
+{
+	return FAppStyle::Get().GetBrush("Icons.FolderOpen");
+}
+
+FText UMovieGraphConditionGroupQuery_EditorFolder::GetRowText(FName InFolderPath)
+{
+	return FText::FromString(InFolderPath.ToString());
+}
+#endif	// WITH_EDITOR
+
+void UMovieGraphConditionGroupQuery_Sublevel::Evaluate(const TArray<AActor*>& InActorsToQuery, const UWorld* InWorld, TSet<AActor*>& OutMatchingActors) const
+{
+	for (const TSoftObjectPtr<UWorld>& World : Sublevels)
+	{
+		// Don't load the level, only use levels which are already loaded
+		const UWorld* LoadedWorld = World.Get();
+		if (!LoadedWorld)
+		{
+			const UMovieGraphCollection* ParentCollection = GetTypedOuter<UMovieGraphCollection>();
+			const FString CollectionName = ParentCollection ? ParentCollection->GetCollectionName() : TEXT("<unknown>");
+			
+			UE_LOG(LogMovieRenderPipeline, Warning, TEXT("Sublevel query in collection '%s' is excluding level (%s) because it is not loaded."), *CollectionName, *World.ToString())
+			continue;
+		}
+
+		ULevel* CurrentLevel = LoadedWorld->GetCurrentLevel();
+		if (!CurrentLevel)
+		{
+			continue;
+		}
+
+		for (TObjectPtr<AActor>& LevelActor : CurrentLevel->Actors)
+		{
+			OutMatchingActors.Add(LevelActor.Get());
+		}
+	}
+}
+
+const FSlateIcon& UMovieGraphConditionGroupQuery_Sublevel::GetIcon() const
+{
+	static const FSlateIcon SublevelIcon = FSlateIcon(FAppStyle::GetAppStyleSetName(), "Icons.Level");
+	return SublevelIcon;
+}
+
+const FText& UMovieGraphConditionGroupQuery_Sublevel::GetDisplayName() const
+{
+	static const FText DisplayName = LOCTEXT("ConditionGroupQueryDisplayName_Sublevel", "Sublevel");
+	return DisplayName;
+}
+
+#if WITH_EDITOR
+TArray<TSharedRef<SWidget>> UMovieGraphConditionGroupQuery_Sublevel::GetWidgets()
+{
+	TArray<TSharedRef<SWidget>> Widgets;
+
+	// Create the data source for the list view
+	ListDataSource.Empty();
+	for (TSoftObjectPtr<UWorld>& Sublevel : Sublevels)
+	{
+		ListDataSource.Add(MakeShared<TSoftObjectPtr<UWorld>>(Sublevel));
+	}
+
+	Widgets.Add(
+		SAssignNew(SublevelsList, SMovieGraphSimpleList<TSharedPtr<TSoftObjectPtr<UWorld>>>)
+			.DataSource(&ListDataSource)
+			.DataType(FText::FromString("Sublevel"))
+			.DataTypePlural(FText::FromString("Sublevels"))
+			.OnGetRowText_Static(&GetRowText)
+			.OnGetRowIcon_Static(&GetRowIcon)
+			.OnDelete_Lambda([this](const TSharedPtr<TSoftObjectPtr<UWorld>> InSublevel)
+			{
+				ListDataSource.Remove(InSublevel);
+				Sublevels.Remove(*InSublevel.Get());
+				
+				SublevelsList->Refresh();
+				
+				constexpr bool bUpdateSources = true;
+				RefreshLevelPicker.ExecuteIfBound(bUpdateSources);
+			})
+	);
+
+	return Widgets;
+}
+
+TSharedRef<SWidget> UMovieGraphConditionGroupQuery_Sublevel::GetAddMenuContents(const FMovieGraphConditionGroupQueryContentsChanged& OnAddFinished)
+{
+	FAssetPickerConfig SublevelPickerConfig;
+	{
+		SublevelPickerConfig.SelectionMode = ESelectionMode::Single;
+		SublevelPickerConfig.SaveSettingsName = TEXT("MovieRenderGraphSublevelPicker");
+		SublevelPickerConfig.RefreshAssetViewDelegates.Add(&RefreshLevelPicker);
+		SublevelPickerConfig.InitialAssetViewType = EAssetViewType::Column;
+		SublevelPickerConfig.bFocusSearchBoxWhenOpened = true;
+		SublevelPickerConfig.bAllowNullSelection = false;
+		SublevelPickerConfig.bShowBottomToolbar = true;
+		SublevelPickerConfig.bAutohideSearchBar = false;
+		SublevelPickerConfig.bAllowDragging = false;
+		SublevelPickerConfig.bCanShowClasses = false;
+		SublevelPickerConfig.bShowPathInColumnView = true;
+		SublevelPickerConfig.bShowTypeInColumnView = false;
+		SublevelPickerConfig.bSortByPathInColumnView = false;
+		SublevelPickerConfig.HiddenColumnNames = {
+			ContentBrowserItemAttributes::ItemDiskSize.ToString(),
+			ContentBrowserItemAttributes::VirtualizedData.ToString(),
+			TEXT("PrimaryAssetType"),
+			TEXT("PrimaryAssetName")
+		};
+		SublevelPickerConfig.AssetShowWarningText = LOCTEXT("ConditionGroupQuery_NoSublevelsFound", "No Sublevels Found");
+		SublevelPickerConfig.Filter.ClassPaths.Add(UWorld::StaticClass()->GetClassPathName());
+		SublevelPickerConfig.OnAssetSelected = FOnAssetSelected::CreateLambda([this, OnAddFinished](const FAssetData& InLevelAsset)
+		{
+			FSlateApplication::Get().DismissAllMenus();
+			
+			Sublevels.AddUnique(InLevelAsset.GetAsset());
+			ListDataSource.AddUnique(MakeShared<TSoftObjectPtr<UWorld>>(InLevelAsset.GetAsset()));
+			OnAddFinished.ExecuteIfBound();
+
+			if (SublevelsList.IsValid())
+			{
+				SublevelsList->Refresh();
+				
+				constexpr bool bUpdateSources = false;
+				RefreshLevelPicker.ExecuteIfBound(bUpdateSources);
+			}
+		});
+		SublevelPickerConfig.OnShouldFilterAsset = FOnShouldFilterAsset::CreateLambda([this](const FAssetData& InLevelAsset)
+		{
+			// Don't show sublevels which have already been picked
+			UWorld* Sublevel = Cast<UWorld>(InLevelAsset.GetAsset());
+			return !Sublevel || Sublevels.Contains(Sublevel);
+		});
+	}
+
+	IContentBrowserSingleton& ContentBrowser = FModuleManager::LoadModuleChecked<FContentBrowserModule>("ContentBrowser").Get();
+
+	return
+		SNew(SBox)
+		.Padding(0, 10.f, 0, 0)
+		.WidthOverride(400.f)
+		.HeightOverride(300.f)
+		[
+			ContentBrowser.CreateAssetPicker(SublevelPickerConfig)
+		];
+}
+
+const FSlateBrush* UMovieGraphConditionGroupQuery_Sublevel::GetRowIcon(TSharedPtr<TSoftObjectPtr<UWorld>> InSublevel)
+{
+	return FAppStyle::Get().GetBrush("Icons.Level");
+}
+
+FText UMovieGraphConditionGroupQuery_Sublevel::GetRowText(TSharedPtr<TSoftObjectPtr<UWorld>> InSublevel)
+{
+	if (InSublevel.IsValid())
+	{
+		if (InSublevel.Get()->IsValid())
+		{
+			// The first Get() returns the TSoftObjectPtr, the second Get() dereferences the TSoftObjectPtr
+			return FText::FromString(InSublevel.Get()->Get()->GetName());
+		}
+	}
+
+	return LOCTEXT("MovieGraphSublevelConditionGroupQuery_InvalidLevel", "(invalid)");
 }
 #endif	// WITH_EDITOR
 
@@ -1185,7 +1483,7 @@ TSet<AActor*> UMovieGraphConditionGroup::Evaluate(const UWorld* InWorld) const
 		// Similar to EvaluationResult, QueryResult is persisted+reset to prevent constantly re-allocating it
 		QueryResult.Reset();
 
-		Query->Evaluate(AllActors, QueryResult);
+		Query->Evaluate(AllActors, InWorld, QueryResult);
 		
 		switch (Query->GetOperationType())
 		{
diff --git a/Engine/Plugins/MovieScene/MovieRenderPipeline/Source/MovieRenderPipelineCore/Public/Graph/MoviePipelineRenderLayerSubsystem.h b/Engine/Plugins/MovieScene/MovieRenderPipeline/Source/MovieRenderPipelineCore/Public/Graph/MoviePipelineRenderLayerSubsystem.h
--- a/Engine/Plugins/MovieScene/MovieRenderPipeline/Source/MovieRenderPipelineCore/Public/Graph/MoviePipelineRenderLayerSubsystem.h
+++ b/Engine/Plugins/MovieScene/MovieRenderPipeline/Source/MovieRenderPipelineCore/Public/Graph/MoviePipelineRenderLayerSubsystem.h
@@ -7,6 +7,10 @@
 #include "Subsystems/WorldSubsystem.h"
 #include "Textures/SlateIcon.h"
 
+#if WITH_EDITOR
+#include "ContentBrowserDelegates.h"
+#endif	// WITH_EDITOR
+
 #include "MoviePipelineRenderLayerSubsystem.generated.h"
 
 class SWidget;
@@ -95,9 +99,9 @@ public:
 	UFUNCTION(BlueprintCallable, Category = "Settings")
 	EMovieGraphConditionGroupQueryOpType GetOperationType() const;
 
-	/** Determines which of the provided actors match the query. Matches are added to OutMatchingActors.  */
+	/** Determines which of the provided actors (in the given world) match the query. Matches are added to OutMatchingActors.  */
 	UFUNCTION(BlueprintCallable, Category = "Settings")
-	virtual void Evaluate(const TArray<AActor*>& InActorsToQuery, TSet<AActor*>& OutMatchingActors) const;
+	virtual void Evaluate(const TArray<AActor*>& InActorsToQuery, const UWorld* InWorld, TSet<AActor*>& OutMatchingActors) const;
 
 	/**
 	 * Determines if the public properties on the query class will have their names hidden in the details panel. Returns
@@ -161,7 +165,7 @@ class MOVIERENDERPIPELINECORE_API UMovieGraphConditionGroupQuery_Actor final : p
 	GENERATED_BODY()
 
 public:
-	virtual void Evaluate(const TArray<AActor*>& InActorsToQuery, TSet<AActor*>& OutMatchingActors) const override;
+	virtual void Evaluate(const TArray<AActor*>& InActorsToQuery, const UWorld* InWorld, TSet<AActor*>& OutMatchingActors) const override;
 	virtual const FSlateIcon& GetIcon() const override;
 	virtual const FText& GetDisplayName() const override;
 
@@ -201,7 +205,7 @@ class MOVIERENDERPIPELINECORE_API UMovieGraphConditionGroupQuery_ActorTagName fi
 	GENERATED_BODY()
 
 public:
-	virtual void Evaluate(const TArray<AActor*>& InActorsToQuery, TSet<AActor*>& OutMatchingActors) const override;
+	virtual void Evaluate(const TArray<AActor*>& InActorsToQuery, const UWorld* InWorld, TSet<AActor*>& OutMatchingActors) const override;
 	virtual const FSlateIcon& GetIcon() const override;
 	virtual const FText& GetDisplayName() const override;
 
@@ -234,7 +238,7 @@ class MOVIERENDERPIPELINECORE_API UMovieGraphConditionGroupQuery_ActorName final
 	GENERATED_BODY()
 
 public:
-	virtual void Evaluate(const TArray<AActor*>& InActorsToQuery, TSet<AActor*>& OutMatchingActors) const override;
+	virtual void Evaluate(const TArray<AActor*>& InActorsToQuery, const UWorld* InWorld, TSet<AActor*>& OutMatchingActors) const override;
 	virtual const FSlateIcon& GetIcon() const override;
 	virtual const FText& GetDisplayName() const override;
 
@@ -269,7 +273,7 @@ class MOVIERENDERPIPELINECORE_API UMovieGraphConditionGroupQuery_ActorType final
 	GENERATED_BODY()
 
 public:
-	virtual void Evaluate(const TArray<AActor*>& InActorsToQuery, TSet<AActor*>& OutMatchingActors) const override;
+	virtual void Evaluate(const TArray<AActor*>& InActorsToQuery, const UWorld* InWorld, TSet<AActor*>& OutMatchingActors) const override;
 	virtual const FSlateIcon& GetIcon() const override;
 	virtual const FText& GetDisplayName() const override;
 
@@ -280,6 +284,7 @@ public:
 
 public:
 	/** The type (class) that the actor needs to have in order to be a match. */
+	UPROPERTY(EditAnywhere, Category="General")
 	TArray<UClass*> ActorTypes;
 
 private:
@@ -302,7 +307,7 @@ class MOVIERENDERPIPELINECORE_API UMovieGraphConditionGroupQuery_ComponentTagNam
 	GENERATED_BODY()
 
 public:
-	virtual void Evaluate(const TArray<AActor*>& InActorsToQuery, TSet<AActor*>& OutMatchingActors) const override;
+	virtual void Evaluate(const TArray<AActor*>& InActorsToQuery, const UWorld* InWorld, TSet<AActor*>& OutMatchingActors) const override;
 	virtual const FSlateIcon& GetIcon() const override;
 	virtual const FText& GetDisplayName() const override;
 
@@ -335,7 +340,7 @@ class MOVIERENDERPIPELINECORE_API UMovieGraphConditionGroupQuery_ComponentType f
 	GENERATED_BODY()
 
 public:
-	virtual void Evaluate(const TArray<AActor*>& InActorsToQuery, TSet<AActor*>& OutMatchingActors) const override;
+	virtual void Evaluate(const TArray<AActor*>& InActorsToQuery, const UWorld* InWorld, TSet<AActor*>& OutMatchingActors) const override;
 	virtual const FSlateIcon& GetIcon() const override;
 	virtual const FText& GetDisplayName() const override;
 
@@ -362,6 +367,79 @@ public:
 	TArray<UClass*> ComponentTypes;
 };
 
+/** Query type which filters actors via the editor folder that they're contained in. */
+UCLASS(BlueprintType)
+class MOVIERENDERPIPELINECORE_API UMovieGraphConditionGroupQuery_EditorFolder final : public UMovieGraphConditionGroupQueryBase
+{
+	GENERATED_BODY()
+
+public:
+	virtual void Evaluate(const TArray<AActor*>& InActorsToQuery, const UWorld* InWorld, TSet<AActor*>& OutMatchingActors) const override;
+	virtual const FSlateIcon& GetIcon() const override;
+	virtual const FText& GetDisplayName() const override;
+	virtual bool IsEditorOnlyQuery() const override;
+
+#if WITH_EDITOR
+	virtual TArray<TSharedRef<SWidget>> GetWidgets() override;
+	virtual TSharedRef<SWidget> GetAddMenuContents(const FMovieGraphConditionGroupQueryContentsChanged& OnAddFinished) override;
+#endif
+
+private:
+#if WITH_EDITOR
+	static const FSlateBrush* GetRowIcon(FName InFolderPath);
+	static FText GetRowText(FName InFolderPath);
+
+	/** Displays the paths of folders which have been chosen. */
+	TSharedPtr<SMovieGraphSimpleList<FName>> FolderPathsList;
+
+	/** The folder browser widget to show in the Add menu. */
+	TSharedPtr<class ISceneOutliner> FolderPickerWidget;
+#endif
+
+public:
+	/** The actor must be in one of the chosen folders in order to be a match. */
+	UPROPERTY(EditAnywhere, Category="General")
+	TArray<FName> FolderPaths;
+};
+
+/** Query type which filters actors via the sublevel that they're contained in. */
+UCLASS(BlueprintType)
+class MOVIERENDERPIPELINECORE_API UMovieGraphConditionGroupQuery_Sublevel final : public UMovieGraphConditionGroupQueryBase
+{
+	GENERATED_BODY()
+
+public:
+	virtual void Evaluate(const TArray<AActor*>& InActorsToQuery, const UWorld* InWorld, TSet<AActor*>& OutMatchingActors) const override;
+	virtual const FSlateIcon& GetIcon() const override;
+	virtual const FText& GetDisplayName() const override;
+
+#if WITH_EDITOR
+	virtual TArray<TSharedRef<SWidget>> GetWidgets() override;
+	virtual TSharedRef<SWidget> GetAddMenuContents(const FMovieGraphConditionGroupQueryContentsChanged& OnAddFinished) override;
+#endif
+
+private:
+#if WITH_EDITOR
+	static const FSlateBrush* GetRowIcon(TSharedPtr<TSoftObjectPtr<UWorld>> InSublevel);
+	static FText GetRowText(TSharedPtr<TSoftObjectPtr<UWorld>> InSublevel);
+
+	/** Displays the names of sublevels which have been chosen. */
+	TSharedPtr<SMovieGraphSimpleList<TSharedPtr<TSoftObjectPtr<UWorld>>>> SublevelsList;
+
+	/** Refreshes the contents of the level picker widget when called. */
+	FRefreshAssetViewDelegate RefreshLevelPicker;
+
+	// Not ideal to store a duplicate of Sublevels, but SListView requires TSharedPtr<...> as the data source, and UPROPERTY does not
+	// support TSharedPtr<...>
+	TArray<TSharedPtr<TSoftObjectPtr<UWorld>>> ListDataSource;
+#endif
+
+public:
+	/** The actor must be in one of the chosen sublevels in order to be a match. */
+	UPROPERTY(EditAnywhere, Category="General")
+	TArray<TSoftObjectPtr<UWorld>> Sublevels;
+};
+
 /** A group of queries which can be added to a collection. */
 UCLASS(BlueprintType)
 class MOVIERENDERPIPELINECORE_API UMovieGraphConditionGroup : public UObject
diff --git a/Engine/Plugins/MovieScene/MovieRenderPipeline/Source/MovieRenderPipelineEditor/Private/Customizations/Graph/MovieGraphModifiersCustomization.cpp b/Engine/Plugins/MovieScene/MovieRenderPipeline/Source/MovieRenderPipelineEditor/Private/Customizations/Graph/MovieGraphModifiersCustomization.cpp
--- a/Engine/Plugins/MovieScene/MovieRenderPipeline/Source/MovieRenderPipelineEditor/Private/Customizations/Graph/MovieGraphModifiersCustomization.cpp
+++ b/Engine/Plugins/MovieScene/MovieRenderPipeline/Source/MovieRenderPipelineEditor/Private/Customizations/Graph/MovieGraphModifiersCustomization.cpp
@@ -60,7 +60,7 @@ public:
 			.VAlign(VAlign_Fill)
 			[
 				SNew(SWidgetSwitcher)
-				.WidgetIndex(DataSource.IsEmpty() ? 0 : 1)
+				.WidgetIndex_Lambda([this] { return DataSource.IsEmpty() ? 0 : 1; })
 
 				+ SWidgetSwitcher::Slot()
 				.VAlign(VAlign_Center)
