MovieGraph: Fix WITH_OCIO endif statements.

#rnx

diff --git a/Engine/Plugins/MovieScene/MovieRenderPipeline/Source/MovieRenderPipelineRenderPasses/Private/MovieGraphImageSequenceOutputNode.cpp b/Engine/Plugins/MovieScene/MovieRenderPipeline/Source/MovieRenderPipelineRenderPasses/Private/MovieGraphImageSequenceOutputNode.cpp
--- a/Engine/Plugins/MovieScene/MovieRenderPipeline/Source/MovieRenderPipelineRenderPasses/Private/MovieGraphImageSequenceOutputNode.cpp
+++ b/Engine/Plugins/MovieScene/MovieRenderPipeline/Source/MovieRenderPipelineRenderPasses/Private/MovieGraphImageSequenceOutputNode.cpp
@@ -23,7 +23,7 @@
 #include "OpenColorIOConfiguration.h"
 #include "OpenColorIOColorTransform.h"
 #include "OpenColorIOWrapper.h"
-#endif
+#endif // WITH_OCIO
 
 namespace UE::MovieGraph::Private
 {
@@ -173,7 +173,7 @@ namespace UE::MovieGraph::Private
 
 		return {};
 	}
-#endif
+#endif // WITH_OCIO
 } //end namespace UE::MovieGraph::Private
 
 UMovieGraphImageSequenceOutputNode::UMovieGraphImageSequenceOutputNode()
@@ -373,7 +373,7 @@ void UMovieGraphImageSequenceOutputNode::OnReceiveImageDataImpl(UMovieGraphPipel
 				bQuantizationEncodeSRGB = false;
 			}
 		}
-#endif
+#endif // WITH_OCIO
 
 		EImagePixelType PixelType = TileImageTask->PixelData->GetType();
 
@@ -478,7 +478,7 @@ void UMovieGraphImageSequenceOutputNode_EXR::UpdateTaskPerLayer(
 			bEnabledOCIO = true;
 		}
 	}
-#endif
+#endif // WITH_OCIO
 
 	if (InLayerIndex == 0)
 	{
@@ -498,7 +498,7 @@ void UMovieGraphImageSequenceOutputNode_EXR::UpdateTaskPerLayer(
 			UE::MoviePipeline::UpdateColorSpaceMetadata(InParentNode->OCIOConfiguration.ColorConfiguration, InOutImageTask);
 		}
 		else
-#endif
+#endif // WITH_OCIO
 		{
 			UE::MoviePipeline::UpdateColorSpaceMetadata(Payload->SceneCaptureSource, InOutImageTask);
 		}
@@ -577,7 +577,7 @@ void UMovieGraphImageSequenceOutputNode_EXR::OnReceiveImageDataImpl(UMovieGraphP
 			InRawFrameData->EvaluatedConfig.Get(),
 			Payload->TraversalContext
 		);
-#endif // WITH_EDITOR
+#endif // WITH_OCIO
 
 		UpdateTaskPerLayer(*ImageWriteTask, ParentNode, MoveTemp(PixelData), LayerIndex, FString(), ResolvedOCIOContext);
 
@@ -678,7 +678,7 @@ void UMovieGraphImageSequenceOutputNode_MultiLayerEXR::OnReceiveImageDataImpl(UM
 				InRawFrameData->EvaluatedConfig.Get(),
 				Payload->TraversalContext
 			);
-#endif // WITH_EDITOR
+#endif // WITH_OCIO
 			UpdateTaskPerLayer(*MultiLayerImageTask, ParentNode, MoveTemp(PixelData), LayerIndex, LayerName, ResolvedOCIOContext);
 
 			LayerIndex++;
