MovieGraph: Fixed Visibility Modifiers not propagating to all components. Added a cvar to help during debugging ("MoviePipeline.FlushLayersHack") which can be used to flush the pipeline between each layer. This is a significant impact to performance (and thus not reccomended to be used) but can help work around some layer rendering issues until their root causes are fixed.

diff --git a/Engine/Plugins/MovieScene/MovieRenderPipeline/Source/MovieRenderPipelineCore/Private/Graph/MoviePipelineRenderLayerSubsystem.cpp b/Engine/Plugins/MovieScene/MovieRenderPipeline/Source/MovieRenderPipelineCore/Private/Graph/MoviePipelineRenderLayerSubsystem.cpp
--- a/Engine/Plugins/MovieScene/MovieRenderPipeline/Source/MovieRenderPipelineCore/Private/Graph/MoviePipelineRenderLayerSubsystem.cpp
+++ b/Engine/Plugins/MovieScene/MovieRenderPipeline/Source/MovieRenderPipelineCore/Private/Graph/MoviePipelineRenderLayerSubsystem.cpp
@@ -162,10 +162,13 @@ void UMoviePipelineVisibilityModifier::ApplyModifier(const UWorld* World)
 	// a differing actor), so create one copy to prevent constantly re-creating structs.
 	FActorVisibilityState NewVisibilityState;
 	NewVisibilityState.bIsHidden = bIsHidden;
-	NewVisibilityState.bCastsShadows = bCastsShadows;
-	NewVisibilityState.bCastShadowWhileHidden = bCastShadowWhileHidden;
-	NewVisibilityState.bAffectIndirectLightingWhileHidden = bAffectIndirectLightingWhileHidden;
-	NewVisibilityState.bHoldout = bHoldout;
+
+	// Create a stub component representation since we're just using this to propagate to all children right now.
+	FActorVisibilityState::FComponentState& NewComponentState = NewVisibilityState.Components.AddDefaulted_GetRef();
+	NewComponentState.bCastsShadows = bCastsShadows;
+	NewComponentState.bCastShadowWhileHidden = bCastShadowWhileHidden;
+	NewComponentState.bAffectIndirectLightingWhileHidden = bAffectIndirectLightingWhileHidden;
+	NewComponentState.bHoldout = bHoldout;
 	
 	for (const UMovieGraphCollection* Collection : Collections)
 	{
@@ -183,19 +186,31 @@ void UMoviePipelineVisibilityModifier::ApplyModifier(const UWorld* World)
 			FActorVisibilityState OriginalVisibilityState;
 			OriginalVisibilityState.Actor = Actor;
 			OriginalVisibilityState.bIsHidden = Actor->IsHidden();
-			if (const UPrimitiveComponent* PrimitiveComponent = Actor->GetComponentByClass<UPrimitiveComponent>())
+
+			const bool bIncludeFromChildActors = true;
+			TInlineComponentArray<UPrimitiveComponent*> Components;
+			Actor->GetComponents<UPrimitiveComponent>(Components, bIncludeFromChildActors);
+
+			OriginalVisibilityState.Components.Reserve(Components.Num());
+			for(const UPrimitiveComponent* PrimitiveComponent : Components)
 			{
-				OriginalVisibilityState.bCastsShadows = PrimitiveComponent->CastShadow;
-				OriginalVisibilityState.bCastShadowWhileHidden = PrimitiveComponent->bCastHiddenShadow;
-				OriginalVisibilityState.bAffectIndirectLightingWhileHidden = PrimitiveComponent->bAffectIndirectLightingWhileHidden;
-				OriginalVisibilityState.bHoldout = PrimitiveComponent->bHoldout;
+				// Cache the state
+				FActorVisibilityState::FComponentState& ComponentState = OriginalVisibilityState.Components.AddDefaulted_GetRef();
+				ComponentState.Component = PrimitiveComponent;
+				ComponentState.bCastsShadows = PrimitiveComponent->CastShadow;
+				ComponentState.bCastShadowWhileHidden = PrimitiveComponent->bCastHiddenShadow;
+				ComponentState.bAffectIndirectLightingWhileHidden = PrimitiveComponent->bAffectIndirectLightingWhileHidden;
+				ComponentState.bHoldout = PrimitiveComponent->bHoldout;
+
+				// Then override it. We override it one component at a time so that we can avoid making a bunch of
+				// copies of NewVisibilityState for the varying number of components you might have.
+				NewVisibilityState.Actor = Actor;
+				NewVisibilityState.Components[0].Component = PrimitiveComponent;
+				SetActorVisibilityState(NewVisibilityState);
 			}
 			
 			ModifiedActors.Add(OriginalVisibilityState);
 
-			// Set new visibility state
-			NewVisibilityState.Actor = Actor;
-			SetActorVisibilityState(NewVisibilityState);
 		}
 	}
 }
@@ -224,14 +239,20 @@ void UMoviePipelineVisibilityModifier::SetActorVisibilityState(const FActorVisib
 	Actor->SetIsTemporarilyHiddenInEditor(NewVisibilityState.bIsHidden);
 #endif
 
-	if (UPrimitiveComponent* PrimitiveComponent = Actor->GetComponentByClass<UPrimitiveComponent>())
+	
+	for (const FActorVisibilityState::FComponentState& ComponentState : NewVisibilityState.Components)
 	{
 		// TODO: These could potentially cause a large rendering penalty due to dirtying the render state; investigate potential
 		// ways to optimize this
-		PrimitiveComponent->SetCastShadow(NewVisibilityState.bCastsShadows);
-		PrimitiveComponent->SetCastHiddenShadow(NewVisibilityState.bCastShadowWhileHidden);
-		PrimitiveComponent->SetAffectIndirectLightingWhileHidden(NewVisibilityState.bAffectIndirectLightingWhileHidden);
-		PrimitiveComponent->SetHoldout(NewVisibilityState.bHoldout);
+		if (!ComponentState.Component)
+		{
+			continue;
+		}
+
+		ComponentState.Component->SetCastShadow(ComponentState.bCastsShadows);
+		ComponentState.Component->SetCastHiddenShadow(ComponentState.bCastShadowWhileHidden);
+		ComponentState.Component->SetAffectIndirectLightingWhileHidden(ComponentState.bAffectIndirectLightingWhileHidden);
+		ComponentState.Component->SetHoldout(ComponentState.bHoldout);
 	}
 }
 
diff --git a/Engine/Plugins/MovieScene/MovieRenderPipeline/Source/MovieRenderPipelineCore/Private/Graph/Nodes/MovieGraphImagePassBaseNode.cpp b/Engine/Plugins/MovieScene/MovieRenderPipeline/Source/MovieRenderPipelineCore/Private/Graph/Nodes/MovieGraphImagePassBaseNode.cpp
--- a/Engine/Plugins/MovieScene/MovieRenderPipeline/Source/MovieRenderPipelineCore/Private/Graph/Nodes/MovieGraphImagePassBaseNode.cpp
+++ b/Engine/Plugins/MovieScene/MovieRenderPipeline/Source/MovieRenderPipelineCore/Private/Graph/Nodes/MovieGraphImagePassBaseNode.cpp
@@ -3,6 +3,17 @@
 #include "Graph/Nodes/MovieGraphImagePassBaseNode.h"
 #include "Graph/MoviePipelineRenderLayerSubsystem.h" 
 
+
+static TAutoConsoleVariable<bool> CVarMoviePipelineFlushRenderingPerLayerHack(
+	TEXT("MoviePipeline.FlushLayersHack"),
+	false,
+	TEXT("This is a debug cvar intended to help diagnose issues in the MoviePipeline multi-layer rendering mode\n")
+	TEXT("that flushes any outstanding rendering commands for each layer. This can have a large performance impact,\n")
+	TEXT("as it effectively removes any parallelism between the Game Thread and Render Thread!\n")
+	TEXT(" 0 - Don't flush (default)\n")
+	TEXT(" 1 - Flush (but negatively impacts performance)"),
+	ECVF_Default);
+
 UMovieGraphImagePassBaseNode::UMovieGraphImagePassBaseNode()
 {
 	ShowFlags = CreateDefaultSubobject<UMovieGraphShowFlags>(TEXT("ShowFlags"));
@@ -51,6 +62,15 @@ void UMovieGraphImagePassBaseNode::RenderImpl(const FMovieGraphTraversalContext&
 			LayerSubsystem->SetActiveRenderLayerByName(Instance->GetBranchName());
 		}
 		
+		// This is a hack to provide debug tooling for rendering issues caused by
+		// multiple layers. Flushing can ensure that all render proxies/etc. are
+		// recreated which can be used to see if it fixes various issues, but
+		// this comes at a big performance cost as it effectively removes GT/RT parallelism.
+		if (CVarMoviePipelineFlushRenderingPerLayerHack.GetValueOnGameThread())
+		{
+			FlushRenderingCommands();
+		}
+
 		Instance->Render(InFrameTraversalContext, InTimeData);
 
 		// Revert all modifiers
@@ -59,6 +79,11 @@ void UMovieGraphImagePassBaseNode::RenderImpl(const FMovieGraphTraversalContext&
 			TRACE_CPUPROFILER_EVENT_SCOPE(MRQ::DeferredRender::RevertRenderLayer);
 			LayerSubsystem->ClearActiveRenderLayer();
 		}
+
+		if (CVarMoviePipelineFlushRenderingPerLayerHack.GetValueOnGameThread())
+		{
+			FlushRenderingCommands();
+		}
 	}
 }
 
diff --git a/Engine/Plugins/MovieScene/MovieRenderPipeline/Source/MovieRenderPipelineCore/Public/Graph/MoviePipelineRenderLayerSubsystem.h b/Engine/Plugins/MovieScene/MovieRenderPipeline/Source/MovieRenderPipelineCore/Public/Graph/MoviePipelineRenderLayerSubsystem.h
--- a/Engine/Plugins/MovieScene/MovieRenderPipeline/Source/MovieRenderPipelineCore/Public/Graph/MoviePipelineRenderLayerSubsystem.h
+++ b/Engine/Plugins/MovieScene/MovieRenderPipeline/Source/MovieRenderPipelineCore/Public/Graph/MoviePipelineRenderLayerSubsystem.h
@@ -679,12 +679,20 @@ private:
 	/** Various visibility properties for an actor. */
 	struct FActorVisibilityState
 	{
+		struct FComponentState
+		{
+			TSoftObjectPtr<UPrimitiveComponent> Component = nullptr;
+			
+			uint8 bCastsShadows : 1 = true;
+			uint8 bCastShadowWhileHidden : 1 = false;
+			uint8 bAffectIndirectLightingWhileHidden : 1 = false;
+			uint8 bHoldout : 1 = false;
+		};
+
 		TSoftObjectPtr<AActor> Actor = nullptr;
+		TArray<FComponentState> Components;
+
 		uint8 bIsHidden : 1 = false;
-		uint8 bCastsShadows : 1 = true;
-		uint8 bCastShadowWhileHidden : 1 = false;
-		uint8 bAffectIndirectLightingWhileHidden : 1 = false;
-		uint8 bHoldout : 1 = false;
 	};
 
 	/** Updates an actor's visibility state to the state contained in NewVisibilityState. */
