MovieGraph: Temporary workaround for crash during export step when a shot-mask is used to turn off some renders.

diff --git a/Engine/Plugins/MovieScene/MovieRenderPipeline/Source/MovieRenderPipelineCore/Private/Graph/MovieGraphPipeline.cpp b/Engine/Plugins/MovieScene/MovieRenderPipeline/Source/MovieRenderPipelineCore/Private/Graph/MovieGraphPipeline.cpp
--- a/Engine/Plugins/MovieScene/MovieRenderPipeline/Source/MovieRenderPipelineCore/Private/Graph/MovieGraphPipeline.cpp
+++ b/Engine/Plugins/MovieScene/MovieRenderPipeline/Source/MovieRenderPipelineCore/Private/Graph/MovieGraphPipeline.cpp
@@ -517,7 +517,13 @@ void UMovieGraphPipeline::TickPostFinalizeExport(const bool bInForceFinish)
 {
 	// This step assumes you have produced data and filled the data structures.
 	check(PipelineState == EMovieRenderPipelineState::Export);
-	UE_LOG(LogMovieRenderPipeline, Verbose, TEXT("[%d] PostFinalize Export (Start)."), GFrameCounter);
+	
+
+	// Temporary change to prevent a crash when rendering only one shot via shot mask.
+	TransitionToState(EMovieRenderPipelineState::Finished);
+	return;
+
+	/*UE_LOG(LogMovieRenderPipeline, Verbose, TEXT("[%d] PostFinalize Export (Start)."), GFrameCounter);
 
 	// Loop through any extensions (such as XML export) and let them export using all of the
 	// data that was generated during this run such as containers, output names and lengths.
@@ -556,7 +562,7 @@ void UMovieGraphPipeline::TickPostFinalizeExport(const bool bInForceFinish)
 		return;
 	}
 
-	TransitionToState(EMovieRenderPipelineState::Finished);
+	TransitionToState(EMovieRenderPipelineState::Finished);*/
 }
 
 void UMovieGraphPipeline::BeginFinalize()
@@ -571,7 +577,10 @@ void UMovieGraphPipeline::BeginFinalize()
 
 void UMovieGraphPipeline::BeginExport()
 {
-	constexpr bool bIncludeCDOs = false;
+	// Temporary change to prevent a crash when rendering only one shot via shot mask.
+	return;
+
+	/*constexpr bool bIncludeCDOs = false;
 	constexpr bool bExactMatch = false;
 	TArray<TPair<FName, UMovieGraphPostRenderNode*>> PostRenderNodes = GetSettingForActiveRenderLayers<UMovieGraphPostRenderNode>(bIncludeCDOs, bExactMatch);
 	for (TPair<FName, UMovieGraphPostRenderNode*>& Pair : PostRenderNodes)
@@ -580,7 +589,7 @@ void UMovieGraphPipeline::BeginExport()
 		FName& BranchName = Pair.Key;
 		
 		PostRenderNode->BeginExport(this, BranchName);
-	}
+	}*/
 }
 
 void UMovieGraphPipeline::SetupShot(const TObjectPtr<UMoviePipelineExecutorShot>& InShot)
