MoviePipeline: Updated the graph to show pin data types in pin tooltips.

Also made connection type checking between pins more accurate (checks enum, struct, class, and object equality now), and prevented the "interface" type from being chosen for graph members because property bags do not support this type.

#rnx

diff --git a/Engine/Plugins/MovieScene/MovieRenderPipeline/Source/MovieRenderPipelineCore/Private/Graph/MovieGraphNode.cpp b/Engine/Plugins/MovieScene/MovieRenderPipeline/Source/MovieRenderPipelineCore/Private/Graph/MovieGraphNode.cpp
--- a/Engine/Plugins/MovieScene/MovieRenderPipeline/Source/MovieRenderPipelineCore/Private/Graph/MovieGraphNode.cpp
+++ b/Engine/Plugins/MovieScene/MovieRenderPipeline/Source/MovieRenderPipelineCore/Private/Graph/MovieGraphNode.cpp
@@ -10,7 +10,7 @@ FString UMovieGraphNode::GlobalsPinNameString = FString("Globals");
 
 namespace UE::MovieGraph::Private
 {
-	EMovieGraphValueType GetValueTypeFromProperty(const FProperty* InSourceProperty)
+	EMovieGraphValueType GetValueTypeFromProperty(const FProperty* InSourceProperty, TObjectPtr<const UObject>& OutValueTypeObject)
 	{
 		if (CastField<FBoolProperty>(InSourceProperty))
 		{
@@ -18,7 +18,13 @@ namespace UE::MovieGraph::Private
 		}
 		if (const FByteProperty* ByteProperty = CastField<FByteProperty>(InSourceProperty))
 		{
-			return ByteProperty->IsEnum() ? EMovieGraphValueType::Enum : EMovieGraphValueType::Byte;
+			if (ByteProperty->IsEnum())
+			{
+				OutValueTypeObject = ByteProperty->Enum.Get();
+				return EMovieGraphValueType::Enum;
+			}
+			
+			return EMovieGraphValueType::Byte;
 		}
 		if (CastField<FIntProperty>(InSourceProperty))
 		{
@@ -48,35 +54,41 @@ namespace UE::MovieGraph::Private
 		{
 			return EMovieGraphValueType::Text;
 		}
-		if (CastField<FEnumProperty>(InSourceProperty))
+		if (const FEnumProperty* EnumProperty = CastField<FEnumProperty>(InSourceProperty))
 		{
+			OutValueTypeObject = EnumProperty->GetEnum();
 			return EMovieGraphValueType::Enum;
 		}
-		if (CastField<FStructProperty>(InSourceProperty))
+		if (const FStructProperty* StructProperty = CastField<FStructProperty>(InSourceProperty))
 		{
+			OutValueTypeObject = StructProperty->Struct.Get();
 			return EMovieGraphValueType::Struct;
 		}
-		if (CastField<FObjectProperty>(InSourceProperty))
+		if (const FObjectProperty* ObjectProperty = CastField<FObjectProperty>(InSourceProperty))
 		{
+			OutValueTypeObject = ObjectProperty->PropertyClass.Get();
 			return EMovieGraphValueType::Object;
 		}
-		if (CastField<FSoftObjectProperty>(InSourceProperty))
+		if (const FSoftObjectProperty* SoftObjectProperty = CastField<FSoftObjectProperty>(InSourceProperty))
 		{
+			OutValueTypeObject = SoftObjectProperty->PropertyClass.Get();
 			return EMovieGraphValueType::SoftObject;
 		}
-		if (CastField<FClassProperty>(InSourceProperty))
+		if (const FClassProperty* ClassProperty = CastField<FClassProperty>(InSourceProperty))
 		{
+			OutValueTypeObject = ClassProperty->PropertyClass.Get();
 			return EMovieGraphValueType::Class;
 		}
-		if (CastField<FSoftClassProperty>(InSourceProperty))
+		if (const FSoftClassProperty* SoftClassProperty = CastField<FSoftClassProperty>(InSourceProperty))
 		{
+			OutValueTypeObject = SoftClassProperty->PropertyClass.Get();
 			return EMovieGraphValueType::SoftClass;
 		}
 
 		// Handle array property
 		if (const FArrayProperty* ArrayProperty = CastField<FArrayProperty>(InSourceProperty))
 		{
-			return GetValueTypeFromProperty(ArrayProperty->Inner);	
+			return GetValueTypeFromProperty(ArrayProperty->Inner, OutValueTypeObject);	
 		}
 
 		return EMovieGraphValueType::None;
@@ -100,7 +112,7 @@ TArray<FMovieGraphPinProperties> UMovieGraphNode::GetExposedPinProperties() cons
 		}
 		
 		constexpr bool bAllowMultipleConnections = false;
-		Properties.Add(FMovieGraphPinProperties(PropertyInfo.Name, PropertyInfo.ValueType, bAllowMultipleConnections));
+		Properties.Add(FMovieGraphPinProperties(PropertyInfo.Name, PropertyInfo.ValueType, PropertyInfo.ValueTypeObject, bAllowMultipleConnections));
 	}
 
 	return Properties;
@@ -441,7 +453,7 @@ TArray<FMovieGraphPropertyInfo> UMovieGraphNode::GetOverrideablePropertyInfo() c
 			FMovieGraphPropertyInfo Info;
 			Info.Name = PropertyIterator->GetFName();
 			Info.bIsDynamicProperty = false;
-			Info.ValueType = UE::MovieGraph::Private::GetValueTypeFromProperty(*PropertyIterator);
+			Info.ValueType = UE::MovieGraph::Private::GetValueTypeFromProperty(*PropertyIterator, Info.ValueTypeObject);
 			
 			OverrideableProperties.Add(MoveTemp(Info));
 		}
@@ -455,6 +467,7 @@ TArray<FMovieGraphPropertyInfo> UMovieGraphNode::GetOverrideablePropertyInfo() c
 			Info.Name = Desc.Name;
 			Info.bIsDynamicProperty = true;
 			Info.ValueType = static_cast<EMovieGraphValueType>(Desc.ValueType);
+			Info.ValueTypeObject = Desc.ValueTypeObject.Get();
 
 			OverrideableProperties.Add(MoveTemp(Info));
 		}
diff --git a/Engine/Plugins/MovieScene/MovieRenderPipeline/Source/MovieRenderPipelineCore/Private/Graph/MovieGraphPin.cpp b/Engine/Plugins/MovieScene/MovieRenderPipeline/Source/MovieRenderPipelineCore/Private/Graph/MovieGraphPin.cpp
--- a/Engine/Plugins/MovieScene/MovieRenderPipeline/Source/MovieRenderPipelineCore/Private/Graph/MovieGraphPin.cpp
+++ b/Engine/Plugins/MovieScene/MovieRenderPipeline/Source/MovieRenderPipelineCore/Private/Graph/MovieGraphPin.cpp
@@ -110,7 +110,8 @@ FPinConnectionResponse UMovieGraphPin::CanCreateConnection_PinConnectionResponse
 
 	const bool bBothPinsAreSameType = bBothPinsAreBranch ||				// Both are branches or
 		(!Properties.bIsBranch && !InOtherPin->Properties.bIsBranch &&	// Neither is branch and
-		Properties.Type == InOtherPin->Properties.Type)	;				// They have the same property type
+		Properties.Type == InOtherPin->Properties.Type &&				// They have the same property type and
+		Properties.TypeObject == InOtherPin->Properties.TypeObject);	// They have the same type object (for enums, structs, objects, classes)
 
 	// Pins need to be the same type
 	if (!bBothPinsAreSameType)				
diff --git a/Engine/Plugins/MovieScene/MovieRenderPipeline/Source/MovieRenderPipelineCore/Private/Graph/Nodes/MovieGraphBranchNode.cpp b/Engine/Plugins/MovieScene/MovieRenderPipeline/Source/MovieRenderPipelineCore/Private/Graph/Nodes/MovieGraphBranchNode.cpp
--- a/Engine/Plugins/MovieScene/MovieRenderPipeline/Source/MovieRenderPipelineCore/Private/Graph/Nodes/MovieGraphBranchNode.cpp
+++ b/Engine/Plugins/MovieScene/MovieRenderPipeline/Source/MovieRenderPipelineCore/Private/Graph/Nodes/MovieGraphBranchNode.cpp
@@ -84,9 +84,10 @@ TArray<FMovieGraphPinProperties> UMovieGraphBranchNode::GetInputPinProperties()
 {
 	TArray<FMovieGraphPinProperties> Properties;
 	
+	const TObjectPtr<UObject> ValueTypeObject = nullptr;
 	Properties.Add(FMovieGraphPinProperties::MakeBranchProperties(UE::MovieGraph::BranchNode::TrueBranch));
 	Properties.Add(FMovieGraphPinProperties::MakeBranchProperties(UE::MovieGraph::BranchNode::FalseBranch));
-	Properties.Add(FMovieGraphPinProperties(UE::MovieGraph::BranchNode::Condition, EMovieGraphValueType::Bool, false));
+	Properties.Add(FMovieGraphPinProperties(UE::MovieGraph::BranchNode::Condition, EMovieGraphValueType::Bool, ValueTypeObject, false));
 	return Properties;
 }
 
diff --git a/Engine/Plugins/MovieScene/MovieRenderPipeline/Source/MovieRenderPipelineCore/Private/Graph/Nodes/MovieGraphInputNode.cpp b/Engine/Plugins/MovieScene/MovieRenderPipeline/Source/MovieRenderPipelineCore/Private/Graph/Nodes/MovieGraphInputNode.cpp
--- a/Engine/Plugins/MovieScene/MovieRenderPipeline/Source/MovieRenderPipelineCore/Private/Graph/Nodes/MovieGraphInputNode.cpp
+++ b/Engine/Plugins/MovieScene/MovieRenderPipeline/Source/MovieRenderPipelineCore/Private/Graph/Nodes/MovieGraphInputNode.cpp
@@ -28,7 +28,7 @@ TArray<FMovieGraphPinProperties> UMovieGraphInputNode::GetOutputPinProperties()
 	{
 		for (const UMovieGraphInput* Input : ParentGraph->GetInputs())
 		{
-			FMovieGraphPinProperties PinProperties(FName(Input->GetMemberName()), Input->GetValueType(), false);
+			FMovieGraphPinProperties PinProperties(FName(Input->GetMemberName()), Input->GetValueType(), Input->GetValueTypeObject(), false);
 			PinProperties.bIsBranch = Input->bIsBranch;
 			Properties.Add(MoveTemp(PinProperties));
 		}
@@ -144,6 +144,7 @@ void UMovieGraphInputNode::UpdateExistingPins(UMovieGraphMember* ChangedInput) c
 			{
 				OutputPins[Index]->Properties.Label = FName(InputMembers[Index]->GetMemberName());
 				OutputPins[Index]->Properties.Type = InputMembers[Index]->GetValueType();
+				OutputPins[Index]->Properties.TypeObject = InputMembers[Index]->GetValueTypeObject();
 				OutputPins[Index]->Properties.bIsBranch = InputMembers[Index]->bIsBranch;
 			}
 		}
diff --git a/Engine/Plugins/MovieScene/MovieRenderPipeline/Source/MovieRenderPipelineCore/Private/Graph/Nodes/MovieGraphOutputNode.cpp b/Engine/Plugins/MovieScene/MovieRenderPipeline/Source/MovieRenderPipelineCore/Private/Graph/Nodes/MovieGraphOutputNode.cpp
--- a/Engine/Plugins/MovieScene/MovieRenderPipeline/Source/MovieRenderPipelineCore/Private/Graph/Nodes/MovieGraphOutputNode.cpp
+++ b/Engine/Plugins/MovieScene/MovieRenderPipeline/Source/MovieRenderPipelineCore/Private/Graph/Nodes/MovieGraphOutputNode.cpp
@@ -27,7 +27,7 @@ TArray<FMovieGraphPinProperties> UMovieGraphOutputNode::GetInputPinProperties()
 	{
 		for (const UMovieGraphOutput* Output : ParentGraph->GetOutputs())
 		{
-			FMovieGraphPinProperties PinProperties(FName(Output->GetMemberName()), Output->GetValueType(), false);
+			FMovieGraphPinProperties PinProperties(FName(Output->GetMemberName()), Output->GetValueType(), Output->GetValueTypeObject(), false);
 			PinProperties.bIsBranch = Output->bIsBranch;
 			Properties.Add(MoveTemp(PinProperties));
 		}
@@ -100,6 +100,7 @@ void UMovieGraphOutputNode::UpdateExistingPins(UMovieGraphMember* ChangedOutput)
 			{
 				InputPins[Index]->Properties.Label = FName(OutputMembers[Index]->GetMemberName());
 				InputPins[Index]->Properties.Type = OutputMembers[Index]->GetValueType();
+				InputPins[Index]->Properties.TypeObject = OutputMembers[Index]->GetValueTypeObject();
 				InputPins[Index]->Properties.bIsBranch = OutputMembers[Index]->bIsBranch;
 			}
 		}
diff --git a/Engine/Plugins/MovieScene/MovieRenderPipeline/Source/MovieRenderPipelineCore/Private/Graph/Nodes/MovieGraphSelectNode.cpp b/Engine/Plugins/MovieScene/MovieRenderPipeline/Source/MovieRenderPipelineCore/Private/Graph/Nodes/MovieGraphSelectNode.cpp
--- a/Engine/Plugins/MovieScene/MovieRenderPipeline/Source/MovieRenderPipelineCore/Private/Graph/Nodes/MovieGraphSelectNode.cpp
+++ b/Engine/Plugins/MovieScene/MovieRenderPipeline/Source/MovieRenderPipelineCore/Private/Graph/Nodes/MovieGraphSelectNode.cpp
@@ -152,7 +152,7 @@ TArray<FMovieGraphPinProperties> UMovieGraphSelectNode::GetInputPinProperties()
 	}
 
 	Properties.Add(FMovieGraphPinProperties::MakeBranchProperties(UE::MovieGraph::SelectNode::DefaultBranch));
-	Properties.Add(FMovieGraphPinProperties(UE::MovieGraph::SelectNode::SelectedOption, SelectOptions->GetValueType(), false));
+	Properties.Add(FMovieGraphPinProperties(UE::MovieGraph::SelectNode::SelectedOption, SelectOptions->GetValueType(), SelectOptions->GetValueTypeObject(), false));
 	
 	return Properties;
 }
diff --git a/Engine/Plugins/MovieScene/MovieRenderPipeline/Source/MovieRenderPipelineCore/Private/Graph/Nodes/MovieGraphSubgraphNode.cpp b/Engine/Plugins/MovieScene/MovieRenderPipeline/Source/MovieRenderPipelineCore/Private/Graph/Nodes/MovieGraphSubgraphNode.cpp
--- a/Engine/Plugins/MovieScene/MovieRenderPipeline/Source/MovieRenderPipelineCore/Private/Graph/Nodes/MovieGraphSubgraphNode.cpp
+++ b/Engine/Plugins/MovieScene/MovieRenderPipeline/Source/MovieRenderPipelineCore/Private/Graph/Nodes/MovieGraphSubgraphNode.cpp
@@ -28,7 +28,7 @@ namespace UE::MovieGraph::Private
 			[](const UMovieGraphInterfaceBase* Member) -> FMovieGraphPinProperties
 			{
 				constexpr bool bAllowMultipleConnections = false;
-				FMovieGraphPinProperties Properties(FName(Member->GetMemberName()), Member->GetValueType(), bAllowMultipleConnections);
+				FMovieGraphPinProperties Properties(FName(Member->GetMemberName()), Member->GetValueType(), Member->GetValueTypeObject(), bAllowMultipleConnections);
 				Properties.bIsBranch = Member->bIsBranch;
 				return Properties;
 			});
diff --git a/Engine/Plugins/MovieScene/MovieRenderPipeline/Source/MovieRenderPipelineCore/Private/Graph/Nodes/MovieGraphVariableNode.cpp b/Engine/Plugins/MovieScene/MovieRenderPipeline/Source/MovieRenderPipelineCore/Private/Graph/Nodes/MovieGraphVariableNode.cpp
--- a/Engine/Plugins/MovieScene/MovieRenderPipeline/Source/MovieRenderPipelineCore/Private/Graph/Nodes/MovieGraphVariableNode.cpp
+++ b/Engine/Plugins/MovieScene/MovieRenderPipeline/Source/MovieRenderPipelineCore/Private/Graph/Nodes/MovieGraphVariableNode.cpp
@@ -20,11 +20,11 @@ TArray<FMovieGraphPinProperties> UMovieGraphVariableNode::GetOutputPinProperties
 	
 	if (GraphVariable)
 	{
-		Properties.Add(FMovieGraphPinProperties(FName(GraphVariable->GetMemberName()), GraphVariable->GetValueType(), false));
+		Properties.Add(FMovieGraphPinProperties(FName(GraphVariable->GetMemberName()), GraphVariable->GetValueType(), GraphVariable->GetValueTypeObject(), false));
 	}
 	else
 	{
-		Properties.Add(FMovieGraphPinProperties(TEXT("Unknown"), EMovieGraphValueType::None, false));
+		Properties.Add(FMovieGraphPinProperties(TEXT("Unknown"), EMovieGraphValueType::None, nullptr, false));
 	}
 	
 	return Properties;
@@ -122,6 +122,7 @@ void UMovieGraphVariableNode::UpdateOutputPin(UMovieGraphMember* ChangedVariable
 		// Update the output pin to reflect the variable data model
 		OutputPins[0]->Properties.Label = FName(ChangedVariable->GetMemberName());
 		OutputPins[0]->Properties.Type = ChangedVariable->GetValueType();
+		OutputPins[0]->Properties.TypeObject = ChangedVariable->GetValueTypeObject();
 	}
 
 	OnNodeChangedDelegate.Broadcast(this);
diff --git a/Engine/Plugins/MovieScene/MovieRenderPipeline/Source/MovieRenderPipelineCore/Public/Graph/MovieGraphNode.h b/Engine/Plugins/MovieScene/MovieRenderPipeline/Source/MovieRenderPipelineCore/Public/Graph/MovieGraphNode.h
--- a/Engine/Plugins/MovieScene/MovieRenderPipeline/Source/MovieRenderPipelineCore/Public/Graph/MovieGraphNode.h
+++ b/Engine/Plugins/MovieScene/MovieRenderPipeline/Source/MovieRenderPipelineCore/Public/Graph/MovieGraphNode.h
@@ -52,6 +52,10 @@ struct FMovieGraphPropertyInfo
 	UPROPERTY()
 	EMovieGraphValueType ValueType = EMovieGraphValueType::None;
 
+	/** The associated value type object if the ValueType is an enum, struct, class, or object. */
+	UPROPERTY()
+	TObjectPtr<const UObject> ValueTypeObject;
+
 	/**
 	 * Determines if this struct represents the same property as another instance of this struct.
 	 *
@@ -69,7 +73,8 @@ struct FMovieGraphPropertyInfo
 	{
 		return (Name == Other.Name)
 			&& (bIsDynamicProperty == Other.bIsDynamicProperty)
-			&& (ValueType == Other.ValueType);
+			&& (ValueType == Other.ValueType)
+			&& (ValueTypeObject == Other.ValueTypeObject);
 	}
 };
 
diff --git a/Engine/Plugins/MovieScene/MovieRenderPipeline/Source/MovieRenderPipelineCore/Public/Graph/MovieGraphPin.h b/Engine/Plugins/MovieScene/MovieRenderPipeline/Source/MovieRenderPipelineCore/Public/Graph/MovieGraphPin.h
--- a/Engine/Plugins/MovieScene/MovieRenderPipeline/Source/MovieRenderPipelineCore/Public/Graph/MovieGraphPin.h
+++ b/Engine/Plugins/MovieScene/MovieRenderPipeline/Source/MovieRenderPipelineCore/Public/Graph/MovieGraphPin.h
@@ -17,15 +17,16 @@ struct FMovieGraphPinProperties
 	GENERATED_BODY()
 
 	FMovieGraphPinProperties() = default;
-	explicit FMovieGraphPinProperties(const FName& InLabel, const EMovieGraphValueType PinType, bool bInAllowMultipleConnections)
+	explicit FMovieGraphPinProperties(const FName& InLabel, const EMovieGraphValueType PinType, const TObjectPtr<const UObject>& TypeObject, bool bInAllowMultipleConnections)
 		: Label(InLabel)
 		, Type(PinType)
+		, TypeObject(TypeObject)
 		, bAllowMultipleConnections(bInAllowMultipleConnections)
 	{}
 
 	static FMovieGraphPinProperties MakeBranchProperties(const FName& InLabel = NAME_None)
 	{
-		FMovieGraphPinProperties Properties(InLabel, EMovieGraphValueType::None, false);
+		FMovieGraphPinProperties Properties(InLabel, EMovieGraphValueType::None, nullptr, false);
 		Properties.bIsBranch = true;
 		return MoveTemp(Properties);
 	}
@@ -38,6 +39,10 @@ struct FMovieGraphPinProperties
 	UPROPERTY(BlueprintReadWrite, EditAnywhere, Category = "Settings")
 	EMovieGraphValueType Type = EMovieGraphValueType::Float;
 
+	/** The value type of the pin, if the type is an enum, struct, class, or object. */
+	UPROPERTY(BlueprintReadWrite, EditAnywhere, Category = "Settings")
+	TObjectPtr<const UObject> TypeObject;
+
 	/** Whether this pin can accept multiple connections. */
 	UPROPERTY(BlueprintReadWrite, EditAnywhere, Category = "Settings")
 	bool bAllowMultipleConnections = true;
@@ -57,6 +62,7 @@ struct FMovieGraphPinProperties
 	{
 		return Label == Other.Label
 			&& Type == Other.Type
+			&& TypeObject == Other.TypeObject
 			&& bAllowMultipleConnections == Other.bAllowMultipleConnections
 			&& bIsBranch == Other.bIsBranch
 			&& bIsBuiltIn == Other.bIsBuiltIn;
diff --git a/Engine/Plugins/MovieScene/MovieRenderPipeline/Source/MovieRenderPipelineEditor/Private/Customizations/Graph/MovieGraphMemberCustomization.h b/Engine/Plugins/MovieScene/MovieRenderPipeline/Source/MovieRenderPipelineEditor/Private/Customizations/Graph/MovieGraphMemberCustomization.h
--- a/Engine/Plugins/MovieScene/MovieRenderPipeline/Source/MovieRenderPipelineEditor/Private/Customizations/Graph/MovieGraphMemberCustomization.h
+++ b/Engine/Plugins/MovieScene/MovieRenderPipeline/Source/MovieRenderPipelineEditor/Private/Customizations/Graph/MovieGraphMemberCustomization.h
@@ -52,15 +52,59 @@ protected:
 		}
 	}
 
+	bool IsPinTypeAllowed(const FEdGraphPinType& InPinType) const
+	{
+		// Property bags do not support interface types
+		if (InPinType.PinCategory == UEdGraphSchema_K2::PC_Interface)
+		{
+			return false;
+		}
+
+		return true;
+	}
+
 	virtual void CustomizeDetails(IDetailLayoutBuilder& DetailBuilder) override
 	{
 		TSharedRef<IPropertyUtilities> PropUtils = DetailBuilder.GetPropertyUtilities();
 		
 		auto GetFilteredVariableTypeTree = [this](TArray<TSharedPtr<UEdGraphSchema_K2::FPinTypeTreeInfo>>& TypeTree, ETypeTreeFilter TypeTreeFilter)
 		{
-			// All types from the schema are allowed
+			// Most types from the schema are allowed
 			check(GetDefault<UEdGraphSchema_K2>());
 			GetDefault<UPropertyBagSchema>()->GetVariableTypeTree(TypeTree, TypeTreeFilter);
+
+			// Filter out disallowed types
+			for (int32 Index = 0; Index < TypeTree.Num(); )
+			{
+				TSharedPtr<UEdGraphSchema_K2::FPinTypeTreeInfo>& PinType = TypeTree[Index];
+				if (!PinType.IsValid())
+				{
+					return;
+				}
+
+				constexpr bool bForceLoadedSubCategoryObject = false;
+				if (!IsPinTypeAllowed(PinType->GetPinType(bForceLoadedSubCategoryObject)))
+				{
+					TypeTree.RemoveAt(Index);
+					continue;
+				}
+
+				for (int32 ChildIndex = 0; ChildIndex < PinType->Children.Num(); )
+				{
+					TSharedPtr<UEdGraphSchema_K2::FPinTypeTreeInfo> Child = PinType->Children[ChildIndex];
+					if (Child.IsValid())
+					{
+						if (!IsPinTypeAllowed(Child->GetPinType(bForceLoadedSubCategoryObject)))
+						{
+							PinType->Children.RemoveAt(ChildIndex);
+							continue;
+						}
+					}
+					++ChildIndex;
+				}
+
+				++Index;
+			}
 		};
 
 		auto PinInfoChanged = [PropUtils](const FEdGraphPinType& PinType, const TWeakObjectPtr<UMovieGraphMember>& GraphMember)
diff --git a/Engine/Plugins/MovieScene/MovieRenderPipeline/Source/MovieRenderPipelineEditor/Private/Graph/MovieEdGraphNode.cpp b/Engine/Plugins/MovieScene/MovieRenderPipeline/Source/MovieRenderPipelineEditor/Private/Graph/MovieEdGraphNode.cpp
--- a/Engine/Plugins/MovieScene/MovieRenderPipeline/Source/MovieRenderPipelineEditor/Private/Graph/MovieEdGraphNode.cpp
+++ b/Engine/Plugins/MovieScene/MovieRenderPipeline/Source/MovieRenderPipelineEditor/Private/Graph/MovieEdGraphNode.cpp
@@ -219,16 +219,18 @@ void UMoviePipelineEdGraphNodeBase::UpdateEnableState() const
 
 void UMoviePipelineEdGraphNode::AllocateDefaultPins()
 {
-	if(RuntimeNode)
+	if (RuntimeNode)
 	{
 		for(const UMovieGraphPin* InputPin : RuntimeNode->GetInputPins())
 		{
-			CreatePin(EEdGraphPinDirection::EGPD_Input, GetPinType(InputPin), InputPin->Properties.Label);
+			UEdGraphPin* NewPin = CreatePin(EGPD_Input, GetPinType(InputPin), InputPin->Properties.Label);
+			NewPin->PinToolTip = GetPinTooltip(InputPin);
 		}
 		
 		for(const UMovieGraphPin* OutputPin : RuntimeNode->GetOutputPins())
 		{
-			CreatePin(EEdGraphPinDirection::EGPD_Output, GetPinType(OutputPin), OutputPin->Properties.Label);
+			UEdGraphPin* NewPin = CreatePin(EGPD_Output, GetPinType(OutputPin), OutputPin->Properties.Label);
+			NewPin->PinToolTip = GetPinTooltip(OutputPin);
 		}
 	}
 }
@@ -322,6 +324,7 @@ void UMoviePipelineEdGraphNodeBase::CreatePins(const TArray<UMovieGraphPin*>& In
 		}
 
 		UEdGraphPin* Pin = CreatePin(EGPD_Input, GetPinType(InputPin), InputPin->Properties.Label);
+		Pin->PinToolTip = GetPinTooltip(InputPin);
 		// Pin->bAdvancedView = InputPin->Properties.bAdvancedPin;
 		bHasAdvancedPin |= Pin->bAdvancedView;
 	}
@@ -334,6 +337,7 @@ void UMoviePipelineEdGraphNodeBase::CreatePins(const TArray<UMovieGraphPin*>& In
 		}
 
 		UEdGraphPin* Pin = CreatePin(EGPD_Output, GetPinType(OutputPin), OutputPin->Properties.Label);
+		Pin->PinToolTip = GetPinTooltip(OutputPin);
 		// Pin->bAdvancedView = OutputPin->Properties.bAdvancedPin;
 		bHasAdvancedPin |= Pin->bAdvancedView;
 	}
@@ -348,6 +352,24 @@ void UMoviePipelineEdGraphNodeBase::CreatePins(const TArray<UMovieGraphPin*>& In
 	}
 }
 
+FString UMoviePipelineEdGraphNodeBase::GetPinTooltip(const UMovieGraphPin* InPin) const
+{
+	const EMovieGraphValueType PinType = InPin->Properties.Type;
+	const FString TypeString = InPin->Properties.bIsBranch
+		? LOCTEXT("PinTypeTooltip_Branch", "Branch").ToString()
+		: StaticEnum<EMovieGraphValueType>()->GetDisplayNameTextByValue(static_cast<int64>(PinType)).ToString();
+
+	TStringBuilder<256> PinTooltip;
+	PinTooltip << LOCTEXT("PinTypeTooltip_Type", "Type: ").ToString() << TypeString;
+
+	if (const TObjectPtr<const UObject> ValueTypeObject = InPin->Properties.TypeObject)
+	{
+		PinTooltip << TEXT(" (") << ValueTypeObject.Get()->GetName() << TEXT(")"); 
+	}
+
+	return PinTooltip.ToString();
+}
+
 void UMoviePipelineEdGraphNodeBase::AutowireNewNode(UEdGraphPin* FromPin)
 {
 	if (RuntimeNode == nullptr || FromPin == nullptr)
diff --git a/Engine/Plugins/MovieScene/MovieRenderPipeline/Source/MovieRenderPipelineEditor/Private/Graph/MovieEdGraphNode.h b/Engine/Plugins/MovieScene/MovieRenderPipeline/Source/MovieRenderPipelineEditor/Private/Graph/MovieEdGraphNode.h
--- a/Engine/Plugins/MovieScene/MovieRenderPipeline/Source/MovieRenderPipelineEditor/Private/Graph/MovieEdGraphNode.h
+++ b/Engine/Plugins/MovieScene/MovieRenderPipeline/Source/MovieRenderPipelineEditor/Private/Graph/MovieEdGraphNode.h
@@ -50,6 +50,9 @@ protected:
 	void RebuildRuntimeEdgesFromPins();
 
 	void CreatePins(const TArray<UMovieGraphPin*>& InInputPins, const TArray<UMovieGraphPin*>& InOutputPins);
+	
+	/** Gets the tooltip for the given pin. */
+	FString GetPinTooltip(const UMovieGraphPin* InPin) const;
 
 	/** Recreate the pins on this node, discarding all existing pins. */
 	void ReconstructPins();
diff --git a/Engine/Plugins/MovieScene/MovieRenderPipeline/Source/MovieRenderPipelineEditor/Private/Graph/MovieEdGraphVariableNode.cpp b/Engine/Plugins/MovieScene/MovieRenderPipeline/Source/MovieRenderPipelineEditor/Private/Graph/MovieEdGraphVariableNode.cpp
--- a/Engine/Plugins/MovieScene/MovieRenderPipeline/Source/MovieRenderPipelineEditor/Private/Graph/MovieEdGraphVariableNode.cpp
+++ b/Engine/Plugins/MovieScene/MovieRenderPipeline/Source/MovieRenderPipelineEditor/Private/Graph/MovieEdGraphVariableNode.cpp
@@ -28,7 +28,8 @@ void UMoviePipelineEdGraphVariableNode::AllocateDefaultPins()
 		const TArray<TObjectPtr<UMovieGraphPin>>& OutputPins = RuntimeNode->GetOutputPins();
 		if (!OutputPins.IsEmpty())
 		{
-			CreatePin(EGPD_Output, GetPinType(OutputPins[0].Get()), FName(VariableNode->GetVariable()->GetMemberName()));
+			UEdGraphPin* NewPin = CreatePin(EGPD_Output, GetPinType(OutputPins[0].Get()), FName(VariableNode->GetVariable()->GetMemberName()));
+			NewPin->PinToolTip = GetPinTooltip(OutputPins[0].Get());
 		}
 	}
 }
