MovieGraph: Add support for modifiers to set the holdout flag in UVolumetricCloudComponent, USkyAtmosphereComponent, UExponentialHeightFogComponent.

diff --git a/Engine/Plugins/MovieScene/MovieRenderPipeline/Source/MovieRenderPipelineCore/Private/Graph/MoviePipelineRenderLayerSubsystem.cpp b/Engine/Plugins/MovieScene/MovieRenderPipeline/Source/MovieRenderPipelineCore/Private/Graph/MoviePipelineRenderLayerSubsystem.cpp
--- a/Engine/Plugins/MovieScene/MovieRenderPipeline/Source/MovieRenderPipelineCore/Private/Graph/MoviePipelineRenderLayerSubsystem.cpp
+++ b/Engine/Plugins/MovieScene/MovieRenderPipeline/Source/MovieRenderPipelineCore/Private/Graph/MoviePipelineRenderLayerSubsystem.cpp
@@ -4,6 +4,8 @@
 
 #include "Components/PrimitiveComponent.h"
 #include "Components/SkyAtmosphereComponent.h"
+#include "Components/VolumetricCloudComponent.h"
+#include "Components/ExponentialHeightFogComponent.h"
 #include "Components/SkyLightComponent.h"
 #include "Components/SphereReflectionCaptureComponent.h"
 #include "EngineUtils.h"
@@ -188,24 +190,41 @@ void UMoviePipelineVisibilityModifier::ApplyModifier(const UWorld* World)
 			OriginalVisibilityState.bIsHidden = Actor->IsHidden();
 
 			const bool bIncludeFromChildActors = true;
-			TInlineComponentArray<UPrimitiveComponent*> Components;
-			Actor->GetComponents<UPrimitiveComponent>(Components, bIncludeFromChildActors);
+			TInlineComponentArray<USceneComponent*> Components;
+			Actor->GetComponents<USceneComponent>(Components, bIncludeFromChildActors);
 
 			OriginalVisibilityState.Components.Reserve(Components.Num());
-			for(const UPrimitiveComponent* PrimitiveComponent : Components)
+			for(const USceneComponent* SceneComponent : Components)
 			{
-				// Cache the state
 				FActorVisibilityState::FComponentState& ComponentState = OriginalVisibilityState.Components.AddDefaulted_GetRef();
-				ComponentState.Component = PrimitiveComponent;
-				ComponentState.bCastsShadows = PrimitiveComponent->CastShadow;
-				ComponentState.bCastShadowWhileHidden = PrimitiveComponent->bCastHiddenShadow;
-				ComponentState.bAffectIndirectLightingWhileHidden = PrimitiveComponent->bAffectIndirectLightingWhileHidden;
-				ComponentState.bHoldout = PrimitiveComponent->bHoldout;
+				ComponentState.Component = SceneComponent;
+
+				// Cache the state
+				if (const UPrimitiveComponent* AsPrimitiveComponent = Cast<UPrimitiveComponent>(SceneComponent))
+				{
+					ComponentState.bCastsShadows = AsPrimitiveComponent->CastShadow;
+					ComponentState.bCastShadowWhileHidden = AsPrimitiveComponent->bCastHiddenShadow;
+					ComponentState.bAffectIndirectLightingWhileHidden = AsPrimitiveComponent->bAffectIndirectLightingWhileHidden;
+					ComponentState.bHoldout = AsPrimitiveComponent->bHoldout;
+				}
+				// Volumetrics are special cases as they don't inherit from UPrimitiveComponent, and don't support all of the flags.
+				else if (const UVolumetricCloudComponent* AsVolumetricCloudComponent = Cast<UVolumetricCloudComponent>(SceneComponent))
+				{
+					ComponentState.bHoldout = AsVolumetricCloudComponent->bHoldout;
+				}
+				else if (const USkyAtmosphereComponent* AsSkyAtmosphereComponent = Cast<USkyAtmosphereComponent>(SceneComponent))
+				{
+					ComponentState.bHoldout = AsSkyAtmosphereComponent->bHoldout;
+				}
+				else if (const UExponentialHeightFogComponent* AsExponentialHeightFogComponent = Cast<UExponentialHeightFogComponent>(SceneComponent))
+				{
+					ComponentState.bHoldout = AsExponentialHeightFogComponent->bHoldout;
+				}
 
 				// Then override it. We override it one component at a time so that we can avoid making a bunch of
 				// copies of NewVisibilityState for the varying number of components you might have.
 				NewVisibilityState.Actor = Actor;
-				NewVisibilityState.Components[0].Component = PrimitiveComponent;
+				NewVisibilityState.Components[0].Component = SceneComponent;
 				SetActorVisibilityState(NewVisibilityState);
 			}
 			
@@ -239,7 +258,6 @@ void UMoviePipelineVisibilityModifier::SetActorVisibilityState(const FActorVisib
 	Actor->SetIsTemporarilyHiddenInEditor(NewVisibilityState.bIsHidden);
 #endif
 
-	
 	for (const FActorVisibilityState::FComponentState& ComponentState : NewVisibilityState.Components)
 	{
 		// TODO: These could potentially cause a large rendering penalty due to dirtying the render state; investigate potential
@@ -249,10 +267,26 @@ void UMoviePipelineVisibilityModifier::SetActorVisibilityState(const FActorVisib
 			continue;
 		}
 
-		ComponentState.Component->SetCastShadow(ComponentState.bCastsShadows);
-		ComponentState.Component->SetCastHiddenShadow(ComponentState.bCastShadowWhileHidden);
-		ComponentState.Component->SetAffectIndirectLightingWhileHidden(ComponentState.bAffectIndirectLightingWhileHidden);
-		ComponentState.Component->SetHoldout(ComponentState.bHoldout);
+		if (UPrimitiveComponent* AsPrimitiveComponent = Cast<UPrimitiveComponent>(ComponentState.Component.Get()))
+		{
+			AsPrimitiveComponent->SetCastShadow(ComponentState.bCastsShadows);
+			AsPrimitiveComponent->SetCastHiddenShadow(ComponentState.bCastShadowWhileHidden);
+			AsPrimitiveComponent->SetAffectIndirectLightingWhileHidden(ComponentState.bAffectIndirectLightingWhileHidden);
+			AsPrimitiveComponent->SetHoldout(ComponentState.bHoldout);
+		}
+		// Volumetrics are special cases as they don't inherit from UPrimitiveComponent, and don't support all of the flags.
+		else if (UVolumetricCloudComponent* AsVolumetricCloudComponent = Cast<UVolumetricCloudComponent>(ComponentState.Component.Get()))
+		{
+			AsVolumetricCloudComponent->SetHoldout(ComponentState.bHoldout);
+		}
+		else if (USkyAtmosphereComponent* AsSkyAtmosphereComponent = Cast<USkyAtmosphereComponent>(ComponentState.Component.Get()))
+		{
+			AsSkyAtmosphereComponent->SetHoldout(ComponentState.bHoldout);
+		}
+		else if (UExponentialHeightFogComponent* AsExponentialHeightFogComponent = Cast<UExponentialHeightFogComponent>(ComponentState.Component.Get()))
+		{
+			AsExponentialHeightFogComponent->SetHoldout(ComponentState.bHoldout);
+		}
 	}
 }
 
diff --git a/Engine/Plugins/MovieScene/MovieRenderPipeline/Source/MovieRenderPipelineCore/Public/Graph/MoviePipelineRenderLayerSubsystem.h b/Engine/Plugins/MovieScene/MovieRenderPipeline/Source/MovieRenderPipelineCore/Public/Graph/MoviePipelineRenderLayerSubsystem.h
--- a/Engine/Plugins/MovieScene/MovieRenderPipeline/Source/MovieRenderPipelineCore/Public/Graph/MoviePipelineRenderLayerSubsystem.h
+++ b/Engine/Plugins/MovieScene/MovieRenderPipeline/Source/MovieRenderPipelineCore/Public/Graph/MoviePipelineRenderLayerSubsystem.h
@@ -681,7 +681,7 @@ private:
 	{
 		struct FComponentState
 		{
-			TSoftObjectPtr<UPrimitiveComponent> Component = nullptr;
+			TSoftObjectPtr<USceneComponent> Component = nullptr;
 			
 			uint8 bCastsShadows : 1 = true;
 			uint8 bCastShadowWhileHidden : 1 = false;
