MoviePipeline: Added the "Promote to Variable" feature to the graph.

Also made the "type" tooltip for pins more correct localization-wise.

#rnx

diff --git a/Engine/Plugins/MovieScene/MovieRenderPipeline/Source/MovieRenderPipelineEditor/Private/Graph/MovieEdGraphNode.cpp b/Engine/Plugins/MovieScene/MovieRenderPipeline/Source/MovieRenderPipelineEditor/Private/Graph/MovieEdGraphNode.cpp
--- a/Engine/Plugins/MovieScene/MovieRenderPipeline/Source/MovieRenderPipelineEditor/Private/Graph/MovieEdGraphNode.cpp
+++ b/Engine/Plugins/MovieScene/MovieRenderPipeline/Source/MovieRenderPipelineEditor/Private/Graph/MovieEdGraphNode.cpp
@@ -5,6 +5,7 @@
 #include "Graph/MovieGraphPin.h"
 #include "Graph/MovieGraphNode.h"
 #include "Graph/MovieGraphConfig.h"
+#include "Graph/Nodes/MovieGraphVariableNode.h"
 #include "EdGraph/EdGraphPin.h"
 #include "Misc/TransactionObjectEvent.h"
 #include "MovieEdGraph.h"
@@ -264,14 +265,44 @@ void UMoviePipelineEdGraphNode::GetNodeContextMenuActions(UToolMenu* Menu, UGrap
 	GetPropertyPromotionContextMenuActions(Menu, Context);
 }
 
-void UMoviePipelineEdGraphNode::GetPropertyPromotionContextMenuActions(UToolMenu* Menu, UGraphNodeContextMenuContext* Context) const
+void UMoviePipelineEdGraphNode::GetPropertyPromotionContextMenuActions(UToolMenu* Menu, const UGraphNodeContextMenuContext* Context) const
 {
-	FToolMenuSection& Section = Menu->AddSection("MoviePipelineGraphExposeAsPin", LOCTEXT("ExposeAsPin", "Expose Property as Pin"));
-
 	const TArray<FMovieGraphPropertyInfo>& OverrideablePropertyInfo = RuntimeNode->GetOverrideablePropertyInfo();
+	
+	FToolMenuSection& PinActionsSection = Menu->FindOrAddSection("EdGraphSchemaPinActions");
+	if (const UEdGraphPin* SelectedPin = Context->Pin)
+	{
+		// Find the property info associated with the selected pin
+		const FMovieGraphPropertyInfo* PropertyInfo = OverrideablePropertyInfo.FindByPredicate([SelectedPin](const FMovieGraphPropertyInfo& PropInfo)
+		{
+			return PropInfo.Name == SelectedPin->GetFName();
+		});
+
+		// Allow promotion of the property to a variable if the property info could be found. Follow the behavior of blueprints, which allows promotion
+		// even if there is an existing connection to the pin.
+		if (PropertyInfo)
+		{
+			const FMovieGraphPropertyInfo& TargetProperty = *PropertyInfo;
+			
+			PinActionsSection.AddMenuEntry(
+				SelectedPin->GetFName(),
+				LOCTEXT("PromotePropertyToVariable_Label", "Promote to Variable"),
+				LOCTEXT("PromotePropertyToVariable_Tooltip", "Promote this property to a new variable and connect the variable to this pin."),
+				FSlateIcon(),
+				FUIAction(
+					FExecuteAction::CreateWeakLambda(this, [this, TargetProperty]()
+					{
+						PromotePropertyToVariable(TargetProperty);
+					}),
+					FCanExecuteAction())
+			);
+		}
+	}
+	
+	FToolMenuSection& ExposeAsPinSection = Menu->AddSection("MoviePipelineGraphExposeAsPin", LOCTEXT("ExposeAsPin", "Expose Property as Pin"));
 	for (const FMovieGraphPropertyInfo& PropertyInfo : OverrideablePropertyInfo)
 	{
-		Section.AddMenuEntry(
+		ExposeAsPinSection.AddMenuEntry(
 			PropertyInfo.Name,
 			FText::FromName(PropertyInfo.Name),
 			LOCTEXT("PromotePropertyToPin", "Promote this property to a pin on this node."),
@@ -290,7 +321,7 @@ void UMoviePipelineEdGraphNode::GetPropertyPromotionContextMenuActions(UToolMenu
 
 	if (OverrideablePropertyInfo.IsEmpty())
 	{
-		Section.AddMenuEntry(
+		ExposeAsPinSection.AddMenuEntry(
 			"NoPropertiesAvailable",
 			FText::FromString("No properties available"),
 			LOCTEXT("PromotePropertyToPin_NoneAvailable", "No properties are available to promote."),
@@ -302,6 +333,29 @@ void UMoviePipelineEdGraphNode::GetPropertyPromotionContextMenuActions(UToolMenu
 	}
 }
 
+void UMoviePipelineEdGraphNode::PromotePropertyToVariable(const FMovieGraphPropertyInfo& TargetProperty) const
+{
+	// Note: AddVariable() will take care of determining a unique name if there is already a variable with the property's name
+	if (UMovieGraphVariable* NewGraphVariable = RuntimeNode->GetGraph()->AddVariable(TargetProperty.Name))
+	{
+		// Set the new variable's type to match the property that is being promoted
+		UObject* ValueTypeObject = const_cast<UObject*>(TargetProperty.ValueTypeObject.Get());
+		NewGraphVariable->SetValueType(TargetProperty.ValueType, ValueTypeObject);
+
+		// When creating the new action, since it's only being used to create a node, the category, display name, and tooltip can just be empty
+		const TSharedPtr<FMovieGraphSchemaAction_NewVariableNode> NewAction = MakeShared<FMovieGraphSchemaAction_NewVariableNode>(
+			FText::GetEmpty(), FText::GetEmpty(), NewGraphVariable->GetGuid(), FText::GetEmpty());
+		NewAction->NodeClass = UMovieGraphVariableNode::StaticClass();
+
+		// Put the new node in a roughly ok-ish position relative to this node
+		const FVector2d NewLocation(NodePosX - 200, NodePosY);
+
+		// Note: Providing FromPin will trigger the action to connect the new node and this node
+		UEdGraphPin* FromPin = FindPin(TargetProperty.Name, EGPD_Input);
+		NewAction->PerformAction(GetGraph(), FromPin, NewLocation);
+	}
+}
+
 void UMoviePipelineEdGraphNode::TogglePromotePropertyToPin(const FName PropertyName) const
 {
 	RuntimeNode->TogglePromotePropertyToPin(PropertyName);
@@ -355,17 +409,21 @@ void UMoviePipelineEdGraphNodeBase::CreatePins(const TArray<UMovieGraphPin*>& In
 FString UMoviePipelineEdGraphNodeBase::GetPinTooltip(const UMovieGraphPin* InPin) const
 {
 	const EMovieGraphValueType PinType = InPin->Properties.Type;
-	const FString TypeString = InPin->Properties.bIsBranch
-		? LOCTEXT("PinTypeTooltip_Branch", "Branch").ToString()
-		: StaticEnum<EMovieGraphValueType>()->GetDisplayNameTextByValue(static_cast<int64>(PinType)).ToString();
+	const FText TypeText = InPin->Properties.bIsBranch
+		? LOCTEXT("PinTypeTooltip_Branch", "Branch")
+		: StaticEnum<EMovieGraphValueType>()->GetDisplayNameTextByValue(static_cast<int64>(PinType));
+	const FText TypeObjectText = InPin->Properties.TypeObject ? FText::FromString(InPin->Properties.TypeObject.Get()->GetName()) : FText::GetEmpty();
 
-	TStringBuilder<256> PinTooltip;
-	PinTooltip << LOCTEXT("PinTypeTooltip_Type", "Type: ").ToString() << TypeString;
+	const FText PinTooltipFormat = LOCTEXT("PinTypeTooltip_NoValueTypeObject", "Type: {ValueType}");
+	const FText PinTooltipFormatWithTypeObject = LOCTEXT("PinTypeTooltip_WithValueTypeObject", "Type: {ValueType} ({ValueTypeObject})");
 
-	if (const TObjectPtr<const UObject> ValueTypeObject = InPin->Properties.TypeObject)
-	{
-		PinTooltip << TEXT(" (") << ValueTypeObject.Get()->GetName() << TEXT(")"); 
-	}
+	FFormatNamedArguments NamedArgs;
+	NamedArgs.Add(TEXT("ValueType"), TypeText);
+	NamedArgs.Add(TEXT("ValueTypeObject"), TypeObjectText);
+
+	const FText PinTooltip = InPin->Properties.TypeObject
+		? FText::Format(PinTooltipFormatWithTypeObject, NamedArgs)
+		: FText::Format(PinTooltipFormat, NamedArgs);
 
 	return PinTooltip.ToString();
 }
diff --git a/Engine/Plugins/MovieScene/MovieRenderPipeline/Source/MovieRenderPipelineEditor/Private/Graph/MovieEdGraphNode.h b/Engine/Plugins/MovieScene/MovieRenderPipeline/Source/MovieRenderPipelineEditor/Private/Graph/MovieEdGraphNode.h
--- a/Engine/Plugins/MovieScene/MovieRenderPipeline/Source/MovieRenderPipelineEditor/Private/Graph/MovieEdGraphNode.h
+++ b/Engine/Plugins/MovieScene/MovieRenderPipeline/Source/MovieRenderPipelineEditor/Private/Graph/MovieEdGraphNode.h
@@ -102,5 +102,11 @@ protected:
 	void TogglePromotePropertyToPin(const FName PropertyName) const;
 
 private:
-	void GetPropertyPromotionContextMenuActions(UToolMenu* Menu, UGraphNodeContextMenuContext* Context) const;
+	void GetPropertyPromotionContextMenuActions(UToolMenu* Menu, const UGraphNodeContextMenuContext* Context) const;
+
+	/**
+	 * Promote the property identified by TargetProperty to a variable. This means 1) creating a new variable w/ the type associated with the
+	 * target property, 2) creating a new variable node linked to this new variable, and 3) connecting the variable node to TargetProperty's pin.
+	 */
+	void PromotePropertyToVariable(const FMovieGraphPropertyInfo& TargetProperty) const;
 };
diff --git a/Engine/Plugins/MovieScene/MovieRenderPipeline/Source/MovieRenderPipelineEditor/Private/Graph/SMovieGraphMembersTabContent.cpp b/Engine/Plugins/MovieScene/MovieRenderPipeline/Source/MovieRenderPipelineEditor/Private/Graph/SMovieGraphMembersTabContent.cpp
--- a/Engine/Plugins/MovieScene/MovieRenderPipeline/Source/MovieRenderPipelineEditor/Private/Graph/SMovieGraphMembersTabContent.cpp
+++ b/Engine/Plugins/MovieScene/MovieRenderPipeline/Source/MovieRenderPipelineEditor/Private/Graph/SMovieGraphMembersTabContent.cpp
@@ -65,6 +65,10 @@ void SMovieGraphMembersTabContent::Construct(const FArguments& InArgs)
 	EditorToolkit = InArgs._Editor;
 	CurrentGraph = InArgs._Graph;
 	OnActionSelected = InArgs._OnActionSelected;
+
+	// Update the UI whenever the graph adds/updates variables. In this case it's not known which variable is added/updated, so just pass nullptr.
+	UMovieGraphMember* UpdatedMember = nullptr;
+	CurrentGraph->OnGraphVariablesChangedDelegate.AddSP(this, &SMovieGraphMembersTabContent::RefreshMemberActions, UpdatedMember);
 	
 	ChildSlot
 	[
@@ -402,7 +406,7 @@ FReply SMovieGraphMembersTabContent::OnAddButtonClickedOnSection(const int32 InS
 	return FReply::Handled();
 }
 
-void SMovieGraphMembersTabContent::RefreshMemberActions(UMovieGraphMember* UpdatedMember) const
+void SMovieGraphMembersTabContent::RefreshMemberActions(UMovieGraphMember* UpdatedMember)
 {
 	// Currently the entire action menu is refreshed rather than a specific action being targeted
 
diff --git a/Engine/Plugins/MovieScene/MovieRenderPipeline/Source/MovieRenderPipelineEditor/Private/Graph/SMovieGraphMembersTabContent.h b/Engine/Plugins/MovieScene/MovieRenderPipeline/Source/MovieRenderPipelineEditor/Private/Graph/SMovieGraphMembersTabContent.h
--- a/Engine/Plugins/MovieScene/MovieRenderPipeline/Source/MovieRenderPipelineEditor/Private/Graph/SMovieGraphMembersTabContent.h
+++ b/Engine/Plugins/MovieScene/MovieRenderPipeline/Source/MovieRenderPipelineEditor/Private/Graph/SMovieGraphMembersTabContent.h
@@ -86,7 +86,7 @@ private:
 	FReply OnAddButtonClickedOnSection(const int32 InSectionID);
 
 	/** Refresh/regenerate the action menu when the given member is updated. */
-	void RefreshMemberActions(class UMovieGraphMember* UpdatedMember = nullptr) const;
+	void RefreshMemberActions(class UMovieGraphMember* UpdatedMember = nullptr);
 
 private:
 	/** The editor that this widget is associated with. */
