MovieGraph: Improvements to (single layer) output filename validation.

We now automatically create missing {layer_name} or {renderer_name} tokens to prevent sequences from overwriting themselves.

#rnx

diff --git a/Engine/Plugins/MovieScene/MovieRenderPipeline/Source/MovieRenderPipelineCore/Private/MoviePipelineUtils.cpp b/Engine/Plugins/MovieScene/MovieRenderPipeline/Source/MovieRenderPipelineCore/Private/MoviePipelineUtils.cpp
--- a/Engine/Plugins/MovieScene/MovieRenderPipeline/Source/MovieRenderPipelineCore/Private/MoviePipelineUtils.cpp
+++ b/Engine/Plugins/MovieScene/MovieRenderPipeline/Source/MovieRenderPipelineCore/Private/MoviePipelineUtils.cpp
@@ -659,6 +659,38 @@ namespace UE
 	{
 		DECLARE_CYCLE_STAT(TEXT("STAT_MoviePipeline_HardwareMetadata"), STAT_HardwareMetadata, STATGROUP_MoviePipeline);
 
+		void ConformOutputFormatStringToken(FString& InOutFilenameFormatString, const FStringView InToken, const FName& InNodeName, const FName& InBranchName)
+		{
+			static const FString FrameNumberIdentifiers[] = { TEXT("{frame_number}"), TEXT("{frame_number_shot}"), TEXT("{frame_number_rel}"), TEXT("{frame_number_shot_rel}") };
+
+			if (!InOutFilenameFormatString.Contains(InToken, ESearchCase::IgnoreCase))
+			{
+				UE_LOG(LogMovieRenderPipeline, Warning, TEXT("Missing expected %s format token on node '%s' in branch '%s'. Automatically adding!"), InToken.GetData(), *InNodeName.ToString(), *InBranchName.ToString());
+
+				// Search for a frame number in the output string
+				int32 FrameNumberIndex = INDEX_NONE;
+				for (const FString& Identifier : FrameNumberIdentifiers)
+				{
+					FrameNumberIndex = InOutFilenameFormatString.Find(Identifier, ESearchCase::IgnoreCase, ESearchDir::FromEnd);
+					if (FrameNumberIndex != INDEX_NONE)
+					{
+						break;
+					}
+				}
+
+				if (FrameNumberIndex == INDEX_NONE)
+				{
+					// No frame number found, so just append the token
+					InOutFilenameFormatString += InToken;
+				}
+				else
+				{
+					// If a frame number is found, we need to insert the token first before it, so various editing
+					// software will still be able to identify if this is an image sequence
+					InOutFilenameFormatString.InsertAt(FrameNumberIndex, FString(InToken) + TEXT("."));
+				}
+			}
+		}
 
 		void ValidateOutputFormatString(FString& InOutFilenameFormatString, const bool bTestRenderPass, const bool bTestFrameNumber, const bool bTestCameraName)
 		{
diff --git a/Engine/Plugins/MovieScene/MovieRenderPipeline/Source/MovieRenderPipelineCore/Public/Graph/MovieGraphDataTypes.h b/Engine/Plugins/MovieScene/MovieRenderPipeline/Source/MovieRenderPipelineCore/Public/Graph/MovieGraphDataTypes.h
--- a/Engine/Plugins/MovieScene/MovieRenderPipeline/Source/MovieRenderPipelineCore/Public/Graph/MovieGraphDataTypes.h
+++ b/Engine/Plugins/MovieScene/MovieRenderPipeline/Source/MovieRenderPipelineCore/Public/Graph/MovieGraphDataTypes.h
@@ -8,6 +8,7 @@
 #include "Templates/SubclassOf.h"
 #include "ImagePixelData.h"
 #include "Containers/Queue.h"
+#include "Containers/Set.h"
 #include "Misc/FrameRate.h"
 #include "MovieRenderPipelineDataTypes.h"
 #include "Engine/EngineCustomTimeStep.h"
@@ -451,6 +452,53 @@ namespace UE::MovieGraph
 
 		/** Additional metadata to be added to the output (if supported by the output container). */
 		TMap<FString, FString> FileMetadata;
+
+		/** Check if the (expected) render passes come from different branches. */
+		bool HasDataFromMultipleBranches() const
+		{
+			TSet<FName> BranchUseCounts;
+			BranchUseCounts.Reserve(ExpectedRenderPasses.Num());
+
+			for (const FMovieGraphRenderDataIdentifier& PassIdentifier : ExpectedRenderPasses)
+			{
+				BranchUseCounts.Add(PassIdentifier.RootBranchName);
+			}
+
+			return BranchUseCounts.Num() > 1;
+		}
+
+		/** Check if the (expected) render passes come from different renderers, excluding composited ones. */
+		bool HasMultipleRendersPerBranch(const FName& InBranchName)
+		{
+			TSet<uint32> BranchRendererUseCounts;
+			BranchRendererUseCounts.Reserve(ExpectedRenderPasses.Num());
+
+			for (const FMovieGraphRenderDataIdentifier& PassIdentifier : ExpectedRenderPasses)
+			{
+				if (PassIdentifier.RootBranchName == InBranchName)
+				{
+					BranchRendererUseCounts.Add(GetTypeHash(PassIdentifier.RendererName));
+				}
+			}
+
+			// Remove any renderers that will be composited on later
+			for (const TPair<FMovieGraphRenderDataIdentifier, TUniquePtr<FImagePixelData>>& RenderData : ImageOutputData)
+			{
+				if (RenderData.Key.RootBranchName != InBranchName)
+				{
+					continue;
+				}
+
+				UE::MovieGraph::FMovieGraphSampleState* Payload = RenderData.Value->GetPayload<UE::MovieGraph::FMovieGraphSampleState>();
+				check(Payload);
+				if (Payload->bCompositeOnOtherRenders)
+				{
+					BranchRendererUseCounts.Remove(GetTypeHash(RenderData.Key.RendererName));
+				}
+			}
+
+			return BranchRendererUseCounts.Num() > 1;
+		}
 	};
 
 	/**
diff --git a/Engine/Plugins/MovieScene/MovieRenderPipeline/Source/MovieRenderPipelineCore/Public/MoviePipelineUtils.h b/Engine/Plugins/MovieScene/MovieRenderPipeline/Source/MovieRenderPipelineCore/Public/MoviePipelineUtils.h
--- a/Engine/Plugins/MovieScene/MovieRenderPipeline/Source/MovieRenderPipelineCore/Public/MoviePipelineUtils.h
+++ b/Engine/Plugins/MovieScene/MovieRenderPipeline/Source/MovieRenderPipelineCore/Public/MoviePipelineUtils.h
@@ -131,6 +131,7 @@ namespace UE
 
 	namespace MoviePipeline
 	{
+		MOVIERENDERPIPELINECORE_API void ConformOutputFormatStringToken(FString& InOutFilenameFormatString, const FStringView InToken, const FName& InNodeName, const FName& InBranchName);
 		MOVIERENDERPIPELINECORE_API void ValidateOutputFormatString(FString& InOutFilenameFormatString, const bool bTestRenderPass, const bool bTestFrameNumber, const bool bIncludeCameraName = false);
 		MOVIERENDERPIPELINECORE_API void RemoveFrameNumberFormatStrings(FString& InOutFilenameFormatString, const bool bIncludeShots);
 		/** De-duplicates the provided array of strings by appending (1), (2), etc to the end of duplicates. */
diff --git a/Engine/Plugins/MovieScene/MovieRenderPipeline/Source/MovieRenderPipelineRenderPasses/Private/MovieGraphImageSequenceOutputNode.cpp b/Engine/Plugins/MovieScene/MovieRenderPipeline/Source/MovieRenderPipelineRenderPasses/Private/MovieGraphImageSequenceOutputNode.cpp
--- a/Engine/Plugins/MovieScene/MovieRenderPipeline/Source/MovieRenderPipelineRenderPasses/Private/MovieGraphImageSequenceOutputNode.cpp
+++ b/Engine/Plugins/MovieScene/MovieRenderPipeline/Source/MovieRenderPipelineRenderPasses/Private/MovieGraphImageSequenceOutputNode.cpp
@@ -232,10 +232,10 @@ TArray<TPair<FMovieGraphRenderDataIdentifier, TUniquePtr<FImagePixelData>>> UMov
 
 FString UMovieGraphImageSequenceOutputNode::CreateFileName(
 	UE::MovieGraph::FMovieGraphOutputMergerFrame* InRawFrameData,
+	const UMovieGraphImageSequenceOutputNode* InParentNode,
 	const UMovieGraphPipeline* InPipeline,
 	const TPair<FMovieGraphRenderDataIdentifier, TUniquePtr<FImagePixelData>>& InRenderData,
 	const EImageFormat InImageFormat,
-	const FString& InFileNameFormat,
 	FMovieGraphResolveArgs& OutMergedFormatArgs) const
 {
 	const TCHAR* Extension = TEXT("");
@@ -254,8 +254,26 @@ FString UMovieGraphImageSequenceOutputNode::CreateFileName(
 	}
 
 	// Generate one string that puts the directory combined with the filename format.
-	FString FileNameFormatString = OutputSettingNode->OutputDirectory.Path / InFileNameFormat;
+	FString FileNameFormatString = OutputSettingNode->OutputDirectory.Path / InParentNode->FileNameFormat;
+
+	if (InRawFrameData->HasDataFromMultipleBranches())
+	{
+		// There should be one render layer per branch, so with multiple branches
+		// we test for {layer_name} presence.
+		// 
+		// ToDo: This is overly protective and could be relaxed later, for instance
+		// if different file write nodes have chosen a separate filepath entirely.
+		UE::MoviePipeline::ConformOutputFormatStringToken(FileNameFormatString, TEXT("{layer_name}"), InParentNode->GetFName(), InRenderData.Key.RootBranchName);
+	}
+
+	if (InRawFrameData->HasMultipleRendersPerBranch(InRenderData.Key.RootBranchName))
+	{
+		UE::MoviePipeline::ConformOutputFormatStringToken(FileNameFormatString, TEXT("{renderer_name}"), InParentNode->GetFName(), InRenderData.Key.RootBranchName);
+	}
+
+	// ToDo: Add {camera_name} validation once relevant
 
+	// Previous method is preserved for output frame number validation.
 	constexpr bool bIncludeRenderPass = false;
 	constexpr bool bTestFrameNumber = true;
 	UE::MoviePipeline::ValidateOutputFormatString(FileNameFormatString, bIncludeRenderPass, bTestFrameNumber);
@@ -326,7 +344,7 @@ void UMovieGraphImageSequenceOutputNode::OnReceiveImageDataImpl(UMovieGraphPipel
 		checkf(ParentNode, TEXT("Image sequence output should not exist without a parent node in the graph."));
 		
 		FMovieGraphResolveArgs FinalResolvedKVPs;
-		FString FileName = CreateFileName(InRawFrameData, InPipeline, RenderData, PreferredOutputFormat, ParentNode->FileNameFormat, FinalResolvedKVPs);
+		FString FileName = CreateFileName(InRawFrameData, ParentNode, InPipeline, RenderData, PreferredOutputFormat, FinalResolvedKVPs);
 		if (!ensureMsgf(!FileName.IsEmpty(), TEXT("Unexpected empty file name, skipping frame.")))
 		{
 			continue;
@@ -544,7 +562,7 @@ void UMovieGraphImageSequenceOutputNode_EXR::OnReceiveImageDataImpl(UMovieGraphP
 		checkf(ParentNode, TEXT("Single-layer EXR should not exist without a parent node in the graph."));
 
 		FMovieGraphResolveArgs ResolvedFormatArgs;
-		FString FileName = CreateFileName(InRawFrameData, InPipeline, RenderData, OutputFormat, ParentNode->FileNameFormat, ResolvedFormatArgs);
+		FString FileName = CreateFileName(InRawFrameData, ParentNode, InPipeline, RenderData, OutputFormat, ResolvedFormatArgs);
 		if (!ensureMsgf(!FileName.IsEmpty(), TEXT("Unexpected empty file name, skipping frame.")))
 		{
 			continue;
diff --git a/Engine/Plugins/MovieScene/MovieRenderPipeline/Source/MovieRenderPipelineRenderPasses/Public/MovieGraphImageSequenceOutputNode.h b/Engine/Plugins/MovieScene/MovieRenderPipeline/Source/MovieRenderPipelineRenderPasses/Public/MovieGraphImageSequenceOutputNode.h
--- a/Engine/Plugins/MovieScene/MovieRenderPipeline/Source/MovieRenderPipelineRenderPasses/Public/MovieGraphImageSequenceOutputNode.h
+++ b/Engine/Plugins/MovieScene/MovieRenderPipeline/Source/MovieRenderPipelineRenderPasses/Public/MovieGraphImageSequenceOutputNode.h
@@ -63,10 +63,10 @@ protected:
 	/** Convenience function to create the output file name. */
 	FString CreateFileName(
 		UE::MovieGraph::FMovieGraphOutputMergerFrame* InRawFrameData,
+		const UMovieGraphImageSequenceOutputNode* InParentNode,
 		const UMovieGraphPipeline* InPipeline,
 		const TPair<FMovieGraphRenderDataIdentifier, TUniquePtr<FImagePixelData>>& InRenderData,
 		const EImageFormat InImageFormat,
-		const FString& InFileNameFormat,
 		FMovieGraphResolveArgs& OutMergedFormatArgs) const;
 
 protected:
