Movie Graph:

-Fix: Render Preview Title Progress Does Not Update

diff --git a/Engine/Plugins/MovieScene/MovieRenderPipeline/Source/MovieRenderPipelineCore/Private/Graph/MovieGraphBlueprintLibrary.cpp b/Engine/Plugins/MovieScene/MovieRenderPipeline/Source/MovieRenderPipelineCore/Private/Graph/MovieGraphBlueprintLibrary.cpp
--- a/Engine/Plugins/MovieScene/MovieRenderPipeline/Source/MovieRenderPipelineCore/Private/Graph/MovieGraphBlueprintLibrary.cpp
+++ b/Engine/Plugins/MovieScene/MovieRenderPipeline/Source/MovieRenderPipelineCore/Private/Graph/MovieGraphBlueprintLibrary.cpp
@@ -370,6 +370,21 @@ FText UMovieGraphBlueprintLibrary::GetJobAuthor(const UMovieGraphPipeline* InMov
 	return InMovieGraphPipeline ? FText::FromString(UE::MoviePipeline::GetJobAuthor(InMovieGraphPipeline->GetCurrentJob())) : FText();
 }
 
+float UMovieGraphBlueprintLibrary::GetCompletionPercentage(const UMovieGraphPipeline* InPipeline)
+{
+	if (!InPipeline)
+	{
+		return 0.f;
+	}
+
+	int32 OutputFrames;
+	int32 TotalOutputFrames;
+	GetOverallOutputFrames(InPipeline, OutputFrames, TotalOutputFrames);
+
+	const float CompletionPercentage = FMath::Clamp(OutputFrames / (float)TotalOutputFrames, 0.f, 1.f);
+	return CompletionPercentage;
+}
+
 void UMovieGraphBlueprintLibrary::GetOverallOutputFrames(const UMovieGraphPipeline* InMovieGraphPipeline, int32& OutCurrentIndex, int32& OutTotalCount)
 {
 	OutCurrentIndex = 0;
diff --git a/Engine/Plugins/MovieScene/MovieRenderPipeline/Source/MovieRenderPipelineCore/Private/MoviePipelineLinearExecutor.cpp b/Engine/Plugins/MovieScene/MovieRenderPipeline/Source/MovieRenderPipelineCore/Private/MoviePipelineLinearExecutor.cpp
--- a/Engine/Plugins/MovieScene/MovieRenderPipeline/Source/MovieRenderPipelineCore/Private/MoviePipelineLinearExecutor.cpp
+++ b/Engine/Plugins/MovieScene/MovieRenderPipeline/Source/MovieRenderPipelineCore/Private/MoviePipelineLinearExecutor.cpp
@@ -1,9 +1,13 @@
 // Copyright Epic Games, Inc. All Rights Reserved.
+
 #include "MoviePipelineLinearExecutor.h"
-#include "MovieRenderPipelineCoreModule.h"
-#include "MoviePipelineQueue.h"
-#include "MoviePipelineBlueprintLibrary.h"
+
+#include "Graph/MovieGraphBlueprintLibrary.h"
+#include "Graph/MovieGraphPipeline.h"
 #include "MoviePipeline.h"
+#include "MoviePipelineBlueprintLibrary.h"
+#include "MoviePipelineQueue.h"
+#include "MovieRenderPipelineCoreModule.h"
 
 #include UE_INLINE_GENERATED_CPP_BY_NAME(MoviePipelineLinearExecutor)
 
@@ -108,18 +112,28 @@ FText UMoviePipelineLinearExecutorBase::GetWindowTitle()
 	PercentFormatOptions.MaximumFractionalDigits = 0;
 	PercentFormatOptions.RoundingMode = ERoundingMode::HalfFromZero;
 
+	const float CompletionPercentage = GetCompletionPercentageFromActivePipeline();
+
+	const FText TitleFormatString = LOCTEXT("MoviePreviewWindowTitleFormat", "Movie Pipeline Render (Preview) [Job {CurrentCount}/{TotalCount} Total] Current Job: {PercentComplete}% Completed.");
+	return FText::FormatNamed(TitleFormatString, TEXT("CurrentCount"), FText::AsNumber(CurrentPipelineIndex + 1), TEXT("TotalCount"), FText::AsNumber(Queue->GetJobs().Num()), TEXT("PercentComplete"), FText::AsNumber(CompletionPercentage, &PercentFormatOptions));
+}
+
+float UMoviePipelineLinearExecutorBase::GetCompletionPercentageFromActivePipeline()
+{
 	float CompletionPercentage = 0.f;
-	if (ActiveMoviePipeline)
+	if (const UMoviePipeline* AsLegacyPipeline = Cast<UMoviePipeline>(ActiveMoviePipeline))
 	{
-		CompletionPercentage = UMoviePipelineBlueprintLibrary::GetCompletionPercentage(Cast<UMoviePipeline>(ActiveMoviePipeline)) * 100.f;
+		CompletionPercentage = UMoviePipelineBlueprintLibrary::GetCompletionPercentage(AsLegacyPipeline) * 100.f;
+	}
+	else if (const UMovieGraphPipeline* AsGraphPipeline = Cast<UMovieGraphPipeline>(ActiveMoviePipeline))
+	{
+		CompletionPercentage = UMovieGraphBlueprintLibrary::GetCompletionPercentage(AsGraphPipeline) * 100.f;
 	}
 
-	FText TitleFormatString = LOCTEXT("MoviePreviewWindowTitleFormat", "Movie Pipeline Render (Preview) [Job {CurrentCount}/{TotalCount} Total] Current Job: {PercentComplete}% Completed.");
-	return FText::FormatNamed(TitleFormatString, TEXT("CurrentCount"), FText::AsNumber(CurrentPipelineIndex + 1), TEXT("TotalCount"), FText::AsNumber(Queue->GetJobs().Num()), TEXT("PercentComplete"), FText::AsNumber(CompletionPercentage, &PercentFormatOptions));
+	return CompletionPercentage;
 }
 
 
-
 void UMoviePipelineLinearExecutorBase::CancelCurrentJob_Implementation()
 {
 	OnExecutorErroredImpl(nullptr, true, LOCTEXT("CancelingExecutorJob", "Executor canceling current job."));
diff --git a/Engine/Plugins/MovieScene/MovieRenderPipeline/Source/MovieRenderPipelineCore/Public/Graph/MovieGraphBlueprintLibrary.h b/Engine/Plugins/MovieScene/MovieRenderPipeline/Source/MovieRenderPipelineCore/Public/Graph/MovieGraphBlueprintLibrary.h
--- a/Engine/Plugins/MovieScene/MovieRenderPipeline/Source/MovieRenderPipelineCore/Public/Graph/MovieGraphBlueprintLibrary.h
+++ b/Engine/Plugins/MovieScene/MovieRenderPipeline/Source/MovieRenderPipelineCore/Public/Graph/MovieGraphBlueprintLibrary.h
@@ -75,6 +75,10 @@ public:
 	UFUNCTION(BlueprintPure, Category = "Movie Graph")
 	static FText GetJobAuthor(const UMovieGraphPipeline* InMovieGraphPipeline);
 
+	/** Gets the completion percent of the Pipeline in 0-1 */
+	UFUNCTION(BlueprintPure, Category = "Movie Graph")
+	static float GetCompletionPercentage(const UMovieGraphPipeline* InPipeline);
+
 	/**
 	* Determines the overall current frame number and total number of frames.
 	* @param	InMovieGraphPipeline	- The pipeline to get the frame information from.
diff --git a/Engine/Plugins/MovieScene/MovieRenderPipeline/Source/MovieRenderPipelineCore/Public/MoviePipelineLinearExecutor.h b/Engine/Plugins/MovieScene/MovieRenderPipeline/Source/MovieRenderPipelineCore/Public/MoviePipelineLinearExecutor.h
--- a/Engine/Plugins/MovieScene/MovieRenderPipeline/Source/MovieRenderPipelineCore/Public/MoviePipelineLinearExecutor.h
+++ b/Engine/Plugins/MovieScene/MovieRenderPipeline/Source/MovieRenderPipelineCore/Public/MoviePipelineLinearExecutor.h
@@ -37,6 +37,7 @@ protected:
 	virtual void StartPipelineByIndex(int32 InPipelineIndex);
 	virtual void Start(const UMoviePipelineExecutorJob* InJob) {}
 	virtual FText GetWindowTitle();
+	virtual float GetCompletionPercentageFromActivePipeline();
 public:
 	virtual void OnIndividualPipelineFinished(UMoviePipeline* /* FinishedPipeline */);
 	virtual void OnExecutorFinishedImpl();
diff --git a/Engine/Plugins/MovieScene/MovieRenderPipeline/Source/MovieRenderPipelineEditor/Private/MoviePipelinePIEExecutor.cpp b/Engine/Plugins/MovieScene/MovieRenderPipeline/Source/MovieRenderPipelineEditor/Private/MoviePipelinePIEExecutor.cpp
--- a/Engine/Plugins/MovieScene/MovieRenderPipeline/Source/MovieRenderPipelineEditor/Private/MoviePipelinePIEExecutor.cpp
+++ b/Engine/Plugins/MovieScene/MovieRenderPipeline/Source/MovieRenderPipelineEditor/Private/MoviePipelinePIEExecutor.cpp
@@ -301,7 +301,7 @@ void UMoviePipelinePIEExecutor::OnTick()
 	
 	if (RemainingInitializationFrames <= 0)
 	{
-		Queue->GetJobs()[CurrentPipelineIndex]->SetStatusProgress(UMoviePipelineBlueprintLibrary::GetCompletionPercentage(Cast<UMoviePipeline>(ActiveMoviePipeline)));
+		Queue->GetJobs()[CurrentPipelineIndex]->SetStatusProgress(GetCompletionPercentageFromActivePipeline());
 	}
 	
 	FText WindowTitle = GetWindowTitle();
