Movie Graph:

-Fix: Some functions still assume that global output options could still be on branches other than global

diff --git a/Engine/Plugins/MovieScene/MovieRenderPipeline/Source/MovieRenderPipelineCore/Private/Graph/MovieGraphBlueprintLibrary.cpp b/Engine/Plugins/MovieScene/MovieRenderPipeline/Source/MovieRenderPipelineCore/Private/Graph/MovieGraphBlueprintLibrary.cpp
--- a/Engine/Plugins/MovieScene/MovieRenderPipeline/Source/MovieRenderPipelineCore/Private/Graph/MovieGraphBlueprintLibrary.cpp
+++ b/Engine/Plugins/MovieScene/MovieRenderPipeline/Source/MovieRenderPipelineCore/Private/Graph/MovieGraphBlueprintLibrary.cpp
@@ -330,7 +330,7 @@ int32 UMovieGraphBlueprintLibrary::ResolveVersionNumber(FMovieGraphFilenameResol
 	return HighestVersion + (bGetNextVersion ? 1 : 0);
 }
 
-FIntPoint UMovieGraphBlueprintLibrary::GetEffectiveOutputResolution(UMovieGraphEvaluatedConfig* InEvaluatedGraph, const FName& InBranchName)
+FIntPoint UMovieGraphBlueprintLibrary::GetEffectiveOutputResolution(UMovieGraphEvaluatedConfig* InEvaluatedGraph)
 {
 	if (!InEvaluatedGraph)
 	{
@@ -339,7 +339,7 @@ FIntPoint UMovieGraphBlueprintLibrary::GetEffectiveOutputResolution(UMovieGraphE
 	}
 	
 	constexpr bool bIncludeCDOs = true;
-	const UMovieGraphGlobalOutputSettingNode* OutputSetting = InEvaluatedGraph->GetSettingForBranch<UMovieGraphGlobalOutputSettingNode>(InBranchName, bIncludeCDOs);
+	const UMovieGraphGlobalOutputSettingNode* OutputSetting = InEvaluatedGraph->GetSettingForBranch<UMovieGraphGlobalOutputSettingNode>(UMovieGraphNode::GlobalsPinName, bIncludeCDOs);
 	if (!ensure(OutputSetting))
 	{
 		return FIntPoint();
diff --git a/Engine/Plugins/MovieScene/MovieRenderPipeline/Source/MovieRenderPipelineCore/Private/Graph/Nodes/MovieGraphWidgetRendererBaseNode.cpp b/Engine/Plugins/MovieScene/MovieRenderPipeline/Source/MovieRenderPipelineCore/Private/Graph/Nodes/MovieGraphWidgetRendererBaseNode.cpp
--- a/Engine/Plugins/MovieScene/MovieRenderPipeline/Source/MovieRenderPipelineCore/Private/Graph/Nodes/MovieGraphWidgetRendererBaseNode.cpp
+++ b/Engine/Plugins/MovieScene/MovieRenderPipeline/Source/MovieRenderPipelineCore/Private/Graph/Nodes/MovieGraphWidgetRendererBaseNode.cpp
@@ -93,7 +93,7 @@ void UMovieGraphWidgetRendererBaseNode::FMovieGraphWidgetPass::Render(const FMov
 {
 	const UMovieGraphPipeline* Pipeline = Renderer->GetOwningGraph();
 	
-	const FIntPoint OutputResolution = UMovieGraphBlueprintLibrary::GetEffectiveOutputResolution(InTimeData.EvaluatedConfig, LayerData.BranchName);
+	const FIntPoint OutputResolution = UMovieGraphBlueprintLibrary::GetEffectiveOutputResolution(InTimeData.EvaluatedConfig);
 	const int32 MaxResolution = GetMax2DTextureDimension();
 	if ((OutputResolution.X > MaxResolution) || (OutputResolution.Y > MaxResolution))
 	{
diff --git a/Engine/Plugins/MovieScene/MovieRenderPipeline/Source/MovieRenderPipelineCore/Private/Graph/Renderers/MovieGraphDeferredPass.cpp b/Engine/Plugins/MovieScene/MovieRenderPipeline/Source/MovieRenderPipelineCore/Private/Graph/Renderers/MovieGraphDeferredPass.cpp
--- a/Engine/Plugins/MovieScene/MovieRenderPipeline/Source/MovieRenderPipelineCore/Private/Graph/Renderers/MovieGraphDeferredPass.cpp
+++ b/Engine/Plugins/MovieScene/MovieRenderPipeline/Source/MovieRenderPipelineCore/Private/Graph/Renderers/MovieGraphDeferredPass.cpp
@@ -119,7 +119,7 @@ void FMovieGraphDeferredPass::Render(const FMovieGraphTraversalContext& InFrameT
 		OverscanFraction = FMath::Clamp(CameraNode->OverscanPercentage / 100.f, 0.f, 1.f);
 	}
 	
-	FIntPoint AccumulatorResolution = UMovieGraphBlueprintLibrary::GetEffectiveOutputResolution(InTimeData.EvaluatedConfig, LayerData.BranchName);
+	FIntPoint AccumulatorResolution = UMovieGraphBlueprintLibrary::GetEffectiveOutputResolution(InTimeData.EvaluatedConfig);
 	// ToDo: When tiling is used, this should be the size of the per-tile backbuffer
 	FIntPoint BackbufferResolution = AccumulatorResolution;
 	// ToDo: This math probably needs the per-tile, pre-overlapped size? 
diff --git a/Engine/Plugins/MovieScene/MovieRenderPipeline/Source/MovieRenderPipelineCore/Public/Graph/MovieGraphBlueprintLibrary.h b/Engine/Plugins/MovieScene/MovieRenderPipeline/Source/MovieRenderPipelineCore/Public/Graph/MovieGraphBlueprintLibrary.h
--- a/Engine/Plugins/MovieScene/MovieRenderPipeline/Source/MovieRenderPipelineCore/Public/Graph/MovieGraphBlueprintLibrary.h
+++ b/Engine/Plugins/MovieScene/MovieRenderPipeline/Source/MovieRenderPipelineCore/Public/Graph/MovieGraphBlueprintLibrary.h
@@ -55,11 +55,10 @@ public:
 	/**
 	* In case of overscan percentage being higher than 0, additional pixels are rendered. This function returns the resolution with overscan taken into account.
 	* @param	InEvaluatedGraph	- The evaluated graph that will provide context for resolving the resolution
-	* @param	InBranchName		- The graph branch that the output resolution should be resolved on
 	* @return						- The output resolution, taking into account overscan
 	*/
 	UFUNCTION(BlueprintCallable, Category = "Movie Graph")
-	static FIntPoint GetEffectiveOutputResolution(UMovieGraphEvaluatedConfig* InEvaluatedGraph, const FName& InBranchName);
+	static FIntPoint GetEffectiveOutputResolution(UMovieGraphEvaluatedConfig* InEvaluatedGraph);
 
 	/**
 	* Gets the name of the current job.
