Movie Graph:

-Rename render_layer token to layer_name (found a straggler)

diff --git a/Engine/Plugins/MovieScene/MovieRenderPipeline/Source/MovieRenderPipelineRenderPasses/Public/MovieGraphImageSequenceOutputNode.h b/Engine/Plugins/MovieScene/MovieRenderPipeline/Source/MovieRenderPipelineRenderPasses/Public/MovieGraphImageSequenceOutputNode.h
--- a/Engine/Plugins/MovieScene/MovieRenderPipeline/Source/MovieRenderPipelineRenderPasses/Public/MovieGraphImageSequenceOutputNode.h
+++ b/Engine/Plugins/MovieScene/MovieRenderPipeline/Source/MovieRenderPipelineRenderPasses/Public/MovieGraphImageSequenceOutputNode.h
@@ -177,7 +177,7 @@ public:
 	UMovieGraphImageSequenceOutputNode_MultiLayerEXR()
 		: UMovieGraphImageSequenceOutputNode_EXR()
 	{
-		// Multi-layer default excludes {render_layer}.
+		// Multi-layer default excludes {layer_name}.
 		FileNameFormat = TEXT("{sequence_name}.{frame_number}");
 	}
 
