Movie Graph:

-Menu to Add Selected Actors from Outliner into Collection Conditions

diff --git a/Engine/Plugins/MovieScene/MovieRenderPipeline/Source/MovieRenderPipelineCore/Private/Graph/MoviePipelineRenderLayerSubsystem.cpp b/Engine/Plugins/MovieScene/MovieRenderPipeline/Source/MovieRenderPipelineCore/Private/Graph/MoviePipelineRenderLayerSubsystem.cpp
--- a/Engine/Plugins/MovieScene/MovieRenderPipeline/Source/MovieRenderPipelineCore/Private/Graph/MoviePipelineRenderLayerSubsystem.cpp
+++ b/Engine/Plugins/MovieScene/MovieRenderPipeline/Source/MovieRenderPipelineCore/Private/Graph/MoviePipelineRenderLayerSubsystem.cpp
@@ -36,6 +36,7 @@
 #include "SceneOutlinerModule.h"
 #include "SceneOutlinerPublicTypes.h"
 #include "SClassViewer.h"
+#include "Selection.h"
 #endif
 
 #define LOCTEXT_NAMESPACE "MovieGraph"
@@ -678,29 +679,84 @@ TSharedRef<SWidget> UMovieGraphConditionGroupQuery_Actor::GetAddMenuContents(con
 	
 	const FSceneOutlinerModule& SceneOutlinerModule = FModuleManager::LoadModuleChecked<FSceneOutlinerModule>("SceneOutliner");
 
-	ActorPickerWidget = SceneOutlinerModule.CreateActorPicker(
-		SceneOutlinerInitOptions,
-		FOnActorPicked::CreateLambda([this, OnAddFinished](AActor* InActor)
+	FMenuBuilder MenuBuilder(false, MakeShared<FUICommandList>());
+
+	auto AddActorToList = [this](const AActor* InActor, const FMovieGraphConditionGroupQueryContentsChanged& OnAddFinished)
+	{
+		ActorsToMatch.Add(InActor);
+		ListDataSource.Add(MakeShared<TSoftObjectPtr<AActor>>(ActorsToMatch.Last()));
+		OnAddFinished.ExecuteIfBound();
+	};
+	
+	auto RefreshActorPickerFilterAndList = [this]()
+	{
+		// Ensure that the actor picker filter runs again so duplicate actors cannot be selected
+		if (ActorPickerWidget.IsValid())
 		{
-			ActorsToMatch.Add(InActor);
-			ListDataSource.Add(MakeShared<TSoftObjectPtr<AActor>>(ActorsToMatch.Last()));
-			OnAddFinished.ExecuteIfBound();
+			ActorPickerWidget->FullRefresh();
+			ActorsList->Refresh();
+		}
+	};
+	
+	MenuBuilder.BeginSection("AddActor", LOCTEXT("AddActor", "Add Actor"));
+	{
+		MenuBuilder.AddMenuEntry(
+			LOCTEXT("AddSelectedInOutliner", "Add Selected In Outliner"),
+			LOCTEXT("AddSelectedInOutlinerTooltip", "Add actors currently selected in the level editor's scene outliner."),
+			FSlateIcon(FAppStyle::GetAppStyleSetName(),"FoliageEditMode.SetSelect"),
+			FUIAction(
+				FExecuteAction::CreateLambda([this, OnAddFinished, RefreshActorPickerFilterAndList, AddActorToList]()
+				{
+					USelection* CurrentSelection = GEditor->GetSelectedActors();
+					TArray<AActor*> SelectedActors;
+					CurrentSelection->GetSelectedObjects<AActor>(SelectedActors);
+					for (const AActor* Actor : SelectedActors)
+					{
+						if (Actor && !ActorsToMatch.Contains(Actor))
+						{
+							AddActorToList(Actor, OnAddFinished);
+						}
+					}
+
+					RefreshActorPickerFilterAndList();
+
+					FSlateApplication::Get().DismissAllMenus();
+				}
+				),
+				FCanExecuteAction::CreateLambda([]()
+				{
+					return GEditor->GetSelectedActors()->Num() > 0;
+					
+				})
+			)
+		);
+	}
+	MenuBuilder.EndSection();
 
-			// Ensure that the filter runs again so duplicate actors cannot be selected
-			if (ActorPickerWidget.IsValid())
+	MenuBuilder.BeginSection("Browse", LOCTEXT("Browse", "Browse"));
+	{
+		ActorPickerWidget = SceneOutlinerModule.CreateActorPicker(
+			SceneOutlinerInitOptions,
+			FOnActorPicked::CreateLambda([this, OnAddFinished, RefreshActorPickerFilterAndList, AddActorToList](AActor* InActor)
 			{
-				ActorPickerWidget->FullRefresh();
-				ActorsList->Refresh();
-			}
-		}));
-	
-	return
-		SNew(SBox)
-		.WidthOverride(400.f)
-		.HeightOverride(300.f)
-		[
-			ActorPickerWidget.ToSharedRef()
-		];
+				AddActorToList(InActor, OnAddFinished);
+
+				RefreshActorPickerFilterAndList();
+			}));
+
+		const TSharedRef<SBox> ActorPickerWidgetBox =
+			SNew(SBox)
+			.WidthOverride(400.f)
+			.HeightOverride(300.f)
+			[
+				ActorPickerWidget.ToSharedRef()
+			];
+
+		MenuBuilder.AddWidget(ActorPickerWidgetBox, FText());
+	}
+	MenuBuilder.EndSection();
+
+	return MenuBuilder.MakeWidget();
 }
 
 const FSlateBrush* UMovieGraphConditionGroupQuery_Actor::GetRowIcon(TSharedPtr<TSoftObjectPtr<AActor>> InActor)
