Movie Graph:

-Fix: Comment Missing From Right-Click Graph Menu

diff --git a/Engine/Plugins/MovieScene/MovieRenderPipeline/Source/MovieRenderPipelineEditor/Private/Graph/MovieGraphSchema.cpp b/Engine/Plugins/MovieScene/MovieRenderPipeline/Source/MovieRenderPipelineEditor/Private/Graph/MovieGraphSchema.cpp
--- a/Engine/Plugins/MovieScene/MovieRenderPipeline/Source/MovieRenderPipelineEditor/Private/Graph/MovieGraphSchema.cpp
+++ b/Engine/Plugins/MovieScene/MovieRenderPipeline/Source/MovieRenderPipelineEditor/Private/Graph/MovieGraphSchema.cpp
@@ -142,6 +142,22 @@ bool UMovieGraphSchema::IsConnectionToBranchAllowed(const UEdGraphPin* InputPin,
 	return FromPin->IsConnectionToBranchAllowed(ToPin, OutError);
 }
 
+void UMovieGraphSchema::AddExtraMenuActions(FGraphActionMenuBuilder& ActionMenuBuilder) const
+{
+	// Comment action
+	ActionMenuBuilder.AddAction(CreateCommentMenuAction());
+}
+
+TSharedRef<FMovieGraphSchemaAction_NewComment> UMovieGraphSchema::CreateCommentMenuAction() const
+{
+	const FText CommentMenuDesc = LOCTEXT("AddComment", "Add Comment");
+	const FText CommentCategory;
+	const FText CommentDescription = LOCTEXT("AddCommentTooltip", "Create a resizable comment box.");
+
+	const TSharedRef<FMovieGraphSchemaAction_NewComment> NewCommentAction = MakeShared<FMovieGraphSchemaAction_NewComment>(CommentCategory, CommentMenuDesc, CommentDescription, 0);
+	return NewCommentAction;
+}
+
 void UMovieGraphSchema::GetGraphContextActions(FGraphContextMenuBuilder& ContextMenuBuilder) const
 {
 	InitMoviePipelineNodeClasses();
@@ -204,6 +220,8 @@ void UMovieGraphSchema::GetGraphContextActions(FGraphContextMenuBuilder& Context
 		
 		ContextMenuBuilder.AddAction(NewAction);
 	}
+
+	AddExtraMenuActions(ContextMenuBuilder);
 }
 
 const FPinConnectionResponse UMovieGraphSchema::CanCreateConnection(const UEdGraphPin* PinA, const UEdGraphPin* PinB) const
diff --git a/Engine/Plugins/MovieScene/MovieRenderPipeline/Source/MovieRenderPipelineEditor/Private/Graph/MovieGraphSchema.h b/Engine/Plugins/MovieScene/MovieRenderPipeline/Source/MovieRenderPipelineEditor/Private/Graph/MovieGraphSchema.h
--- a/Engine/Plugins/MovieScene/MovieRenderPipeline/Source/MovieRenderPipelineEditor/Private/Graph/MovieGraphSchema.h
+++ b/Engine/Plugins/MovieScene/MovieRenderPipeline/Source/MovieRenderPipelineEditor/Private/Graph/MovieGraphSchema.h
@@ -45,6 +45,17 @@ private:
 	 */
 	bool IsConnectionToBranchAllowed(const UEdGraphPin* InputPin, const UEdGraphPin* OutputPin, FText& OutError) const;
 
+	/**
+	 * Adds extra menu actions to the context/palette menu.
+	 */
+	void AddExtraMenuActions(FGraphActionMenuBuilder& ActionMenuBuilder) const;
+
+	/**
+	 * Returns a menu action for creating a new comment in the graph.
+	 * Note that this is not the same as adding a new comment to the graph via hotkey.
+	 */
+	TSharedRef<FMovieGraphSchemaAction_NewComment> CreateCommentMenuAction() const;
+
 public:
 	// Allowed "PinCategory" values for use on EdGraphPin
 	static const FName PC_Branch;
@@ -161,6 +172,9 @@ struct FMovieGraphSchemaAction_NewComment : public FMovieGraphSchemaAction
 		return Type;
 	}
 
+	// Inherit the base class's constructors for context menu/palette
+	using FMovieGraphSchemaAction::FMovieGraphSchemaAction;
+
 	FMovieGraphSchemaAction_NewComment()
 		: FMovieGraphSchemaAction()
 	{}
diff --git a/Engine/Plugins/MovieScene/MovieRenderPipeline/Source/MovieRenderPipelineEditor/Private/Widgets/Graph/SMovieGraphConfigPanel.cpp b/Engine/Plugins/MovieScene/MovieRenderPipeline/Source/MovieRenderPipelineEditor/Private/Widgets/Graph/SMovieGraphConfigPanel.cpp
--- a/Engine/Plugins/MovieScene/MovieRenderPipeline/Source/MovieRenderPipelineEditor/Private/Widgets/Graph/SMovieGraphConfigPanel.cpp
+++ b/Engine/Plugins/MovieScene/MovieRenderPipeline/Source/MovieRenderPipelineEditor/Private/Widgets/Graph/SMovieGraphConfigPanel.cpp
@@ -208,10 +208,6 @@ void SMoviePipelineGraphPanel::MakeEditorCommands()
 			FExecuteAction::CreateSP(this, &SMoviePipelineGraphPanel::OnStraightenConnections)
 		);
 
-		GraphEditorCommands->MapAction(FGraphEditorCommands::Get().CreateComment,
-			FExecuteAction::CreateSP(this, &SMoviePipelineGraphPanel::OnCreateComment)
-		);
-
 		// Distribution Commands
 		GraphEditorCommands->MapAction(FGraphEditorCommands::Get().DistributeNodesHorizontally,
 			FExecuteAction::CreateSP(this, &SMoviePipelineGraphPanel::OnDistributeNodesH)
