Movie Graph:

-Fix Log Spam: Unable to statfs('/render_output_path.{ext}')

diff --git a/Engine/Plugins/MovieScene/MovieRenderPipeline/Source/MovieRenderPipelineCore/Private/Graph/MovieGraphBlueprintLibrary.cpp b/Engine/Plugins/MovieScene/MovieRenderPipeline/Source/MovieRenderPipelineCore/Private/Graph/MovieGraphBlueprintLibrary.cpp
--- a/Engine/Plugins/MovieScene/MovieRenderPipeline/Source/MovieRenderPipelineCore/Private/Graph/MovieGraphBlueprintLibrary.cpp
+++ b/Engine/Plugins/MovieScene/MovieRenderPipeline/Source/MovieRenderPipelineCore/Private/Graph/MovieGraphBlueprintLibrary.cpp
@@ -187,9 +187,11 @@ FString UMovieGraphBlueprintLibrary::ResolveFilenameFormatArguments(const FStrin
 	// after resolving file format strings, ie: {sequence_name}.{frame_number} becomes {sequence_name}. for
 	// videos (which can't use frame_numbers).
 	BaseFilename.RemoveFromEnd(TEXT("."));
-
-	FString Extension = FString::Format(TEXT(".{ext}"), NamedArgs);
-	FString ThisTry = BaseFilename + Extension;
+	
+	// If the extension is not resolved, ThisTry will be a path
+	const FString ExtToken = TEXT(".{ext}");
+	FString Extension = FString::Format(*ExtToken, NamedArgs);
+	FString ThisTry = Extension == ExtToken ? BaseFilename : BaseFilename + Extension;
 	
 	// Check that it's a valid path that we can write to.
 	if (UE::MoviePipeline::CanWriteToFile(*ThisTry, bOverwriteExisting))
