Movie graph:

-Fix: Named Resolution changes are not transacted (undo)

diff --git a/Engine/Plugins/MovieScene/MovieRenderPipeline/Source/MovieRenderPipelineEditor/Private/Customizations/Graph/MovieGraphNamedResolutionCustomization.h b/Engine/Plugins/MovieScene/MovieRenderPipeline/Source/MovieRenderPipelineEditor/Private/Customizations/Graph/MovieGraphNamedResolutionCustomization.h
--- a/Engine/Plugins/MovieScene/MovieRenderPipeline/Source/MovieRenderPipelineEditor/Private/Customizations/Graph/MovieGraphNamedResolutionCustomization.h
+++ b/Engine/Plugins/MovieScene/MovieRenderPipeline/Source/MovieRenderPipelineEditor/Private/Customizations/Graph/MovieGraphNamedResolutionCustomization.h
@@ -485,7 +485,7 @@ protected:
 	/**
 	 * Updates the struct which owns this customization with the data from the matched named resolution.
 	 */
-	void UpdateOwningStruct() const
+	void UpdateOwningStruct()
 	{
 		const TSharedPtr<IPropertyHandle> PinnedStructPropertyHandle = StructPropertyHandle.Pin();
 		if (!PinnedStructPropertyHandle)
@@ -506,17 +506,17 @@ protected:
 			// Copy the data from our matched NamedResolution to the owning struct
 			FMemory::Memcpy(StructData, NamedResolution, sizeof(FMovieGraphNamedResolution));
 
-			// Mark owning packages dirty
+			// Transact on outer objects
 			TArray<UObject*> OutObjects;
 			PinnedStructPropertyHandle->GetOuterObjects(OutObjects);
 
-			for (const UObject* Object : OutObjects)
+			for (UObject* Object : OutObjects)
 			{
-				if (!Object)
+				if (Object)
 				{
-					continue;
+					constexpr bool bAlwaysMarkDirty = true;	
+					Object->Modify(bAlwaysMarkDirty);
 				}
-				ensureAlwaysMsgf(Object->MarkPackageDirty(), TEXT("%hs: Failed to mark package of %s dirty!"), __FUNCTION__, *Object->GetName());
 			}
 		}
 		else
