Movie Graph:

-Fix: Format Strings Autocomplete Does Not Work

diff --git a/Engine/Plugins/MovieScene/MovieRenderPipeline/Source/MovieRenderPipelineEditor/Private/Customizations/Graph/MovieGraphFormatTokenCustomization.cpp b/Engine/Plugins/MovieScene/MovieRenderPipeline/Source/MovieRenderPipelineEditor/Private/Customizations/Graph/MovieGraphFormatTokenCustomization.cpp
new file mode 100644
--- /dev/null
+++ b/Engine/Plugins/MovieScene/MovieRenderPipeline/Source/MovieRenderPipelineEditor/Private/Customizations/Graph/MovieGraphFormatTokenCustomization.cpp
@@ -0,0 +1,91 @@
+﻿// Copyright Epic Games, Inc. All Rights Reserved.
+
+#include "MovieGraphFormatTokenCustomization.h"
+
+#include "Graph/Nodes/MovieGraphFileOutputNode.h"
+#include "MoviePipelineUtils.h"
+
+#include "DetailLayoutBuilder.h"
+#include "DetailWidgetRow.h"
+#include "Layout/WidgetPath.h"
+#include "Widgets/SMoviePipelineFormatTokenAutoCompleteBox.h"
+
+void FMovieGraphFormatTokenCustomization::CustomizeDetails(IDetailLayoutBuilder& DetailBuilder)
+{
+	CustomizedObject = DetailBuilder.GetSelectedObjects()[0];
+	check(CustomizedObject.IsValid() && !CustomizedObject.IsStale());
+
+	// This should work just fine for UMovieGraphFileOutputNode and UMovieGraphCommandLineEncoderNode as long as the property names stay in sync
+	OutputFormatPropertyHandle = DetailBuilder.GetProperty(GET_MEMBER_NAME_CHECKED(UMovieGraphFileOutputNode, FileNameFormat));
+	const bool bShowChildren = false;
+	FText InitialText;
+	OutputFormatPropertyHandle->GetValueAsDisplayText(InitialText);
+
+	DetailBuilder.EditDefaultProperty(OutputFormatPropertyHandle)->CustomWidget(bShowChildren)
+		.NameContent()
+		[
+			OutputFormatPropertyHandle->CreatePropertyNameWidget()
+		]
+		.ValueContent()
+		.MinDesiredWidth(200.0f)
+		[
+			SNew(SMoviePipelineFormatTokenAutoCompleteBox)
+			.Text(this, &FMovieGraphFormatTokenCustomization::GetText)
+			.Suggestions(this, &FMovieGraphFormatTokenCustomization::GetSuggestions)
+			.OnTextChanged(this, &FMovieGraphFormatTokenCustomization::OnTextChanged)
+		];
+}
+
+TArray<FString> FMovieGraphFormatTokenCustomization::GetSuggestions() const
+{
+	TArray<FString> Suggestions;
+	FMoviePipelineFormatArgs FormatArgs;
+	GetFormatArguments(FormatArgs);
+
+	for (const TPair<FString, FString>& KVP : FormatArgs.FilenameArguments)
+	{
+		Suggestions.Add(KVP.Key);
+	}
+
+	return Suggestions;
+}
+
+void FMovieGraphFormatTokenCustomization::OnTextChanged(const FText& InValue) const
+{
+	OutputFormatPropertyHandle->SetValue(InValue.ToString());
+}
+
+FText FMovieGraphFormatTokenCustomization::GetText() const
+{
+	FText DisplayText;
+	OutputFormatPropertyHandle->GetValueAsDisplayText(DisplayText);
+	return DisplayText;
+}
+
+void FMovieGraphFormatTokenCustomization::GetFormatArguments(FMoviePipelineFormatArgs& InOutFormatArgs)
+{
+	static const FString LevelName = TEXT("Level Name");
+	static const FString SequenceName = TEXT("Sequence Name");
+	static const FString JobAuthor = TEXT("Job Author");
+	static const FString JobName = TEXT("Job Name");
+	static const FString JobComment = TEXT("Job Comment");
+	static constexpr double FrameRate = 0.0;
+	static const FString FramePlaceholderNumber = TEXT("0");
+
+	MoviePipeline::GetOutputStateFormatArgs(
+		InOutFormatArgs.FilenameArguments, InOutFormatArgs.FileMetadata,
+		FramePlaceholderNumber, FramePlaceholderNumber,
+		FramePlaceholderNumber, FramePlaceholderNumber,
+		TEXT("CameraName"), TEXT("ShotName"));
+
+	InOutFormatArgs.FilenameArguments.Add(TEXT("level_name"), LevelName);
+	InOutFormatArgs.FilenameArguments.Add(TEXT("sequence_name"), SequenceName);
+	InOutFormatArgs.FilenameArguments.Add(TEXT("job_author"), JobAuthor);
+	InOutFormatArgs.FilenameArguments.Add(TEXT("job_name"), JobName);
+	InOutFormatArgs.FilenameArguments.Add(TEXT("frame_rate"), FString::SanitizeFloat(FrameRate));
+
+	static const FDateTime CurrentTime = FDateTime::Now();
+	static constexpr int32 DummyVersionNumber = 1;
+	UE::MoviePipeline::GetSharedFormatArguments(
+		InOutFormatArgs.FilenameArguments, InOutFormatArgs.FileMetadata, CurrentTime, DummyVersionNumber, InOutFormatArgs.InJob);
+}
diff --git a/Engine/Plugins/MovieScene/MovieRenderPipeline/Source/MovieRenderPipelineEditor/Private/Customizations/Graph/MovieGraphFormatTokenCustomization.h b/Engine/Plugins/MovieScene/MovieRenderPipeline/Source/MovieRenderPipelineEditor/Private/Customizations/Graph/MovieGraphFormatTokenCustomization.h
new file mode 100644
--- /dev/null
+++ b/Engine/Plugins/MovieScene/MovieRenderPipeline/Source/MovieRenderPipelineEditor/Private/Customizations/Graph/MovieGraphFormatTokenCustomization.h
@@ -0,0 +1,33 @@
+﻿#pragma once
+
+#include "MoviePipelineQueue.h"
+#include "IDetailCustomization.h"
+#include "PropertyHandle.h"
+
+class FMovieGraphFormatTokenCustomization : public IDetailCustomization
+{
+public:
+	/** Creates a detail customization instance */
+	static TSharedRef<IDetailCustomization> MakeInstance()
+	{
+		return MakeShared<FMovieGraphFormatTokenCustomization>();
+	}
+
+	FMovieGraphFormatTokenCustomization(){}
+
+protected:
+	//~ Begin IDetailCustomization interface
+	virtual void CustomizeDetails(IDetailLayoutBuilder& DetailBuilder) override;
+	//~ End IDetailCustomization interface
+
+	TArray<FString> GetSuggestions() const;
+
+	void OnTextChanged(const FText& InValue) const;
+
+	FText GetText() const;
+
+	static void GetFormatArguments(FMoviePipelineFormatArgs& InOutFormatArgs);
+	
+	TWeakObjectPtr<UObject> CustomizedObject;
+	TSharedPtr<IPropertyHandle> OutputFormatPropertyHandle;
+};
diff --git a/Engine/Plugins/MovieScene/MovieRenderPipeline/Source/MovieRenderPipelineEditor/Private/Graph/MovieGraphAssetToolkit.cpp b/Engine/Plugins/MovieScene/MovieRenderPipeline/Source/MovieRenderPipelineEditor/Private/Graph/MovieGraphAssetToolkit.cpp
--- a/Engine/Plugins/MovieScene/MovieRenderPipeline/Source/MovieRenderPipelineEditor/Private/Graph/MovieGraphAssetToolkit.cpp
+++ b/Engine/Plugins/MovieScene/MovieRenderPipeline/Source/MovieRenderPipelineEditor/Private/Graph/MovieGraphAssetToolkit.cpp
@@ -3,6 +3,7 @@
 #include "MovieGraphAssetToolkit.h"
 
 #include "Customizations/Graph/MovieGraphCollectionsCustomization.h"
+#include "Customizations/Graph/MovieGraphFormatTokenCustomization.h"
 #include "Customizations/Graph/MovieGraphMemberCustomization.h"
 #include "Customizations/Graph/MovieGraphModifiersCustomization.h"
 #include "Customizations/Graph/MovieGraphNamedResolutionCustomization.h"
@@ -14,6 +15,8 @@
 
 #include "Graph/MovieGraphConfig.h"
 #include "Graph/Nodes/MovieGraphCollectionNode.h"
+#include "Graph/Nodes/MovieGraphCommandLineEncoderNode.h"
+#include "Graph/Nodes/MovieGraphFileOutputNode.h"
 #include "Graph/Nodes/MovieGraphModifierNode.h"
 #include "Graph/Nodes/MovieGraphSelectNode.h"
 #include "MovieEdGraphNode.h"
@@ -400,6 +403,14 @@ TSharedRef<SDockTab> FMovieGraphAssetToolkit::SpawnTab_RenderGraphDetails(const
 		UMovieGraphModifierNode::StaticClass(),
 		FOnGetDetailCustomizationInstance::CreateStatic(&FMovieGraphModifiersCustomization::MakeInstance));
 	
+	SelectedGraphObjectsDetailsWidget->RegisterInstancedCustomPropertyLayout(
+		UMovieGraphFileOutputNode::StaticClass(),
+		FOnGetDetailCustomizationInstance::CreateStatic(&FMovieGraphFormatTokenCustomization::MakeInstance));
+
+	SelectedGraphObjectsDetailsWidget->RegisterInstancedCustomPropertyLayout(
+		UMovieGraphCommandLineEncoderNode::StaticClass(),
+		FOnGetDetailCustomizationInstance::CreateStatic(&FMovieGraphFormatTokenCustomization::MakeInstance));
+	
 	TSharedRef<SWidget> CustomContent = SAssignNew(NameAreaCustomContent, SHorizontalBox)
 	+ SHorizontalBox::Slot()
 	[
diff --git a/Engine/Plugins/MovieScene/MovieRenderPipeline/Source/MovieRenderPipelineEditor/Private/Widgets/SMoviePipelineConfigEditor.cpp b/Engine/Plugins/MovieScene/MovieRenderPipeline/Source/MovieRenderPipelineEditor/Private/Widgets/SMoviePipelineConfigEditor.cpp
--- a/Engine/Plugins/MovieScene/MovieRenderPipeline/Source/MovieRenderPipelineEditor/Private/Widgets/SMoviePipelineConfigEditor.cpp
+++ b/Engine/Plugins/MovieScene/MovieRenderPipeline/Source/MovieRenderPipelineEditor/Private/Widgets/SMoviePipelineConfigEditor.cpp
@@ -1,34 +1,26 @@
 // Copyright Epic Games, Inc. All Rights Reserved.
 
 #include "Widgets/SMoviePipelineConfigEditor.h"
-#include "Widgets/MoviePipelineWidgetConstants.h"
-#include "Widgets/SMoviePipelineConfigSettings.h"
-#include "MoviePipelineShotConfig.h"
+
+#include "MoviePipelineConsoleVariableSetting.h"
+#include "MoviePipelineOutputSetting.h"
 #include "MoviePipelineSetting.h"
 #include "MoviePipelineQueue.h"
 #include "MoviePipelineUtils.h"
+#include "SMoviePipelineFormatTokenAutoCompleteBox.h"
+#include "Widgets/SMoviePipelineConfigSettings.h"
 
 // Core includes
-#include "UObject/UObjectIterator.h"
+#include "Modules/ModuleManager.h"
 #include "Templates/SubclassOf.h"
-#include "ClassIconFinder.h"
 
 // AssetRegistry includes
-#include "AssetRegistry/AssetRegistryModule.h"
 #include "AssetRegistry/AssetData.h"
  
-// ContentBrowser includes
-#include "IContentBrowserSingleton.h"
-#include "ContentBrowserModule.h"
- 
-// AssetTools includes
-#include "AssetToolsModule.h"
- 
 // Slate includes
 #include "Widgets/SBoxPanel.h"
 #include "Widgets/Text/STextBlock.h"
 #include "Widgets/Text/SMultiLineEditableText.h"
-#include "Widgets/Input/SMultiLineEditableTextBox.h"
 #include "Widgets/Input/SButton.h"
 #include "Widgets/Input/SComboButton.h"
 #include "Widgets/Layout/SSplitter.h"
@@ -39,329 +31,23 @@
 #include "Styling/SlateIconFinder.h"
 #include "SPositiveActionButton.h"
 
-
 // EditorStyle includes
 #include "Styling/AppStyle.h"
 #include "EditorFontGlyphs.h"
 #include "ScopedTransaction.h"
 #include "PropertyEditorModule.h"
 #include "IDetailsView.h"
-#include "MovieRenderPipelineStyle.h"
 #include "FrameNumberDetailsCustomization.h"
 #include "IDetailCustomization.h"
 #include "Editor.h"
 
 #include "DetailLayoutBuilder.h"
 #include "DetailWidgetRow.h"
-#include "IPropertyUtilities.h"
 
-#include "Customizations/ConsoleVariableCustomization.h"
 #include "Customizations/ConsoleVariableSettingCustomization.h"
-#include "MoviePipelineConsoleVariableSetting.h"
-#include "MoviePipelineOutputSetting.h"
 
 #define LOCTEXT_NAMESPACE "SMoviePipelineEditor"
 
-
-
-
-/** This is a specific auto-complete widget made for the MoviePipeline that isn't very flexible. It's
-* similar to SSuggestionTextBox, but SSuggestionTextBox doesn't handle more than one word/suggestions
-* mid string. This widget is hardcoded to look for '{' characters (for {format_tokens}) and then auto
-* completes them from a list and fixes up the {} braces.
-*/
-class SMoviePipelineAutoCompleteTextBox : public SCompoundWidget
-{
-public:
-
-	SLATE_BEGIN_ARGS(SMoviePipelineAutoCompleteTextBox)
-		: _Text()
-		, _Suggestions()
-	{}
-
-	SLATE_ATTRIBUTE(FText, Text)
-	SLATE_ATTRIBUTE(TArray<FString>, Suggestions)
-	/** Called whenever the text is changed programmatically or interactively by the user. */
-	SLATE_EVENT(FOnTextChanged, OnTextChanged)
-	SLATE_END_ARGS()
-
-	void Construct(const FArguments& InArgs)
-	{
-		ChildSlot
-		[
-			SAssignNew(MenuAnchor, SMenuAnchor)
-			.Placement(MenuPlacement_ComboBox)
-			[
-				SAssignNew(TextBox, SMultiLineEditableTextBox)
-				.Text(InArgs._Text)
-				.OnKeyDownHandler(this, &SMoviePipelineAutoCompleteTextBox::OnKeyDown)
-				.OnTextChanged(this, &SMoviePipelineAutoCompleteTextBox::HandleTextBoxTextChanged)
-				.SelectWordOnMouseDoubleClick(true)
-				.AllowMultiLine(false)
-			]
-			.MenuContent
-			(
-				SNew(SBorder)
-				.Padding(FMargin(2))
-				[
-					SAssignNew(VerticalBox, SVerticalBox)
-					+ SVerticalBox::Slot()
-					.AutoHeight()
-					[
-						SAssignNew(SuggestionListView, SListView<TSharedPtr<FString>>)
-						.ItemHeight(18.f)
-						.ListItemsSource(&Suggestions)
-						.SelectionMode(ESelectionMode::Single)
-						.OnGenerateRow(this, &SMoviePipelineAutoCompleteTextBox::HandleSuggestionListViewGenerateRow)
-						.OnSelectionChanged(this, &SMoviePipelineAutoCompleteTextBox::HandleSuggestionListViewSelectionChanged)
-					]
-				]
-			)
-		];
-
-		// We just call it once and cache it for now as the selection code isn't tested against
-		// the amount of suggestions changing.
-		AllSuggestions.Append(InArgs._Suggestions.Get());
-		OnTextChanged = InArgs._OnTextChanged;
-	}
-
-	FReply OnKeyDown(const FGeometry& MyGeometry, const FKeyEvent& KeyEvent)
-	{
-		if (MenuAnchor->IsOpen())
-		{
-			if (KeyEvent.GetKey() == EKeys::Up)
-			{
-				// Because the pop-up dialog is below the text, 'up' actually goes to an earlier item in the list.
-				int32 NewSuggestionIndex = CurrentSuggestionIndex - 1;
-				if (NewSuggestionIndex < 0)
-				{
-					NewSuggestionIndex = Suggestions.Num() - 1;
-				}
-
-				SetActiveSuggestionIndex(NewSuggestionIndex);
-				return FReply::Handled();
-			}
-			else if(KeyEvent.GetKey() == EKeys::Down)
-			{
-				int32 NewSuggestionIndex = CurrentSuggestionIndex + 1;
-				if (NewSuggestionIndex > Suggestions.Num() - 1)
-				{
-					NewSuggestionIndex = 0;
-				}
-
-				SetActiveSuggestionIndex(NewSuggestionIndex);
-				return FReply::Handled();
-			}
-			else if (KeyEvent.GetKey() == EKeys::Tab || KeyEvent.GetKey() == EKeys::Enter)
-			{
-				if(CurrentSuggestionIndex >= 0 && CurrentSuggestionIndex <= Suggestions.Num() - 1)
-				{
-					// Trigger the auto-complete for the highlighted suggestion
-					FString SuggestionText = *Suggestions[CurrentSuggestionIndex];
-					ReplaceRelevantTextWithSuggestion(SuggestionText);
-					return FReply::Handled();
-				}
-			}
-		}
-		return FReply::Unhandled();
-	}
-
-
-	void FindAutoCompleteableTextAtPos(const FString& InWholeString, int32 InCursorPos, FString& OutStr, bool& bShowAutoComplete)
-	{
-		OutStr = FString();
-		bShowAutoComplete = false;
-
-		// We want to find a { brace on or to the left of InCursorPos, but if we find a } we
-		// stop looking, because that's a brace for another text. (+1 for ::FromEnd off by one)
-		int32 StartingBracePos = InWholeString.Find(TEXT("{"), ESearchCase::Type::IgnoreCase, ESearchDir::Type::FromEnd, InCursorPos + 1);
-		int32 PreviousEndBracePos = InWholeString.Find(TEXT("}"), ESearchCase::Type::IgnoreCase, ESearchDir::Type::FromEnd, InCursorPos);
-
-		if (StartingBracePos < PreviousEndBracePos)
-		{
-			return;
-		}
-
-		FString AutoCompleteText;
-
-		// Now that we found a {, take the substring between it and either the next }, or the end of the string.
-		int32 NextEndBracePos = InWholeString.Find(TEXT("}"), ESearchCase::Type::IgnoreCase, ESearchDir::Type::FromStart, InCursorPos);
-		if (StartingBracePos >= 0)
-		{
-			int32 Count = InWholeString.Len() - StartingBracePos;
-			if (NextEndBracePos >= 0)
-			{
-				Count = NextEndBracePos - StartingBracePos;
-			}
-
-			AutoCompleteText = InWholeString.Mid(StartingBracePos + 1, Count - 1);
-		}
-
-		OutStr = AutoCompleteText;
-		bShowAutoComplete = StartingBracePos >= 0 && OutStr.Len() == 0;
-	}
-
-	void ReplaceRelevantTextWithSuggestion(const FString& InSuggestionText)
-	{
-		FString TextBoxText = TextBox->GetText().ToString();
-		int32 CursorPos = TextBoxText.Len();
-		FTextLocation CursorLoc = TextBox->GetCursorLocation();
-		if (CursorLoc.IsValid())
-		{
-			CursorPos = FMath::Clamp(CursorLoc.GetOffset(), 0, CursorPos);
-		}
-		
-		// Look for the { to the left of the cursor. We search StrPositionIndex from +1 here due to a bug in ::FromEnd being off by one.
-		int32 StartingBracePos = TextBoxText.Find(TEXT("{"), ESearchCase::Type::IgnoreCase, ESearchDir::Type::FromEnd, CursorPos + 1);
-
-		// Now that we found a {, take the substring between it and either the next }, or the end of the string.
-		int32 NextEndBracePos = TextBoxText.Find(TEXT("}"), ESearchCase::Type::IgnoreCase, ESearchDir::Type::FromStart, CursorPos);
-		int32 NewCursorPos = 0;
-		if (StartingBracePos >= 0)	
-		{
-			// +1 to keep the left { brace
-			FString Left = TextBoxText.Left(StartingBracePos+1);
-			FString Right;
-
-			if (NextEndBracePos >= 0)
-			{
-				Right = TextBoxText.RightChop(NextEndBracePos);
-			}
-
-			// Since the user chose the suggestion ensure there's already a } brace to close off the pair.
-			if (!Right.StartsWith(TEXT("}")))
-			{
-				Right = FString::Printf(TEXT("}%s"), *Right);
-			}
-
-			TextBoxText = Left + InSuggestionText + Right;
-			// We subtract 1 from the Right as we want to put the cursor after the automatically generated "}" token.
-			NewCursorPos = TextBoxText.Len() - (Right.Len() - 1);
-		}
-
-		TextBox->SetText(FText::FromString(TextBoxText));
-		TextBox->GoTo(FTextLocation(0, NewCursorPos));
-	}
-
-	void HandleTextBoxTextChanged(const FText& InText)
-	{
-		OnTextChanged.ExecuteIfBound(InText);
-
-		FString TextAsStr = InText.ToString();
-		if (TextAsStr.Len() > 0)
-		{
-			FString OutStr;
-			bool bShowAutoComplete;
-
-			int32 CursorPos = TextAsStr.Len();
-			FTextLocation CursorLoc = TextBox->GetCursorLocation();
-			if (CursorLoc.IsValid())
-			{
-				CursorPos = FMath::Clamp(CursorLoc.GetOffset(), 0, CursorPos);
-			}
-
-			FindAutoCompleteableTextAtPos(TextAsStr, CursorPos, OutStr, bShowAutoComplete);
-			FilterVisibleSuggestions(OutStr, bShowAutoComplete);
-		}	
-		else
-		{
-			// If they have no text, suggest all possible solutions
-			FilterVisibleSuggestions(FString(), false);
-		}
-	}
-
-	void FilterVisibleSuggestions(const FString& StrToMatch, const bool bForceShowAll)
-	{
-		Suggestions.Reset();
-		for (const FString& Suggestion : AllSuggestions)
-		{
-			if (Suggestion.Contains(StrToMatch) || bForceShowAll)
-			{
-				Suggestions.Add(MakeShared<FString>(Suggestion));
-			}
-		}
-
-		if (Suggestions.Num() > 0)
-		{
-			// We don't focus the menu (because then you can't type on the keyboard) and instead
-			// keep the focus on the text field and bubble the keyboard commands to it.
-			const bool bIsOpen = true;
-			const bool bFocusMenu = false;
-			MenuAnchor->SetIsOpen(bIsOpen, bFocusMenu);
-			SuggestionListView->RequestScrollIntoView(Suggestions[0]);
-		}
-		else
-		{
-			CloseMenuAndReset();
-		}
-	}
-
-	void CloseMenuAndReset()
-	{
-		const bool bIsOpen = false;
-		MenuAnchor->SetIsOpen(bIsOpen);
-
-		// Reset their index when the drawer closes so that the first item is always selected when we re-open.
-		CurrentSuggestionIndex = -1;
-	}
-
-	void SetActiveSuggestionIndex(int32 InIndex)
-	{
-		if (InIndex < 0 || InIndex >= Suggestions.Num())
-		{
-			return;
-		}
-
-		TSharedPtr<FString> Suggestion = Suggestions[InIndex];
-		SuggestionListView->SetSelection(Suggestion);
-		if (!SuggestionListView->IsItemVisible(Suggestion))
-		{
-			SuggestionListView->RequestScrollIntoView(Suggestion);
-		}
-		CurrentSuggestionIndex = InIndex;
-	}
-
-	TSharedRef<ITableRow> HandleSuggestionListViewGenerateRow(TSharedPtr<FString> Text, const TSharedRef<STableViewBase>& OwnerTable)
-	{
-		FString SuggestionText = *Text;
-
-		return SNew(STableRow<TSharedPtr<FString> >, OwnerTable)
-			[
-				SNew(SBox)
-				[
-					SNew(STextBlock)
-					.Text(FText::FromString(SuggestionText))
-				]
-			];
-	}
-
-	void HandleSuggestionListViewSelectionChanged(TSharedPtr<FString> NewValue, ESelectInfo::Type SelectInfo)
-	{
-		// This is called when clicking on an item and when navigating the menu via arrow keys, but we already
-		// handle arrow keys in OnKeyDown so we only want to handle mouse click here.
-		if (SelectInfo == ESelectInfo::Type::OnMouseClick)
-		{
-			// Trigger the auto-complete for the highlighted suggestion
-			ReplaceRelevantTextWithSuggestion(*NewValue);
-			CloseMenuAndReset();
-		}
-	}
-
-private:
-	TSharedPtr<SListView<TSharedPtr<FString>>> SuggestionListView;
-	TSharedPtr<SMultiLineEditableTextBox> TextBox;
-	TSharedPtr<SMenuAnchor> MenuAnchor;
-	TSharedPtr<SVerticalBox> VerticalBox;
-	// Holds a delegate that is executed when the text has changed.
-	FOnTextChanged OnTextChanged;
-
-	// The pool of suggestions to show
-	TArray<FString> AllSuggestions;
-	// The currently filtered suggestion list
-	TArray<TSharedPtr<FString>> Suggestions;
-	int32 CurrentSuggestionIndex = -1;
-};
-
 class FOutputFormatDetailsCustomization : public IDetailCustomization
 {
 public:
@@ -393,7 +79,7 @@ protected:
 			.ValueContent()
 			.MinDesiredWidth(200.0f)
 			[
-				SNew(SMoviePipelineAutoCompleteTextBox)
+				SNew(SMoviePipelineFormatTokenAutoCompleteBox)
 				.Text(this, &FOutputFormatDetailsCustomization::GetText)
 				.Suggestions(this, &FOutputFormatDetailsCustomization::GetSuggestions)
 				.OnTextChanged(this, &FOutputFormatDetailsCustomization::OnTextChanged)
diff --git a/Engine/Plugins/MovieScene/MovieRenderPipeline/Source/MovieRenderPipelineEditor/Private/Widgets/SMoviePipelineFormatTokenAutoCompleteBox.cpp b/Engine/Plugins/MovieScene/MovieRenderPipeline/Source/MovieRenderPipelineEditor/Private/Widgets/SMoviePipelineFormatTokenAutoCompleteBox.cpp
new file mode 100644
--- /dev/null
+++ b/Engine/Plugins/MovieScene/MovieRenderPipeline/Source/MovieRenderPipelineEditor/Private/Widgets/SMoviePipelineFormatTokenAutoCompleteBox.cpp
@@ -0,0 +1,292 @@
+﻿// Copyright Epic Games, Inc. All Rights Reserved.
+
+#include "SMoviePipelineFormatTokenAutoCompleteBox.h"
+
+#include "DetailLayoutBuilder.h"
+#include "Layout/WidgetPath.h"
+
+void SMoviePipelineFormatTokenAutoCompleteBox::Construct(const FArguments& InArgs)
+{
+	ChildSlot
+	[
+		SAssignNew(MenuAnchor, SMenuAnchor)
+		.Placement(MenuPlacement_ComboBox)
+		[
+			SAssignNew(TextBox, SMultiLineEditableTextBox)
+			.Text(InArgs._Text)
+			.Font(IDetailLayoutBuilder::GetDetailFont())
+			.OnKeyDownHandler(this, &SMoviePipelineFormatTokenAutoCompleteBox::OnKeyDown)
+			.OnTextChanged(this, &SMoviePipelineFormatTokenAutoCompleteBox::HandleTextBoxTextChanged)
+			.SelectWordOnMouseDoubleClick(true)
+			.AllowMultiLine(false)
+		]
+		.MenuContent
+		(
+			SNew(SBorder)
+			.Padding(FMargin(2))
+			[
+				SAssignNew(VerticalBox, SVerticalBox)
+				+ SVerticalBox::Slot()
+				.AutoHeight()
+				[
+					SAssignNew(SuggestionListView, SListView<TSharedPtr<FString>>)
+					.ItemHeight(18.f)
+					.ListItemsSource(&Suggestions)
+					.SelectionMode(ESelectionMode::Single)
+					.OnGenerateRow(this, &SMoviePipelineFormatTokenAutoCompleteBox::HandleSuggestionListViewGenerateRow)
+					.OnSelectionChanged(this, &SMoviePipelineFormatTokenAutoCompleteBox::HandleSuggestionListViewSelectionChanged)
+					.OnMouseButtonClick(this, &SMoviePipelineFormatTokenAutoCompleteBox::OnItemClicked)
+				]
+			]
+		)
+	];
+
+	// We just call it once and cache it for now as the selection code isn't tested against
+	// the amount of suggestions changing.
+	AllSuggestions.Append(InArgs._Suggestions.Get());
+	OnTextChanged = InArgs._OnTextChanged;
+}
+
+void SMoviePipelineFormatTokenAutoCompleteBox::OnFocusChanging(const FWeakWidgetPath& PreviousFocusPath, const FWidgetPath& NewWidgetPath, const FFocusEvent& InFocusEvent)
+{
+	// Close the suggestion list if the text box has lost focus to anything other than the suggestion list
+	if (PreviousFocusPath.ContainsWidget(TextBox.Get()) && !NewWidgetPath.ContainsWidget(SuggestionListView.Get()))
+	{
+		CloseMenuAndReset();
+	}
+	SWidget::OnFocusChanging(PreviousFocusPath, NewWidgetPath, InFocusEvent);
+}
+
+FReply SMoviePipelineFormatTokenAutoCompleteBox::OnKeyDown(const FGeometry& MyGeometry, const FKeyEvent& KeyEvent)
+{
+	if (MenuAnchor->IsOpen())
+	{
+		if (KeyEvent.GetKey() == EKeys::Up)
+		{
+			// Because the pop-up dialog is below the text, 'up' actually goes to an earlier item in the list.
+			int32 NewSuggestionIndex = CurrentSuggestionIndex - 1;
+			if (NewSuggestionIndex < 0)
+			{
+				NewSuggestionIndex = Suggestions.Num() - 1;
+			}
+
+			SetActiveSuggestionIndex(NewSuggestionIndex);
+			return FReply::Handled();
+		}
+		else if(KeyEvent.GetKey() == EKeys::Down)
+		{
+			int32 NewSuggestionIndex = CurrentSuggestionIndex + 1;
+			if (NewSuggestionIndex > Suggestions.Num() - 1)
+			{
+				NewSuggestionIndex = 0;
+			}
+
+			SetActiveSuggestionIndex(NewSuggestionIndex);
+			return FReply::Handled();
+		}
+		else if (KeyEvent.GetKey() == EKeys::Escape)
+		{
+			CloseMenuAndReset();
+			return FReply::Handled();
+		}
+		else if (KeyEvent.GetKey() == EKeys::Tab || KeyEvent.GetKey() == EKeys::Enter)
+		{
+			if(CurrentSuggestionIndex >= 0 && CurrentSuggestionIndex <= Suggestions.Num() - 1)
+			{
+				// Trigger the auto-complete for the highlighted suggestion
+				const FString SuggestionText = *Suggestions[CurrentSuggestionIndex];
+				ReplaceRelevantTextWithSuggestion(SuggestionText);
+				return FReply::Handled();
+			}
+		}
+	}
+	return FReply::Unhandled();
+}
+
+void SMoviePipelineFormatTokenAutoCompleteBox::OnItemClicked(TSharedPtr<FString> Item) const
+{
+	ReplaceRelevantTextWithSuggestion(*Item);
+}
+
+void SMoviePipelineFormatTokenAutoCompleteBox::FindAutoCompletableTextAtPos(const FString& InWholeString, int32 InCursorPos, FString& OutStr, bool& bShowAutoComplete)
+{
+	OutStr = FString();
+	bShowAutoComplete = false;
+
+	// We want to find a { brace on or to the left of InCursorPos, but if we find a } we
+	// stop looking, because that's a brace for another text. (+1 for ::FromEnd off by one)
+	const int32 StartingBracePos = InWholeString.Find(
+		TEXT("{"), ESearchCase::Type::IgnoreCase, ESearchDir::Type::FromEnd, InCursorPos + 1);
+	const int32 PreviousEndBracePos = InWholeString.Find(
+		TEXT("}"), ESearchCase::Type::IgnoreCase, ESearchDir::Type::FromEnd, InCursorPos);
+
+	if (StartingBracePos < PreviousEndBracePos)
+	{
+		return;
+	}
+
+	FString AutoCompleteText;
+
+	// Now that we found a {, take the substring between it and either the next }, or the end of the string.
+	const int32 NextEndBracePos = InWholeString.Find(TEXT("}"), ESearchCase::Type::IgnoreCase, ESearchDir::Type::FromStart, InCursorPos);
+	if (StartingBracePos >= 0)
+	{
+		int32 Count = InWholeString.Len() - StartingBracePos;
+		if (NextEndBracePos >= 0)
+		{
+			Count = NextEndBracePos - StartingBracePos;
+		}
+
+		AutoCompleteText = InWholeString.Mid(StartingBracePos + 1, Count - 1);
+	}
+
+	OutStr = AutoCompleteText;
+	bShowAutoComplete = StartingBracePos >= 0 && OutStr.Len() == 0;
+}
+
+void SMoviePipelineFormatTokenAutoCompleteBox::ReplaceRelevantTextWithSuggestion(const FString& InSuggestionText) const
+{
+	FString TextBoxText = TextBox->GetText().ToString();
+	int32 CursorPos = TextBoxText.Len();
+	const FTextLocation CursorLoc = TextBox->GetCursorLocation();
+	if (CursorLoc.IsValid())
+	{
+		CursorPos = FMath::Clamp(CursorLoc.GetOffset(), 0, CursorPos);
+	}
+		
+	// Look for the { to the left of the cursor. We search StrPositionIndex from +1 here due to a bug in ::FromEnd being off by one.
+	const int32 StartingBracePos = TextBoxText.Find(TEXT("{"), ESearchCase::Type::IgnoreCase, ESearchDir::Type::FromEnd, CursorPos + 1);
+
+	// Now that we found a {, take the substring between it and either the next }, or the end of the string.
+	const int32 NextEndBracePos = TextBoxText.Find(TEXT("}"), ESearchCase::Type::IgnoreCase, ESearchDir::Type::FromStart, CursorPos);
+	int32 NewCursorPos = 0;
+	if (StartingBracePos >= 0)	
+	{
+		// +1 to keep the left { brace
+		const FString Left = TextBoxText.Left(StartingBracePos+1);
+		FString Right;
+
+		if (NextEndBracePos >= 0)
+		{
+			Right = TextBoxText.RightChop(NextEndBracePos);
+		}
+
+		// Since the user chose the suggestion ensure there's already a } brace to close off the pair.
+		if (!Right.StartsWith(TEXT("}")))
+		{
+			Right = FString::Printf(TEXT("}%s"), *Right);
+		}
+
+		TextBoxText = Left + InSuggestionText + Right;
+		// We subtract 1 from the Right as we want to put the cursor after the automatically generated "}" token.
+		NewCursorPos = TextBoxText.Len() - (Right.Len() - 1);
+	}
+
+	TextBox->SetText(FText::FromString(TextBoxText));
+	TextBox->GoTo(FTextLocation(0, NewCursorPos));
+}
+
+void SMoviePipelineFormatTokenAutoCompleteBox::HandleTextBoxTextChanged(const FText& InText)
+{
+	OnTextChanged.ExecuteIfBound(InText);
+
+	const FString TextAsStr = InText.ToString();
+	if (TextAsStr.Len() > 0)
+	{
+		FString OutStr;
+		bool bShowAutoComplete;
+
+		int32 CursorPos = TextAsStr.Len();
+		const FTextLocation CursorLoc = TextBox->GetCursorLocation();
+		if (CursorLoc.IsValid())
+		{
+			CursorPos = FMath::Clamp(CursorLoc.GetOffset(), 0, CursorPos);
+		}
+
+		FindAutoCompletableTextAtPos(TextAsStr, CursorPos, OutStr, bShowAutoComplete);
+		FilterVisibleSuggestions(OutStr, bShowAutoComplete);
+	}	
+	else
+	{
+		// If they have no text, suggest all possible solutions
+		FilterVisibleSuggestions(FString(), false);
+	}
+}
+
+void SMoviePipelineFormatTokenAutoCompleteBox::FilterVisibleSuggestions(const FString& StrToMatch, const bool bForceShowAll)
+{
+	Suggestions.Reset();
+	for (const FString& Suggestion : AllSuggestions)
+	{
+		if (Suggestion.Contains(StrToMatch) || bForceShowAll)
+		{
+			Suggestions.Add(MakeShared<FString>(Suggestion));
+		}
+	}
+
+	if (Suggestions.Num() > 0)
+	{
+		// We don't focus the menu (because then you can't type on the keyboard) and instead
+		// keep the focus on the text field and bubble the keyboard commands to it.
+		constexpr bool bIsOpen = true;
+		constexpr bool bFocusMenu = false;
+		MenuAnchor->SetIsOpen(bIsOpen, bFocusMenu);
+		SuggestionListView->RequestScrollIntoView(Suggestions[0]);
+	}
+	else
+	{
+		CloseMenuAndReset();
+	}
+}
+
+void SMoviePipelineFormatTokenAutoCompleteBox::CloseMenuAndReset()
+{
+	constexpr bool bIsOpen = false;
+	MenuAnchor->SetIsOpen(bIsOpen);
+
+	// Reset their index when the drawer closes so that the first item is always selected when we re-open.
+	CurrentSuggestionIndex = -1;
+}
+
+void SMoviePipelineFormatTokenAutoCompleteBox::SetActiveSuggestionIndex(int32 InIndex)
+{
+	if (InIndex < 0 || InIndex >= Suggestions.Num())
+	{
+		return;
+	}
+
+	const TSharedPtr<FString> Suggestion = Suggestions[InIndex];
+	SuggestionListView->SetSelection(Suggestion);
+	if (!SuggestionListView->IsItemVisible(Suggestion))
+	{
+		SuggestionListView->RequestScrollIntoView(Suggestion);
+	}
+	CurrentSuggestionIndex = InIndex;
+}
+
+TSharedRef<ITableRow> SMoviePipelineFormatTokenAutoCompleteBox::HandleSuggestionListViewGenerateRow(
+	TSharedPtr<FString> Text, const TSharedRef<STableViewBase>& OwnerTable) const
+{
+	const FString SuggestionText = *Text;
+
+	return SNew(STableRow<TSharedPtr<FString> >, OwnerTable)
+		[
+			SNew(SBox)
+			[
+				SNew(STextBlock)
+				.Text(FText::FromString(SuggestionText))
+			]
+		];
+}
+
+void SMoviePipelineFormatTokenAutoCompleteBox::HandleSuggestionListViewSelectionChanged(TSharedPtr<FString> NewValue, ESelectInfo::Type SelectInfo)
+{
+	// This is called when clicking on an item and when navigating the menu via arrow keys, but we already
+	// handle arrow keys in OnKeyDown so we only want to handle mouse click here.
+	if (SelectInfo == ESelectInfo::Type::OnMouseClick)
+	{
+		// Trigger the auto-complete for the highlighted suggestion
+		ReplaceRelevantTextWithSuggestion(*NewValue);
+		CloseMenuAndReset();
+	}
+}
diff --git a/Engine/Plugins/MovieScene/MovieRenderPipeline/Source/MovieRenderPipelineEditor/Private/Widgets/SMoviePipelineFormatTokenAutoCompleteBox.h b/Engine/Plugins/MovieScene/MovieRenderPipeline/Source/MovieRenderPipelineEditor/Private/Widgets/SMoviePipelineFormatTokenAutoCompleteBox.h
new file mode 100644
--- /dev/null
+++ b/Engine/Plugins/MovieScene/MovieRenderPipeline/Source/MovieRenderPipelineEditor/Private/Widgets/SMoviePipelineFormatTokenAutoCompleteBox.h
@@ -0,0 +1,67 @@
+﻿#pragma once
+
+#include "PropertyHandle.h"
+#include "Widgets/Input/SMenuAnchor.h"
+#include "Widgets/Input/SMultiLineEditableTextBox.h"
+#include "Widgets/Views/SListView.h"
+
+/** This is a specific auto-complete widget made for the MoviePipeline that isn't very flexible. It's
+* similar to SSuggestionTextBox, but SSuggestionTextBox doesn't handle more than one word/suggestions
+* mid string. This widget is hardcoded to look for '{' characters (for {format_tokens}) and then auto
+* completes them from a list and fixes up the {} braces.
+*/
+class SMoviePipelineFormatTokenAutoCompleteBox : public SCompoundWidget
+{
+public:
+
+	SLATE_BEGIN_ARGS(SMoviePipelineFormatTokenAutoCompleteBox)
+		: _Text()
+		, _Suggestions()
+	{}
+
+	SLATE_ATTRIBUTE(FText, Text)
+	SLATE_ATTRIBUTE(TArray<FString>, Suggestions)
+	/** Called whenever the text is changed programmatically or interactively by the user. */
+	SLATE_EVENT(FOnTextChanged, OnTextChanged)
+	SLATE_END_ARGS()
+
+	void Construct(const FArguments& InArgs);
+
+	// SWidget interface
+	virtual void OnFocusChanging(const FWeakWidgetPath& PreviousFocusPath, const FWidgetPath& NewWidgetPath, const FFocusEvent& InFocusEvent) override;
+	// End of SWidget interface
+
+	FReply OnKeyDown(const FGeometry& MyGeometry, const FKeyEvent& KeyEvent);
+
+	void OnItemClicked(TSharedPtr<FString>) const;
+
+	static void FindAutoCompletableTextAtPos(const FString& InWholeString, int32 InCursorPos, FString& OutStr, bool& bShowAutoComplete);
+
+	void ReplaceRelevantTextWithSuggestion(const FString& InSuggestionText) const;
+
+	void HandleTextBoxTextChanged(const FText& InText);
+
+	void FilterVisibleSuggestions(const FString& StrToMatch, const bool bForceShowAll);
+
+	void CloseMenuAndReset();
+
+	void SetActiveSuggestionIndex(int32 InIndex);
+
+	TSharedRef<ITableRow> HandleSuggestionListViewGenerateRow(TSharedPtr<FString> Text, const TSharedRef<STableViewBase>& OwnerTable) const;
+
+	void HandleSuggestionListViewSelectionChanged(TSharedPtr<FString> NewValue, ESelectInfo::Type SelectInfo);
+
+private:
+	TSharedPtr<SListView<TSharedPtr<FString>>> SuggestionListView;
+	TSharedPtr<SMultiLineEditableTextBox> TextBox;
+	TSharedPtr<SMenuAnchor> MenuAnchor;
+	TSharedPtr<SVerticalBox> VerticalBox;
+	// Holds a delegate that is executed when the text has changed.
+	FOnTextChanged OnTextChanged;
+
+	// The pool of suggestions to show
+	TArray<FString> AllSuggestions;
+	// The currently filtered suggestion list
+	TArray<TSharedPtr<FString>> Suggestions;
+	int32 CurrentSuggestionIndex = -1;
+};
