Movie Graph:

-Fix: [Main CIS] :horde-error: Issue 414017: Compile errors in MovieGraphFormatTokenCustomization.h and SMoviePipelineFormatTokenAutoCompleteBox.h

diff --git a/Engine/Plugins/MovieScene/MovieRenderPipeline/Source/MovieRenderPipelineEditor/Private/Customizations/Graph/MovieGraphFormatTokenCustomization.h b/Engine/Plugins/MovieScene/MovieRenderPipeline/Source/MovieRenderPipelineEditor/Private/Customizations/Graph/MovieGraphFormatTokenCustomization.h
--- a/Engine/Plugins/MovieScene/MovieRenderPipeline/Source/MovieRenderPipelineEditor/Private/Customizations/Graph/MovieGraphFormatTokenCustomization.h
+++ b/Engine/Plugins/MovieScene/MovieRenderPipeline/Source/MovieRenderPipelineEditor/Private/Customizations/Graph/MovieGraphFormatTokenCustomization.h
@@ -1,4 +1,5 @@
-﻿#pragma once
+﻿// Copyright Epic Games, Inc. All Rights Reserved.
+#pragma once
 
 #include "MoviePipelineQueue.h"
 #include "IDetailCustomization.h"
diff --git a/Engine/Plugins/MovieScene/MovieRenderPipeline/Source/MovieRenderPipelineEditor/Private/Widgets/SMoviePipelineFormatTokenAutoCompleteBox.h b/Engine/Plugins/MovieScene/MovieRenderPipeline/Source/MovieRenderPipelineEditor/Private/Widgets/SMoviePipelineFormatTokenAutoCompleteBox.h
--- a/Engine/Plugins/MovieScene/MovieRenderPipeline/Source/MovieRenderPipelineEditor/Private/Widgets/SMoviePipelineFormatTokenAutoCompleteBox.h
+++ b/Engine/Plugins/MovieScene/MovieRenderPipeline/Source/MovieRenderPipelineEditor/Private/Widgets/SMoviePipelineFormatTokenAutoCompleteBox.h
@@ -1,4 +1,5 @@
-﻿#pragma once
+﻿// Copyright Epic Games, Inc. All Rights Reserved.
+#pragma once
 
 #include "PropertyHandle.h"
 #include "Widgets/Input/SMenuAnchor.h"
