Movie Graph:

-Fix: Can't Undo Delete Node

diff --git a/Engine/Plugins/MovieScene/MovieRenderPipeline/Source/MovieRenderPipelineCore/Private/Graph/MovieGraphConfig.cpp b/Engine/Plugins/MovieScene/MovieRenderPipeline/Source/MovieRenderPipelineCore/Private/Graph/MovieGraphConfig.cpp
--- a/Engine/Plugins/MovieScene/MovieRenderPipeline/Source/MovieRenderPipelineCore/Private/Graph/MovieGraphConfig.cpp
+++ b/Engine/Plugins/MovieScene/MovieRenderPipeline/Source/MovieRenderPipelineCore/Private/Graph/MovieGraphConfig.cpp
@@ -629,6 +629,8 @@ bool UMovieGraphConfig::RemoveNode(UMovieGraphNode* InNode)
 		return false;
 	}
 
+	Modify();
+
 	RemoveAllInboundEdges(InNode);
 	RemoveAllOutboundEdges(InNode);
 
diff --git a/Engine/Plugins/MovieScene/MovieRenderPipeline/Source/MovieRenderPipelineCore/Private/Graph/MovieGraphNode.cpp b/Engine/Plugins/MovieScene/MovieRenderPipeline/Source/MovieRenderPipelineCore/Private/Graph/MovieGraphNode.cpp
--- a/Engine/Plugins/MovieScene/MovieRenderPipeline/Source/MovieRenderPipelineCore/Private/Graph/MovieGraphNode.cpp
+++ b/Engine/Plugins/MovieScene/MovieRenderPipeline/Source/MovieRenderPipelineCore/Private/Graph/MovieGraphNode.cpp
@@ -246,7 +246,7 @@ void UMovieGraphNode::UpdatePins()
 		for (const FMovieGraphPinProperties& UnmatchedProperty : UnmatchedProperties)
 		{
 			const int32 InsertIndex = PinProperties.IndexOfByKey(UnmatchedProperty);
-			UMovieGraphPin* NewPin = NewObject<UMovieGraphPin>(this);
+			UMovieGraphPin* NewPin = NewObject<UMovieGraphPin>(this, NAME_None, RF_Transactional);
 			NewPin->Node = this;
 			NewPin->Properties = UnmatchedProperty;
 			Pins.Insert(NewPin, InsertIndex);
diff --git a/Engine/Plugins/MovieScene/MovieRenderPipeline/Source/MovieRenderPipelineCore/Private/Graph/MovieGraphPin.cpp b/Engine/Plugins/MovieScene/MovieRenderPipeline/Source/MovieRenderPipelineCore/Private/Graph/MovieGraphPin.cpp
--- a/Engine/Plugins/MovieScene/MovieRenderPipeline/Source/MovieRenderPipelineCore/Private/Graph/MovieGraphPin.cpp
+++ b/Engine/Plugins/MovieScene/MovieRenderPipeline/Source/MovieRenderPipelineCore/Private/Graph/MovieGraphPin.cpp
@@ -10,6 +10,14 @@
 #include "MovieRenderPipelineCoreModule.h"
 #include "Graph/MovieGraphConfig.h"
 
+#if WITH_EDITOR
+bool UMovieGraphPin::Modify(bool bAlwaysMarkDirty)
+{
+	SetFlags(RF_Transactional);
+	return Super::Modify(bAlwaysMarkDirty);
+}
+#endif
+
 bool UMovieGraphPin::AddEdgeTo(UMovieGraphPin* InOtherPin)
 {
 	if (!InOtherPin)
diff --git a/Engine/Plugins/MovieScene/MovieRenderPipeline/Source/MovieRenderPipelineCore/Public/Graph/MovieGraphPin.h b/Engine/Plugins/MovieScene/MovieRenderPipeline/Source/MovieRenderPipelineCore/Public/Graph/MovieGraphPin.h
--- a/Engine/Plugins/MovieScene/MovieRenderPipeline/Source/MovieRenderPipelineCore/Public/Graph/MovieGraphPin.h
+++ b/Engine/Plugins/MovieScene/MovieRenderPipeline/Source/MovieRenderPipelineCore/Public/Graph/MovieGraphPin.h
@@ -95,6 +95,12 @@ class MOVIERENDERPIPELINECORE_API UMovieGraphPin : public UObject
 	GENERATED_BODY()
 
 public:
+	// UObject Interface
+#if WITH_EDITOR
+	virtual bool Modify(bool bAlwaysMarkDirty = true) override;
+#endif
+	// End UObject Interface
+	
 	bool AddEdgeTo(UMovieGraphPin* InOtherPin);
 	bool BreakEdgeTo(UMovieGraphPin* InOtherPin);
 	bool BreakAllEdges();
diff --git a/Engine/Plugins/MovieScene/MovieRenderPipeline/Source/MovieRenderPipelineEditor/Private/Widgets/Graph/SMovieGraphConfigPanel.cpp b/Engine/Plugins/MovieScene/MovieRenderPipeline/Source/MovieRenderPipelineEditor/Private/Widgets/Graph/SMovieGraphConfigPanel.cpp
--- a/Engine/Plugins/MovieScene/MovieRenderPipeline/Source/MovieRenderPipelineEditor/Private/Widgets/Graph/SMovieGraphConfigPanel.cpp
+++ b/Engine/Plugins/MovieScene/MovieRenderPipeline/Source/MovieRenderPipelineEditor/Private/Widgets/Graph/SMovieGraphConfigPanel.cpp
@@ -287,19 +287,30 @@ void SMoviePipelineGraphPanel::DeleteSelectedNodes()
 		}
 	}
 
-	// Remove all runtime graph nodes
-	const UMoviePipelineEdGraph* Graph = Cast<UMoviePipelineEdGraph>(GraphEditorWidget->GetCurrentGraph());
-	if (Graph && !NodesToDelete.IsEmpty())
-	{
-		Graph->GetPipelineGraph()->RemoveNodes(NodesToDelete);
-	}
-
-	// Remove all editor nodes (nodes not backed by a runtime node, like comments)
-	if (Graph && !EdNodesToDelete.IsEmpty())
+	UMoviePipelineEdGraph* Graph = Cast<UMoviePipelineEdGraph>(GraphEditorWidget->GetCurrentGraph());
+	
+	const bool bShouldActuallyTransact = Graph && (NodesToDelete.Num() > 0 || EdNodesToDelete.Num() > 0);
+	const FScopedTransaction Transaction(
+		FGenericCommands::Get().Delete->GetDescription(), bShouldActuallyTransact);
+	
+	if (bShouldActuallyTransact)
 	{
-		for (UEdGraphNode* EdGraphNode : EdNodesToDelete)
+		// Remove all runtime graph nodes
+		if (!NodesToDelete.IsEmpty())
+		{
+			for (UMovieGraphNode* Node : NodesToDelete)
+			{
+				Graph->GetPipelineGraph()->RemoveNode(Node);
+			}
+		}
+		
+		// Remove all editor nodes (nodes not backed by a runtime node, like comments)
+		if (!EdNodesToDelete.IsEmpty())
 		{
-			GraphEditorWidget->GetCurrentGraph()->RemoveNode(EdGraphNode);
+			for (UEdGraphNode* EdGraphNode : EdNodesToDelete)
+			{
+				Graph->RemoveNode(EdGraphNode);
+			}
 		}
 	}
 
