Movie Graph:

-Fix: Render Preview UI Does Not Show Sub-Sample Count

diff --git a/Engine/Plugins/MovieScene/MovieRenderPipeline/Source/MovieRenderPipelineCore/Private/Graph/MovieGraphCoreTimeStep.cpp b/Engine/Plugins/MovieScene/MovieRenderPipeline/Source/MovieRenderPipelineCore/Private/Graph/MovieGraphCoreTimeStep.cpp
--- a/Engine/Plugins/MovieScene/MovieRenderPipeline/Source/MovieRenderPipelineCore/Private/Graph/MovieGraphCoreTimeStep.cpp
+++ b/Engine/Plugins/MovieScene/MovieRenderPipeline/Source/MovieRenderPipelineCore/Private/Graph/MovieGraphCoreTimeStep.cpp
@@ -441,6 +441,7 @@ void UMovieGraphCoreTimeStep::TickProducingFrames()
 			// If we've rendered the last temporal sub-sample, we've started a new output frame
 			// and we need to reset our temporal sample index.
 			CurrentFrameData.TemporalSampleIndex = 0;
+			CurrentCameraCut->ShotInfo.WorkMetrics.OutputSubSampleIndex = 0;
 		}
 		else
 		{
@@ -452,6 +453,7 @@ void UMovieGraphCoreTimeStep::TickProducingFrames()
 			{
 				// Each tick we increment the temporal sample we're on.
 				CurrentFrameData.TemporalSampleIndex++;
+				CurrentCameraCut->ShotInfo.WorkMetrics.OutputSubSampleIndex++;
 			}
 		}
 	}
diff --git a/Engine/Plugins/MovieScene/MovieRenderPipeline/Source/MovieRenderPipelineCore/Private/Graph/MovieGraphPipeline.cpp b/Engine/Plugins/MovieScene/MovieRenderPipeline/Source/MovieRenderPipelineCore/Private/Graph/MovieGraphPipeline.cpp
--- a/Engine/Plugins/MovieScene/MovieRenderPipeline/Source/MovieRenderPipelineCore/Private/Graph/MovieGraphPipeline.cpp
+++ b/Engine/Plugins/MovieScene/MovieRenderPipeline/Source/MovieRenderPipelineCore/Private/Graph/MovieGraphPipeline.cpp
@@ -351,6 +351,16 @@ void UMovieGraphPipeline::BuildShotListFromDataSource()
 		const FFrameRate TickResolution = GetDataSourceInstance()->GetTickResolution();
 
 		UMovieGraphWarmUpSettingNode* WarmUpNode = EvaluatedConfig->GetSettingForBranch<UMovieGraphWarmUpSettingNode>(UMovieGraphSettingNode::GlobalsPinName);
+
+		Shot->ShotInfo.NumTemporalSamples = SamplingMethodNode->TemporalSampleCount;
+		Shot->ShotInfo.NumSpatialSamples = 1;
+		Shot->ShotInfo.NumTiles = FIntPoint(1,1);
+		Shot->ShotInfo.CachedFrameRate = FinalFrameRate;
+		Shot->ShotInfo.CachedTickResolution = Shot->ShotInfo.CachedShotTickResolution = TickResolution;
+		if (Shot->ShotInfo.SubSectionHierarchy.IsValid() && Shot->ShotInfo.SubSectionHierarchy->MovieScene.IsValid())
+		{
+			Shot->ShotInfo.CachedShotTickResolution = Shot->ShotInfo.SubSectionHierarchy->MovieScene->GetTickResolution();
+		}
 		
 		const bool bPrePass = true;
 		const bool bExpandForTemporalSubSample = GraphTimeStepInstances.Last()->IsExpansionForTSRequired(EvaluatedConfig);
