Movie Graph:

-Fix: Navigating to output directory creates directory with cut off name

diff --git a/Engine/Plugins/MovieScene/MovieRenderPipeline/Source/MovieRenderPipelineEditor/Private/MoviePipelineEditorBlueprintLibrary.cpp b/Engine/Plugins/MovieScene/MovieRenderPipeline/Source/MovieRenderPipelineEditor/Private/MoviePipelineEditorBlueprintLibrary.cpp
--- a/Engine/Plugins/MovieScene/MovieRenderPipeline/Source/MovieRenderPipelineEditor/Private/MoviePipelineEditorBlueprintLibrary.cpp
+++ b/Engine/Plugins/MovieScene/MovieRenderPipeline/Source/MovieRenderPipelineEditor/Private/MoviePipelineEditorBlueprintLibrary.cpp
@@ -287,8 +287,12 @@ FString UMoviePipelineEditorBlueprintLibrary::ResolveOutputDirectoryFromJob(UMov
 		UMoviePipelineBlueprintLibrary::ResolveFilenameFormatArguments(FormatString, Params, OutResolvedPath, Dummy);
 	}
 
-	// Drop the .{ext} resolving always puts on.
-	OutResolvedPath.LeftChopInline(6);
+	// Drop the extension if it exists
+	FString Extension = FPaths::GetExtension(OutResolvedPath);
+	if (Extension.Len() > 0 && OutResolvedPath.EndsWith(Extension))
+	{
+		OutResolvedPath.RemoveFromEnd(Extension);
+	}
 
 	if (FPaths::IsRelative(OutResolvedPath))
 	{
