Movie Graph:

-Fix: No custom editor for MRQ shot settings

diff --git a/Engine/Plugins/MovieScene/MovieRenderPipeline/Source/MovieRenderPipelineEditor/Private/AssetDefinition_PipelineShotConfig.cpp b/Engine/Plugins/MovieScene/MovieRenderPipeline/Source/MovieRenderPipelineEditor/Private/AssetDefinition_PipelineShotConfig.cpp
new file mode 100644
--- /dev/null
+++ b/Engine/Plugins/MovieScene/MovieRenderPipeline/Source/MovieRenderPipelineEditor/Private/AssetDefinition_PipelineShotConfig.cpp
@@ -0,0 +1,29 @@
+// Copyright Epic Games, Inc. All Rights Reserved.
+
+#include "AssetDefinition_PipelineShotConfig.h"
+
+#include "MoviePipelineAssetEditor.h"
+#include "MoviePipelineQueueSubsystem.h"
+
+#define LOCTEXT_NAMESPACE "AssetTypeActions"
+
+EAssetCommandResult UAssetDefinition_PipelineShotConfig::OpenAssets(const FAssetOpenArgs& OpenArgs) const
+{
+	UAssetEditorSubsystem* AssetEditorSubsystem = GEditor->GetEditorSubsystem<UAssetEditorSubsystem>();
+	if (AssetEditorSubsystem)
+	{
+		UMoviePipelineAssetEditor* AssetEditor = NewObject<UMoviePipelineAssetEditor>(AssetEditorSubsystem, NAME_None, RF_Transient);
+		if (AssetEditor)
+		{
+			for (UMoviePipelineShotConfig* ShotConfig : OpenArgs.LoadObjects<UMoviePipelineShotConfig>())
+			{
+				AssetEditor->SetObjectToEdit(ShotConfig);
+				AssetEditor->Initialize();
+			}
+		}
+	}
+
+	return EAssetCommandResult::Handled;
+}
+
+#undef LOCTEXT_NAMESPACE
diff --git a/Engine/Plugins/MovieScene/MovieRenderPipeline/Source/MovieRenderPipelineEditor/Private/AssetDefinition_PipelineShotConfig.h b/Engine/Plugins/MovieScene/MovieRenderPipeline/Source/MovieRenderPipelineEditor/Private/AssetDefinition_PipelineShotConfig.h
new file mode 100644
--- /dev/null
+++ b/Engine/Plugins/MovieScene/MovieRenderPipeline/Source/MovieRenderPipelineEditor/Private/AssetDefinition_PipelineShotConfig.h
@@ -0,0 +1,27 @@
+// Copyright Epic Games, Inc. All Rights Reserved.
+
+#pragma once
+
+#include "MoviePipelineShotConfig.h"
+#include "AssetDefinitionDefault.h"
+
+#include "AssetDefinition_PipelineShotConfig.generated.h"
+
+UCLASS()
+class UAssetDefinition_PipelineShotConfig : public UAssetDefinitionDefault
+{
+	GENERATED_BODY()
+
+public:
+	// UAssetDefinition Begin
+	virtual FText GetAssetDisplayName() const override { return NSLOCTEXT("AssetTypeActions", "AssetTypeActions_PipelineShotConfig", "Movie Pipeline Shot Config"); }
+	virtual FLinearColor GetAssetColor() const override { return FLinearColor(FColor(255, 255, 255)); }
+	virtual TSoftClassPtr<UObject> GetAssetClass() const override { return UMoviePipelineShotConfig::StaticClass(); }
+	virtual TConstArrayView<FAssetCategoryPath> GetAssetCategories() const override
+	{
+		static const auto Categories = { EAssetCategoryPaths::Cinematics };
+		return Categories;
+	}
+	virtual EAssetCommandResult OpenAssets(const FAssetOpenArgs& OpenArgs) const override;
+	// UAssetDefinition End
+};
