MovieGraph: Initial Python Support for creating configuration assets via Python Scripting. This adds a sample script (MovieGraphCreateConfigExample.py) which documents a fair number of use cases - creating nodes, adding inputs and outputs, linking nodes together, adding branches and creating variables, converting jobs to use graph configs and overriding variables on the job itself.

Currently not supported by the Python API is Select Nodes which are difficult to work with due to their dynamic internal structure.

diff --git a/Engine/Plugins/MovieScene/MovieRenderPipeline/Source/MovieRenderPipelineCore/Private/Graph/MovieGraphConfig.cpp b/Engine/Plugins/MovieScene/MovieRenderPipeline/Source/MovieRenderPipelineCore/Private/Graph/MovieGraphConfig.cpp
--- a/Engine/Plugins/MovieScene/MovieRenderPipeline/Source/MovieRenderPipelineCore/Private/Graph/MovieGraphConfig.cpp
+++ b/Engine/Plugins/MovieScene/MovieRenderPipeline/Source/MovieRenderPipelineCore/Private/Graph/MovieGraphConfig.cpp
@@ -13,6 +13,7 @@
 #include "Graph/Nodes/MovieGraphRemoveRenderSettingNode.h"
 #include "Graph/Nodes/MovieGraphSubgraphNode.h"
 #include "Graph/Nodes/MovieGraphVariableNode.h"
+#include "Graph/Nodes/MovieGraphSelectNode.h"
 #include "MovieGraphUtils.h"
 #include "MoviePipelineQueue.h"
 #include "MovieRenderPipelineCoreModule.h"
@@ -93,11 +94,13 @@ bool UMovieGraphVariable::CanRename(const FText& InNewName, FText& OutError) con
 
 bool UMovieGraphVariable::SetMemberName(const FString& InNewName)
 {
+	bool bSuccess = Super::SetMemberName(InNewName);
+
 #if WITH_EDITOR
 	OnMovieGraphVariableChangedDelegate.Broadcast(this);
 #endif
 	
-	return Super::SetMemberName(InNewName);
+	return bSuccess;
 }
 
 #if WITH_EDITOR
@@ -211,11 +214,12 @@ bool UMovieGraphInput::CanRename(const FText& InNewName, FText& OutError) const
 
 bool UMovieGraphInput::SetMemberName(const FString& InNewName)
 {
+	bool bSuccess = Super::SetMemberName(InNewName);
+
 #if WITH_EDITOR
 	OnMovieGraphInputChangedDelegate.Broadcast(this);
 #endif
-	
-	return Super::SetMemberName(InNewName);
+	return bSuccess;
 }
 
 #if WITH_EDITOR
@@ -254,11 +258,13 @@ bool UMovieGraphOutput::CanRename(const FText& InNewName, FText& OutError) const
 
 bool UMovieGraphOutput::SetMemberName(const FString& InNewName)
 {
+	bool bSuccess = Super::SetMemberName(InNewName);
+
 #if WITH_EDITOR
 	OnMovieGraphOutputChangedDelegate.Broadcast(this);
 #endif
 	
-	return Super::SetMemberName(InNewName);
+	return bSuccess;
 }
 
 #if WITH_EDITOR
@@ -438,7 +444,7 @@ bool UMovieGraphConfig::AddLabeledEdge(UMovieGraphNode* FromNode, const FName& F
 	return bConnectionBrokeOtherEdges;
 }
 
-bool UMovieGraphConfig::RemoveEdge(UMovieGraphNode* FromNode, const FName& FromPinLabel, UMovieGraphNode* ToNode, const FName& ToPinLabel)
+bool UMovieGraphConfig::RemoveLabeledEdge(UMovieGraphNode* FromNode, const FName& FromPinLabel, UMovieGraphNode* ToNode, const FName& ToPinLabel)
 {
 	if (!FromNode || !ToNode)
 	{
@@ -1324,6 +1330,31 @@ bool UMovieGraphConfig::CreateFlattenedGraph_Recursive(UMovieGraphEvaluatedConfi
 
 	if (bShouldIncludeNode)
 	{
+#if WITH_EDITOR
+		// Normally we copy properties if we find a matching bOverride_ property. Unfortunately this creates a somewhat common
+		// scenario where you've created a bOverride_ property but typo'd the real property name, so the real property doesn't
+		// actually get updated, but we don't produce a warning (as it's valid to have properties with no matching bOverride_).
+		// So to avoid this we have this editor ensure to prompt you when we find a bOverride_ property with no matching "real" property.
+		for (TFieldIterator<FProperty> PropertyIterator(Node->GetClass()); PropertyIterator; ++PropertyIterator)
+		{
+			// If we're looking at an override property... 
+			if (PropertyIterator->GetName().StartsWith(TEXT("bOverride_")))
+			{
+				const FString RealPropertyName = PropertyIterator->GetName().RightChop(10); // Chop off bOverride_ to get the name of the property we're searching for.
+				bool bFoundProperty = false;
+				for (TFieldIterator<FProperty> InnerPropertyIterator(Node->GetClass()); InnerPropertyIterator; ++InnerPropertyIterator)
+				{
+					if (InnerPropertyIterator->GetName() == RealPropertyName)
+					{
+						bFoundProperty = true;
+						break;
+					}
+				}
+
+				ensureAlwaysMsgf(bFoundProperty, TEXT("Found override property named %s, but could not find real property named %s"), *PropertyIterator->GetName(), *RealPropertyName);
+			}
+		}
+#endif
 		const UMovieGraphSettingNode* NodeAsSetting = CastChecked<UMovieGraphSettingNode>(Node);
 		const FString& NodeInstanceName = NodeAsSetting->GetNodeInstanceName();
 		
@@ -1591,6 +1622,4 @@ UMovieGraphEvaluatedConfig* UMovieGraphConfig::CreateFlattenedGraph(const FMovie
 
 	return NewContext;
 }
-
-
 #undef LOCTEXT_NAMESPACE
diff --git a/Engine/Plugins/MovieScene/MovieRenderPipeline/Source/MovieRenderPipelineCore/Private/Graph/MovieGraphNode.cpp b/Engine/Plugins/MovieScene/MovieRenderPipeline/Source/MovieRenderPipelineCore/Private/Graph/MovieGraphNode.cpp
--- a/Engine/Plugins/MovieScene/MovieRenderPipeline/Source/MovieRenderPipelineCore/Private/Graph/MovieGraphNode.cpp
+++ b/Engine/Plugins/MovieScene/MovieRenderPipeline/Source/MovieRenderPipelineCore/Private/Graph/MovieGraphNode.cpp
@@ -153,6 +153,14 @@ void UMovieGraphNode::TogglePromotePropertyToPin(const FName& PropertyName)
 			});
 		}
 	}
+	else
+	{
+		// They requested a property that didn't exist, throw a kismet exception so scripts can know.
+		FFrame::KismetExecutionMessage(
+			*FString::Printf(
+				TEXT("%hs: Could not find a property with the name '%s' to promote to pin."), __FUNCTION__, *PropertyName.ToString()),
+			ELogVerbosity::Error);
+	}
 
 	UpdatePins();
 }
diff --git a/Engine/Plugins/MovieScene/MovieRenderPipeline/Source/MovieRenderPipelineCore/Private/Graph/MovieGraphPin.cpp b/Engine/Plugins/MovieScene/MovieRenderPipeline/Source/MovieRenderPipelineCore/Private/Graph/MovieGraphPin.cpp
--- a/Engine/Plugins/MovieScene/MovieRenderPipeline/Source/MovieRenderPipelineCore/Private/Graph/MovieGraphPin.cpp
+++ b/Engine/Plugins/MovieScene/MovieRenderPipeline/Source/MovieRenderPipelineCore/Private/Graph/MovieGraphPin.cpp
@@ -230,6 +230,33 @@ TArray<UMovieGraphPin*> UMovieGraphPin::GetAllConnectedPins() const
 	return ConnectedPins;
 }
 
+TArray<UMovieGraphNode*> UMovieGraphPin::GetConnectedNodes() const
+{
+	TArray<TObjectPtr<UMovieGraphNode>> OutNodes;
+	for (const TObjectPtr<UMovieGraphEdge>& Edge : Edges)
+	{
+		if (!Edge)
+		{
+			continue;
+		}
+
+		UMovieGraphPin* OtherPin = Edge->GetOtherPin(this);
+		if (!OtherPin)
+		{
+			continue;
+		}
+
+		if (!OtherPin->Node)
+		{
+			continue;
+		}
+
+		OutNodes.Add(OtherPin->Node);
+	}
+
+	return OutNodes;
+}
+
 bool UMovieGraphPin::IsConnectionToBranchAllowed(const UMovieGraphPin* OtherPin, FText& OutError) const
 {
 	const UMovieGraphPin* InputPin = IsInputPin() ? this : OtherPin;
diff --git a/Engine/Plugins/MovieScene/MovieRenderPipeline/Source/MovieRenderPipelineCore/Public/Graph/MovieGraphCommon.h b/Engine/Plugins/MovieScene/MovieRenderPipeline/Source/MovieRenderPipelineCore/Public/Graph/MovieGraphCommon.h
--- a/Engine/Plugins/MovieScene/MovieRenderPipeline/Source/MovieRenderPipelineCore/Public/Graph/MovieGraphCommon.h
+++ b/Engine/Plugins/MovieScene/MovieRenderPipeline/Source/MovieRenderPipelineCore/Public/Graph/MovieGraphCommon.h
@@ -22,12 +22,12 @@ enum class EMovieGraphValueType : uint8
 	Name,
 	String,
 	Text,
-	Enum UMETA(Hidden),
-	Struct UMETA(Hidden),
-	Object UMETA(Hidden),
-	SoftObject UMETA(Hidden),
-	Class UMETA(Hidden),
-	SoftClass UMETA(Hidden),
+	Enum,
+	Struct,
+	Object,
+	SoftObject,
+	Class,
+	SoftClass,
 	UInt32,	// Type not fully supported at UI, will work with restrictions to type editing
 	UInt64, // Type not fully supported at UI, will work with restrictions to type editing
 
diff --git a/Engine/Plugins/MovieScene/MovieRenderPipeline/Source/MovieRenderPipelineCore/Public/Graph/MovieGraphConfig.h b/Engine/Plugins/MovieScene/MovieRenderPipeline/Source/MovieRenderPipelineCore/Public/Graph/MovieGraphConfig.h
--- a/Engine/Plugins/MovieScene/MovieRenderPipeline/Source/MovieRenderPipelineCore/Public/Graph/MovieGraphConfig.h
+++ b/Engine/Plugins/MovieScene/MovieRenderPipeline/Source/MovieRenderPipelineCore/Public/Graph/MovieGraphConfig.h
@@ -36,9 +36,11 @@ public:
 	UMovieGraphConfig* GetOwningGraph() const;
 
 	/** Gets the name of this member. */
+	UFUNCTION(BlueprintCallable, Category = "Movie Graph")
 	FString GetMemberName() const { return Name; }
 
 	/** Sets the name of this member. Returns true if the rename was successful, else false. */
+	UFUNCTION(BlueprintCallable, Category ="Movie Graph")
 	virtual bool SetMemberName(const FString& InNewName);
 
 	/**
@@ -48,6 +50,7 @@ public:
 	virtual bool CanRename(const FText& InNewName, FText& OutError) const;
 
 	/** Gets the GUID that uniquely identifies this member. */
+	UFUNCTION(BlueprintCallable, Category = "Movie Graph")
 	const FGuid& GetGuid() const { return Guid; }
 
 	/** Sets the GUID that uniquely identifies this member. */
@@ -560,40 +563,86 @@ public:
 	virtual void PostLoad() override;
 	//~ End UObject interface
 
+	/**
+	* Add a connection in the graph between the given nodes and pin names. Pin name may be empty for basic
+	* nodes (if no name is displayed in the UI). Can be used for either input or output pins.
+	* Returns False if the pin could not be found, or the connection could not be made (type mismatches).
+	*/
+	UFUNCTION(BlueprintCallable, Category = "Movie Graph")
 	bool AddLabeledEdge(UMovieGraphNode* FromNode, const FName& FromPinLabel, UMovieGraphNode* ToNode, const FName& ToPinLabel);
-	bool RemoveEdge(UMovieGraphNode* FromNode, const FName& FromPinName, UMovieGraphNode* ToNode, const FName& ToPinName);
+
+	/**
+	* Like AddLabeledEdge, removes the given connection between Node A and Node B (for the specified pins by name).
+	*/
+	UFUNCTION(BlueprintCallable, Category = "Movie Graph")
+	bool RemoveLabeledEdge(UMovieGraphNode* FromNode, const FName& FromPinName, UMovieGraphNode* ToNode, const FName& ToPinName);
+	
+	
+	/**
+	* Convinence function which removes all Inbound (pins on the left side of a node) edges for the given node.
+	*/
+	UFUNCTION(BlueprintCallable, Category = "Movie Graph")
 	bool RemoveAllInboundEdges(UMovieGraphNode* InNode);
+	/**
+	* Convinence function which removes all Outbound (pins on the right side of a node) edges for the given node.
+	*/
+	UFUNCTION(BlueprintCallable, Category = "Movie Graph")
 	bool RemoveAllOutboundEdges(UMovieGraphNode* InNode);
+
+	/**
+	* Convinence function which removes all Inbound (pins on the left side of a node) edges connected to the given inbound pin by name, for the given node.
+	*/
+	UFUNCTION(BlueprintCallable, Category = "Movie Graph")
 	bool RemoveInboundEdges(UMovieGraphNode* InNode, const FName& InPinName);
+	
+	/**
+	* Convinence function which removes all Outobund (pins on the right side of a node) edges connected to the given outbound pin by name, for the given node.
+	*/
+	UFUNCTION(BlueprintCallable, Category = "Movie Graph")
 	bool RemoveOutboundEdges(UMovieGraphNode* InNode, const FName& InPinName);
 
 	/** 
 	* Add the specified node instance to the graph. This will rename the node to ensure the graph is the outer
 	* and then it will add it to the internal list of nodes used by the graph. See ConstructRuntimeNode if you
 	* want to construct a node by class and don't already have an instance.
+	* 
+	* Not currently exposed to the Blueprint API as it's generally internal use only.
 	*/
 	void AddNode(UMovieGraphNode* InNode);
 
-	/** Removes the specified node from the graph. */
+	/** 
+	* Removes the specified node from the graph, disconnecting connected edges as it goes. */
+	UFUNCTION(BlueprintCallable, Category = "Movie Graph")
 	bool RemoveNode(UMovieGraphNode* InNode);
-	/** Removes the specified nodes from the graph. */
+
+	/** Like RemoveNode but takes an entire array at once for convinence. */
+	UFUNCTION(BlueprintCallable, Category = "Movie Graph")
 	bool RemoveNodes(TArray<UMovieGraphNode*> InNodes);
 
+	/** Gets the automatically generated "Inputs" node in the Graph. */
+	UFUNCTION(BlueprintPure, Category = "Movie Graph")
 	UMovieGraphNode* GetInputNode() const { return InputNode; }
+
+	/** Gets the automatically generated "Outputs" node in the Graph. */
+	UFUNCTION(BlueprintPure, Category = "Movie Graph")
 	UMovieGraphNode* GetOutputNode() const { return OutputNode; }
+
+
 	const TArray<TObjectPtr<UMovieGraphNode>>& GetNodes() const { return AllNodes; }
 
 	/**
 	 * Adds a new variable member with default values to the graph. The new variable will have a base name of
 	 * "Variable" unless specified in InCustomBaseName. Returns the new variable on success, else nullptr.
 	 */
-	UFUNCTION(BlueprintCallable, Category="Experimental")
+	UFUNCTION(BlueprintCallable, Category="Movie Graph")
 	UMovieGraphVariable* AddVariable(const FName InCustomBaseName = NAME_None);
 
 	/** Adds a new input member to the graph. Returns the new input on success, else nullptr. */
+	UFUNCTION(BlueprintCallable, Category = "Movie Graph")
 	UMovieGraphInput* AddInput();
 
 	/** Adds a new output member to the graph. Returns the new output on success, else nullptr. */
+	UFUNCTION(BlueprintCallable, Category = "Movie Graph")
 	UMovieGraphOutput* AddOutput();
 
 	/** Gets the variable in the graph with the specified GUID, else nullptr if one could not be found. */
@@ -603,7 +652,7 @@ public:
 	 * Gets all variables that are available to be used in the graph. Global variables can optionally be included if
 	 * bIncludeGlobal is set to true.
 	 */
-	UFUNCTION(BlueprintCallable, Category="Experimental")
+	UFUNCTION(BlueprintCallable, Category="Movie Graph")
 	TArray<UMovieGraphVariable*> GetVariables(const bool bIncludeGlobal = false) const;
 
 	/** Updates the values of all global variables. */
@@ -611,12 +660,15 @@ public:
 	void UpdateGlobalVariableValues(const UMovieGraphPipeline* InPipeline);
 
 	/** Gets all inputs that have been defined on the graph. */
+	UFUNCTION(BlueprintCallable, Category = "Movie Graph")
 	TArray<UMovieGraphInput*> GetInputs() const;
 
 	/** Gets all outputs that have been defined on the graph. */
+	UFUNCTION(BlueprintCallable, Category = "Movie Graph")
 	TArray<UMovieGraphOutput*> GetOutputs() const;
 
 	/** Remove the specified member (input, output, variable) from the graph. */
+	UFUNCTION(BlueprintCallable, Category = "Movie Graph")
 	bool DeleteMember(UMovieGraphMember* MemberToDelete);
 
 #if WITH_EDITOR
@@ -719,29 +771,39 @@ public:
 	TObjectPtr<UEdGraph> PipelineEdGraph;
 #endif
 
-	template<class T>
-	T* ConstructRuntimeNode(TSubclassOf<UMovieGraphNode> PipelineGraphNodeClass = T::StaticClass())
+	/**
+	* Creates the given node type in this graph. Does not create any connections, and a node will not be
+	* considered during evaluation unless it is connected to other nodes in the graph.
+	*/
+	UFUNCTION(BlueprintCallable, meta = (DeterminesOutputType = "InClass"), Category = "Movie Graph")
+	UMovieGraphNode* CreateNodeByClass(const TSubclassOf<UMovieGraphNode> InClass)
 	{
-		if (!PipelineGraphNodeClass)
+		if (!InClass)
 		{
 			FFrame::KismetExecutionMessage(
 				*FString::Printf(
 					TEXT("%hs: Invalid PipelineGraphNodeClass. Please specify a valid class."), __FUNCTION__),
-					ELogVerbosity::Error);
-		
+				ELogVerbosity::Error);
+
 			return nullptr;
 		}
-		
+
 		// Construct a new object with ourselves as the outer, then keep track of it.
-		T* RuntimeNode = NewObject<T>(this, PipelineGraphNodeClass, NAME_None, RF_Transactional);
+		UMovieGraphNode* RuntimeNode = NewObject<UMovieGraphNode>(this, InClass, NAME_None, RF_Transactional);
 		RuntimeNode->UpdateDynamicProperties();
 		RuntimeNode->UpdatePins();
 		RuntimeNode->Guid = FGuid::NewGuid();
-		
+
 		AddNode(RuntimeNode);
 		return RuntimeNode;
 	}
 
+	template<class T>
+	T* ConstructRuntimeNode(TSubclassOf<UMovieGraphNode> PipelineGraphNodeClass = T::StaticClass())
+	{
+		return Cast<T>(CreateNodeByClass(PipelineGraphNodeClass));
+	}
+
 private:
 	/** Remove the specified variable member from the graph. */
 	bool DeleteVariableMember(UMovieGraphVariable* VariableMemberToDelete);
diff --git a/Engine/Plugins/MovieScene/MovieRenderPipeline/Source/MovieRenderPipelineCore/Public/Graph/MovieGraphNode.h b/Engine/Plugins/MovieScene/MovieRenderPipeline/Source/MovieRenderPipelineCore/Public/Graph/MovieGraphNode.h
--- a/Engine/Plugins/MovieScene/MovieRenderPipeline/Source/MovieRenderPipelineCore/Public/Graph/MovieGraphNode.h
+++ b/Engine/Plugins/MovieScene/MovieRenderPipeline/Source/MovieRenderPipelineCore/Public/Graph/MovieGraphNode.h
@@ -41,15 +41,15 @@ struct FMovieGraphPropertyInfo
 	GENERATED_BODY()
 	
 	/** The name of the property. */
-	UPROPERTY()
+	UPROPERTY(BlueprintReadOnly, Category = "Movie Graph")
 	FName Name;
 
 	/** Whether this property is dynamic (ie, it does not correspond to a native UPROPERTY on the node). */
-	UPROPERTY()
+	UPROPERTY(BlueprintReadOnly, Category = "Movie Graph")
 	bool bIsDynamicProperty = false;
 
 	/** The type of the value pointed to by the property. */
-	UPROPERTY()
+	UPROPERTY(BlueprintReadOnly, Category = "Movie Graph")
 	EMovieGraphValueType ValueType = EMovieGraphValueType::None;
 
 	/** The associated value type object if the ValueType is an enum, struct, class, or object. */
@@ -177,6 +177,7 @@ public:
 	virtual TArray<FMovieGraphPropertyInfo> GetOverrideablePropertyInfo() const;
 
 	/** Gets the information about properties which are currently exposed as pins on the node. */
+	UFUNCTION(BlueprintCallable, Category = "Movie Graph")
 	virtual TArray<FMovieGraphPropertyInfo> GetExposedProperties() const
 	{
 		return ExposedPropertyInfo;
@@ -205,6 +206,7 @@ public:
 	}
 
 	/** Toggles the promotion of the property with the given name to a pin on the node. */
+	UFUNCTION(BlueprintCallable, Category = "Movie Graph")
 	virtual void TogglePromotePropertyToPin(const FName& PropertyName);
 
 	/**
@@ -228,10 +230,12 @@ public:
 	 * Gets the input pin with the specified name, or nullptr if one could not be found. Most pins on a node are
 	 * "built-in", meaning they ship with the node. Dynamic pins (pins which are not built-in) can potentially have the
 	 * same name as a built-in (eg, the option pins on the Select node). To disambiguate between built-in and dynamic
-	 * pins, PinRequirement can be specified. */
+	 * pins, specify bIsBuiltInPin = false if trying to fetch a pin that is not built-in.  */
+	UFUNCTION(BlueprintCallable, Category = "Movie Graph")
 	UMovieGraphPin* GetInputPin(const FName& InPinLabel, const EMovieGraphPinQueryRequirement PinRequirement = EMovieGraphPinQueryRequirement::BuiltInOrDynamic) const;
 
 	/** Gets the output pin with the specified name, or nullptr if one could not be found. */
+	UFUNCTION(BlueprintCallable, Category = "Movie Graph")
 	UMovieGraphPin* GetOutputPin(const FName& InPinLabel) const;
 
 	/** Gets the first input pin on the node which has a connection, or nullptr if no pins are connected. */
@@ -253,10 +257,14 @@ public:
 	bool IsDisabled() const;
 
 #if WITH_EDITOR
+	UFUNCTION(BlueprintPure, Category = "Movie Graph")
 	int32 GetNodePosX() const { return NodePosX; }
+	UFUNCTION(BlueprintPure, Category = "Movie Graph")
 	int32 GetNodePosY() const { return NodePosY; }
 
+	UFUNCTION(BlueprintCallable, Category = "Movie Graph")
 	void SetNodePosX(const int32 InNodePosX) { NodePosX = InNodePosX; }
+	UFUNCTION(BlueprintCallable, Category = "Movie Graph")
 	void SetNodePosY(const int32 InNodePosY) { NodePosY = InNodePosY; }
 
 	/** Gets the node's title color, as visible in the graph. */
@@ -265,13 +273,19 @@ public:
 	/** Gets the node's icon and icon tint, as visible in the graph. */
 	virtual FSlateIcon GetIconAndTint(FLinearColor& OutColor) const;
 
+	UFUNCTION(BlueprintPure, Category = "Movie Graph")
 	FString GetNodeComment() const { return NodeComment; }
+	UFUNCTION(BlueprintCallable, Category = "Movie Graph")
 	void SetNodeComment(const FString& InNodeComment) { NodeComment = InNodeComment; }
 
+	UFUNCTION(BlueprintPure, Category = "Movie Graph")
 	bool IsCommentBubblePinned() const { return bIsCommentBubblePinned; }
+	UFUNCTION(BlueprintCallable, Category = "Movie Graph")
 	void SetIsCommentBubblePinned(const uint8 bIsPinned) { bIsCommentBubblePinned = bIsPinned; }
 
+	UFUNCTION(BlueprintPure, Category = "Movie Graph")
 	bool IsCommentBubbleVisible() const { return bIsCommentBubbleVisible; }
+	UFUNCTION(BlueprintCallable, Category = "Movie Graph")
 	void SetIsCommentBubbleVisible(uint8 bIsVisible) { bIsCommentBubbleVisible = bIsVisible; }
 #endif
 
diff --git a/Engine/Plugins/MovieScene/MovieRenderPipeline/Source/MovieRenderPipelineCore/Public/Graph/MovieGraphPin.h b/Engine/Plugins/MovieScene/MovieRenderPipeline/Source/MovieRenderPipelineCore/Public/Graph/MovieGraphPin.h
--- a/Engine/Plugins/MovieScene/MovieRenderPipeline/Source/MovieRenderPipelineCore/Public/Graph/MovieGraphPin.h
+++ b/Engine/Plugins/MovieScene/MovieRenderPipeline/Source/MovieRenderPipelineCore/Public/Graph/MovieGraphPin.h
@@ -118,6 +118,14 @@ public:
 	/** Gets all connected pins. */
 	TArray<UMovieGraphPin*> GetAllConnectedPins() const;
 
+	/**
+	* Utility function for scripting which gathers all of the nodes connected
+	* to this particular pin. Equivalent to looping through all of the edges,
+	* getting the connected pin, and then getting the node associated with that pin.
+	*/
+	UFUNCTION(BlueprintCallable, Category = "Movie Graph")
+	TArray<UMovieGraphNode*> GetConnectedNodes() const;
+
 	/**
 	 * Determines if the connection between this pin and OtherPin follows branch restriction rules. OutError is populated
 	 * with an error if the connection should be rejected and the function will return false.
diff --git a/Engine/Plugins/MovieScene/MovieRenderPipeline/Source/MovieRenderPipelineCore/Public/Graph/Nodes/MovieGraphVariableNode.h b/Engine/Plugins/MovieScene/MovieRenderPipeline/Source/MovieRenderPipelineCore/Public/Graph/Nodes/MovieGraphVariableNode.h
--- a/Engine/Plugins/MovieScene/MovieRenderPipeline/Source/MovieRenderPipelineCore/Public/Graph/Nodes/MovieGraphVariableNode.h
+++ b/Engine/Plugins/MovieScene/MovieRenderPipeline/Source/MovieRenderPipelineCore/Public/Graph/Nodes/MovieGraphVariableNode.h
@@ -1,4 +1,4 @@
-﻿// Copyright Epic Games, Inc. All Rights Reserved.
+// Copyright Epic Games, Inc. All Rights Reserved.
 
 #pragma once
 
@@ -21,9 +21,11 @@ public:
 	virtual bool GetResolvedValueForOutputPin(const FName& InPinName, const FMovieGraphTraversalContext* InContext, TObjectPtr<UMovieGraphValueContainer>& OutValueContainer) const override;
 
 	/** Gets the variable that this node represents. */
+	UFUNCTION(BlueprintCallable, Category = "Movie Graph")
 	UMovieGraphVariable* GetVariable() const { return GraphVariable; }
 
 	/** Sets the variable that this node represents. */
+	UFUNCTION(BlueprintCallable, Category = "Movie Graph")
 	void SetVariable(UMovieGraphVariable* InVariable);
 
 	/** Returns true if this node represents a global variable, else false. */
diff --git a/Engine/Plugins/MovieScene/MovieRenderPipeline/Source/MovieRenderPipelineEditor/Private/Graph/MovieGraphSchema.cpp b/Engine/Plugins/MovieScene/MovieRenderPipeline/Source/MovieRenderPipelineEditor/Private/Graph/MovieGraphSchema.cpp
--- a/Engine/Plugins/MovieScene/MovieRenderPipeline/Source/MovieRenderPipelineEditor/Private/Graph/MovieGraphSchema.cpp
+++ b/Engine/Plugins/MovieScene/MovieRenderPipeline/Source/MovieRenderPipelineEditor/Private/Graph/MovieGraphSchema.cpp
@@ -307,7 +307,7 @@ void UMovieGraphSchema::BreakSinglePinLink(UEdGraphPin* SourcePin, UEdGraphPin*
 	UMovieGraphConfig* RuntimeGraph = SourceRuntimeNode->GetGraph();
 	check(RuntimeGraph);
 
-	RuntimeGraph->RemoveEdge(SourceRuntimeNode, SourcePin->PinName, TargetRuntimeNode, TargetPin->PinName);
+	RuntimeGraph->RemoveLabeledEdge(SourceRuntimeNode, SourcePin->PinName, TargetRuntimeNode, TargetPin->PinName);
 }
 
 FLinearColor UMovieGraphSchema::GetTypeColor(const FName& InPinCategory, const FName& InPinSubCategory)
diff --git a/Engine/Plugins/MovieScene/MovieRenderPipeline/Source/MovieRenderPipelineEditor/Private/Tests/MovieGraphTests.cpp b/Engine/Plugins/MovieScene/MovieRenderPipeline/Source/MovieRenderPipelineEditor/Private/Tests/MovieGraphTests.cpp
--- a/Engine/Plugins/MovieScene/MovieRenderPipeline/Source/MovieRenderPipelineEditor/Private/Tests/MovieGraphTests.cpp
+++ b/Engine/Plugins/MovieScene/MovieRenderPipeline/Source/MovieRenderPipelineEditor/Private/Tests/MovieGraphTests.cpp
@@ -159,7 +159,7 @@ namespace UE::MovieGraph::Private::Tests
 		UMovieGraphNode* NodeFrom, UMovieGraphNode* NodeTo,
 		UMovieGraphPin* OutputPin, UMovieGraphPin* InputPin)
 	{
-		GraphConfig->RemoveEdge(NodeFrom, OutputPin->Properties.Label,
+		GraphConfig->RemoveLabeledEdge(NodeFrom, OutputPin->Properties.Label,
 			NodeTo, InputPin->Properties.Label);
 
 		InTestBase->TestTrue(
