MoviePipeline: Enabled the maximize button on the PIE window in case artists want to have a bigger preview.

diff --git a/Engine/Plugins/MovieScene/MovieRenderPipeline/Source/MovieRenderPipelineEditor/Private/MoviePipelinePIEExecutor.cpp b/Engine/Plugins/MovieScene/MovieRenderPipeline/Source/MovieRenderPipelineEditor/Private/MoviePipelinePIEExecutor.cpp
--- a/Engine/Plugins/MovieScene/MovieRenderPipeline/Source/MovieRenderPipelineEditor/Private/MoviePipelinePIEExecutor.cpp
+++ b/Engine/Plugins/MovieScene/MovieRenderPipeline/Source/MovieRenderPipelineEditor/Private/MoviePipelinePIEExecutor.cpp
@@ -134,10 +134,10 @@ void UMoviePipelinePIEExecutor::Start(const UMoviePipelineExecutorJob* InJob)
 		.ClientSize(WindowSize)
 		.AutoCenter(EAutoCenter::PrimaryWorkArea)
 		.UseOSWindowBorder(true)
-		.FocusWhenFirstShown(false)
+		.FocusWhenFirstShown(true)
 		.ActivationPolicy(EWindowActivationPolicy::Never)
 		.HasCloseButton(true)
-		.SupportsMaximize(false)
+		.SupportsMaximize(true)
 		.SupportsMinimize(true)
 		.SizingRule(ESizingRule::UserSized);
 
