MoviePipeline: Show variables from all subgraphs in the MRQ queue UI.

Along with this change, made several other improvements:
* Fixed many refresh issues with variable assignments not updating properly (both in the data model and in the UI).
* Shot jobs can now have their own variable assignments which affect a graph preset assigned to the shot.
* Shot jobs can also additionally override variables defined on the primary job that owns it.

To make the shot-specific workflows possible, the queue UI can now show details for a selected shot (previously, selecting a shot would show the details for the associated primary job).

#rnx

diff --git a/Engine/Plugins/MovieScene/MovieRenderPipeline/Source/MovieRenderPipelineCore/Private/Graph/Nodes/MovieGraphVariableNode.cpp b/Engine/Plugins/MovieScene/MovieRenderPipeline/Source/MovieRenderPipelineCore/Private/Graph/Nodes/MovieGraphVariableNode.cpp
--- a/Engine/Plugins/MovieScene/MovieRenderPipeline/Source/MovieRenderPipelineCore/Private/Graph/Nodes/MovieGraphVariableNode.cpp
+++ b/Engine/Plugins/MovieScene/MovieRenderPipeline/Source/MovieRenderPipelineCore/Private/Graph/Nodes/MovieGraphVariableNode.cpp
@@ -34,9 +34,10 @@ FString UMovieGraphVariableNode::GetResolvedValueForOutputPin(const FName& InPin
 {
 	if (GraphVariable && (GraphVariable->GetMemberName() == InPinName))
 	{
-		if (ContextHasEnabledAssignmentForVariable(InContext))
+		TObjectPtr<UMovieJobVariableAssignmentContainer> VariableAssignment;
+		if (ContextHasEnabledAssignmentForVariable(InContext, VariableAssignment))
 		{
-			return InContext->Job->VariableAssignments->GetValueSerializedString(GraphVariable);
+			return VariableAssignment->GetValueSerializedString(GraphVariable);
 		}
 
 		// No valid variable assignment: just get the value from the variable
@@ -50,9 +51,10 @@ bool UMovieGraphVariableNode::GetResolvedValueForOutputPin(const FName& InPinNam
 {
 	if (GraphVariable && (GraphVariable->GetMemberName() == InPinName))
 	{
-		if (ContextHasEnabledAssignmentForVariable(InContext))
+		TObjectPtr<UMovieJobVariableAssignmentContainer> VariableAssignment;
+		if (ContextHasEnabledAssignmentForVariable(InContext, VariableAssignment))
 		{
-			return InContext->Job->VariableAssignments->GetValueContainer(GraphVariable, OutValueContainer);
+			return VariableAssignment->GetValueContainer(GraphVariable, OutValueContainer);
 		}
 
 		// No valid variable assignment: just get the value from the variable
@@ -128,17 +130,61 @@ void UMovieGraphVariableNode::UpdateOutputPin(UMovieGraphMember* ChangedVariable
 	OnNodeChangedDelegate.Broadcast(this);
 }
 
-bool UMovieGraphVariableNode::ContextHasEnabledAssignmentForVariable(const FMovieGraphTraversalContext* InContext) const
+bool UMovieGraphVariableNode::ContextHasEnabledAssignmentForVariable(const FMovieGraphTraversalContext* InContext, TObjectPtr<UMovieJobVariableAssignmentContainer>& OutVariableAssignment) const
 {
-	if (InContext && InContext->Job)
+	if (!InContext)
 	{
-		bool bIsEnabled = false;
-		if (InContext->Job->VariableAssignments->GetVariableAssignmentEnableState(GraphVariable, bIsEnabled))
+		return false;
+	}
+
+	TObjectPtr<UMovieJobVariableAssignmentContainer> ShotVariableAssignments = nullptr;
+	TObjectPtr<UMovieJobVariableAssignmentContainer> ShotVariableAssignments_PrimaryOverrides = nullptr;
+	TObjectPtr<UMovieJobVariableAssignmentContainer> JobVariableAssignments = nullptr;
+	
+	if (InContext->Job->ShotInfo.IsValidIndex(InContext->ShotIndex))
+	{
+		const TObjectPtr<UMoviePipelineExecutorShot> ShotJob = InContext->Job->ShotInfo[InContext->ShotIndex];
+		
+		ShotVariableAssignments = ShotJob->GetOrCreateJobVariableAssignmentsForGraph(ShotJob->GetGraphPreset());
+
+		// The shot can also override variables on the primary job's graph
+		constexpr bool bIsForPrimaryOverrides = true;
+		ShotVariableAssignments_PrimaryOverrides = ShotJob->GetOrCreateJobVariableAssignmentsForGraph(InContext->Job->GetGraphPreset(), bIsForPrimaryOverrides);
+	}
+
+	if (InContext->Job)
+	{
+		JobVariableAssignments = InContext->Job->GetOrCreateJobVariableAssignmentsForGraph(InContext->Job->GetGraphPreset());
+	}
+	
+	// Check the shot job first for an enabled job variable assignment for this variable. Shot jobs take precedence over primary jobs.
+	bool bIsEnabled = false;
+	if (ShotVariableAssignments && ShotVariableAssignments->GetVariableAssignmentEnableState(GraphVariable, bIsEnabled))
+	{
+		if (bIsEnabled)
+		{
+			OutVariableAssignment = ShotVariableAssignments;
+			return true;
+		}
+	}
+
+	// Next check for shot-level overrides to the primary graph variables.
+	if (ShotVariableAssignments_PrimaryOverrides && ShotVariableAssignments_PrimaryOverrides->GetVariableAssignmentEnableState(GraphVariable, bIsEnabled))
+	{
+		if (bIsEnabled)
+		{
+			OutVariableAssignment = ShotVariableAssignments_PrimaryOverrides;
+			return true;
+		}
+	}
+
+	// Check the primary job last.
+	if (JobVariableAssignments && JobVariableAssignments->GetVariableAssignmentEnableState(GraphVariable, bIsEnabled))
+	{
+		if (bIsEnabled)
 		{
-			if (bIsEnabled)
-			{
-				return true;
-			}
+			OutVariableAssignment = JobVariableAssignments;
+			return true;
 		}
 	}
 
diff --git a/Engine/Plugins/MovieScene/MovieRenderPipeline/Source/MovieRenderPipelineCore/Private/MovieJobVariableAssignmentContainer.cpp b/Engine/Plugins/MovieScene/MovieRenderPipeline/Source/MovieRenderPipelineCore/Private/MovieJobVariableAssignmentContainer.cpp
--- a/Engine/Plugins/MovieScene/MovieRenderPipeline/Source/MovieRenderPipelineCore/Private/MovieJobVariableAssignmentContainer.cpp
+++ b/Engine/Plugins/MovieScene/MovieRenderPipeline/Source/MovieRenderPipelineCore/Private/MovieJobVariableAssignmentContainer.cpp
@@ -13,6 +13,11 @@ void UMovieJobVariableAssignmentContainer::SetGraphConfig(const TSoftObjectPtr<U
 	GraphPreset = InGraphConfig;
 }
 
+TSoftObjectPtr<UMovieGraphConfig> UMovieJobVariableAssignmentContainer::GetGraphConfig() const
+{
+	return GraphPreset;
+}
+
 uint32 UMovieJobVariableAssignmentContainer::GetNumAssignments() const
 {
 	uint32 NumAssignments = 0;
diff --git a/Engine/Plugins/MovieScene/MovieRenderPipeline/Source/MovieRenderPipelineCore/Private/MoviePipelineQueue.cpp b/Engine/Plugins/MovieScene/MovieRenderPipeline/Source/MovieRenderPipelineCore/Private/MoviePipelineQueue.cpp
--- a/Engine/Plugins/MovieScene/MovieRenderPipeline/Source/MovieRenderPipelineCore/Private/MoviePipelineQueue.cpp
+++ b/Engine/Plugins/MovieScene/MovieRenderPipeline/Source/MovieRenderPipelineCore/Private/MoviePipelineQueue.cpp
@@ -10,6 +10,65 @@
 
 #include UE_INLINE_GENERATED_CPP_BY_NAME(MoviePipelineQueue)
 
+namespace UE::MovieGraph::Private
+{
+	/**
+	 * Gets the job variable assignments for a specific graph. Creates a new variable assignments container if one was not found for the given graph.
+	 * The owner of the variable assignments must be provided in case a new assignment needs to be created.
+	 */
+	TObjectPtr<UMovieJobVariableAssignmentContainer> GetOrCreateJobVariableAssignmentsForGraph(
+		const UMovieGraphConfig* InGraph, TArray<TObjectPtr<UMovieJobVariableAssignmentContainer>>& InVariableAssignments, UObject* InAssignmentsOwner)
+	{
+		for (TObjectPtr<UMovieJobVariableAssignmentContainer>& VariableAssignment : InVariableAssignments)
+		{
+			const TSoftObjectPtr<UMovieGraphConfig> SoftGraphConfig = VariableAssignment->GetGraphConfig();
+			if (SoftGraphConfig.Get() == InGraph)
+			{
+#if WITH_EDITOR
+				VariableAssignment->UpdateGraphVariableOverrides();
+#endif
+				
+				return VariableAssignment;
+			}
+		}
+
+		// Create the variable assignments container if it wasn't found
+		TObjectPtr<UMovieJobVariableAssignmentContainer> NewVariableAssignments = NewObject<UMovieJobVariableAssignmentContainer>(InAssignmentsOwner);
+		InVariableAssignments.Add(NewVariableAssignments);
+		NewVariableAssignments->SetGraphConfig(InGraph);
+
+#if WITH_EDITOR
+		NewVariableAssignments->UpdateGraphVariableOverrides();
+#endif
+	
+		return NewVariableAssignments;
+	}
+
+	/** Refreshes the variable assignments for the given graph and its associated subgraphs. */
+	void RefreshVariableAssignments(UMovieGraphConfig* InRootGraph, TArray<TObjectPtr<UMovieJobVariableAssignmentContainer>>& InVariableAssignments, UObject* InAssignmentsOwner)
+	{
+		if (!InRootGraph)
+		{
+			return;
+		}
+	
+		TSet<UMovieGraphConfig*> AllGraphs = {InRootGraph};
+		InRootGraph->GetAllContainedSubgraphs(AllGraphs);
+
+		// Add/update variable assignments for the graph on the job and all of its subgraphs
+		for (const UMovieGraphConfig* Graph : AllGraphs)
+		{
+			GetOrCreateJobVariableAssignmentsForGraph(Graph, InVariableAssignments, InAssignmentsOwner);
+		}
+
+		// Remove any stale variable assignments for graphs/subgraphs which are no longer part of the job
+		InVariableAssignments.RemoveAll([&AllGraphs](const TObjectPtr<UMovieJobVariableAssignmentContainer>& VariableAssignment)
+		{
+			return !AllGraphs.Contains(VariableAssignment->GetGraphConfig().LoadSynchronous());
+		});
+	}
+}
+
 UMoviePipelineQueue::UMoviePipelineQueue()
 	: QueueSerialNumber(0)
 	, bIsDirty(false)
@@ -199,30 +258,52 @@ void UMoviePipelineExecutorJob::PreSave(FObjectPreSaveContext ObjectSaveContext)
 {
 	Super::PreSave(ObjectSaveContext);
 
-#if WITH_EDITOR
-	VariableAssignments->UpdateGraphVariableOverrides();
-#endif
+	RefreshAllVariableAssignments();
 }
 
 void UMoviePipelineExecutorJob::PostLoad()
 {
 	Super::PostLoad();
 	
-	if (VariableAssignments)
-	{
-		VariableAssignments->SetGraphConfig(GraphPreset.LoadSynchronous());
-	}
+	RefreshAllVariableAssignments();
+
+#if WITH_EDITOR
+	// Listen for any graph/subgraph saves. It's difficult/error-prone to track which graphs to monitor because the subgraph hierarchy can
+	// constantly change. Instead, just run any logic needed when any graph is saved.
+	FCoreUObjectDelegates::OnObjectPreSave.AddUObject(this, &UMoviePipelineExecutorJob::OnGraphPreSave);
+#endif
+}
+
+void UMoviePipelineExecutorJob::BeginDestroy()
+{
+	Super::BeginDestroy();
 
-	// Update the job's variable overrides whenever the graph's variables change
-	if (UMovieGraphConfig* Config = GraphPreset.LoadSynchronous())
-	{
 #if WITH_EDITOR
-		Config->OnGraphVariablesChangedDelegate.AddUObject(VariableAssignments, &UMovieJobVariableAssignmentContainer::UpdateGraphVariableOverrides);
+	FCoreUObjectDelegates::OnObjectPreSave.RemoveAll(this);
 #endif
+}
+
+void UMoviePipelineExecutorJob::RefreshAllVariableAssignments()
+{
+	UE::MovieGraph::Private::RefreshVariableAssignments(GraphPreset.LoadSynchronous(), VariableAssignments, this);
+
+	// Notify all shots that they should refresh their assignments as well (this could be needed, for example, when the primary graph preset changes)
+	for (TObjectPtr<UMoviePipelineExecutorShot>& ShotJob : ShotInfo)
+	{
+		if (ShotJob)
+		{
+			ShotJob->RefreshAllVariableAssignments();
+		}
 	}
+}
+
+void UMoviePipelineExecutorJob::SetGraphPreset(const UMovieGraphConfig* InGraphPreset)
+{
+	GraphPreset = InGraphPreset;
+
+	RefreshAllVariableAssignments();
 
-	// TODO: The variable assignments need to be updated here, but currently GraphPreset.LoadSynchronous() does not
-	// provide a fully-loaded object (ie, GraphPreset will only be partially loaded at this point).
+	OnJobGraphPresetChanged.Broadcast(this, GraphPreset.Get());
 }
 
 void UMoviePipelineExecutorJob::SetSequence(FSoftObjectPath InSequence)
@@ -247,6 +328,16 @@ void UMoviePipelineExecutorJob::SetSequence(FSoftObjectPath InSequence)
 	}
 }
 
+TObjectPtr<UMovieJobVariableAssignmentContainer> UMoviePipelineExecutorJob::GetOrCreateJobVariableAssignmentsForGraph(const UMovieGraphConfig* InGraph)
+{
+	if (InGraph)
+	{
+		return UE::MovieGraph::Private::GetOrCreateJobVariableAssignmentsForGraph(InGraph, VariableAssignments, this);
+	}
+
+	return nullptr;
+}
+
 void UMoviePipelineExecutorJob::SetConfiguration(UMoviePipelinePrimaryConfig* InPreset)
 {
 	if (InPreset)
@@ -273,6 +364,36 @@ void UMoviePipelineExecutorJob::OnDuplicated_Implementation()
 	SetConsumed(false);
 }
 
+void UMoviePipelineExecutorJob::OnGraphPreSave(UObject* InObject, FObjectPreSaveContext InObjectPreSaveContext)
+{
+	const UMovieGraphConfig* SavedGraph = Cast<UMovieGraphConfig>(InObject);
+	if (!SavedGraph || InObjectPreSaveContext.IsProceduralSave())
+	{
+		return;
+	}
+
+	RefreshAllVariableAssignments();
+}
+
+void UMoviePipelineExecutorShot::SetGraphPreset(const UMovieGraphConfig* InGraphPreset)
+{
+	GraphPreset = InGraphPreset;
+
+	RefreshAllVariableAssignments();
+
+	OnShotGraphPresetChanged.Broadcast(this, GraphPreset.Get());
+}
+
+TObjectPtr<UMovieJobVariableAssignmentContainer> UMoviePipelineExecutorShot::GetOrCreateJobVariableAssignmentsForGraph(const UMovieGraphConfig* InGraph, const bool bIsForPrimaryOverrides)
+{
+	if (InGraph)
+	{
+		return UE::MovieGraph::Private::GetOrCreateJobVariableAssignmentsForGraph(InGraph, bIsForPrimaryOverrides ? PrimaryGraphVariableAssignments : VariableAssignments, this);
+	}
+
+	return nullptr;
+}
+
 UMoviePipelineShotConfig* UMoviePipelineExecutorShot::AllocateNewShotOverrideConfig(TSubclassOf<UMoviePipelineShotConfig> InConfigType)
 {
 	if (!ensureAlwaysMsgf(InConfigType, TEXT("Failed to specify a Config Type. Use the default in project setting or UMoviePipelineShotConfig.")))
@@ -338,28 +459,52 @@ void UMoviePipelineExecutorShot::PreSave(FObjectPreSaveContext ObjectSaveContext
 {
 	Super::PreSave(ObjectSaveContext);
 
-#if WITH_EDITOR
-	VariableAssignments->UpdateGraphVariableOverrides();
-#endif
+	RefreshAllVariableAssignments();
 }
 
 void UMoviePipelineExecutorShot::PostLoad()
 {
 	Super::PostLoad();
 
-	if (VariableAssignments)
+	RefreshAllVariableAssignments();
+
+#if WITH_EDITOR
+	// Listen for any graph/subgraph saves. It's difficult/error-prone to track which graphs to monitor because the subgraph hierarchy can
+	// constantly change. Instead, just run any logic needed when any graph is saved.
+	FCoreUObjectDelegates::OnObjectPreSave.AddUObject(this, &UMoviePipelineExecutorShot::OnGraphPreSave);
+#endif
+}
+
+void UMoviePipelineExecutorShot::BeginDestroy()
+{
+	Super::BeginDestroy();
+
+#if WITH_EDITOR
+	FCoreUObjectDelegates::OnObjectPreSave.RemoveAll(this);
+#endif
+}
+
+void UMoviePipelineExecutorShot::RefreshAllVariableAssignments()
+{
+	UE::MovieGraph::Private::RefreshVariableAssignments(GraphPreset.LoadSynchronous(), VariableAssignments, this);
+
+	// Refresh the associated primary graph variable override assignments as well
+	if (const UMoviePipelineExecutorJob* PrimaryJob = GetTypedOuter<UMoviePipelineExecutorJob>())
 	{
-		VariableAssignments->SetGraphConfig(GraphPreset.LoadSynchronous());
+		if (UMovieGraphConfig* PrimaryGraph = PrimaryJob->GetGraphPreset())
+		{
+			UE::MovieGraph::Private::RefreshVariableAssignments(PrimaryGraph, PrimaryGraphVariableAssignments, this);
+		}
 	}
+}
 
-	// Update the job's variable overrides whenever the graph's variables change
-	if (UMovieGraphConfig* Config = GraphPreset.LoadSynchronous())
+void UMoviePipelineExecutorShot::OnGraphPreSave(UObject* InObject, FObjectPreSaveContext InObjectPreSaveContext)
+{
+	const UMovieGraphConfig* SavedGraph = Cast<UMovieGraphConfig>(InObject);
+	if (!SavedGraph || InObjectPreSaveContext.IsProceduralSave())
 	{
-#if WITH_EDITOR
-		Config->OnGraphVariablesChangedDelegate.AddUObject(VariableAssignments, &UMovieJobVariableAssignmentContainer::UpdateGraphVariableOverrides);
-#endif
+		return;
 	}
 
-	// TODO: The variable assignments need to be updated here, but currently GraphPreset.LoadSynchronous() does not
-	// provide a fully-loaded object (ie, GraphPreset will only be partially loaded at this point).
+	RefreshAllVariableAssignments();
 }
diff --git a/Engine/Plugins/MovieScene/MovieRenderPipeline/Source/MovieRenderPipelineCore/Public/Graph/Nodes/MovieGraphVariableNode.h b/Engine/Plugins/MovieScene/MovieRenderPipeline/Source/MovieRenderPipelineCore/Public/Graph/Nodes/MovieGraphVariableNode.h
--- a/Engine/Plugins/MovieScene/MovieRenderPipeline/Source/MovieRenderPipelineCore/Public/Graph/Nodes/MovieGraphVariableNode.h
+++ b/Engine/Plugins/MovieScene/MovieRenderPipeline/Source/MovieRenderPipelineCore/Public/Graph/Nodes/MovieGraphVariableNode.h
@@ -7,6 +7,8 @@
 
 #include "MovieGraphVariableNode.generated.h"
 
+class UMovieJobVariableAssignmentContainer;
+
 /** A node which gets the value of a variable which has been defined on the graph. */
 UCLASS()
 class MOVIERENDERPIPELINECORE_API UMovieGraphVariableNode : public UMovieGraphNode
@@ -44,8 +46,8 @@ private:
 	/** Updates the output pin on the node to match the provided variable. */
 	void UpdateOutputPin(UMovieGraphMember* ChangedVariable) const;
 
-	/** Determines if the job in the given context has a valid, enabled variable assignment for the variable this node represents. */
-	bool ContextHasEnabledAssignmentForVariable(const FMovieGraphTraversalContext* InContext) const;
+	/** Determines if the job in the given context has a valid, enabled variable assignment for the variable this node represents. Sets OutVariableAssignment if it does. */
+	bool ContextHasEnabledAssignmentForVariable(const FMovieGraphTraversalContext* InContext, TObjectPtr<UMovieJobVariableAssignmentContainer>& OutVariableAssignment) const;
 
 private:
 	/** The underlying graph variable this node represents. */
diff --git a/Engine/Plugins/MovieScene/MovieRenderPipeline/Source/MovieRenderPipelineCore/Public/MovieJobVariableAssignmentContainer.h b/Engine/Plugins/MovieScene/MovieRenderPipeline/Source/MovieRenderPipelineCore/Public/MovieJobVariableAssignmentContainer.h
--- a/Engine/Plugins/MovieScene/MovieRenderPipeline/Source/MovieRenderPipelineCore/Public/MovieJobVariableAssignmentContainer.h
+++ b/Engine/Plugins/MovieScene/MovieRenderPipeline/Source/MovieRenderPipelineCore/Public/MovieJobVariableAssignmentContainer.h
@@ -31,6 +31,9 @@ public:
 	UFUNCTION(BlueprintCallable, Category="Experimental")
 	void SetGraphConfig(const TSoftObjectPtr<UMovieGraphConfig>& InGraphConfig);
 
+	/** Gets the graph that is associated with this container. */
+	TSoftObjectPtr<UMovieGraphConfig> GetGraphConfig() const;
+
 	/** Gets the number of variable assignments present in this container. Assignments that are disabled are counted. */
 	uint32 GetNumAssignments() const;
 
@@ -221,13 +224,13 @@ public:
 	UFUNCTION(BlueprintCallable, Category = "Experimental")
 	bool GetVariableAssignmentEnableState(const UMovieGraphVariable* InGraphVariable, bool& bOutIsEnabled);
 
+#if WITH_EDITOR
 	/**
 	 * Updates the stored variable overrides to reflect the graph preset. Existing overrides will be updated to match
 	 * the graph variable name, value type, object type, and container type. Additionally, stale overrides that have no
 	 * corresponding graph variable will be removed, and overrides will be created for graph variables which do not have
 	 * existing overrides.
 	 */
-#if WITH_EDITOR
 	UFUNCTION(BlueprintCallable, Category = "Experimental")
 	void UpdateGraphVariableOverrides();
 #endif
@@ -252,6 +255,7 @@ protected:
 
 	/** Given a graph variable, retrieve the hidden internal name used by our actual property bag. */
 	FName ConvertVariableToInternalName(const UMovieGraphVariable* InGraphVariable) const;
+
 private:
 	/** The properties managed by this object. */
 	UPROPERTY(EditAnywhere, meta=(ShowOnlyInnerProperties, FixedLayout), Category = "Value")
diff --git a/Engine/Plugins/MovieScene/MovieRenderPipeline/Source/MovieRenderPipelineCore/Public/MoviePipelineQueue.h b/Engine/Plugins/MovieScene/MovieRenderPipeline/Source/MovieRenderPipelineCore/Public/MoviePipelineQueue.h
--- a/Engine/Plugins/MovieScene/MovieRenderPipeline/Source/MovieRenderPipelineCore/Public/MoviePipelineQueue.h
+++ b/Engine/Plugins/MovieScene/MovieRenderPipeline/Source/MovieRenderPipelineCore/Public/MoviePipelineQueue.h
@@ -38,6 +38,8 @@ public:
 	FString Name;
 };
 
+DECLARE_MULTICAST_DELEGATE_TwoParams(FOnMoviePipelineShotGraphPresetChanged, UMoviePipelineExecutorShot*, UMovieGraphConfig*);
+
 /**
 * This class represents a segment of work within the Executor Job. This should be owned
 * by the UMoviePipelineExecutorJob and can be created before the movie pipeline starts to
@@ -53,7 +55,6 @@ public:
 		: bEnabled(true)
 	{
 		Progress = 0.f;
-		VariableAssignments = CreateDefaultSubobject<UMovieJobVariableAssignmentContainer>("VariableAssignments");
 	}
 
 public:
@@ -154,15 +155,10 @@ public:
 	 * graph-style configuration if it is not already using it.
 	 */
 	UFUNCTION(BlueprintCallable, Category = "Movie Render Pipeline")
-	void SetGraphPreset(const UMovieGraphConfig* InGraphPreset)
-	{
-		GraphPreset = InGraphPreset;
-		VariableAssignments->SetGraphConfig(GraphPreset);
+	void SetGraphPreset(const UMovieGraphConfig* InGraphPreset);
 
-#if WITH_EDITOR
-		VariableAssignments->UpdateGraphVariableOverrides();
-#endif
-	}
+	/** Gets the job variable assignments for a specific graph. Creates a new variable assignments container if one was not found for the given graph. */
+	TObjectPtr<UMovieJobVariableAssignmentContainer> GetOrCreateJobVariableAssignmentsForGraph(const UMovieGraphConfig* InGraph, const bool bIsForPrimaryOverrides = false);
 
 	/** Returns whether this should should be rendered */
 	UFUNCTION(BlueprintPure, Category = "Movie Render Pipeline")
@@ -184,7 +180,11 @@ public:
 	// UObject Interface
 	virtual void PreSave(FObjectPreSaveContext ObjectSaveContext) override;
 	virtual void PostLoad() override;
+	virtual void BeginDestroy() override;
 	// ~UObject Interface
+
+	/** Refreshes the variable assignments associated with this shot, both for the shot's own graph preset and the associated primary graph. */
+	void RefreshAllVariableAssignments();
 	
 protected:
 	// UMoviePipipelineExecutorShot Interface
@@ -193,28 +193,42 @@ protected:
 	virtual FString GetStatusMessage_Implementation() const { return StatusMessage; }
 	virtual float GetStatusProgress_Implementation() const { return Progress; }
 	// ~UMoviePipipelineExecutorShot Interface
+	
+	void OnGraphPreSave(UObject* InObject, FObjectPreSaveContext InObjectPreSaveContext);
 
 public:
 
 	/** Does the user want to render this shot? */
-	UPROPERTY(BlueprintReadWrite, EditAnywhere, Category = "Movie Render Pipeline")
+	UPROPERTY(BlueprintReadWrite, Category = "Movie Render Pipeline")
 	bool bEnabled;
 
 	/** The name of the shot section that contains this shot. Can be empty. */
-	UPROPERTY(BlueprintReadWrite, EditAnywhere, Category = "Movie Render Pipeline")
+	UPROPERTY(BlueprintReadWrite, Category = "Movie Render Pipeline")
 	FString OuterName;
 
 	/** The name of the camera cut section that this shot represents. Can be empty. */
-	UPROPERTY(BlueprintReadWrite, EditAnywhere, Category = "Movie Render Pipeline")
+	UPROPERTY(BlueprintReadWrite, Category = "Movie Render Pipeline")
 	FString InnerName;
 
 	/** List of cameras to render for this shot. Only used if the setting flag is set in the Camera setting. */
-	UPROPERTY(BlueprintReadWrite, EditAnywhere, Category = "Movie Render Pipeline")
+	UPROPERTY(BlueprintReadWrite, Category = "Movie Render Pipeline")
 	TArray<FMoviePipelineSidecarCamera> SidecarCameras;
 
-	/** Overrides on the variables in graph preset associated with this job. */
-	UPROPERTY(BlueprintReadWrite, EditAnywhere, Category = "Variable Assignments")
-	TObjectPtr<UMovieJobVariableAssignmentContainer> VariableAssignments;
+	/**
+	 * Overrides on the variables in graph presets associated with this job. A job can have multiple graphs associated with it if the job's assigned
+	 * graph contains subgraphs. The job's graph and each subgraph will have an entry in this array.
+	 */
+	UPROPERTY(BlueprintReadWrite, EditAnywhere, Instanced, Category = "Variable Assignments")
+	TArray<TObjectPtr<UMovieJobVariableAssignmentContainer>> VariableAssignments;
+	
+	/**
+	 * Overrides on the variables in the primary graph (and its subgraphs) associated with this job.
+	 */
+	UPROPERTY(BlueprintReadWrite, EditAnywhere, Instanced, Category = "Variable Assignments")
+	TArray<TObjectPtr<UMovieJobVariableAssignmentContainer>> PrimaryGraphVariableAssignments;
+
+	/** Called when the graph preset assigned to the shot changes. */
+	FOnMoviePipelineShotGraphPresetChanged OnShotGraphPresetChanged;
 
 public:
 	/** Transient information used by the active Movie Pipeline working on this shot. */
@@ -238,6 +252,8 @@ private:
 	TSoftObjectPtr<UMovieGraphConfig> GraphPreset;
 };
 
+DECLARE_MULTICAST_DELEGATE_TwoParams(FOnMoviePipelineJobGraphPresetChanged, UMoviePipelineExecutorJob*, UMovieGraphConfig*);
+
 /**
 * A particular job within the Queue
 */
@@ -252,7 +268,6 @@ public:
 		StatusProgress = 0.f;
 		bIsConsumed = false;
 		Configuration = CreateDefaultSubobject<UMoviePipelinePrimaryConfig>("DefaultConfig");
-		VariableAssignments = CreateDefaultSubobject<UMovieJobVariableAssignmentContainer>("VariableAssignments");
 	}
 
 public:	
@@ -437,19 +452,14 @@ public:
 	 * graph-style configuration if it is not already using it.
 	 */
 	UFUNCTION(BlueprintCallable, Category = "Movie Render Pipeline")
-	void SetGraphPreset(const UMovieGraphConfig* InGraphPreset)
-	{
-		GraphPreset = InGraphPreset;
-		VariableAssignments->SetGraphConfig(InGraphPreset);
-
-#if WITH_EDITOR
-		VariableAssignments->UpdateGraphVariableOverrides();
-#endif
-	}
+	void SetGraphPreset(const UMovieGraphConfig* InGraphPreset);
 
 	UFUNCTION(BlueprintSetter, Category = "Movie Render Pipeline")
 	void SetSequence(FSoftObjectPath InSequence);
 
+	/** Gets the job variable assignments for a specific graph. Creates a new variable assignments container if one was not found for the given graph. */
+	TObjectPtr<UMovieJobVariableAssignmentContainer> GetOrCreateJobVariableAssignmentsForGraph(const UMovieGraphConfig* InGraph);
+
 public:
 	// UObject Interface
 #if WITH_EDITOR
@@ -457,8 +467,12 @@ public:
 #endif
 	virtual void PreSave(FObjectPreSaveContext ObjectSaveContext) override;
 	virtual void PostLoad() override;
+	virtual void BeginDestroy() override;
 	// ~UObject Interface
 
+	/** Refreshes the variable assignments associated with this job, both for the jobs's own graph preset and the associated shot graphs. */
+	void RefreshAllVariableAssignments();
+
 protected:
 	// UMoviePipelineExecutorJob Interface
 	virtual void SetStatusMessage_Implementation(const FString& InMessage) { StatusMessage = InMessage; }
@@ -478,6 +492,8 @@ protected:
 	virtual bool IsEnabled_Implementation() const { return bEnabled; }
 	// ~UMoviePipelineExecutorJob Interface
 
+	void OnGraphPreSave(UObject* InObject, FObjectPreSaveContext InObjectPreSaveContext);
+	
 public:
 	/** (Optional) Name of the job. Shown on the default burn-in. */
 	UPROPERTY(BlueprintReadWrite, EditAnywhere, Category = "Movie Render Pipeline")
@@ -511,9 +527,15 @@ public:
 	UPROPERTY(BlueprintReadWrite, Category = "Movie Render Pipeline")
 	FString UserData;
 
-	/** Overrides on the variables in graph preset associated with this job. */
+	/**
+	 * Overrides on the variables in graph presets associated with this job. A job can have multiple graphs associated with it if the job's assigned
+	 * graph contains subgraphs. The job's graph and each subgraph will have an entry in this array.
+	 */
 	UPROPERTY(BlueprintReadWrite, EditAnywhere, Instanced, Category = "Variable Assignments")
-	TObjectPtr<UMovieJobVariableAssignmentContainer> VariableAssignments;
+	TArray<TObjectPtr<UMovieJobVariableAssignmentContainer>> VariableAssignments;
+
+	/** Called when the graph preset assigned to the job changes. */
+	FOnMoviePipelineJobGraphPresetChanged OnJobGraphPresetChanged;
 
 private:
 	UPROPERTY(Transient)
diff --git a/Engine/Plugins/MovieScene/MovieRenderPipeline/Source/MovieRenderPipelineEditor/Private/Customizations/JobCustomization.h b/Engine/Plugins/MovieScene/MovieRenderPipeline/Source/MovieRenderPipelineEditor/Private/Customizations/JobCustomization.h
--- a/Engine/Plugins/MovieScene/MovieRenderPipeline/Source/MovieRenderPipelineEditor/Private/Customizations/JobCustomization.h
+++ b/Engine/Plugins/MovieScene/MovieRenderPipeline/Source/MovieRenderPipelineEditor/Private/Customizations/JobCustomization.h
@@ -1,15 +1,14 @@
-﻿// Copyright Epic Games, Inc. All Rights Reserved.
+// Copyright Epic Games, Inc. All Rights Reserved.
 
 #pragma once
 
 #include "DetailCategoryBuilder.h"
 #include "DetailLayoutBuilder.h"
 #include "IDetailCustomization.h"
+#include "IDetailGroup.h"
 #include "MoviePipelineQueue.h"
-#include "PropertyCustomizationHelpers.h"
 #include "PropertyHandle.h"
-#include "Widgets/Input/SNumericEntryBox.h"
-#include "Widgets/Text/STextBlock.h"
+#include "UObject/Package.h"
 
 #define LOCTEXT_NAMESPACE "MoviePipelineEditor"
 
@@ -22,46 +21,151 @@ public:
 		return MakeShared<FJobDetailsCustomization>();
 	}
 
+	virtual ~FJobDetailsCustomization() override
+	{
+		UPackage::PackageSavedWithContextEvent.RemoveAll(this);
+	}
+
 protected:
 	//~ Begin IDetailCustomization interface
-	virtual void CustomizeDetails(const TSharedPtr<IDetailLayoutBuilder>& DetailBuilder) override
+	virtual void CustomizeDetails(const TSharedPtr<IDetailLayoutBuilder>& InDetailBuilder) override
+	{
+		DetailBuilder = InDetailBuilder.Get();
+		CustomizeDetails(*InDetailBuilder);
+	}
+
+	void RefreshLayout(const FString&, UPackage*, FObjectPostSaveContext) const
+	{
+		DetailBuilder->ForceRefreshDetails();
+	}
+
+	void RefreshLayout(UMoviePipelineExecutorShot*, UMovieGraphConfig*) const
+	{
+		DetailBuilder->ForceRefreshDetails();
+	}
+
+	void RefreshLayout(UMoviePipelineExecutorJob*, UMovieGraphConfig*) const
 	{
-		CustomizeDetails(*DetailBuilder);
+		DetailBuilder->ForceRefreshDetails();
 	}
 
-	virtual void CustomizeDetails(IDetailLayoutBuilder& DetailBuilder) override
+	virtual void CustomizeDetails(IDetailLayoutBuilder& InDetailBuilder) override
 	{
+		// Refresh the customization every time a save happens. Use this opportunity to update the variables in the UI. We could update the UI before
+		// a save occurs, but this would be very difficult to get right when multiple subgraphs are involved.
+		UPackage::PackageSavedWithContextEvent.AddSP(this, &FJobDetailsCustomization::RefreshLayout);
+		
 		TArray<TWeakObjectPtr<UObject>> ObjectsBeingCustomized;
-		DetailBuilder.GetObjectsBeingCustomized(ObjectsBeingCustomized);
+		InDetailBuilder.GetObjectsBeingCustomized(ObjectsBeingCustomized);
+
+		TArray<UMoviePipelineExecutorJob*> SelectedJobs;
+		TArray<UMoviePipelineExecutorShot*> SelectedShots;
 
-		// Only display variables if there is one job selected
-		if (ObjectsBeingCustomized.Num() != 1)
+		for (const TWeakObjectPtr<UObject>& SelectedObject : ObjectsBeingCustomized)
 		{
-			return;
+			if (UMoviePipelineExecutorJob* SelectedJob = Cast<UMoviePipelineExecutorJob>(SelectedObject.Get()))
+			{
+				SelectedJobs.Add(SelectedJob);
+			}
+			else if (UMoviePipelineExecutorShot* SelectedShot = Cast<UMoviePipelineExecutorShot>(SelectedObject.Get()))
+			{
+				SelectedShots.Add(SelectedShot);
+			}
 		}
 
-		// Hide the original assignments property (since it presents an asset picker by default), and add the Value
-		// property (which is a property bag) instead -- that will present each of the variable overrides
-		if (UMoviePipelineExecutorJob* Job = CastChecked<UMoviePipelineExecutorJob>(ObjectsBeingCustomized[0]))
+		// Hide the original assignments property (since it presents an asset picker by default) for both jobs and shots
+		const TSharedRef<IPropertyHandle> JobAssignmentsProperty =
+			InDetailBuilder.GetProperty(GET_MEMBER_NAME_CHECKED(UMoviePipelineExecutorJob, VariableAssignments));
+		const TSharedRef<IPropertyHandle> ShotAssignmentsProperty =
+			InDetailBuilder.GetProperty(GET_MEMBER_NAME_CHECKED(UMoviePipelineExecutorShot, VariableAssignments));
+		const TSharedRef<IPropertyHandle> ShotPrimaryGraphAssignmentsProperty =
+			InDetailBuilder.GetProperty(GET_MEMBER_NAME_CHECKED(UMoviePipelineExecutorShot, PrimaryGraphVariableAssignments));
+		InDetailBuilder.HideProperty(JobAssignmentsProperty);
+		InDetailBuilder.HideProperty(ShotAssignmentsProperty);
+		InDetailBuilder.HideProperty(ShotPrimaryGraphAssignmentsProperty);
+
+		// Only display the customized variables UI if there is one job or shot selected
+		const bool bIsPrimaryJob = (SelectedJobs.Num() == 1) && SelectedShots.IsEmpty();
+		const bool bIsShot = (SelectedShots.Num() == 1) && SelectedJobs.IsEmpty();
+		if (!bIsPrimaryJob && !bIsShot)
 		{
-			const TSharedRef<IPropertyHandle> AssignmentsProperty =
-				DetailBuilder.GetProperty(GET_MEMBER_NAME_CHECKED(UMoviePipelineExecutorJob, VariableAssignments));
+			return;
+		}
 
-			// Always hide the assignments property
-			DetailBuilder.HideProperty(AssignmentsProperty);
+		// Refresh the UI if the graph preset changes (so the new variable assignments are displayed)
+		if (bIsShot)
+		{
+			SelectedShots[0]->OnShotGraphPresetChanged.AddSP(this, &FJobDetailsCustomization::RefreshLayout);
 
-			if (Job->VariableAssignments->GetNumAssignments() > 0)
+			// Also listen for changes to the primary job. Changes to the primary job can trigger an update to shot variable assignments.
+			if (UMoviePipelineExecutorJob* PrimaryJob = SelectedShots[0]->GetTypedOuter<UMoviePipelineExecutorJob>())
 			{
-				// Add a new "Graph Variables" category if there is at least one assignment. Set as "Uncommon" priority to
-				// push variables down below the other properties.
-				IDetailCategoryBuilder& GraphVariablesCategory = DetailBuilder.EditCategory(
-					"GraphVariables", LOCTEXT("GraphVariablesCategory", "Graph Variables"), ECategoryPriority::Uncommon);
-			
-				GraphVariablesCategory.AddExternalObjectProperty({Job->VariableAssignments.Get()}, FName("Value"));
+				PrimaryJob->OnJobGraphPresetChanged.AddSP(this, &FJobDetailsCustomization::RefreshLayout);
 			}
 		}
+		else
+		{
+			SelectedJobs[0]->OnJobGraphPresetChanged.AddSP(this, &FJobDetailsCustomization::RefreshLayout);
+		}
+
+		// Set up the categories for variable assignments. Set as "Uncommon" priority to push variables down below the other properties.
+		IDetailCategoryBuilder& PrimaryGraphVariablesCategory = InDetailBuilder.EditCategory(
+			"PrimaryGraphVariables", LOCTEXT("PrimaryGraphVariablesCategory", "Primary Graph Variables"), ECategoryPriority::Uncommon);
+		IDetailCategoryBuilder& PrimaryGraphVariablesShotOverridesCategory = InDetailBuilder.EditCategory(
+			"PrimaryGraphVariablesShotOverrides", LOCTEXT("PrimaryGraphVariablesShotOverridesCategory", "Primary Graph Variables (shot overrides)"), ECategoryPriority::Uncommon);
+		IDetailCategoryBuilder& ShotGraphVariablesCategory = InDetailBuilder.EditCategory(
+			"ShotGraphVariables", LOCTEXT("ShotGraphVariablesCategory", "Shot Graph Variables"), ECategoryPriority::Uncommon);
+
+		// Set all categories as hidden by default. Individual categories will be made visible if variables are added under them.
+		PrimaryGraphVariablesCategory.SetCategoryVisibility(false);
+		PrimaryGraphVariablesShotOverridesCategory.SetCategoryVisibility(false);
+		ShotGraphVariablesCategory.SetCategoryVisibility(false);
+		
+		if (bIsShot)
+		{
+			AddVariableAssignments(SelectedShots[0]->VariableAssignments, ShotGraphVariablesCategory);
+			AddVariableAssignments(SelectedShots[0]->PrimaryGraphVariableAssignments, PrimaryGraphVariablesShotOverridesCategory);
+		}
+		else
+		{
+			AddVariableAssignments(SelectedJobs[0]->VariableAssignments, PrimaryGraphVariablesCategory);
+		}
 	}
 	//~ End IDetailCustomization interface
+
+private:
+	void AddVariableAssignments(TArray<TObjectPtr<UMovieJobVariableAssignmentContainer>>& VariableAssignments, IDetailCategoryBuilder& InCategory)
+	{
+		// Add a sub-category for each graph (including subgraphs). Each entry in the array represents the assignments for one graph.
+		for (TObjectPtr<UMovieJobVariableAssignmentContainer>& VariableAssignment : VariableAssignments)
+		{
+			// Skip if the graph associated with this container has no variables in it
+			if (VariableAssignment->GetNumAssignments() <= 0)
+			{
+				continue;
+			}
+
+			// If the graph can be found, display its variable assignments under its own category (group)
+			TSoftObjectPtr<UMovieGraphConfig> SoftGraphConfig = VariableAssignment->GetGraphConfig();
+			if (const UMovieGraphConfig* GraphConfig = SoftGraphConfig.Get())
+			{
+				constexpr bool bForAdvanced = false;
+				constexpr bool bStartExpanded = true;
+				IDetailGroup& GraphGroup = InCategory.AddGroup(GraphConfig->GetFName(), FText::FromString(GraphConfig->GetName()), bForAdvanced, bStartExpanded);
+
+				// "Value" is private so we can't use GET_MEMBER_NAME_CHECKED unfortunately
+				TSharedPtr<IPropertyHandle> ValueProperty = DetailBuilder->AddObjectPropertyData({VariableAssignment}, FName("Value"));
+				GraphGroup.AddPropertyRow(ValueProperty.ToSharedRef());
+
+				// Un-hide the category if it's currently visible
+				InCategory.SetCategoryVisibility(true);
+			}
+		}
+	}
+
+private:
+	/** The details builder associated with the customization. */
+	IDetailLayoutBuilder* DetailBuilder = nullptr;
 };
 
 #undef LOCTEXT_NAMESPACE
\ No newline at end of file
diff --git a/Engine/Plugins/MovieScene/MovieRenderPipeline/Source/MovieRenderPipelineEditor/Private/Widgets/Graph/SMovieGraphConfigPanel.cpp b/Engine/Plugins/MovieScene/MovieRenderPipeline/Source/MovieRenderPipelineEditor/Private/Widgets/Graph/SMovieGraphConfigPanel.cpp
--- a/Engine/Plugins/MovieScene/MovieRenderPipeline/Source/MovieRenderPipelineEditor/Private/Widgets/Graph/SMovieGraphConfigPanel.cpp
+++ b/Engine/Plugins/MovieScene/MovieRenderPipeline/Source/MovieRenderPipelineEditor/Private/Widgets/Graph/SMovieGraphConfigPanel.cpp
@@ -906,7 +906,7 @@ void SMoviePipelineGraphPanel::OnConfigUpdatedForJobToPreset(TWeakObjectPtr<UMov
 	OnConfigWindowClosed();
 }
 
-void SMoviePipelineGraphPanel::OnSelectionChanged(const TArray<UMoviePipelineExecutorJob*>& InSelectedJobs)
+void SMoviePipelineGraphPanel::OnSelectionChanged(const TArray<UMoviePipelineExecutorJob*>& InSelectedJobs, const TArray<UMoviePipelineExecutorShot*>& InSelectedShots)
 {
 	TArray<UObject*> Jobs;
 	for (UMoviePipelineExecutorJob* Job : InSelectedJobs)
diff --git a/Engine/Plugins/MovieScene/MovieRenderPipeline/Source/MovieRenderPipelineEditor/Private/Widgets/Graph/SMovieGraphConfigPanel.h b/Engine/Plugins/MovieScene/MovieRenderPipeline/Source/MovieRenderPipelineEditor/Private/Widgets/Graph/SMovieGraphConfigPanel.h
--- a/Engine/Plugins/MovieScene/MovieRenderPipeline/Source/MovieRenderPipelineEditor/Private/Widgets/Graph/SMovieGraphConfigPanel.h
+++ b/Engine/Plugins/MovieScene/MovieRenderPipeline/Source/MovieRenderPipelineEditor/Private/Widgets/Graph/SMovieGraphConfigPanel.h
@@ -64,7 +64,7 @@ private:
 	void OnConfigUpdatedForJobToPreset(TWeakObjectPtr<UMoviePipelineExecutorJob> InJob, TWeakObjectPtr<UMoviePipelineExecutorShot> InShot, UMoviePipelineConfigBase* InConfig);
 	void OnConfigWindowClosed();
 
-	void OnSelectionChanged(const TArray<UMoviePipelineExecutorJob*>& InSelectedJobs);
+	void OnSelectionChanged(const TArray<UMoviePipelineExecutorJob*>& InSelectedJobs, const TArray<UMoviePipelineExecutorShot*>& InSelectedShots);
 	TArray<UMovieGraphNode*> GetSelectedModelNodes() const;
 
 	TSharedRef<SWidget> OnGenerateSavedQueuesMenu();
diff --git a/Engine/Plugins/MovieScene/MovieRenderPipeline/Source/MovieRenderPipelineEditor/Private/Widgets/SMoviePipelineQueueEditor.cpp b/Engine/Plugins/MovieScene/MovieRenderPipeline/Source/MovieRenderPipelineEditor/Private/Widgets/SMoviePipelineQueueEditor.cpp
--- a/Engine/Plugins/MovieScene/MovieRenderPipeline/Source/MovieRenderPipelineEditor/Private/Widgets/SMoviePipelineQueueEditor.cpp
+++ b/Engine/Plugins/MovieScene/MovieRenderPipeline/Source/MovieRenderPipelineEditor/Private/Widgets/SMoviePipelineQueueEditor.cpp
@@ -88,6 +88,7 @@ struct IMoviePipelineQueueTreeItem : TSharedFromThis<IMoviePipelineQueueTreeItem
 
 	virtual TSharedPtr<FMoviePipelineQueueJobTreeItem> AsJob() { return nullptr; }
 	virtual UMoviePipelineExecutorJob* GetOwningJob() { return nullptr; }
+	virtual UMoviePipelineExecutorShot* GetOwningShot() { return nullptr; }
 	virtual void Delete(UMoviePipelineQueue* InOwningQueue) {}
 	virtual void ResetStatus() {}
 	virtual UMoviePipelineExecutorJob* Duplicate(UMoviePipelineQueue* InOwningQueue) { return nullptr; }
@@ -834,6 +835,11 @@ struct FMoviePipelineShotItem : IMoviePipelineQueueTreeItem
 		return WeakJob.Get();
 	}
 
+	virtual UMoviePipelineExecutorShot* GetOwningShot() override
+	{
+		return WeakShot.Get();
+	}
+
 	virtual TSharedRef<ITableRow> ConstructWidget(TWeakPtr<SMoviePipelineQueueEditor> InQueueWidget, const TSharedRef<STableViewBase>& OwnerTable) override
 	{
 		WeakQueueEditor = InQueueWidget;
@@ -1691,21 +1697,27 @@ void SMoviePipelineQueueEditor::SetSelectedJobs_Impl(const TArray<UMoviePipeline
 void SMoviePipelineQueueEditor::OnJobSelectionChanged_Impl(TSharedPtr<IMoviePipelineQueueTreeItem> TreeItem, ESelectInfo::Type SelectInfo)
 {
 	TArray<UMoviePipelineExecutorJob*> SelectedJobs;
+	TArray<UMoviePipelineExecutorShot*> SelectedShots;
+	
 	if (TreeItem.IsValid())
 	{
 		// Iterate the tree and get all selected items.
 		TArray<TSharedPtr<IMoviePipelineQueueTreeItem>> SelectedTreeItems = TreeView->GetSelectedItems();
 		for (TSharedPtr<IMoviePipelineQueueTreeItem> Item : SelectedTreeItems)
 		{
-			UMoviePipelineExecutorJob* Job = Item->GetOwningJob();
-			if (Job)
+			if (UMoviePipelineExecutorJob* Job = Item->GetOwningJob())
 			{
 				SelectedJobs.AddUnique(Job);
 			}
+
+			if (UMoviePipelineExecutorShot* Shot = Item->GetOwningShot())
+			{
+				SelectedShots.AddUnique(Shot);
+			}
 		}
 	}
 
-	OnJobSelectionChanged.ExecuteIfBound(SelectedJobs);
+	OnJobSelectionChanged.ExecuteIfBound(SelectedJobs, SelectedShots);
 }
 
 void SMoviePipelineQueueEditor::OnQueueLoaded()
diff --git a/Engine/Plugins/MovieScene/MovieRenderPipeline/Source/MovieRenderPipelineEditor/Private/Widgets/SMoviePipelineQueueEditor.h b/Engine/Plugins/MovieScene/MovieRenderPipeline/Source/MovieRenderPipelineEditor/Private/Widgets/SMoviePipelineQueueEditor.h
--- a/Engine/Plugins/MovieScene/MovieRenderPipeline/Source/MovieRenderPipelineEditor/Private/Widgets/SMoviePipelineQueueEditor.h
+++ b/Engine/Plugins/MovieScene/MovieRenderPipeline/Source/MovieRenderPipelineEditor/Private/Widgets/SMoviePipelineQueueEditor.h
@@ -25,7 +25,7 @@ class UMovieSceneCinematicShotSection;
 struct FMoviePipelineQueueJobTreeItem;
 
 DECLARE_DELEGATE_TwoParams(FOnMoviePipelineEditConfig, TWeakObjectPtr<UMoviePipelineExecutorJob>, TWeakObjectPtr<UMoviePipelineExecutorShot>)
-DECLARE_DELEGATE_OneParam(FOnMoviePipelineJobSelection, const TArray<UMoviePipelineExecutorJob*>&)
+DECLARE_DELEGATE_TwoParams(FOnMoviePipelineJobSelection, const TArray<UMoviePipelineExecutorJob*>&, const TArray<UMoviePipelineExecutorShot*>&)
 
 /**
  * Widget used to edit a Movie Pipeline Queue
diff --git a/Engine/Plugins/MovieScene/MovieRenderPipeline/Source/MovieRenderPipelineEditor/Private/Widgets/SMoviePipelineQueuePanel.cpp b/Engine/Plugins/MovieScene/MovieRenderPipeline/Source/MovieRenderPipelineEditor/Private/Widgets/SMoviePipelineQueuePanel.cpp
--- a/Engine/Plugins/MovieScene/MovieRenderPipeline/Source/MovieRenderPipelineEditor/Private/Widgets/SMoviePipelineQueuePanel.cpp
+++ b/Engine/Plugins/MovieScene/MovieRenderPipeline/Source/MovieRenderPipelineEditor/Private/Widgets/SMoviePipelineQueuePanel.cpp
@@ -72,6 +72,10 @@ void SMoviePipelineQueuePanel::Construct(const FArguments& InArgs)
 		UMoviePipelineExecutorJob::StaticClass(),
 		FOnGetDetailCustomizationInstance::CreateStatic(&FJobDetailsCustomization::MakeInstance));
 
+	JobDetailsPanelWidget->RegisterInstancedCustomPropertyLayout(
+		UMoviePipelineExecutorShot::StaticClass(),
+		FOnGetDetailCustomizationInstance::CreateStatic(&FJobDetailsCustomization::MakeInstance));
+
 	JobDetailsPanelWidget->RegisterInstancedCustomPropertyTypeLayout(
 		FMovieGraphNamedResolution::StaticStruct()->GetFName(),
 		FOnGetPropertyTypeCustomizationInstance::CreateStatic(&FMovieGraphNamedResolutionCustomization::MakeInstance));
@@ -499,12 +503,18 @@ void SMoviePipelineQueuePanel::OnConfigUpdatedForJobToPreset(TWeakObjectPtr<UMov
 	OnConfigWindowClosed();
 }
 
-void SMoviePipelineQueuePanel::OnSelectionChanged(const TArray<UMoviePipelineExecutorJob*>& InSelectedJobs)
+void SMoviePipelineQueuePanel::OnSelectionChanged(const TArray<UMoviePipelineExecutorJob*>& InSelectedJobs, const TArray<UMoviePipelineExecutorShot*>& InSelectedShots)
 {
 	TArray<UObject*> Jobs;
-	for (UMoviePipelineExecutorJob* Job : InSelectedJobs)
+
+	// Select the shot if a shot is selected, otherwise select the primary job
+	if (!InSelectedShots.IsEmpty())
+	{
+		Jobs.Append(InSelectedShots);
+	}
+	else
 	{
-		Jobs.Add(Job);
+		Jobs.Append(InSelectedJobs);
 	}
 	
 	JobDetailsPanelWidget->SetObjects(Jobs);
diff --git a/Engine/Plugins/MovieScene/MovieRenderPipeline/Source/MovieRenderPipelineEditor/Private/Widgets/SMoviePipelineQueuePanel.h b/Engine/Plugins/MovieScene/MovieRenderPipeline/Source/MovieRenderPipelineEditor/Private/Widgets/SMoviePipelineQueuePanel.h
--- a/Engine/Plugins/MovieScene/MovieRenderPipeline/Source/MovieRenderPipelineEditor/Private/Widgets/SMoviePipelineQueuePanel.h
+++ b/Engine/Plugins/MovieScene/MovieRenderPipeline/Source/MovieRenderPipelineEditor/Private/Widgets/SMoviePipelineQueuePanel.h
@@ -48,7 +48,7 @@ private:
 	void OnConfigUpdatedForJobToPreset(TWeakObjectPtr<UMoviePipelineExecutorJob> InJob, TWeakObjectPtr<UMoviePipelineExecutorShot> InShot, UMoviePipelineConfigBase* InConfig);
 	void OnConfigWindowClosed();
 
-	void OnSelectionChanged(const TArray<UMoviePipelineExecutorJob*>& InSelectedJobs);
+	void OnSelectionChanged(const TArray<UMoviePipelineExecutorJob*>& InSelectedJobs, const TArray<UMoviePipelineExecutorShot*>& InSelectedShots);
 	int32 GetDetailsViewWidgetIndex() const;
 	bool IsDetailsViewEnabled() const;
 
