Movie Graph:

-Fix: Move the MovieGraphimagePassBaseNode to the MovieRenderPipelineRenderPasses module

diff --git a/Engine/Plugins/MovieScene/MovieRenderPipeline/Config/BaseMovieRenderPipeline.ini b/Engine/Plugins/MovieScene/MovieRenderPipeline/Config/BaseMovieRenderPipeline.ini
--- a/Engine/Plugins/MovieScene/MovieRenderPipeline/Config/BaseMovieRenderPipeline.ini
+++ b/Engine/Plugins/MovieScene/MovieRenderPipeline/Config/BaseMovieRenderPipeline.ini
@@ -3,3 +3,7 @@
 +ClassRedirects=(OldName="/Script/MovieRenderPipelineCore.MovieGraphOutputSettingNode",NewName="/Script/MovieRenderPipelineCore.MovieGraphGlobalOutputSettingNode")
 +PropertyRedirects=(OldName="/Script/MovieRenderPipelineCore.MovieGraphCommandLineEncoderNode.FileNameFormatOverride",NewName="/Script/MovieRenderPipelineCore.MovieGraphCommandLineEncoderNode.FileNameFormat")
 +PropertyRedirects=(OldName="/Script/MovieRenderPipelineCore.MovieGraphCommandLineEncoderNode.bOverride_FileNameFormatOverride",NewName="/Script/MovieRenderPipelineCore.MovieGraphCommandLineEncoderNode.bOverride_FileNameFormat")
+
++ClassRedirects=(OldName="/Script/MovieRenderPipelineCore.MovieGraphDeferredRenderPassNode",NewName="/Script/MovieRenderPipelineRenderPasses.MovieGraphDeferredRenderPassNode")
++ClassRedirects=(OldName="/Script/MovieRenderPipelineCore.MovieGraphImagePassBaseNode",NewName="/Script/MovieRenderPipelineRenderPasses.MovieGraphImagePassBaseNode")
++ClassRedirects=(OldName="/Script/MovieRenderPipelineCore.MovieGraphPathTracedRenderPassNode",NewName="/Script/MovieRenderPipelineRenderPasses.MovieGraphPathTracedRenderPassNode")
diff --git a/Engine/Plugins/MovieScene/MovieRenderPipeline/Source/MovieRenderPipelineCore/Private/Graph/Nodes/MovieGraphDeferredPassNode.cpp b/Engine/Plugins/MovieScene/MovieRenderPipeline/Source/MovieRenderPipelineCore/Private/Graph/Nodes/MovieGraphDeferredPassNode.cpp
deleted file mode 100644
diff --git a/Engine/Plugins/MovieScene/MovieRenderPipeline/Source/MovieRenderPipelineCore/Private/Graph/Nodes/MovieGraphImagePassBaseNode.cpp b/Engine/Plugins/MovieScene/MovieRenderPipeline/Source/MovieRenderPipelineCore/Private/Graph/Nodes/MovieGraphImagePassBaseNode.cpp
deleted file mode 100644
diff --git a/Engine/Plugins/MovieScene/MovieRenderPipeline/Source/MovieRenderPipelineCore/Private/Graph/Nodes/MovieGraphPathTracedRenderPassNode.cpp b/Engine/Plugins/MovieScene/MovieRenderPipeline/Source/MovieRenderPipelineCore/Private/Graph/Nodes/MovieGraphPathTracedRenderPassNode.cpp
deleted file mode 100644
diff --git a/Engine/Plugins/MovieScene/MovieRenderPipeline/Source/MovieRenderPipelineCore/Private/Graph/Renderers/MovieGraphDeferredPass.cpp b/Engine/Plugins/MovieScene/MovieRenderPipeline/Source/MovieRenderPipelineCore/Private/Graph/Renderers/MovieGraphDeferredPass.cpp
deleted file mode 100644
diff --git a/Engine/Plugins/MovieScene/MovieRenderPipeline/Source/MovieRenderPipelineCore/Private/Graph/Renderers/MovieGraphImagePassBase.cpp b/Engine/Plugins/MovieScene/MovieRenderPipeline/Source/MovieRenderPipelineCore/Private/Graph/Renderers/MovieGraphImagePassBase.cpp
deleted file mode 100644
diff --git a/Engine/Plugins/MovieScene/MovieRenderPipeline/Source/MovieRenderPipelineCore/Private/Graph/Renderers/MovieGraphPathTracerPass.cpp b/Engine/Plugins/MovieScene/MovieRenderPipeline/Source/MovieRenderPipelineCore/Private/Graph/Renderers/MovieGraphPathTracerPass.cpp
deleted file mode 100644
diff --git a/Engine/Plugins/MovieScene/MovieRenderPipeline/Source/MovieRenderPipelineCore/Public/Graph/Nodes/MovieGraphDeferredPassNode.h b/Engine/Plugins/MovieScene/MovieRenderPipeline/Source/MovieRenderPipelineCore/Public/Graph/Nodes/MovieGraphDeferredPassNode.h
deleted file mode 100644
diff --git a/Engine/Plugins/MovieScene/MovieRenderPipeline/Source/MovieRenderPipelineCore/Public/Graph/Nodes/MovieGraphImagePassBaseNode.h b/Engine/Plugins/MovieScene/MovieRenderPipeline/Source/MovieRenderPipelineCore/Public/Graph/Nodes/MovieGraphImagePassBaseNode.h
deleted file mode 100644
diff --git a/Engine/Plugins/MovieScene/MovieRenderPipeline/Source/MovieRenderPipelineCore/Public/Graph/Nodes/MovieGraphPathTracedRenderPassNode.h b/Engine/Plugins/MovieScene/MovieRenderPipeline/Source/MovieRenderPipelineCore/Public/Graph/Nodes/MovieGraphPathTracedRenderPassNode.h
deleted file mode 100644
diff --git a/Engine/Plugins/MovieScene/MovieRenderPipeline/Source/MovieRenderPipelineCore/Public/Graph/Renderers/MovieGraphDeferredPass.h b/Engine/Plugins/MovieScene/MovieRenderPipeline/Source/MovieRenderPipelineCore/Public/Graph/Renderers/MovieGraphDeferredPass.h
deleted file mode 100644
diff --git a/Engine/Plugins/MovieScene/MovieRenderPipeline/Source/MovieRenderPipelineCore/Public/Graph/Renderers/MovieGraphImagePassBase.h b/Engine/Plugins/MovieScene/MovieRenderPipeline/Source/MovieRenderPipelineCore/Public/Graph/Renderers/MovieGraphImagePassBase.h
deleted file mode 100644
diff --git a/Engine/Plugins/MovieScene/MovieRenderPipeline/Source/MovieRenderPipelineCore/Public/Graph/Renderers/MovieGraphPathTracerPass.h b/Engine/Plugins/MovieScene/MovieRenderPipeline/Source/MovieRenderPipelineCore/Public/Graph/Renderers/MovieGraphPathTracerPass.h
deleted file mode 100644
diff --git a/Engine/Plugins/MovieScene/MovieRenderPipeline/Source/MovieRenderPipelineRenderPasses/Private/Graph/Nodes/MovieGraphDeferredPassNode.cpp b/Engine/Plugins/MovieScene/MovieRenderPipeline/Source/MovieRenderPipelineRenderPasses/Private/Graph/Nodes/MovieGraphDeferredPassNode.cpp
new file mode 100644
--- /dev/null
+++ b/Engine/Plugins/MovieScene/MovieRenderPipeline/Source/MovieRenderPipelineRenderPasses/Private/Graph/Nodes/MovieGraphDeferredPassNode.cpp
@@ -0,0 +1,68 @@
+﻿// Copyright Epic Games, Inc. All Rights Reserved.
+
+#include "Graph/Nodes/MovieGraphDeferredPassNode.h"
+#include "Graph/Renderers/MovieGraphDeferredPass.h"
+
+TUniquePtr<UE::MovieGraph::Rendering::FMovieGraphImagePassBase> UMovieGraphDeferredRenderPassNode::CreateInstance() const
+{
+	return MakeUnique<UE::MovieGraph::Rendering::FMovieGraphDeferredPass>();
+}
+
+UMovieGraphDeferredRenderPassNode::UMovieGraphDeferredRenderPassNode()
+	: SpatialSampleCount(1)
+	, AntiAliasingMethod(EAntiAliasingMethod::AAM_TSR)
+	, bDisableToneCurve(false)
+	, bAllowOCIO(true)
+	, bWriteAllSamples(false)
+	, ViewModeIndex(VMI_Lit)
+{
+}
+
+void UMovieGraphDeferredRenderPassNode::GetFormatResolveArgs(FMovieGraphResolveArgs& OutMergedFormatArgs, const FMovieGraphRenderDataIdentifier& InRenderDataIdentifier) const
+{
+	OutMergedFormatArgs.FilenameArguments.Add(TEXT("ss_count"), FString::FromInt(SpatialSampleCount));
+	OutMergedFormatArgs.FileMetadata.Add(TEXT("unreal/sampling/spatialSampleCount"), FString::FromInt(SpatialSampleCount));
+}
+
+#if WITH_EDITOR
+FText UMovieGraphDeferredRenderPassNode::GetNodeTitle(const bool bGetDescriptive) const
+{
+	return NSLOCTEXT("MovieGraphNodes", "DeferredRenderPassGraphNode_Description", "Deferred Renderer");
+}
+#endif
+
+FString UMovieGraphDeferredRenderPassNode::GetRendererNameImpl() const
+{
+	static const FString RendererNameImpl(TEXT("Deferred"));
+	return RendererNameImpl;
+}
+
+EViewModeIndex UMovieGraphDeferredRenderPassNode::GetViewModeIndex() const
+{
+	return ViewModeIndex;
+}
+
+bool UMovieGraphDeferredRenderPassNode::GetWriteAllSamples() const
+{
+	return bWriteAllSamples;
+}
+
+int32 UMovieGraphDeferredRenderPassNode::GetNumSpatialSamples() const
+{
+	return SpatialSampleCount;
+}
+
+bool UMovieGraphDeferredRenderPassNode::GetDisableToneCurve() const
+{
+	return bDisableToneCurve;
+}
+
+bool UMovieGraphDeferredRenderPassNode::GetAllowOCIO() const
+{
+	return bAllowOCIO;
+}
+
+EAntiAliasingMethod UMovieGraphDeferredRenderPassNode::GetAntiAliasingMethod() const
+{
+	return AntiAliasingMethod;
+}
\ No newline at end of file
diff --git a/Engine/Plugins/MovieScene/MovieRenderPipeline/Source/MovieRenderPipelineRenderPasses/Private/Graph/Nodes/MovieGraphImagePassBaseNode.cpp b/Engine/Plugins/MovieScene/MovieRenderPipeline/Source/MovieRenderPipelineRenderPasses/Private/Graph/Nodes/MovieGraphImagePassBaseNode.cpp
new file mode 100644
--- /dev/null
+++ b/Engine/Plugins/MovieScene/MovieRenderPipeline/Source/MovieRenderPipelineRenderPasses/Private/Graph/Nodes/MovieGraphImagePassBaseNode.cpp
@@ -0,0 +1,122 @@
+// Copyright Epic Games, Inc. All Rights Reserved.
+
+#include "Graph/Nodes/MovieGraphImagePassBaseNode.h"
+#include "Graph/MoviePipelineRenderLayerSubsystem.h" 
+
+
+static TAutoConsoleVariable<bool> CVarMoviePipelineFlushRenderingPerLayerHack(
+	TEXT("MoviePipeline.FlushLayersHack"),
+	false,
+	TEXT("This is a debug cvar intended to help diagnose issues in the MoviePipeline multi-layer rendering mode\n")
+	TEXT("that flushes any outstanding rendering commands for each layer. This can have a large performance impact,\n")
+	TEXT("as it effectively removes any parallelism between the Game Thread and Render Thread!\n")
+	TEXT(" 0 - Don't flush (default)\n")
+	TEXT(" 1 - Flush (but negatively impacts performance)"),
+	ECVF_Default);
+
+UMovieGraphImagePassBaseNode::UMovieGraphImagePassBaseNode()
+{
+	ShowFlags = CreateDefaultSubobject<UMovieGraphShowFlags>(TEXT("ShowFlags"));
+}
+
+void UMovieGraphImagePassBaseNode::SetupImpl(const FMovieGraphRenderPassSetupData& InSetupData)
+{
+	// To make the implementation simpler, we make one instance of FMovieGraphDeferredRenderPas
+	// per camera, and per render layer. These objects can pull from common pools to share state,
+	// which gives us a better overview of how many resources are being used by MRQ.
+	for (const FMovieGraphRenderPassLayerData& LayerData : InSetupData.Layers)
+	{
+		// Each node can change which renderer instance actually gets created.
+		TUniquePtr<UE::MovieGraph::Rendering::FMovieGraphImagePassBase> RendererInstance = CreateInstance();
+		if(RendererInstance)
+		{
+			RendererInstance->Setup(InSetupData.Renderer, this, LayerData);
+			CurrentInstances.Add(MoveTemp(RendererInstance));
+		}
+	}
+}
+
+void UMovieGraphImagePassBaseNode::TeardownImpl()
+{
+	// We don't need to flush the rendering commands as we assume the MovieGraph
+	// Renderer has already done that once, so all data for all passes should
+	// have been submitted to the GPU (and subsequently read back) by now.
+	for (TUniquePtr<UE::MovieGraph::Rendering::FMovieGraphImagePassBase>& Instance : CurrentInstances)
+	{
+		Instance->Teardown();
+	}
+	CurrentInstances.Reset();
+}
+
+void UMovieGraphImagePassBaseNode::RenderImpl(const FMovieGraphTraversalContext& InFrameTraversalContext, const FMovieGraphTimeStepData& InTimeData)
+{
+	for (const TUniquePtr<UE::MovieGraph::Rendering::FMovieGraphImagePassBase>& Instance : CurrentInstances)
+	{
+		UMoviePipelineRenderLayerSubsystem* LayerSubsystem =
+			Instance->GetRenderer()->GetWorld()->GetSubsystem<UMoviePipelineRenderLayerSubsystem>();
+		
+		// Apply all modifiers in the evaluated graph
+		if (LayerSubsystem)
+		{
+			TRACE_CPUPROFILER_EVENT_SCOPE(MRQ::DeferredRender::ActivateRenderLayer);
+			LayerSubsystem->SetActiveRenderLayerByName(Instance->GetBranchName());
+		}
+		
+		// This is a hack to provide debug tooling for rendering issues caused by
+		// multiple layers. Flushing can ensure that all render proxies/etc. are
+		// recreated which can be used to see if it fixes various issues, but
+		// this comes at a big performance cost as it effectively removes GT/RT parallelism.
+		if (CVarMoviePipelineFlushRenderingPerLayerHack.GetValueOnGameThread())
+		{
+			FlushRenderingCommands();
+		}
+
+		Instance->Render(InFrameTraversalContext, InTimeData);
+
+		// Revert all modifiers
+		if (LayerSubsystem)
+		{
+			TRACE_CPUPROFILER_EVENT_SCOPE(MRQ::DeferredRender::RevertRenderLayer);
+			LayerSubsystem->ClearActiveRenderLayer();
+		}
+
+		if (CVarMoviePipelineFlushRenderingPerLayerHack.GetValueOnGameThread())
+		{
+			FlushRenderingCommands();
+		}
+	}
+}
+
+void UMovieGraphImagePassBaseNode::GatherOutputPassesImpl(TArray<FMovieGraphRenderDataIdentifier>& OutExpectedPasses) const
+{
+	for (const TUniquePtr<UE::MovieGraph::Rendering::FMovieGraphImagePassBase>& Instance : CurrentInstances)
+	{
+		Instance->GatherOutputPasses(OutExpectedPasses);
+	}
+}
+
+EViewModeIndex UMovieGraphImagePassBaseNode::GetViewModeIndex() const
+{
+	return VMI_Lit;
+}
+
+FEngineShowFlags UMovieGraphImagePassBaseNode::GetShowFlags() const
+{
+	return ShowFlags->GetShowFlags();
+}
+
+TArray<uint32> UMovieGraphImagePassBaseNode::GetDefaultShowFlags() const
+{
+	return TArray<uint32>();
+}
+
+void UMovieGraphImagePassBaseNode::AddReferencedObjects(UObject* InThis, FReferenceCollector& Collector)
+{
+	Super::AddReferencedObjects(InThis, Collector);
+	
+	UMovieGraphImagePassBaseNode* This = CastChecked<UMovieGraphImagePassBaseNode>(InThis);
+	for (TUniquePtr<UE::MovieGraph::Rendering::FMovieGraphImagePassBase>& Instance : This->CurrentInstances)
+	{
+		Instance->AddReferencedObjects(Collector);
+	}
+}
\ No newline at end of file
diff --git a/Engine/Plugins/MovieScene/MovieRenderPipeline/Source/MovieRenderPipelineRenderPasses/Private/Graph/Nodes/MovieGraphPathTracedRenderPassNode.cpp b/Engine/Plugins/MovieScene/MovieRenderPipeline/Source/MovieRenderPipelineRenderPasses/Private/Graph/Nodes/MovieGraphPathTracedRenderPassNode.cpp
new file mode 100644
--- /dev/null
+++ b/Engine/Plugins/MovieScene/MovieRenderPipeline/Source/MovieRenderPipelineRenderPasses/Private/Graph/Nodes/MovieGraphPathTracedRenderPassNode.cpp
@@ -0,0 +1,82 @@
+// Copyright Epic Games, Inc. All Rights Reserved.
+
+#include "Graph/Nodes/MovieGraphPathTracedRenderPassNode.h"
+#include "Graph/Renderers/MovieGraphPathTracerPass.h"
+#include "Engine/EngineBaseTypes.h"
+#include "ShowFlags.h"
+
+TUniquePtr<UE::MovieGraph::Rendering::FMovieGraphImagePassBase> UMovieGraphPathTracedRenderPassNode::CreateInstance() const
+{
+	return MakeUnique<UE::MovieGraph::Rendering::FMovieGraphPathTracerPass>();
+}
+
+UMovieGraphPathTracedRenderPassNode::UMovieGraphPathTracedRenderPassNode()
+	: SpatialSampleCount(1)
+	, bDisableToneCurve(false)
+	, bAllowOCIO(true)
+	, bWriteAllSamples(false)
+{
+	ShowFlags->ApplyDefaultShowFlagValue(VMI_PathTracing, true);
+	// TODO: Showflag for SetMotionBlur()?
+}
+
+#if WITH_EDITOR
+FText UMovieGraphPathTracedRenderPassNode::GetNodeTitle(const bool bGetDescriptive) const
+{
+	return NSLOCTEXT("MovieGraphNodes", "PathTracedRenderPassGraphNode_Description", "Path Traced Renderer");
+}
+#endif
+
+void UMovieGraphPathTracedRenderPassNode::SetupImpl(const FMovieGraphRenderPassSetupData& InSetupData)
+{
+	Super::SetupImpl(InSetupData);
+	
+	// Hide the progress display during the render
+	if (IConsoleVariable* ProgressDisplayCvar = IConsoleManager::Get().FindConsoleVariable(TEXT("r.PathTracing.ProgressDisplay")))
+	{
+		bOriginalProgressDisplayCvarValue = ProgressDisplayCvar->GetBool();
+		ProgressDisplayCvar->Set(false);
+	}
+}
+
+void UMovieGraphPathTracedRenderPassNode::TeardownImpl()
+{
+	Super::TeardownImpl();
+
+	// Restore the original setting for the progress display
+	if (IConsoleVariable* ProgressDisplayCvar = IConsoleManager::Get().FindConsoleVariable(TEXT("r.PathTracing.ProgressDisplay")))
+	{
+		ProgressDisplayCvar->Set(bOriginalProgressDisplayCvarValue);
+	}
+}
+
+FString UMovieGraphPathTracedRenderPassNode::GetRendererNameImpl() const
+{
+	static const FString RendererNameImpl(TEXT("PathTraced"));
+	return RendererNameImpl;
+}
+
+EViewModeIndex UMovieGraphPathTracedRenderPassNode::GetViewModeIndex() const
+{
+	return VMI_PathTracing;
+}
+
+bool UMovieGraphPathTracedRenderPassNode::GetWriteAllSamples() const
+{
+	return bWriteAllSamples;
+}
+
+int32 UMovieGraphPathTracedRenderPassNode::GetNumSpatialSamples() const
+{
+	return SpatialSampleCount;
+}
+
+bool UMovieGraphPathTracedRenderPassNode::GetDisableToneCurve() const
+{
+	return bDisableToneCurve;
+}
+
+bool UMovieGraphPathTracedRenderPassNode::GetAllowOCIO() const
+{
+	return bAllowOCIO;
+}
diff --git a/Engine/Plugins/MovieScene/MovieRenderPipeline/Source/MovieRenderPipelineRenderPasses/Private/Graph/Renderers/MovieGraphDeferredPass.cpp b/Engine/Plugins/MovieScene/MovieRenderPipeline/Source/MovieRenderPipelineRenderPasses/Private/Graph/Renderers/MovieGraphDeferredPass.cpp
new file mode 100644
--- /dev/null
+++ b/Engine/Plugins/MovieScene/MovieRenderPipeline/Source/MovieRenderPipelineRenderPasses/Private/Graph/Renderers/MovieGraphDeferredPass.cpp
@@ -0,0 +1,341 @@
+// Copyright Epic Games, Inc. All Rights Reserved.
+
+#include "Graph/Renderers/MovieGraphDeferredPass.h"
+#include "Graph/Nodes/MovieGraphDeferredPassNode.h"
+#include "Graph/Nodes/MovieGraphCameraNode.h"
+#include "Graph/MovieGraphDataTypes.h"
+#include "Graph/MovieGraphDefaultRenderer.h"
+#include "Graph/MovieGraphPipeline.h"
+#include "Graph/MoviePipelineRenderLayerSubsystem.h"
+#include "Graph/MovieGraphBlueprintLibrary.h"
+#include "MovieRenderOverlappedImage.h"
+#include "MoviePipelineSurfaceReader.h"
+#include "EngineModule.h"
+#include "SceneManagement.h"
+#include "CanvasTypes.h"
+#include "Engine/EngineBaseTypes.h"
+#include "Engine/TextureRenderTarget2D.h"
+#include "LegacyScreenPercentageDriver.h"
+#include "OpenColorIODisplayExtension.h"
+#include "TextureResource.h"
+#include "Tasks/Task.h"
+#include "SceneView.h"
+
+// For the 1D Weight table for accumulation
+#include "MovieRenderPipelineDataTypes.h"
+// For Sub-Pixel Jitter
+#include "MoviePipelineUtils.h"
+
+namespace UE::MovieGraph::Rendering
+{
+
+void FMovieGraphDeferredPass::Setup(TWeakObjectPtr<UMovieGraphDefaultRenderer> InRenderer, TWeakObjectPtr<UMovieGraphImagePassBaseNode> InRenderPassNode, const FMovieGraphRenderPassLayerData& InLayer)
+{
+	FMovieGraphImagePassBase::Setup(InRenderer, InRenderPassNode, InLayer);
+	
+	LayerData = InLayer;
+
+
+	RenderDataIdentifier.RootBranchName = LayerData.BranchName;
+	RenderDataIdentifier.RendererName = InRenderPassNode->GetRendererName();
+	RenderDataIdentifier.SubResourceName = TEXT("beauty");
+	
+	UE::MovieGraph::DefaultRenderer::FCameraInfo CameraInfo = InRenderer->GetCameraInfo(LayerData.CameraIdentifier);
+	RenderDataIdentifier.CameraName =  CameraInfo.CameraName;
+
+	SceneViewState.Allocate(InRenderer->GetWorld()->GetFeatureLevel());
+}
+
+void FMovieGraphDeferredPass::Teardown()
+{
+	FSceneViewStateInterface* SceneViewRef = SceneViewState.GetReference();
+	if (SceneViewRef)
+	{
+		SceneViewRef->ClearMIDPool();
+	}
+	SceneViewState.Destroy();
+}
+
+void FMovieGraphDeferredPass::GatherOutputPasses(TArray<FMovieGraphRenderDataIdentifier>& OutExpectedPasses) const
+{
+	FMovieGraphImagePassBase::GatherOutputPasses(OutExpectedPasses);
+	
+	// Add our pre-calculated identifier
+	OutExpectedPasses.Add(RenderDataIdentifier);
+}
+
+void FMovieGraphDeferredPass::AddReferencedObjects(FReferenceCollector& Collector)
+{
+	FSceneViewStateInterface* SceneViewRef = SceneViewState.GetReference();
+	if (SceneViewRef)
+	{
+		SceneViewRef->AddReferencedObjects(Collector);
+	}
+}
+
+FName FMovieGraphDeferredPass::GetBranchName() const
+{
+	return LayerData.BranchName;
+}
+
+UMovieGraphImagePassBaseNode* FMovieGraphDeferredPass::GetParentNode(UMovieGraphEvaluatedConfig* InConfig) const
+{
+	// This is a bit of a workaround for the fact that the pass doesn't have a strong pointer to the node it's supposed to be associated with,
+	// since that instance changes every frame. So instead we have a virtual function here so the node can look it up by type, and then we can
+	// call a bunch of virtual functions on the right instance to fetch values.
+	const bool bIncludeCDOs = true;
+	UMovieGraphDeferredRenderPassNode* ParentNode = InConfig->GetSettingForBranch<UMovieGraphDeferredRenderPassNode>(GetBranchName(), bIncludeCDOs);
+	if (!ensureMsgf(ParentNode, TEXT("DeferredPass should not exist without parent node in graph.")))
+	{
+		return nullptr;
+	}
+
+	return ParentNode;
+}
+
+void FMovieGraphDeferredPass::Render(const FMovieGraphTraversalContext& InFrameTraversalContext, const FMovieGraphTimeStepData& InTimeData)
+{
+	// ToDo: InFrameTraversalContext includes a copy of TimeData, but may be the one cached at the first temporal sample,
+	// maybe we can combine, maybe we can't?
+	TObjectPtr<UMovieGraphDefaultRenderer> GraphRenderer = GetRenderer().Get();
+	if (!GraphRenderer)
+	{
+		return;
+	}
+
+	UMovieGraphImagePassBaseNode* ParentNodeThisFrame = GetParentNode(InTimeData.EvaluatedConfig);
+	const bool bWriteAllSamples = ParentNodeThisFrame->GetWriteAllSamples();
+	int32 NumSpatialSamples = FMath::Max(1, ParentNodeThisFrame->GetNumSpatialSamples());
+	const ESceneCaptureSource SceneCaptureSource = ParentNodeThisFrame->GetDisableToneCurve() ? ESceneCaptureSource::SCS_FinalColorHDR : ESceneCaptureSource::SCS_FinalToneCurveHDR;
+	const EAntiAliasingMethod AntiAliasingMethod = ParentNodeThisFrame->GetAntiAliasingMethod();
+	float OverscanFraction = 0.f;
+	const float TileOverlapPadRatio = 0.0f; // No tiling support right now
+
+	// Camera nodes are optional
+	const bool bIncludeCDOs = false;
+	const UMovieGraphCameraSettingNode* CameraNode = InTimeData.EvaluatedConfig->GetSettingForBranch<UMovieGraphCameraSettingNode>(LayerData.BranchName, bIncludeCDOs);
+	if (CameraNode)
+	{
+		OverscanFraction = FMath::Clamp(CameraNode->OverscanPercentage / 100.f, 0.f, 1.f);
+	}
+	
+	FIntPoint AccumulatorResolution = UMovieGraphBlueprintLibrary::GetEffectiveOutputResolution(InTimeData.EvaluatedConfig);
+	// ToDo: When tiling is used, this should be the size of the per-tile backbuffer
+	FIntPoint BackbufferResolution = AccumulatorResolution;
+	// ToDo: This math probably needs the per-tile, pre-overlapped size? 
+	FIntPoint OverlappedPad = FIntPoint(FMath::CeilToInt(BackbufferResolution.X * TileOverlapPadRatio), FMath::CeilToInt(BackbufferResolution.Y * TileOverlapPadRatio));
+	// Calculate a backbuffer
+	UE::MovieGraph::DefaultRenderer::FRenderTargetInitParams RenderTargetInitParams;
+	{
+		RenderTargetInitParams.Size = BackbufferResolution;
+
+		// OCIO: Since this is a manually created Render target we don't need Gamma to be applied.
+		// We use this render target to render to via a display extension that utilizes Display Gamma
+		// which has a default value of 2.2 (DefaultDisplayGamma), therefore we need to set Gamma on this render target to 2.2 to cancel out any unwanted effects.
+		RenderTargetInitParams.TargetGamma = FOpenColorIORendering::DefaultDisplayGamma;
+		RenderTargetInitParams.PixelFormat = EPixelFormat::PF_FloatRGBA;
+	}
+
+	UTextureRenderTarget2D* RenderTarget = GraphRenderer->GetOrCreateViewRenderTarget(RenderTargetInitParams, RenderDataIdentifier);
+	FRenderTarget* RenderTargetResource = RenderTarget->GameThread_GetRenderTargetResource();
+	check(RenderTargetResource);
+
+	// Now we can actually construct our ViewFamily, SceneView, and submit it for Rendering + Readback
+	for(int32 SpatialIndex = 0; SpatialIndex < NumSpatialSamples; SpatialIndex++)
+	{
+		// World should be paused for every spatial sample except the last one, so that
+		// the view doesn't update histories until the end, allowing us to render the same
+		// scene multiple times.
+		const bool bWorldIsPaused = SpatialIndex < (NumSpatialSamples - 1);
+		const int32 FrameIndex = InTimeData.RenderedFrameNumber * ((InTimeData.TemporalSampleCount * NumSpatialSamples) + (InTimeData.TemporalSampleIndex * NumSpatialSamples)) + SpatialIndex;
+
+		// We only allow a spatial jitter if we have more than one sample
+		FVector2f SpatialShiftAmount = FVector2f(0.f, 0.f);
+		const bool bAntiAliasingAllowsJitter = AntiAliasingMethod == EAntiAliasingMethod::AAM_None;
+		const bool bSampleCountsAllowsJitter = NumSpatialSamples > 1 || InTimeData.TemporalSampleCount > 1;
+		if (bAntiAliasingAllowsJitter && bSampleCountsAllowsJitter)
+		{
+			const int32 NumSamplesPerOutputFrame = NumSpatialSamples * InTimeData.TemporalSampleCount;
+			SpatialShiftAmount = UE::MoviePipeline::GetSubPixelJitter(FrameIndex, NumSamplesPerOutputFrame);
+		}
+		 
+		// These are the parameters of our camera 
+		UE::MovieGraph::DefaultRenderer::FCameraInfo CameraInfo;
+
+		// ToDo: Get this from the renderer based on LayerData.CameraIdentifier for eventual multi-camera support
+		APlayerController* LocalPlayerController = GraphRenderer->GetWorld()->GetFirstPlayerController();
+		// CameraAnim override
+		if (LocalPlayerController->PlayerCameraManager)
+		{
+			CameraInfo.ViewInfo = LocalPlayerController->PlayerCameraManager->GetCameraCacheView();
+		}
+
+
+		CameraInfo.bAllowCameraAspectRatio = true;
+		CameraInfo.TilingParams.TileSize = BackbufferResolution;
+		CameraInfo.TilingParams.OverlapPad = FVector2f(0.f, 0.f); // No tile support
+		CameraInfo.TilingParams.TileCount = FIntPoint(1, 1); // No tile support
+		CameraInfo.TilingParams.TileIndexes = FIntPoint(0, 0); // No tile support
+		CameraInfo.SamplingParams.TemporalSampleIndex = InTimeData.TemporalSampleIndex;
+		CameraInfo.SamplingParams.TemporalSampleCount = InTimeData.TemporalSampleCount;
+		CameraInfo.SamplingParams.SpatialSampleIndex = SpatialIndex;
+		CameraInfo.SamplingParams.SpatialSampleCount = NumSpatialSamples;
+		CameraInfo.OverscanFraction = OverscanFraction;
+		CameraInfo.ProjectionMatrixJitterAmount = FVector2D((SpatialShiftAmount.X) * 2.0f / (float)BackbufferResolution.X, SpatialShiftAmount.Y * -2.0f / (float)BackbufferResolution.Y);
+
+		// For this particular tile, what is the offset into the output image
+		FIntPoint OverlappedOffset = FIntPoint(CameraInfo.TilingParams.TileIndexes.X * BackbufferResolution.X - OverlappedPad.X, CameraInfo.TilingParams.TileIndexes.Y * BackbufferResolution.Y - OverlappedPad.Y);
+		
+		// Move the final render by this much in the accumulator to counteract the offset put into the view matrix.
+		// Note that when bAllowSpatialJitter is false, SpatialShiftX/Y will always be zero.
+		FVector2D OverlappedSubpixelShift = FVector2D(0.5f - SpatialShiftAmount.X, 0.5f - SpatialShiftAmount.Y);
+		
+		FMatrix ProjectionMatrix = CalculateProjectionMatrix(CameraInfo);
+		float DoFSensorScale = 1.f;
+
+		// Modify the perspective matrix to do an off center projection, with overlap for high-res tiling
+		const bool bOrthographic = CameraInfo.ViewInfo.ProjectionMode == ECameraProjectionMode::Type::Orthographic;
+		ModifyProjectionMatrixForTiling(CameraInfo.TilingParams, bOrthographic, ProjectionMatrix, DoFSensorScale);
+		
+		CameraInfo.ProjectionMatrix = ProjectionMatrix;
+		CameraInfo.DoFSensorScale = DoFSensorScale;
+
+		// The Scene View Family must be constructed first as the FSceneView needs it to be constructed
+		UE::MovieGraph::Rendering::FViewFamilyInitData ViewFamilyInitData;
+		ViewFamilyInitData.RenderTarget = RenderTargetResource;
+		ViewFamilyInitData.World = GraphRenderer->GetWorld();
+		ViewFamilyInitData.TimeData = InTimeData;
+		ViewFamilyInitData.SceneCaptureSource = SceneCaptureSource;
+		ViewFamilyInitData.bWorldIsPaused = bWorldIsPaused;
+		ViewFamilyInitData.FrameIndex = FrameIndex;
+		ViewFamilyInitData.AntiAliasingMethod = AntiAliasingMethod;
+		ViewFamilyInitData.ShowFlags = ParentNodeThisFrame->GetShowFlags();
+		ViewFamilyInitData.ViewModeIndex = ParentNodeThisFrame->GetViewModeIndex();
+		
+		TSharedRef<FSceneViewFamilyContext> ViewFamily = CreateSceneViewFamily(ViewFamilyInitData, CameraInfo);
+		 
+		// Now we can construct a View to go within this family.
+		FSceneViewInitOptions SceneViewInitOptions = CreateViewInitOptions(CameraInfo, ViewFamily.ToSharedPtr().Get(), SceneViewState);
+		FSceneView* NewView = CreateSceneView(SceneViewInitOptions, ViewFamily, CameraInfo);
+		
+		// Then apply Movie Render Queue specific overrides to the ViewFamily, and then to the SceneView.
+		ApplyMovieGraphOverridesToViewFamily(ViewFamily, ViewFamilyInitData);
+		
+		// ToDo: This really only needs access to the ViewFamily for path tracer related things,
+		// and would rather just take a FSceneView* 
+		ApplyMovieGraphOverridesToSceneView(ViewFamily, ViewFamilyInitData, CameraInfo);
+		
+		FHitProxyConsumer* HitProxyConsumer = nullptr;
+		const float DPIScale = 1.0f;
+		FCanvas Canvas = FCanvas(RenderTargetResource, HitProxyConsumer, GraphRenderer->GetWorld(), GraphRenderer->GetWorld()->GetFeatureLevel(), FCanvas::CDM_DeferDrawing, DPIScale);
+		
+		// Submit the renderer to be rendered
+		GetRendererModule().BeginRenderingViewFamily(&Canvas, ViewFamily.ToSharedPtr().Get());
+
+		// Construct the sample state that reflects the current render sample
+		UE::MovieGraph::FMovieGraphSampleState SampleState;
+		{
+			// Take our per-frame Traversal Context and update it with context specific to this sample.
+			FMovieGraphTraversalContext UpdatedTraversalContext = InFrameTraversalContext;
+			UpdatedTraversalContext.Time = InTimeData;
+			UpdatedTraversalContext.RenderDataIdentifier = RenderDataIdentifier;
+
+			SampleState.TraversalContext = MoveTemp(UpdatedTraversalContext);
+			SampleState.BackbufferResolution = BackbufferResolution;
+			SampleState.AccumulatorResolution = AccumulatorResolution;
+			SampleState.bWriteSampleToDisk = bWriteAllSamples;
+			SampleState.bRequiresAccumulator = InTimeData.bRequiresAccumulator || (NumSpatialSamples > 1);
+			SampleState.bFetchFromAccumulator = InTimeData.bIsLastTemporalSampleForFrame && (SpatialIndex == (NumSpatialSamples - 1));
+			SampleState.OverlappedPad = OverlappedPad;
+			SampleState.OverlappedOffset = OverlappedOffset;
+			SampleState.OverlappedSubpixelShift = OverlappedSubpixelShift;
+			SampleState.OverscanFraction = OverscanFraction;
+			SampleState.bAllowOCIO = ParentNodeThisFrame->GetAllowOCIO();
+			SampleState.SceneCaptureSource = SceneCaptureSource;
+		}
+
+		// If this was just to contribute to the history buffer, no need to go any further.
+		bool bDiscardOutput = InTimeData.bDiscardOutput || ShouldDiscardOutput(ViewFamily, CameraInfo);
+		if (bDiscardOutput)
+		{
+			continue;
+		}
+
+		// Readback + Accumulate.
+		PostRendererSubmission(SampleState, RenderTargetInitParams, Canvas, CameraInfo);
+	}
+}
+
+void FMovieGraphDeferredPass::PostRendererSubmission(
+	const UE::MovieGraph::FMovieGraphSampleState& InSampleState,
+	const UE::MovieGraph::DefaultRenderer::FRenderTargetInitParams& InRenderTargetInitParams, FCanvas& InCanvas, const UE::MovieGraph::DefaultRenderer::FCameraInfo& InCameraInfo)
+{
+	TObjectPtr<UMovieGraphDefaultRenderer> GraphRenderer = GetRenderer().Get();
+	if (!GraphRenderer)
+	{
+		return;
+	}
+
+	// Draw letterboxing
+	// ToDo: Multi-camera support
+	APlayerCameraManager* PlayerCameraManager = GraphRenderer->GetWorld()->GetFirstPlayerController()->PlayerCameraManager;
+	if(PlayerCameraManager && PlayerCameraManager->GetCameraCacheView().bConstrainAspectRatio)
+	{
+		const FMinimalViewInfo CameraCache = PlayerCameraManager->GetCameraCacheView();
+		
+		// Taking overscan into account.
+		const FIntPoint FullOutputSize = InSampleState.AccumulatorResolution;
+	
+		const float OutputSizeAspectRatio = FullOutputSize.X / (float)FullOutputSize.Y;
+		const float CameraAspectRatio = InCameraInfo.bAllowCameraAspectRatio ? CameraCache.AspectRatio : OutputSizeAspectRatio;
+	
+		const FIntPoint ConstrainedFullSize = CameraAspectRatio > OutputSizeAspectRatio ?
+			FIntPoint(FullOutputSize.X, FMath::CeilToInt((double)FullOutputSize.X / (double)CameraAspectRatio)) :
+			FIntPoint(FMath::CeilToInt(CameraAspectRatio * FullOutputSize.Y), FullOutputSize.Y);
+	
+		const FIntPoint TileViewMin = InSampleState.OverlappedOffset;
+		const FIntPoint TileViewMax = TileViewMin + InSampleState.BackbufferResolution;
+	
+		// Camera ratio constrained rect, clipped by the tile rect
+		FIntPoint ConstrainedViewMin = (FullOutputSize - ConstrainedFullSize) / 2;
+		FIntPoint ConstrainedViewMax = ConstrainedViewMin + ConstrainedFullSize;
+		ConstrainedViewMin = FIntPoint(FMath::Clamp(ConstrainedViewMin.X, TileViewMin.X, TileViewMax.X),
+			FMath::Clamp(ConstrainedViewMin.Y, TileViewMin.Y, TileViewMax.Y));
+		ConstrainedViewMax = FIntPoint(FMath::Clamp(ConstrainedViewMax.X, TileViewMin.X, TileViewMax.X),
+			FMath::Clamp(ConstrainedViewMax.Y, TileViewMin.Y, TileViewMax.Y));
+	
+		// Difference between the clipped constrained rect and the tile rect
+		const FIntPoint OffsetMin = ConstrainedViewMin - TileViewMin;
+		const FIntPoint OffsetMax = TileViewMax - ConstrainedViewMax;
+	
+		// Clear left
+		if (OffsetMin.X > 0)
+		{
+			InCanvas.DrawTile(0, 0, OffsetMin.X, InSampleState.BackbufferResolution.Y,
+				0.0f, 0.0f, 1.0f, 1.0f, FLinearColor::Black, nullptr, false);
+		}
+		// Clear right
+		if (OffsetMax.X > 0)
+		{
+			InCanvas.DrawTile(InSampleState.BackbufferResolution.X - OffsetMax.X, 0, InSampleState.BackbufferResolution.X, InSampleState.BackbufferResolution.Y,
+				0.0f, 0.0f, 1.0f, 1.0f, FLinearColor::Black, nullptr, false);
+		}
+		// Clear top
+		if (OffsetMin.Y > 0)
+		{
+			InCanvas.DrawTile(0, 0, InSampleState.BackbufferResolution.X, OffsetMin.Y,
+				0.0f, 0.0f, 1.0f, 1.0f, FLinearColor::Black, nullptr, false);
+		}
+		// Clear bottom
+		if (OffsetMax.Y > 0)
+		{
+			InCanvas.DrawTile(0, InSampleState.BackbufferResolution.Y - OffsetMax.Y, InSampleState.BackbufferResolution.X, InSampleState.BackbufferResolution.Y,
+				0.0f, 0.0f, 1.0f, 1.0f, FLinearColor::Black, nullptr, false);
+		}
+	
+		InCanvas.Flush_GameThread(true);
+	}
+	
+	FMovieGraphImagePassBase::PostRendererSubmission(InSampleState, InRenderTargetInitParams, InCanvas, InCameraInfo);
+}
+} // namespace UE::MovieGraph::Rendering
diff --git a/Engine/Plugins/MovieScene/MovieRenderPipeline/Source/MovieRenderPipelineRenderPasses/Private/Graph/Renderers/MovieGraphImagePassBase.cpp b/Engine/Plugins/MovieScene/MovieRenderPipeline/Source/MovieRenderPipelineRenderPasses/Private/Graph/Renderers/MovieGraphImagePassBase.cpp
new file mode 100644
--- /dev/null
+++ b/Engine/Plugins/MovieScene/MovieRenderPipeline/Source/MovieRenderPipelineRenderPasses/Private/Graph/Renderers/MovieGraphImagePassBase.cpp
@@ -0,0 +1,656 @@
+// Copyright Epic Games, Inc. All Rights Reserved.
+
+#include "Graph/Renderers/MovieGraphImagePassBase.h"
+#include "Graph/MovieGraphDataTypes.h"
+#include "Graph/MovieGraphPipeline.h"
+#include "SceneView.h"
+#include "MoviePipelineUtils.h"
+#include "LegacyScreenPercentageDriver.h"
+#include "MoviePipelineSurfaceReader.h"
+#include "CanvasTypes.h"
+#include "MovieRenderOverlappedImage.h"
+#include "Engine/RendererSettings.h"
+#include "UnrealClient.h"
+#include "SceneViewExtensionContext.h"
+#include "SceneViewExtension.h"
+#include "Engine/Engine.h"
+
+namespace UE::MovieGraph::Rendering
+{
+
+void FMovieGraphImagePassBase::Setup(TWeakObjectPtr<UMovieGraphDefaultRenderer> InRenderer, TWeakObjectPtr<UMovieGraphImagePassBaseNode> InRenderPassNode, const FMovieGraphRenderPassLayerData& InLayer)
+{
+	// This is a pointer to the UMovieGraphPipeline's renderer which is valid throughout the entire render.
+	WeakGraphRenderer = InRenderer;
+}
+
+FSceneViewInitOptions FMovieGraphImagePassBase::CreateViewInitOptions(const UE::MovieGraph::DefaultRenderer::FCameraInfo& InCameraInfo, FSceneViewFamilyContext* InViewFamily, FSceneViewStateReference& InViewStateRef) const
+{
+	check(InViewFamily);
+
+	FIntPoint RenderResolution = InViewFamily->RenderTarget->GetSizeXY();
+
+	FSceneViewInitOptions ViewInitOptions;
+	ViewInitOptions.ViewFamily = InViewFamily;
+	ViewInitOptions.ViewOrigin = InCameraInfo.ViewInfo.Location;
+	ViewInitOptions.ViewLocation = InCameraInfo.ViewInfo.Location;
+	ViewInitOptions.ViewRotation = InCameraInfo.ViewInfo.Rotation;
+	ViewInitOptions.SetViewRectangle(FIntRect(FIntPoint(0, 0), RenderResolution));
+	ViewInitOptions.ViewRotationMatrix = FInverseRotationMatrix(InCameraInfo.ViewInfo.Rotation);
+	ViewInitOptions.ViewActor = InCameraInfo.ViewActor;
+	ViewInitOptions.ProjectionMatrix = InCameraInfo.ProjectionMatrix;
+
+	// Rotate the view 90 degrees to match the rest of the engine.
+	ViewInitOptions.ViewRotationMatrix = ViewInitOptions.ViewRotationMatrix * FMatrix(
+		FPlane(0, 0, 1, 0),
+		FPlane(1, 0, 0, 0),
+		FPlane(0, 1, 0, 0),
+		FPlane(0, 0, 0, 1));
+	
+	float ViewFOV = InCameraInfo.ViewInfo.FOV;
+
+	// Inflate our FOV to support the overscan 
+	// ToDo: This is a duplicate of the logic in CalculateProjectionMatrix, should combine.
+	ViewFOV = 2.0f * FMath::RadiansToDegrees(FMath::Atan((1.0f + InCameraInfo.OverscanFraction) * FMath::Tan(FMath::DegreesToRadians( ViewFOV * 0.5f ))));
+
+	ViewInitOptions.SceneViewStateInterface = InViewStateRef.GetReference();
+	ViewInitOptions.FOV = ViewFOV;
+	ViewInitOptions.DesiredFOV = ViewFOV;
+	
+	return ViewInitOptions;
+}
+
+
+
+FSceneView* FMovieGraphImagePassBase::CreateSceneView(const FSceneViewInitOptions& InInitOptions, TSharedRef<FSceneViewFamilyContext> InViewFamily, const UE::MovieGraph::DefaultRenderer::FCameraInfo& InCameraInfo) const
+{
+	// Create our view based on the Init Options
+	FSceneView* View = new FSceneView(InInitOptions);
+	InViewFamily->Views.Add(View);
+
+	View->StartFinalPostprocessSettings(InInitOptions.ViewLocation);
+	ApplyCameraManagerPostProcessBlends(View);
+
+	// Scaling sensor size inversely with the the projection matrix [0][0] should physically
+	// cause the circle of confusion to be unchanged.
+	View->FinalPostProcessSettings.DepthOfFieldSensorWidth *= InCameraInfo.DoFSensorScale;
+	// Modify the 'center' of the lens to be offset for high-res tiling, helps some effects (vignette) etc. still work.
+	View->LensPrincipalPointOffsetScale = CalculatePrinciplePointOffsetForTiling(InCameraInfo.TilingParams);
+	View->EndFinalPostprocessSettings(InInitOptions);
+
+	return View;
+}
+
+void FMovieGraphImagePassBase::ApplyCameraManagerPostProcessBlends(FSceneView* InView) const
+{
+	check(InView);
+
+	TObjectPtr<UMovieGraphDefaultRenderer> GraphRenderer = GetRenderer().Get();
+	if (!GraphRenderer)
+	{
+		return;
+	}
+
+	APlayerController* LocalPlayerController = GraphRenderer->GetWorld()->GetFirstPlayerController();
+	// CameraAnim override
+	if (LocalPlayerController->PlayerCameraManager)
+	{
+		TArray<FPostProcessSettings> const* CameraAnimPPSettings;
+		TArray<float> const* CameraAnimPPBlendWeights;
+		LocalPlayerController->PlayerCameraManager->GetCachedPostProcessBlends(CameraAnimPPSettings, CameraAnimPPBlendWeights);
+
+		if (LocalPlayerController->PlayerCameraManager->bEnableFading)
+		{
+			InView->OverlayColor = LocalPlayerController->PlayerCameraManager->FadeColor;
+			InView->OverlayColor.A = FMath::Clamp(LocalPlayerController->PlayerCameraManager->FadeAmount, 0.f, 1.f);
+		}
+
+		if (LocalPlayerController->PlayerCameraManager->bEnableColorScaling)
+		{
+			FVector ColorScale = LocalPlayerController->PlayerCameraManager->ColorScale;
+			InView->ColorScale = FLinearColor(ColorScale.X, ColorScale.Y, ColorScale.Z);
+		}
+
+		FMinimalViewInfo ViewInfo = LocalPlayerController->PlayerCameraManager->GetCameraCacheView();
+		for (int32 PPIdx = 0; PPIdx < CameraAnimPPBlendWeights->Num(); ++PPIdx)
+		{
+			InView->OverridePostProcessSettings((*CameraAnimPPSettings)[PPIdx], (*CameraAnimPPBlendWeights)[PPIdx]);
+		}
+
+		InView->OverridePostProcessSettings(ViewInfo.PostProcessSettings, ViewInfo.PostProcessBlendWeight);
+	}
+}
+
+TSharedRef<FSceneViewFamilyContext> FMovieGraphImagePassBase::CreateSceneViewFamily(const FViewFamilyInitData& InInitData, const UE::MovieGraph::DefaultRenderer::FCameraInfo& InCameraInfo) const
+{
+
+	EViewModeIndex ViewModeIndex = InInitData.ViewModeIndex;
+	FEngineShowFlags ShowFlags = InInitData.ShowFlags;
+
+	const bool bIsPerspective = InCameraInfo.ViewInfo.ProjectionMode == ECameraProjectionMode::Type::Perspective;
+
+	// Allow the Engine Showflag system to override our engine showflags, based on our view mode index.
+	// This is required for certain debug view modes (to have matching show flags set for rendering).
+	ApplyViewMode(/*In*/ ViewModeIndex, bIsPerspective, /*InOut*/ShowFlags);
+
+	// And then we have to let another system override them again (based on cvars, etc.)
+	EngineShowFlagOverride(ESFIM_Game, ViewModeIndex, ShowFlags, false);
+
+	TSharedRef<FSceneViewFamilyContext> OutViewFamily = MakeShared<FSceneViewFamilyContext>(FSceneViewFamily::ConstructionValues(
+		InInitData.RenderTarget,
+		InInitData.World->Scene,
+		ShowFlags)
+		.SetTime(FGameTime::CreateUndilated(InInitData.TimeData.WorldSeconds, InInitData.TimeData.FrameDeltaTime))
+		.SetRealtimeUpdate(true));
+
+	// Need to add the engine-wide view extensions, as rendering code may depend on them (ie: landscapes)
+	OutViewFamily->ViewExtensions.Append(GEngine->ViewExtensions->GatherActiveExtensions(FSceneViewExtensionContext(InInitData.World->Scene)));
+
+	return OutViewFamily;
+}
+
+void FMovieGraphImagePassBase::ApplyMovieGraphOverridesToViewFamily(TSharedRef<FSceneViewFamilyContext> InOutFamily, const FViewFamilyInitData& InInitData) const
+{
+	// Used to specify if the Tone Curve is being applied or not to our Linear Output data
+	InOutFamily->SceneCaptureSource = InInitData.SceneCaptureSource;
+	InOutFamily->bWorldIsPaused = InInitData.bWorldIsPaused;
+	// InOutFamily->ViewMode = ViewModeIndex;
+	InOutFamily->bOverrideVirtualTextureThrottle = true;
+	
+	// This is for a workaround involving GFrameCounter already being incremented by the time MRQ renders are submitted
+	// meaning that it's out of sync with some parts of the renderer's idea of the current frame.
+	PRAGMA_DISABLE_DEPRECATION_WARNINGS
+	InOutFamily->OverrideFrameCounter = UE::MovieRenderPipeline::GetRendererFrameCount();
+	PRAGMA_ENABLE_DEPRECATION_WARNINGS
+	
+	// We need to check if this is the first FSceneView being submitted to the renderer module, and set some flags on the ViewFamily for ensuring some
+	// parts of the renderer only get updated once per frame. Kept as an if/else statement to avoid the confusion with setting all of these values to 
+	// some permutation of !/!!bHasRenderedFirstViewThisFrame.
+	TObjectPtr<UMovieGraphDefaultRenderer> GraphRenderer = GetRenderer().Get();
+	if (GraphRenderer)
+	{
+		if (!GraphRenderer->GetHasRenderedFirstViewThisFrame())
+		{
+			// Update our renderer
+			GraphRenderer->SetHasRenderedFirstViewThisFrame(true);
+
+			InOutFamily->bIsFirstViewInMultipleViewFamily = true;
+			InOutFamily->bAdditionalViewFamily = false;
+		}
+		else
+		{
+			InOutFamily->bIsFirstViewInMultipleViewFamily = false;
+			InOutFamily->bAdditionalViewFamily = true;
+		}
+
+		InOutFamily->bIsMultipleViewFamily = true;
+	}
+	
+	// Skip the whole pass if they don't want motion blur.
+	if (FMath::IsNearlyZero(InInitData.TimeData.MotionBlurFraction))
+	{
+		InOutFamily->EngineShowFlags.SetMotionBlur(false);
+	}
+	
+
+	// ToDo: Let settings modify the ScreenPercentageInterface so third party screen percentages are supported.
+
+	// If UMoviePipelineViewFamilySetting never set a Screen percentage interface we fallback to default.
+	if (InOutFamily->GetScreenPercentageInterface() == nullptr)
+	{
+		InOutFamily->SetScreenPercentageInterface(new FLegacyScreenPercentageDriver(*InOutFamily, FLegacyScreenPercentageDriver::GetCVarResolutionFraction()));
+	}
+}
+
+void FMovieGraphImagePassBase::ApplyMovieGraphOverridesToSceneView(TSharedRef<FSceneViewFamilyContext> InOutFamily, const FViewFamilyInitData& InInitData, const UE::MovieGraph::DefaultRenderer::FCameraInfo& InCameraInfo) const
+{
+	// Override the view's FrameIndex to be based on our progress through the sequence. This greatly increases
+	// determinism with things like TAA.
+	FSceneView* View = const_cast<FSceneView*>(InOutFamily->Views[0]);
+	View->OverrideFrameIndexValue = InInitData.FrameIndex;
+	// Each shot should initialize a scene history from scratch so there should be no need to do an extra camera cut flag.
+	View->bCameraCut = false; 
+	View->AntiAliasingMethod = InInitData.AntiAliasingMethod;
+	View->bIsOfflineRender = true;
+
+	// Override the Motion Blur settings since these are controlled by the movie pipeline.
+	{
+		// We need to inversly scale the target FPS by time dilation to counteract slowmo. If scaling isn't applied then motion blur length
+		// stays the same length despite the smaller delta time and the blur ends up too long.
+		View->FinalPostProcessSettings.MotionBlurTargetFPS = FMath::RoundToInt(InInitData.TimeData.FrameRate.AsDecimal() / FMath::Max(SMALL_NUMBER, InInitData.TimeData.WorldTimeDilation));
+		View->FinalPostProcessSettings.MotionBlurAmount = InInitData.TimeData.MotionBlurFraction;
+		View->FinalPostProcessSettings.MotionBlurMax = 100.f;
+		View->FinalPostProcessSettings.bOverride_MotionBlurAmount = true;
+		View->FinalPostProcessSettings.bOverride_MotionBlurTargetFPS = true;
+		View->FinalPostProcessSettings.bOverride_MotionBlurMax = true;
+
+	}
+
+	// Warn the user for invalid setting combinations / enforce hardware limitations
+	// Locked Exposure
+	const bool bAutoExposureAllowed = true; // IsAutoExposureAllowed(InOutSampleState);
+	{
+		// If the rendering pass doesn't allow autoexposure and they dont' have manual exposure set up, warn.
+		if (!bAutoExposureAllowed && (View->FinalPostProcessSettings.AutoExposureMethod != EAutoExposureMethod::AEM_Manual))
+		{
+			// Skip warning if the project setting is disabled though, as exposure will be forced off in the renderer anyways.
+			const URendererSettings* RenderSettings = GetDefault<URendererSettings>();
+			if (RenderSettings->bDefaultFeatureAutoExposure != false)
+			{
+				UE_LOG(LogMovieRenderPipeline, Warning, TEXT("Camera Auto Exposure Method not supported by one or more render passes. Change the Auto Exposure Method to Manual!"));
+				View->FinalPostProcessSettings.AutoExposureMethod = EAutoExposureMethod::AEM_Manual;
+			}
+		}
+	}
+
+	// Orthographic cameras don't support anti-aliasing outside the path tracer (other than FXAA)
+	const bool bIsOrthographicCamera = !View->IsPerspectiveProjection();
+	if (bIsOrthographicCamera)
+	{
+		bool bIsSupportedAAMethod = View->AntiAliasingMethod == EAntiAliasingMethod::AAM_FXAA;
+		bool bIsPathTracer = InOutFamily->EngineShowFlags.PathTracing;
+		bool bWarnJitters = InCameraInfo.ProjectionMatrixJitterAmount.SquaredLength() > SMALL_NUMBER;
+		if ((!bIsPathTracer && !bIsSupportedAAMethod) || bWarnJitters)
+		{
+			UE_LOG(LogMovieRenderPipeline, Warning, TEXT("Orthographic Cameras are only supported with PathTracer or Deferred with FXAA Anti-Aliasing"));
+		}
+	}
+
+	{
+		bool bMethodWasUnsupported = false;
+		if (View->AntiAliasingMethod == AAM_TemporalAA && !SupportsGen4TAA(View->GetShaderPlatform()))
+		{
+			UE_LOG(LogMovieRenderPipeline, Error, TEXT("TAA was requested but this hardware does not support it."));
+			bMethodWasUnsupported = true;
+		}
+		else if (View->AntiAliasingMethod == AAM_TSR && !SupportsTSR(View->GetShaderPlatform()))
+		{
+			UE_LOG(LogMovieRenderPipeline, Error, TEXT("TSR was requested but this hardware does not support it."));
+			bMethodWasUnsupported = true;
+		}
+
+		if (bMethodWasUnsupported)
+		{
+			View->AntiAliasingMethod = AAM_None;
+		}
+	}
+
+	// Anti Aliasing
+	{
+		// If we're not using TAA, TSR, or Path Tracing we will apply the View Matrix projection jitter. Normally TAA sets this
+		// inside FSceneRenderer::PreVisibilityFrameSetup. Path Tracing does its own anti-aliasing internally.
+		bool bApplyProjectionJitter = !bIsOrthographicCamera
+			&& !InOutFamily->EngineShowFlags.PathTracing
+			&& !IsTemporalAccumulationBasedMethod(View->AntiAliasingMethod);
+		if (bApplyProjectionJitter)
+		{
+			View->ViewMatrices.HackAddTemporalAAProjectionJitter(InCameraInfo.ProjectionMatrixJitterAmount);
+		}
+	}
+}
+
+FMatrix FMovieGraphImagePassBase::CalculateProjectionMatrix(const UE::MovieGraph::DefaultRenderer::FCameraInfo& InCameraInfo) const
+{
+	// Calculate a Projection Matrix. This code unfortunately ends up similar to, but not quite the same as FMinimalViewInfo::CalculateProjectionMatrixGivenView
+	FMatrix BaseProjMatrix;
+	
+	// TileSize should respect the actual backbuffer size being used by the render.
+	float ViewRectWidth = InCameraInfo.TilingParams.TileSize.X;
+	float ViewRectHeight = InCameraInfo.TilingParams.TileSize.Y;
+
+	const float DestAspectRatio = ViewRectWidth / ViewRectHeight;
+	const float CameraAspectRatio = InCameraInfo.bAllowCameraAspectRatio ? InCameraInfo.ViewInfo.AspectRatio : DestAspectRatio;
+	
+	float ViewFOV = InCameraInfo.ViewInfo.FOV;
+
+	// Inflate our FOV to support the overscan 
+	ViewFOV = 2.0f * FMath::RadiansToDegrees(FMath::Atan((1.0f + InCameraInfo.OverscanFraction) * FMath::Tan(FMath::DegreesToRadians(ViewFOV * 0.5f))));
+	
+	if (InCameraInfo.ViewInfo.ProjectionMode == ECameraProjectionMode::Orthographic)
+	{
+		const float YScale = 1.0f / InCameraInfo.ViewInfo.AspectRatio;
+		const float OverscanScale = 1.0f + (InCameraInfo.OverscanFraction);
+
+		const float HalfOrthoWidth = (InCameraInfo.ViewInfo.OrthoWidth / 2.0f) * OverscanScale;
+		const float ScaledOrthoHeight = (InCameraInfo.ViewInfo.OrthoWidth / 2.0f) * OverscanScale * YScale;
+
+		const float NearPlane = InCameraInfo.ViewInfo.OrthoNearClipPlane;
+		const float FarPlane = InCameraInfo.ViewInfo.OrthoFarClipPlane;
+
+		const float ZScale = 1.0f / (FarPlane - NearPlane);
+		const float ZOffset = -NearPlane;
+
+		BaseProjMatrix = FReversedZOrthoMatrix(
+			HalfOrthoWidth,
+			ScaledOrthoHeight,
+			ZScale,
+			ZOffset
+		);
+	}
+	else
+	{
+		float XAxisMultiplier;
+		float YAxisMultiplier;
+
+		if (InCameraInfo.ViewInfo.bConstrainAspectRatio)
+		{
+			// If the camera's aspect ratio has a thinner width, then stretch the horizontal fov more than usual to 
+			// account for the extra with of (before constraining - after constraining)
+			if (InCameraInfo.ViewInfo.AspectRatio < DestAspectRatio)
+			{
+				const float ConstrainedWidth = ViewRectHeight * InCameraInfo.ViewInfo.AspectRatio;
+				XAxisMultiplier = ConstrainedWidth / (float)ViewRectWidth;
+				YAxisMultiplier = InCameraInfo.ViewInfo.AspectRatio;
+			}
+			// Simplified some math here but effectively functions similarly to the above, the unsimplified code would look like:
+			// const float ConstrainedHeight = ViewRectWidth / CameraCache.AspectRatio;
+			// YAxisMultiplier = (ConstrainedHeight / ViewInitOptions.GetViewRect.Height()) * CameraCache.AspectRatio;
+			else
+			{
+				XAxisMultiplier = 1.0f;
+				YAxisMultiplier = ViewRectWidth / ViewRectHeight;
+			}
+		}
+		else
+		{
+			const EAspectRatioAxisConstraint AspectRatioAxisConstraint = GetDefault<ULocalPlayer>()->AspectRatioAxisConstraint;
+			if (((ViewRectWidth > ViewRectHeight) && (AspectRatioAxisConstraint == AspectRatio_MajorAxisFOV)) || (AspectRatioAxisConstraint == AspectRatio_MaintainXFOV))
+			{
+				//if the viewport is wider than it is tall
+				XAxisMultiplier = 1.0f;
+				YAxisMultiplier = ViewRectWidth / ViewRectHeight;
+			}
+			else
+			{
+				//if the viewport is taller than it is wide
+				XAxisMultiplier = ViewRectHeight / ViewRectWidth;
+				YAxisMultiplier = 1.0f;
+			}
+		}
+
+		const float MinZ = InCameraInfo.ViewInfo.GetFinalPerspectiveNearClipPlane();
+		const float MaxZ = MinZ;
+		// Avoid zero ViewFOV's which cause divide by zero's in projection matrix
+		const float MatrixFOV = FMath::Max(0.001f, ViewFOV) * (float)PI / 360.0f;
+
+
+		if ((bool)ERHIZBuffer::IsInverted)
+		{
+			BaseProjMatrix = FReversedZPerspectiveMatrix(
+				MatrixFOV,
+				MatrixFOV,
+				XAxisMultiplier,
+				YAxisMultiplier,
+				MinZ,
+				MaxZ
+			);
+		}
+		else
+		{
+			BaseProjMatrix = FPerspectiveMatrix(
+				MatrixFOV,
+				MatrixFOV,
+				XAxisMultiplier,
+				YAxisMultiplier,
+				MinZ,
+				MaxZ
+			);
+		}
+	}
+	
+	return BaseProjMatrix;
+}
+
+FVector4f FMovieGraphImagePassBase::CalculatePrinciplePointOffsetForTiling(const UE::MovieGraph::DefaultRenderer::FMovieGraphTilingParams& InTilingParams) const 
+{
+	// We need our final view parameters to be in the space of [-1,1], including all the tiles.
+	// Starting with a single tile, the middle of the tile in offset screen space is:
+	FVector2f TilePrincipalPointOffset;
+
+	TilePrincipalPointOffset.X = (float(InTilingParams.TileIndexes.X) + 0.5f - (0.5f * float(InTilingParams.TileCount.X))) * 2.0f;
+	TilePrincipalPointOffset.Y = (float(InTilingParams.TileIndexes.Y) + 0.5f - (0.5f * float(InTilingParams.TileCount.Y))) * 2.0f;
+
+	// For the tile size ratio, we have to multiply by (1.0 + overlap) and then divide by tile num
+	FVector2D OverlapScale;
+	OverlapScale.X = (1.0f + float(2 * InTilingParams.OverlapPad.X) / float(InTilingParams.TileSize.X));
+	OverlapScale.Y = (1.0f + float(2 * InTilingParams.OverlapPad.Y) / float(InTilingParams.TileSize.Y));
+
+	TilePrincipalPointOffset.X /= OverlapScale.X;
+	TilePrincipalPointOffset.Y /= OverlapScale.Y;
+
+	FVector2D TilePrincipalPointScale;
+	TilePrincipalPointScale.X = OverlapScale.X / float(InTilingParams.TileCount.X);
+	TilePrincipalPointScale.Y = OverlapScale.Y / float(InTilingParams.TileCount.Y);
+
+	TilePrincipalPointOffset.X *= TilePrincipalPointScale.X;
+	TilePrincipalPointOffset.Y *= TilePrincipalPointScale.Y;
+
+	return FVector4f(TilePrincipalPointOffset.X, -TilePrincipalPointOffset.Y, TilePrincipalPointScale.X, TilePrincipalPointScale.Y);
+}
+
+void FMovieGraphImagePassBase::ModifyProjectionMatrixForTiling(const UE::MovieGraph::DefaultRenderer::FMovieGraphTilingParams& InTilingParams, const bool bInOrthographic, FMatrix& InOutProjectionMatrix, float& OutDoFSensorScale) const
+{
+	float PadRatioX = 1.0f;
+	float PadRatioY = 1.0f;
+
+	if (InTilingParams.OverlapPad.X > 0 && InTilingParams.OverlapPad.Y > 0)
+	{
+		PadRatioX = float(InTilingParams.OverlapPad.X * 2 + InTilingParams.TileSize.X) / float(InTilingParams.TileSize.X);
+		PadRatioY = float(InTilingParams.OverlapPad.Y * 2 + InTilingParams.TileSize.Y) / float(InTilingParams.TileSize.Y);
+	}
+
+	float ScaleX = PadRatioX / float(InTilingParams.TileCount.X);
+	float ScaleY = PadRatioY / float(InTilingParams.TileCount.Y);
+
+	InOutProjectionMatrix.M[0][0] /= ScaleX;
+	InOutProjectionMatrix.M[1][1] /= ScaleY;
+	OutDoFSensorScale = ScaleX;
+
+	// this offset would be correct with no pad
+	float OffsetX = -((float(InTilingParams.TileIndexes.X) + 0.5f - float(InTilingParams.TileCount.X) / 2.0f) * 2.0f);
+	float OffsetY = ((float(InTilingParams.TileIndexes.Y) + 0.5f - float(InTilingParams.TileCount.Y) / 2.0f) * 2.0f);
+
+	if (bInOrthographic)
+	{
+		InOutProjectionMatrix.M[3][0] += OffsetX / PadRatioX;
+		InOutProjectionMatrix.M[3][1] += OffsetY / PadRatioY;
+	}
+	else
+	{
+		InOutProjectionMatrix.M[2][0] += OffsetX / PadRatioX;
+		InOutProjectionMatrix.M[2][1] += OffsetY / PadRatioY;
+	}
+}
+
+
+void FMovieGraphImagePassBase::PostRendererSubmission(
+	const UE::MovieGraph::FMovieGraphSampleState& InSampleState,
+	const UE::MovieGraph::DefaultRenderer::FRenderTargetInitParams& InRenderTargetInitParams, FCanvas& InCanvas, const UE::MovieGraph::DefaultRenderer::FCameraInfo& InCameraInfo)
+{
+	// We have a pool of accumulators - we multi-thread the accumulation on the task graph, and for each frame,
+	// the task has the previous samples as pre-reqs to keep the accumulation in order. However, each accumulator
+	// can only work on one frame at a time, so we create a pool of them to work concurrently. This needs a limit
+	// as large accumulations (16k) can take a lot of system RAM.
+	TObjectPtr<UMovieGraphDefaultRenderer> GraphRenderer = GetRenderer().Get();
+	if (!GraphRenderer)
+	{
+		return;
+	}
+
+	FMoviePipelineAccumulatorPoolPtr SampleAccumulatorPool = GraphRenderer->GetOrCreateAccumulatorPool<FImageOverlappedAccumulator>();
+	UE::MovieGraph::DefaultRenderer::FSurfaceAccumulatorPool::FInstancePtr AccumulatorInstance = SampleAccumulatorPool->GetAccumulatorInstance_GameThread<FImageOverlappedAccumulator>(InSampleState.TraversalContext.Time.RenderedFrameNumber, InSampleState.TraversalContext.RenderDataIdentifier);
+	
+	FMoviePipelineSurfaceQueuePtr LocalSurfaceQueue = GraphRenderer->GetOrCreateSurfaceQueue(InRenderTargetInitParams);
+	LocalSurfaceQueue->BlockUntilAnyAvailable();
+
+	FMovieGraphRenderDataAccumulationArgs AccumulationArgs;
+	{
+		AccumulationArgs.OutputMerger = GraphRenderer->GetOwningGraph()->GetOutputMerger();
+		AccumulationArgs.ImageAccumulator = StaticCastSharedPtr<FImageOverlappedAccumulator>(AccumulatorInstance->Accumulator);
+		AccumulationArgs.bIsFirstSample = InSampleState.TraversalContext.Time.bIsFirstTemporalSampleForFrame;
+		AccumulationArgs.bIsLastSample = InSampleState.TraversalContext.Time.bIsLastTemporalSampleForFrame;
+	}
+
+	auto OnSurfaceReadbackFinished = [this, InSampleState, AccumulationArgs, AccumulatorInstance](TUniquePtr<FImagePixelData>&& InPixelData)
+	{
+		UE::Tasks::TTask<void> Task = UE::Tasks::Launch(UE_SOURCE_LOCATION, [PixelData = MoveTemp(InPixelData), InSampleState, AccumulationArgs, AccumulatorInstance]() mutable
+		{
+			// Enqueue a encode for this frame onto our worker thread.
+			AccumulateSample_TaskThread(MoveTemp(PixelData), InSampleState, AccumulationArgs);
+
+			// We have to defer clearing the accumulator until after sample accumulation has finished
+			if (AccumulationArgs.bIsLastSample)
+			{
+				// Final sample has now been executed, free the accumulator for reuse.
+				AccumulatorInstance->bIsActive = false;
+			}
+		}, AccumulatorInstance->TaskPrereq);
+
+		// Make the next accumulation task that uses this accumulator use the task we just created as a pre-req.
+		AccumulatorInstance->TaskPrereq = Task;
+
+		// Because we're run on a separate thread, we need to check validity differently. The standard
+		// TWeakObjectPtr will report non-valid during GC (even if the object it's pointing to isn't being
+		// GC'd).
+		const bool bEvenIfPendingKill = false;
+		const bool bThreadsafeTest = true;
+		const bool bValid = this->WeakGraphRenderer.IsValid(bEvenIfPendingKill, bThreadsafeTest);
+		if (ensureMsgf(bValid, TEXT("Renderer was garbage collected while outstanding tasks existed, outstanding tasks were not flushed properly during shutdown!")))
+		{
+			// The regular Get() will fail during GC so we use the above check to see if it's valid
+			// before ignoring it and directly getting the object.
+			this->WeakGraphRenderer.GetEvenIfUnreachable()->AddOutstandingRenderTask_AnyThread(Task);
+		}
+	};
+
+	FRenderTarget* RenderTarget = InCanvas.GetRenderTarget();
+
+	ENQUEUE_RENDER_COMMAND(CanvasRenderTargetResolveCommand)(
+		[LocalSurfaceQueue, OnSurfaceReadbackFinished, RenderTarget](FRHICommandListImmediate& RHICmdList) mutable
+		{
+			// The legacy surface reader takes the payload just so it can shuffle it into our callback, but we can just include the data
+			// directly in the callback, so this is just a dummy payload.
+			TSharedRef<FImagePixelDataPayload, ESPMode::ThreadSafe> FramePayload = MakeShared<FImagePixelDataPayload, ESPMode::ThreadSafe>();
+			LocalSurfaceQueue->OnRenderTargetReady_RenderThread(RenderTarget->GetRenderTargetTexture(), FramePayload, MoveTemp(OnSurfaceReadbackFinished));
+		});
+}
+
+void AccumulateSample_TaskThread(TUniquePtr<FImagePixelData>&& InPixelData, const ::UE::MovieGraph::FMovieGraphSampleState InSampleState, const FMovieGraphRenderDataAccumulationArgs& InAccumulationParams)
+{
+	TRACE_CPUPROFILER_EVENT_SCOPE(MoviePipeline_AccumulateSample);
+		
+	TUniquePtr<FImagePixelData> SamplePixelData = MoveTemp(InPixelData);
+
+	// Associate the sample state with the image as payload data, this allows downstream systems to fetch the values without us having to store the data
+	// separately and ensure they stay paired the whole way down.
+	TSharedPtr<FMovieGraphSampleState> SampleStatePayload = MakeShared<FMovieGraphSampleState>(InSampleState);
+	SamplePixelData->SetPayload(StaticCastSharedPtr<IImagePixelDataPayload>(SampleStatePayload));
+
+	TSharedPtr<IMovieGraphOutputMerger, ESPMode::ThreadSafe> OutputMergerPin = InAccumulationParams.OutputMerger.Pin();
+	if (!OutputMergerPin.IsValid())
+	{
+		return;
+	}
+
+	const bool bIsWellFormed = SamplePixelData->IsDataWellFormed();
+	check(bIsWellFormed);
+
+	if (SampleStatePayload->bWriteSampleToDisk)
+	{
+		// Debug Feature: Write the raw sample to disk for debugging purposes. We copy the data here,
+		// as we don't want to disturb the memory flow below.
+		TUniquePtr<FImagePixelData> SampleData = SamplePixelData->CopyImageData();
+		OutputMergerPin->OnSingleSampleDataAvailable_AnyThread(MoveTemp(SampleData));
+	}
+
+	// Optimization! If we don't need the accumulator (no tiling, no sub-sampling) then we'll skip it and just send it straight to the output stage.
+	// This reduces memory requirements and improves performance in the baseline case.
+	if (!SampleStatePayload->bRequiresAccumulator)
+	{
+		OutputMergerPin->OnCompleteRenderPassDataAvailable_AnyThread(MoveTemp(SamplePixelData));
+		return;
+	}
+
+	TSharedPtr<FImageOverlappedAccumulator> AccumulatorPin = InAccumulationParams.ImageAccumulator.Pin();
+	if (AccumulatorPin->NumChannels == 0)
+	{
+		LLM_SCOPE_BYNAME(TEXT("MoviePipeline/ImageAccumulatorInitMemory"));
+		const int32 ChannelCount = 4;
+		AccumulatorPin->InitMemory(SampleStatePayload->AccumulatorResolution, ChannelCount);
+		AccumulatorPin->ZeroPlanes();
+		AccumulatorPin->AccumulationGamma = 1.f;
+	}
+
+	// Accumulate the new sample to our target
+	{
+		TRACE_CPUPROFILER_EVENT_SCOPE(MoviePipeline_AccumulatePixelData);
+				
+		const FIntPoint TileSize = SampleStatePayload->BackbufferResolution;
+		const FIntPoint OverlappedPad = SampleStatePayload->OverlappedPad;
+		const FIntPoint OverlappedOffset = SampleStatePayload->OverlappedOffset;
+		const FVector2D OverlappedSubpixelShift = SampleStatePayload->OverlappedSubpixelShift;
+		::MoviePipeline::FTileWeight1D WeightFunctionX;
+		::MoviePipeline::FTileWeight1D WeightFunctionY;
+		WeightFunctionX.InitHelper(OverlappedPad.X, TileSize.X, OverlappedPad.X);
+		WeightFunctionY.InitHelper(OverlappedPad.Y, TileSize.Y, OverlappedPad.Y);
+				
+		AccumulatorPin->AccumulatePixelData(*SamplePixelData, OverlappedOffset, OverlappedSubpixelShift, WeightFunctionX, WeightFunctionY);
+	}
+
+	// Finally on our last sample, we fetch the data out of the accumulator
+	// and move it to the Output Merger.
+	if (SampleStatePayload->bFetchFromAccumulator)
+	{
+		TRACE_CPUPROFILER_EVENT_SCOPE(MoviePipeline_FetchAccumulatedPixelData);
+		int32 FullSizeX = AccumulatorPin->PlaneSize.X;
+		int32 FullSizeY = AccumulatorPin->PlaneSize.Y;
+
+		// Now that a tile is fully built and accumulated we can notify the output builder that the
+		// data is ready so it can pass that onto the output containers (if needed).
+		if (SamplePixelData->GetType() == EImagePixelType::Float32)
+		{
+			// 32 bit FLinearColor
+			TUniquePtr<TImagePixelData<FLinearColor> > FinalPixelData = MakeUnique<TImagePixelData<FLinearColor>>(FIntPoint(FullSizeX, FullSizeY), SampleStatePayload);
+			AccumulatorPin->FetchFinalPixelDataLinearColor(FinalPixelData->Pixels);
+
+			// Send the data to the Output Builder
+			OutputMergerPin->OnCompleteRenderPassDataAvailable_AnyThread(MoveTemp(FinalPixelData));
+		}
+		else if (SamplePixelData->GetType() == EImagePixelType::Float16)
+		{
+			// 16 bit FLinearColor
+			TUniquePtr<TImagePixelData<FFloat16Color> > FinalPixelData = MakeUnique<TImagePixelData<FFloat16Color>>(FIntPoint(FullSizeX, FullSizeY), SampleStatePayload);
+			AccumulatorPin->FetchFinalPixelDataHalfFloat(FinalPixelData->Pixels);
+
+			// Send the data to the Output Builder
+			OutputMergerPin->OnCompleteRenderPassDataAvailable_AnyThread(MoveTemp(FinalPixelData));
+		}
+		else if (SamplePixelData->GetType() == EImagePixelType::Color)
+		{
+			// 8bit FColors
+			TUniquePtr<TImagePixelData<FColor>> FinalPixelData = MakeUnique<TImagePixelData<FColor>>(FIntPoint(FullSizeX, FullSizeY), SampleStatePayload);
+			AccumulatorPin->FetchFinalPixelDataByte(FinalPixelData->Pixels);
+
+			// Send the data to the Output Builder
+			OutputMergerPin->OnCompleteRenderPassDataAvailable_AnyThread(MoveTemp(FinalPixelData));
+		}
+		else
+		{
+			check(0);
+		}
+
+		{
+			TRACE_CPUPROFILER_EVENT_SCOPE(MoviePipeline_FreeAccumulatedPixelData);
+			// Free the memory in the accumulator.
+			AccumulatorPin->Reset();
+		}
+	}
+
+	{
+		TRACE_CPUPROFILER_EVENT_SCOPE(ReleaseSampleData);
+		SamplePixelData.Reset();
+	}
+}
+
+} // namespace UE::MovieGraph::Rendering
+
diff --git a/Engine/Plugins/MovieScene/MovieRenderPipeline/Source/MovieRenderPipelineRenderPasses/Private/Graph/Renderers/MovieGraphPathTracerPass.cpp b/Engine/Plugins/MovieScene/MovieRenderPipeline/Source/MovieRenderPipelineRenderPasses/Private/Graph/Renderers/MovieGraphPathTracerPass.cpp
new file mode 100644
--- /dev/null
+++ b/Engine/Plugins/MovieScene/MovieRenderPipeline/Source/MovieRenderPipelineRenderPasses/Private/Graph/Renderers/MovieGraphPathTracerPass.cpp
@@ -0,0 +1,72 @@
+// Copyright Epic Games, Inc. All Rights Reserved.
+
+#include "Graph/Renderers/MovieGraphPathTracerPass.h"
+#include "Graph/Nodes/MovieGraphPathTracedRenderPassNode.h"
+
+namespace UE::MovieGraph::Rendering
+{
+
+UMovieGraphImagePassBaseNode* FMovieGraphPathTracerPass::GetParentNode(UMovieGraphEvaluatedConfig* InConfig) const
+{
+	// This is a bit of a workaround for the fact that the pass doesn't have a strong pointer to the node it's supposed to be associated with,
+	// since that instance changes every frame. So instead we have a virtual function here so the node can look it up by type, and then we can
+	// call a bunch of virtual functions on the right instance to fetch values.
+	const bool bIncludeCDOs = true;
+	UMovieGraphPathTracedRenderPassNode* ParentNode = InConfig->GetSettingForBranch<UMovieGraphPathTracedRenderPassNode>(GetBranchName(), bIncludeCDOs);
+	if (!ensureMsgf(ParentNode, TEXT("PathTracerPass should not exist without parent node in graph.")))
+	{
+		return nullptr;
+	}
+
+	return ParentNode;
+}
+namespace Private
+{
+	void GetSampleData(TSharedRef<FSceneViewFamilyContext> InOutFamily, const UE::MovieGraph::DefaultRenderer::FCameraInfo& InCameraInfo, int32& OutSampleCount, int32& OutSampleIndex)
+	{
+		// If motion blur is enabled:
+		//    blend all spatial samples together while leaving the handling of temporal samples up to MRQ
+		//    each temporal sample will include denoising and post-process effects
+		// If motion blur is NOT enabled:
+		//    blend all temporal+spatial samples within the path tracer and only apply denoising on the last temporal sample
+		//    this way we minimize denoising cost and also allow a much higher number of temporal samples to be used which
+		//    can help reduce strobing
+
+		// NOTE: Tiling is not compatible with the reference motion blur mode because it changes the order of the loops over the image.
+		const bool bHasTiles = InCameraInfo.TilingParams.TileCount.X > 1 || InCameraInfo.TilingParams.TileCount.Y;
+		const bool bAccumulateSpatialSamplesOnly = InOutFamily->EngineShowFlags.MotionBlur || bHasTiles;
+
+		OutSampleCount = bAccumulateSpatialSamplesOnly ? InCameraInfo.SamplingParams.SpatialSampleCount : InCameraInfo.SamplingParams.TemporalSampleCount * InCameraInfo.SamplingParams.SpatialSampleCount;
+		OutSampleIndex = bAccumulateSpatialSamplesOnly ? InCameraInfo.SamplingParams.SpatialSampleIndex : InCameraInfo.SamplingParams.TemporalSampleIndex * InCameraInfo.SamplingParams.SpatialSampleCount + InCameraInfo.SamplingParams.SpatialSampleIndex;
+	}
+}
+
+bool FMovieGraphPathTracerPass::ShouldDiscardOutput(const TSharedRef<FSceneViewFamilyContext>& InFamily, const UE::MovieGraph::DefaultRenderer::FCameraInfo& InCameraInfo) const
+{ 
+	int32 SampleCount, SampleIndex;
+	Private::GetSampleData(InFamily, InCameraInfo, SampleCount, SampleIndex);
+
+	// Discard the result, unless its the last sample
+	return !(SampleIndex == SampleCount - 1);
+}
+
+void FMovieGraphPathTracerPass::ApplyMovieGraphOverridesToSceneView(TSharedRef<FSceneViewFamilyContext> InOutFamily, const FViewFamilyInitData& InInitData, const UE::MovieGraph::DefaultRenderer::FCameraInfo& InCameraInfo) const
+{
+	FMovieGraphImagePassBase::ApplyMovieGraphOverridesToSceneView(InOutFamily, InInitData, InCameraInfo);
+
+	// Override whatever settings came from PostProcessVolume or Camera
+	int32 SampleCount, SampleIndex;
+	Private::GetSampleData(InOutFamily, InCameraInfo, SampleCount, SampleIndex);
+
+	// TODO: pass along FrameIndex (which includes SampleIndex) to make sure sampling is fully deterministic
+		
+	// Overwrite whatever sampling count came from the PostProcessVolume
+	FSceneView* View = const_cast<FSceneView*>(InOutFamily->Views[0]);
+	View->FinalPostProcessSettings.bOverride_PathTracingSamplesPerPixel = true;
+	View->FinalPostProcessSettings.PathTracingSamplesPerPixel = SampleCount;
+		
+	// reset path tracer's accumulation at the start of each sample
+	View->bForcePathTracerReset = SampleIndex == 0;
+}
+
+} // UE::MovieGraph::Rendering
\ No newline at end of file
diff --git a/Engine/Plugins/MovieScene/MovieRenderPipeline/Source/MovieRenderPipelineRenderPasses/Public/Graph/Nodes/MovieGraphDeferredPassNode.h b/Engine/Plugins/MovieScene/MovieRenderPipeline/Source/MovieRenderPipelineRenderPasses/Public/Graph/Nodes/MovieGraphDeferredPassNode.h
new file mode 100644
--- /dev/null
+++ b/Engine/Plugins/MovieScene/MovieRenderPipeline/Source/MovieRenderPipelineRenderPasses/Public/Graph/Nodes/MovieGraphDeferredPassNode.h
@@ -0,0 +1,107 @@
+﻿// Copyright Epic Games, Inc. All Rights Reserved.
+
+#pragma once
+#include "Graph/Nodes/MovieGraphImagePassBaseNode.h"
+#include "MovieGraphDeferredPassNode.generated.h"
+
+/**
+* A render node which uses the Deferred Renderer.
+*/
+UCLASS()
+class MOVIERENDERPIPELINERENDERPASSES_API UMovieGraphDeferredRenderPassNode : public UMovieGraphImagePassBaseNode 
+{
+	GENERATED_BODY()
+
+public:
+	UMovieGraphDeferredRenderPassNode();
+
+	virtual void GetFormatResolveArgs(FMovieGraphResolveArgs& OutMergedFormatArgs, const FMovieGraphRenderDataIdentifier& InRenderDataIdentifier) const override;
+
+#if WITH_EDITOR
+	virtual FText GetNodeTitle(const bool bGetDescriptive = false) const override;
+#endif
+
+protected:
+	// UMovieGraphRenderPassNode Interface
+	virtual FString GetRendererNameImpl() const override;
+	virtual TUniquePtr<UE::MovieGraph::Rendering::FMovieGraphImagePassBase> CreateInstance() const override;
+	virtual bool GetWriteAllSamples() const override;
+	virtual int32 GetNumSpatialSamples() const override;
+	virtual bool GetDisableToneCurve() const override;
+	virtual bool GetAllowOCIO() const override;
+	virtual EAntiAliasingMethod GetAntiAliasingMethod() const override;
+	// ~UMovieGraphRenderPassNode Interface
+
+	// UMovieGraphImagePassBaseNode Interface
+	virtual EViewModeIndex GetViewModeIndex() const override;
+	// ~UMovieGraphImagePassBaseNode Interface
+
+public:
+	UPROPERTY(EditAnywhere, BlueprintReadWrite, Category = Overrides, meta = (InlineEditConditionToggle))
+	uint8 bOverride_SpatialSampleCount: 1;
+
+	UPROPERTY(EditAnywhere, BlueprintReadWrite, Category = Overrides, meta = (InlineEditConditionToggle))
+	uint8 bOverride_AntiAliasingMethod : 1;
+	
+	UPROPERTY(EditAnywhere, BlueprintReadWrite, Category = Overrides, meta = (InlineEditConditionToggle))
+	uint8 bOverride_bDisableToneCurve : 1;
+
+	UPROPERTY(EditAnywhere, BlueprintReadWrite, Category = Overrides, meta = (InlineEditConditionToggle))
+	uint8 bOverride_bAllowOCIO : 1;
+	
+	UPROPERTY(EditAnywhere, BlueprintReadWrite, Category = Overrides, meta = (InlineEditConditionToggle))
+	uint8 bOverride_ViewModeIndex : 1;
+
+	UPROPERTY(EditAnywhere, BlueprintReadWrite, Category = Overrides, meta = (InlineEditConditionToggle))
+	uint8 bOverride_bWriteAllSamples : 1;
+
+	/**
+	* How many sub-pixel jitter renders should we do per temporal sample? This can be used to achieve high
+	* sample counts without Temporal Sub-Sampling (allowing high sample counts without motion blur being enabled),
+	* but we generally recommend using Temporal Sub-Samples when possible. It can also be combined with
+	* temporal samples and you will get SpatialSampleCount many renders per temporal sample.
+	*/
+	UPROPERTY(EditAnywhere, BlueprintReadWrite, meta = (UIMin = 1, ClampMin = 1), Category = "Settings", meta=(EditCondition="bOverride_SpatialSampleCount"))
+	int32 SpatialSampleCount;
+
+	/**
+	* Which anti-aliasing method should this render use. If this is set to None, then Movie Render Graph
+	* will handle anti-aliasing by doing a sub-pixel jitter (one for each temporal/spatial sample). Some
+	* rendering effects rely on TSR or TAA to reduce noise so we reccomend leaving them enabled
+	* where possible. All options work with Spatial and Temporal samples, but TSR/TAA may introduce minor
+	* visual artifacts (such as ghosting). MSAA is not supported in the deferred renderer.
+	*/
+	UPROPERTY(EditAnywhere, BlueprintReadWrite, Category = "Settings", meta = (EditCondition = "bOverride_AntiAliasingMethod"))
+	TEnumAsByte<EAntiAliasingMethod> AntiAliasingMethod;
+
+	/**
+	* If true, the tone curve will be disabled for this render pass. This will result in values greater than 1.0 in final renders
+	* and can optionally be combined with OCIO profiles on the file output nodes to convert from Linear Values in Working Color Space
+	* (which is sRGB  (Rec. 709) by default, unless changed in the project settings).
+	*/
+	UPROPERTY(EditAnywhere, BlueprintReadWrite, Category = "Settings", meta = (EditCondition = "bOverride_bDisableToneCurve"))
+	bool bDisableToneCurve;
+
+	/**
+	* Allow the output file OpenColorIO transform to be used on this render.
+	*/
+	UPROPERTY(EditAnywhere, BlueprintReadWrite, Category = "Settings", meta = (EditCondition = "bOverride_bAllowOCIO"))
+	bool bAllowOCIO;
+
+	/**
+	* Debug Feature. Can use this to write out each individual Temporal and Spatial sample rendered by this render pass,
+	* which allows you to see which images are being accumulated together. Can be useful for debugging incorrect looking
+	* frames to see which sub-frame evaluations were incorrect.
+	*/
+	UPROPERTY(EditAnywhere, BlueprintReadWrite, Category = "Debug", meta = (EditCondition = "bOverride_bWriteAllsamples"))
+	bool bWriteAllSamples;
+
+	/** 
+	* The view mode index that will be applied to renders. These mirror the View Modes you find in the Viewport,
+	* but most view modes other than Lit are used for debugging so they may not do what you expect, or may
+	* have to be used in combination with certain Show Flags to produce a result similar to what you see in
+	* the viewport.
+	*/
+	UPROPERTY(EditAnywhere, BlueprintReadWrite, Category = "Settings", meta=(EditCondition="bOverride_ViewModeIndex", InvalidEnumValues = "VMI_PathTracing"))
+	TEnumAsByte<EViewModeIndex> ViewModeIndex;
+};
\ No newline at end of file
diff --git a/Engine/Plugins/MovieScene/MovieRenderPipeline/Source/MovieRenderPipelineRenderPasses/Public/Graph/Nodes/MovieGraphImagePassBaseNode.h b/Engine/Plugins/MovieScene/MovieRenderPipeline/Source/MovieRenderPipelineRenderPasses/Public/Graph/Nodes/MovieGraphImagePassBaseNode.h
new file mode 100644
--- /dev/null
+++ b/Engine/Plugins/MovieScene/MovieRenderPipeline/Source/MovieRenderPipelineRenderPasses/Public/Graph/Nodes/MovieGraphImagePassBaseNode.h
@@ -0,0 +1,76 @@
+// Copyright Epic Games, Inc. All Rights Reserved.
+
+#pragma once
+#include "Graph/MovieGraphNode.h"
+#include "Graph/MovieGraphDataTypes.h"
+#include "Graph/MovieGraphConfig.h"
+#include "Graph/Renderers/MovieGraphImagePassBase.h"
+#include "Graph/Renderers/MovieGraphShowFlags.h"
+#include "Graph/Nodes/MovieGraphRenderPassNode.h"
+#include "MovieGraphImagePassBaseNode.generated.h"
+
+// Forward Declare
+struct FMovieGraphRenderPassSetupData;
+struct FMovieGraphTimeStepData;
+
+/**
+* The UMovieGraphImagePassBaseNode is an abstract base-class for render nodes that wish to create
+* renders of the 3d scene. You are not required to inherit from this node (can inherit from
+* UMovieGraphRenderPassNode), but this node provides a helpful set of functions and default values
+* for constructing the required matrices and settings for viewport-like renders.
+*/
+UCLASS(Abstract)
+class MOVIERENDERPIPELINERENDERPASSES_API UMovieGraphImagePassBaseNode : public UMovieGraphRenderPassNode
+{
+	GENERATED_BODY()
+public:
+	UMovieGraphImagePassBaseNode();
+
+	static void AddReferencedObjects(UObject* InThis, FReferenceCollector& Collector);
+		
+	/** Gets the view mode index that should be active for this renderer. */
+	virtual EViewModeIndex GetViewModeIndex() const;
+
+	/** Gets the show flags that should be active for this renderer. */
+	virtual FEngineShowFlags GetShowFlags() const;
+
+	/** Gets any show flags that should be applied as defaults, before user changes are applied. */
+	virtual TArray<uint32> GetDefaultShowFlags() const;
+
+	/** Should each individual sample rendered be written out for debugging? */
+	virtual bool GetWriteAllSamples() const { return false; }
+
+	/** How many spatial samples should be rendered each frame? */
+	virtual int32 GetNumSpatialSamples() const { return 1; }
+
+	/** Should the tone curve be disabled while rendering? Allows for linear values in exrs but changes the look of the final image. */
+	virtual bool GetDisableToneCurve() const { return false; }
+
+	/** Should the output file be allowed to apply an OCIO transform on this render? */
+	virtual bool GetAllowOCIO() const { return true; }
+
+	/** Which AA Method should be used? */
+	virtual EAntiAliasingMethod GetAntiAliasingMethod() const { return EAntiAliasingMethod::AAM_None; }
+
+protected:
+	// Note: Since *individual* show flags are overridden instead of the entire ShowFlags property, manually set to
+	// overridden so the traversal picks the changes up (otherwise they will be ignored).
+	UPROPERTY(EditAnywhere, BlueprintReadWrite, Category = Overrides, meta = (InlineEditConditionToggle))
+	uint8 bOverride_ShowFlags : 1 = 1; //-V570
+
+	/** The show flags that should be active during a render for this node. */
+	UPROPERTY(EditAnywhere, BlueprintReadWrite, Category = "Show Flags")
+	TObjectPtr<UMovieGraphShowFlags> ShowFlags;
+
+	TArray<TUniquePtr<UE::MovieGraph::Rendering::FMovieGraphImagePassBase>> CurrentInstances;
+protected:
+	// UMovieGraphRenderPassNode Interface
+	virtual FString GetRendererNameImpl() const { return TEXT("UnnamedImageRenderPass"); }
+	virtual void SetupImpl(const FMovieGraphRenderPassSetupData& InSetupData) override;
+	virtual void TeardownImpl() override;
+	virtual void RenderImpl(const FMovieGraphTraversalContext& InFrameTraversalContext, const FMovieGraphTimeStepData& InTimeData) override;
+	virtual void GatherOutputPassesImpl(TArray<FMovieGraphRenderDataIdentifier>& OutExpectedPasses) const override;
+	// ~UMovieGraphRenderPassNode Interface
+
+	virtual TUniquePtr<UE::MovieGraph::Rendering::FMovieGraphImagePassBase> CreateInstance() const { return nullptr; }
+};
\ No newline at end of file
diff --git a/Engine/Plugins/MovieScene/MovieRenderPipeline/Source/MovieRenderPipelineRenderPasses/Public/Graph/Nodes/MovieGraphPathTracedRenderPassNode.h b/Engine/Plugins/MovieScene/MovieRenderPipeline/Source/MovieRenderPipelineRenderPasses/Public/Graph/Nodes/MovieGraphPathTracedRenderPassNode.h
new file mode 100644
--- /dev/null
+++ b/Engine/Plugins/MovieScene/MovieRenderPipeline/Source/MovieRenderPipelineRenderPasses/Public/Graph/Nodes/MovieGraphPathTracedRenderPassNode.h
@@ -0,0 +1,96 @@
+// Copyright Epic Games, Inc. All Rights Reserved.
+
+#pragma once
+
+#include "Graph/Nodes/MovieGraphImagePassBaseNode.h"
+
+#include "MovieGraphPathTracedRenderPassNode.generated.h"
+
+/** A render node which uses the path tracer. */
+UCLASS()
+class MOVIERENDERPIPELINERENDERPASSES_API UMovieGraphPathTracedRenderPassNode : public UMovieGraphImagePassBaseNode
+{
+	GENERATED_BODY()
+
+public:
+	UMovieGraphPathTracedRenderPassNode();
+
+#if WITH_EDITOR
+	virtual FText GetNodeTitle(const bool bGetDescriptive = false) const override;
+#endif
+
+	// UMovieGraphRenderPassNode Interface
+	virtual void SetupImpl(const FMovieGraphRenderPassSetupData& InSetupData) override;
+	virtual void TeardownImpl() override;
+	// ~UMovieGraphRenderPassNode Interface
+
+	// UMovieGraphImagePassBaseNode Interface
+	virtual bool GetWriteAllSamples() const override;
+	virtual int32 GetNumSpatialSamples() const override;
+	virtual bool GetDisableToneCurve() const override;
+	virtual bool GetAllowOCIO() const override;
+	virtual TUniquePtr<UE::MovieGraph::Rendering::FMovieGraphImagePassBase> CreateInstance() const;
+	// ~UMovieGraphImagePassBaseNode Interface
+
+protected:
+	// UMovieGraphRenderPassNode Interface
+	virtual FString GetRendererNameImpl() const override;
+
+	// ~UMovieGraphRenderPassNode Interface
+
+	// UMovieGraphCoreRenderPassNode Interface
+	virtual EViewModeIndex GetViewModeIndex() const override;
+	// ~UMovieGraphCoreRenderPassNode Interface
+
+public:
+
+	UPROPERTY(EditAnywhere, BlueprintReadWrite, Category = Overrides, meta = (InlineEditConditionToggle))
+	uint8 bOverride_SpatialSampleCount : 1;
+
+	UPROPERTY(EditAnywhere, BlueprintReadWrite, Category = Overrides, meta = (InlineEditConditionToggle))
+	uint8 bOverride_bDisableToneCurve : 1;
+
+	UPROPERTY(EditAnywhere, BlueprintReadWrite, Category = Overrides, meta = (InlineEditConditionToggle))
+	uint8 bOverride_bAllowOCIO : 1;
+
+	UPROPERTY(EditAnywhere, BlueprintReadWrite, Category = Overrides, meta = (InlineEditConditionToggle))
+	uint8 bOverride_bWriteAllSamples : 1;
+
+	/**
+	* How many sub-pixel jitter renders should we do per temporal sample? This can be used to achieve high
+	* sample counts without Temporal Sub-Sampling (allowing high sample counts without motion blur being enabled),
+	* but we generally recommend using Temporal Sub-Samples when possible. It can also be combined with
+	* temporal samples and you will get SpatialSampleCount many renders per temporal sample.
+	*/
+	UPROPERTY(EditAnywhere, BlueprintReadWrite, meta = (UIMin = 1, ClampMin = 1), Category = "Settings", meta = (EditCondition = "bOverride_SpatialSampleCount"))
+	int32 SpatialSampleCount;
+
+	/**
+	* If true, the tone curve will be disabled for this render pass. This will result in values greater than 1.0 in final renders
+	* and can optionally be combined with OCIO profiles on the file output nodes to convert from Linear Values in Working Color Space
+	* (which is sRGB  (Rec. 709) by default, unless changed in the project settings).
+	*/
+	UPROPERTY(EditAnywhere, BlueprintReadWrite, Category = "Settings", meta = (EditCondition = "bOverride_bDisableToneCurve"))
+	bool bDisableToneCurve;
+
+	/**
+	* Allow the output file OpenColorIO transform to be used on this render.
+	*/
+	UPROPERTY(EditAnywhere, BlueprintReadWrite, Category = "Settings", meta = (EditCondition = "bOverride_bAllowOCIO"))
+	bool bAllowOCIO;
+
+	/**
+	* Debug Feature. Can use this to write out each individual Temporal and Spatial sample rendered by this render pass,
+	* which allows you to see which images are being accumulated together. Can be useful for debugging incorrect looking
+	* frames to see which sub-frame evaluations were incorrect.
+	*/
+	UPROPERTY(EditAnywhere, BlueprintReadWrite, Category = "Debug", meta = (EditCondition = "bOverride_bWriteAllsamples"))
+	bool bWriteAllSamples;
+
+private:
+	/**
+	 * The original value of the "r.PathTracing.ProgressDisplay" cvar before the render starts. The progress display
+	 * will be hidden during the render.
+	 */
+	bool bOriginalProgressDisplayCvarValue = false;
+};
\ No newline at end of file
diff --git a/Engine/Plugins/MovieScene/MovieRenderPipeline/Source/MovieRenderPipelineRenderPasses/Public/Graph/Renderers/MovieGraphDeferredPass.h b/Engine/Plugins/MovieScene/MovieRenderPipeline/Source/MovieRenderPipelineRenderPasses/Public/Graph/Renderers/MovieGraphDeferredPass.h
new file mode 100644
--- /dev/null
+++ b/Engine/Plugins/MovieScene/MovieRenderPipeline/Source/MovieRenderPipelineRenderPasses/Public/Graph/Renderers/MovieGraphDeferredPass.h
@@ -0,0 +1,29 @@
+// Copyright Epic Games, Inc. All Rights Reserved.
+#pragma once
+
+#include "MovieGraphImagePassBase.h"
+
+namespace UE::MovieGraph::Rendering
+{
+	struct FMovieGraphDeferredPass : public FMovieGraphImagePassBase
+	{
+		virtual void Setup(TWeakObjectPtr<UMovieGraphDefaultRenderer> InRenderer, TWeakObjectPtr<UMovieGraphImagePassBaseNode> InRenderPassNode, const FMovieGraphRenderPassLayerData& InLayer) override;
+		virtual void Teardown() override;
+		virtual void Render(const FMovieGraphTraversalContext& InFrameTraversalContext, const FMovieGraphTimeStepData& InTimeData) override;
+		virtual void GatherOutputPasses(TArray<FMovieGraphRenderDataIdentifier>& OutExpectedPasses) const override;
+		virtual void AddReferencedObjects(FReferenceCollector& Collector) override;
+		virtual FName GetBranchName() const override;
+		virtual UMovieGraphImagePassBaseNode* GetParentNode(UMovieGraphEvaluatedConfig* InConfig) const override;
+			
+	protected:
+		virtual void PostRendererSubmission(const UE::MovieGraph::FMovieGraphSampleState& InSampleState, const UE::MovieGraph::DefaultRenderer::FRenderTargetInitParams& InRenderTargetInitParams, FCanvas& InCanvas, const UE::MovieGraph::DefaultRenderer::FCameraInfo& InCameraInfo) override;
+	protected:
+		FMovieGraphRenderPassLayerData LayerData;
+
+		/** Unique identifier passed in GatherOutputPasses and with each render that identifies the data produced by this renderer. */
+		FMovieGraphRenderDataIdentifier RenderDataIdentifier;
+
+		// Scene View history used by the renderer 
+		FSceneViewStateReference SceneViewState;
+	}; 
+}
\ No newline at end of file
diff --git a/Engine/Plugins/MovieScene/MovieRenderPipeline/Source/MovieRenderPipelineRenderPasses/Public/Graph/Renderers/MovieGraphImagePassBase.h b/Engine/Plugins/MovieScene/MovieRenderPipeline/Source/MovieRenderPipelineRenderPasses/Public/Graph/Renderers/MovieGraphImagePassBase.h
new file mode 100644
--- /dev/null
+++ b/Engine/Plugins/MovieScene/MovieRenderPipeline/Source/MovieRenderPipelineRenderPasses/Public/Graph/Renderers/MovieGraphImagePassBase.h
@@ -0,0 +1,92 @@
+// Copyright Epic Games, Inc. All Rights Reserved.
+
+#pragma once
+#include "CoreTypes.h"
+#include "Graph/MovieGraphDefaultRenderer.h"
+#include "SceneView.h"
+
+// Forward Declares
+class UMovieGraphImagePassBaseNode;
+
+namespace UE::MovieGraph::Rendering
+{
+	struct FMovieGraphRenderDataAccumulationArgs
+	{
+	public:
+		TWeakPtr<FImageOverlappedAccumulator, ESPMode::ThreadSafe> ImageAccumulator;
+		TWeakPtr<IMovieGraphOutputMerger, ESPMode::ThreadSafe> OutputMerger;
+
+		// If it's the first sample then we will reset the accumulator to a clean slate before accumulating into it.
+		bool bIsFirstSample;
+		// If it's the last sample, then we will trigger moving the data to the output merger. It can be both the first and last sample at the same time.
+		bool bIsLastSample;
+	};
+
+	// Forward Declare
+	void AccumulateSample_TaskThread(TUniquePtr<FImagePixelData>&& InPixelData, const ::UE::MovieGraph::FMovieGraphSampleState InSampleState, const FMovieGraphRenderDataAccumulationArgs& InAccumulationParams);
+
+	struct FViewFamilyInitData
+	{
+		FViewFamilyInitData()
+			: RenderTarget(nullptr)
+			, World(nullptr)
+			, SceneCaptureSource(ESceneCaptureSource::SCS_MAX)
+			, bWorldIsPaused(false)
+			, FrameIndex(-1)
+			, AntiAliasingMethod(EAntiAliasingMethod::AAM_None)
+			, ShowFlags(ESFIM_Game)
+			, ViewModeIndex(VMI_Lit)
+		{
+		}
+
+		class FRenderTarget* RenderTarget;
+		class UWorld* World;
+		FMovieGraphTimeStepData TimeData;
+		ESceneCaptureSource SceneCaptureSource;
+		bool bWorldIsPaused;
+		int32 FrameIndex;
+		EAntiAliasingMethod AntiAliasingMethod;
+		FEngineShowFlags ShowFlags;
+		EViewModeIndex ViewModeIndex;
+	};
+	
+
+
+	struct FMovieGraphImagePassBase
+	{
+		FMovieGraphImagePassBase() = default;
+		virtual ~FMovieGraphImagePassBase() = default;
+
+		virtual void Setup(TWeakObjectPtr<UMovieGraphDefaultRenderer> InRenderer, TWeakObjectPtr<UMovieGraphImagePassBaseNode> InRenderPassNode, const FMovieGraphRenderPassLayerData& InLayer);
+		virtual void Teardown() {}
+		virtual void Render(const FMovieGraphTraversalContext& InFrameTraversalContext, const FMovieGraphTimeStepData& InTimeData) {}
+		virtual void GatherOutputPasses(TArray<FMovieGraphRenderDataIdentifier>& OutExpectedPasses) const {}
+		virtual void AddReferencedObjects(FReferenceCollector& Collector) {}
+		virtual FName GetBranchName() const { return NAME_None; }
+		virtual TWeakObjectPtr<UMovieGraphDefaultRenderer> GetRenderer() const { return WeakGraphRenderer; }
+
+	protected:
+		/** Utility function for calculating a Projection Matrix (Orthographic or Perspective), and modify it based on the overscan percentage, aspect ratio, etc. */
+		virtual FMatrix CalculateProjectionMatrix(const UE::MovieGraph::DefaultRenderer::FCameraInfo& InCameraInfo) const;
+		/** Utility function for modifying the projection matrix to match the given tiling parameters. */
+		virtual void ModifyProjectionMatrixForTiling(const UE::MovieGraph::DefaultRenderer::FMovieGraphTilingParams& InTilingParams, const bool bInOrthographic, FMatrix& InOutProjectionMatrix, float& OutDoFSensorScale) const;
+		/** Utility function for calculating the PrinciplePointOffset with the given tiling parameters. Used to make some effects (like vignette) work with tiling. */
+		virtual FVector4f CalculatePrinciplePointOffsetForTiling(const UE::MovieGraph::DefaultRenderer::FMovieGraphTilingParams& InTilingParams) const;
+		/** Utility function for creating FSceneViewInitOptions based on the specified camera info. */
+		virtual FSceneViewInitOptions CreateViewInitOptions(const UE::MovieGraph::DefaultRenderer::FCameraInfo& InCameraInfo, FSceneViewFamilyContext* InViewFamily, FSceneViewStateReference& InViewStateRef) const;
+		/** Utility function for creating a FSceneView for the given InitOptions, Family, and Camera. */
+		virtual FSceneView* CreateSceneView(const FSceneViewInitOptions& InInitOptions, TSharedRef<FSceneViewFamilyContext> InViewFamily, const UE::MovieGraph::DefaultRenderer::FCameraInfo& InCameraInfo) const;
+
+		virtual void ApplyCameraManagerPostProcessBlends(FSceneView* InView) const;
+		virtual TSharedRef<FSceneViewFamilyContext> CreateSceneViewFamily(const FViewFamilyInitData& InInitData, const UE::MovieGraph::DefaultRenderer::FCameraInfo& InCameraInfo) const;
+		virtual void ApplyMovieGraphOverridesToSceneView(TSharedRef<FSceneViewFamilyContext> InOutFamily, const FViewFamilyInitData& InInitData, const UE::MovieGraph::DefaultRenderer::FCameraInfo& InCameraInfo) const;
+		virtual void ApplyMovieGraphOverridesToViewFamily(TSharedRef<FSceneViewFamilyContext> InOutFamily, const FViewFamilyInitData& InInitData) const;
+		virtual void PostRendererSubmission(const UE::MovieGraph::FMovieGraphSampleState& InSampleState, const UE::MovieGraph::DefaultRenderer::FRenderTargetInitParams& InRenderTargetInitParams, FCanvas& InCanvas, const UE::MovieGraph::DefaultRenderer::FCameraInfo& InCameraInfo);
+		virtual bool ShouldDiscardOutput(const TSharedRef<FSceneViewFamilyContext>& InFamily, const UE::MovieGraph::DefaultRenderer::FCameraInfo& InCameraInfo) const { return false; }
+		/** For this image pass, look up the associated node by type for the given config. */
+		virtual UMovieGraphImagePassBaseNode* GetParentNode(UMovieGraphEvaluatedConfig* InConfig) const { return nullptr; }
+
+	protected:
+		TWeakObjectPtr<UMovieGraphDefaultRenderer> WeakGraphRenderer;
+	};
+}
\ No newline at end of file
diff --git a/Engine/Plugins/MovieScene/MovieRenderPipeline/Source/MovieRenderPipelineRenderPasses/Public/Graph/Renderers/MovieGraphPathTracerPass.h b/Engine/Plugins/MovieScene/MovieRenderPipeline/Source/MovieRenderPipelineRenderPasses/Public/Graph/Renderers/MovieGraphPathTracerPass.h
new file mode 100644
--- /dev/null
+++ b/Engine/Plugins/MovieScene/MovieRenderPipeline/Source/MovieRenderPipelineRenderPasses/Public/Graph/Renderers/MovieGraphPathTracerPass.h
@@ -0,0 +1,14 @@
+// Copyright Epic Games, Inc. All Rights Reserved.
+#pragma once
+
+#include "MovieGraphDeferredPass.h"
+
+namespace UE::MovieGraph::Rendering
+{
+	struct FMovieGraphPathTracerPass : public FMovieGraphDeferredPass
+	{
+		virtual UMovieGraphImagePassBaseNode* GetParentNode(UMovieGraphEvaluatedConfig* InConfig) const override;
+		virtual void ApplyMovieGraphOverridesToSceneView(TSharedRef<FSceneViewFamilyContext> InOutFamily, const FViewFamilyInitData& InInitData, const UE::MovieGraph::DefaultRenderer::FCameraInfo& InCameraInfo) const override;
+		virtual bool ShouldDiscardOutput(const TSharedRef<FSceneViewFamilyContext>& InFamily, const UE::MovieGraph::DefaultRenderer::FCameraInfo& InCameraInfo) const override;
+	}; 
+}
\ No newline at end of file
