MovieGraph: Performance improvement by preventing the utility FindMoviePipelineSettingClasses function from searching through all blueprint assets every rendered frame.

Each call now takes ~700us down from ~15ms in our test project. This change assumes that MRG nodes cannot be created from blueprint assets.

diff --git a/Engine/Plugins/MovieScene/MovieRenderPipeline/Source/MovieRenderPipelineCore/Private/Graph/MovieGraphBlueprintLibrary.cpp b/Engine/Plugins/MovieScene/MovieRenderPipeline/Source/MovieRenderPipelineCore/Private/Graph/MovieGraphBlueprintLibrary.cpp
--- a/Engine/Plugins/MovieScene/MovieRenderPipeline/Source/MovieRenderPipelineCore/Private/Graph/MovieGraphBlueprintLibrary.cpp
+++ b/Engine/Plugins/MovieScene/MovieRenderPipeline/Source/MovieRenderPipelineCore/Private/Graph/MovieGraphBlueprintLibrary.cpp
@@ -126,7 +126,7 @@ FString UMovieGraphBlueprintLibrary::ResolveFilenameFormatArguments(const FStrin
 	//  Now get the settings from our config. We need to gather KVP data from all possible nodes, even if not expressed in your configuration. This is because you might want to
 	// always use the {ts_count} token even if you don't have a Temporal Sample Count node to add it. So we loop through all the possible class types, and call a function on the 
 	// CDO, but then we pass that class type from the evaluated config (if it exists), and we pass the CDO as an argument if it doesn't.
-	TArray<UClass*> AllSettingsNodeClasses = UE::MovieRenderPipeline::FindMoviePipelineSettingClasses(UMovieGraphSettingNode::StaticClass());
+	TArray<UClass*> AllSettingsNodeClasses = UE::MovieRenderPipeline::FindMoviePipelineSettingClasses(UMovieGraphSettingNode::StaticClass(), false);
 
 	// ToDo: This loops through class iterators every frame, we should probably initialize a copy of everything into the flattened config, since we could cache the classes
 	// once per run there. We don't cache the returned results here because you could potentially add/remove classes (via Blueprints) which would invalidate our cache.
diff --git a/Engine/Plugins/MovieScene/MovieRenderPipeline/Source/MovieRenderPipelineCore/Private/MoviePipelineUtils.cpp b/Engine/Plugins/MovieScene/MovieRenderPipeline/Source/MovieRenderPipelineCore/Private/MoviePipelineUtils.cpp
--- a/Engine/Plugins/MovieScene/MovieRenderPipeline/Source/MovieRenderPipelineCore/Private/MoviePipelineUtils.cpp
+++ b/Engine/Plugins/MovieScene/MovieRenderPipeline/Source/MovieRenderPipelineCore/Private/MoviePipelineUtils.cpp
@@ -38,7 +38,7 @@ namespace UE
 {
 	namespace MovieRenderPipeline
 	{
-		TArray<UClass*> FindMoviePipelineSettingClasses(UClass* InBaseClass)
+		TArray<UClass*> FindMoviePipelineSettingClasses(UClass* InBaseClass, const bool bIncludeBlueprints)
 		{
 			FAssetRegistryModule& AssetRegistryModule = FModuleManager::LoadModuleChecked<FAssetRegistryModule>(TEXT("AssetRegistry"));
 
@@ -46,7 +46,10 @@ namespace UE
 
 			FARFilter Filter;
 			Filter.ClassPaths.Add(InBaseClass->GetClassPathName());
-			Filter.ClassPaths.Add(UBlueprint::StaticClass()->GetClassPathName());
+			if (bIncludeBlueprints)
+			{
+				Filter.ClassPaths.Add(UBlueprint::StaticClass()->GetClassPathName());
+			}
 
 			// Include any Blueprint based objects as well, this includes things like Blutilities, UMG, and GameplayAbility objects
 			Filter.bRecursiveClasses = true;
diff --git a/Engine/Plugins/MovieScene/MovieRenderPipeline/Source/MovieRenderPipelineCore/Public/MoviePipelineUtils.h b/Engine/Plugins/MovieScene/MovieRenderPipeline/Source/MovieRenderPipelineCore/Public/MoviePipelineUtils.h
--- a/Engine/Plugins/MovieScene/MovieRenderPipeline/Source/MovieRenderPipelineCore/Public/MoviePipelineUtils.h
+++ b/Engine/Plugins/MovieScene/MovieRenderPipeline/Source/MovieRenderPipelineCore/Public/MoviePipelineUtils.h
@@ -122,7 +122,7 @@ namespace UE
 {
 	namespace MovieRenderPipeline
 	{
-		MOVIERENDERPIPELINECORE_API TArray<UClass*> FindMoviePipelineSettingClasses(UClass* InBaseClass);
+		MOVIERENDERPIPELINECORE_API TArray<UClass*> FindMoviePipelineSettingClasses(UClass* InBaseClass, const bool bIncludeBlueprints = true);
 		MOVIERENDERPIPELINECORE_API EAntiAliasingMethod GetEffectiveAntiAliasingMethod(const UMoviePipelineAntiAliasingSetting* InSetting);
 		
 		UE_DEPRECATED(5.3, "Do not use, this is here as a temporary workaround for another issue.")
diff --git a/Engine/Plugins/MovieScene/MovieRenderPipeline/Source/MovieRenderPipelineEditor/Private/MoviePipelineNewProcessExecutor.cpp b/Engine/Plugins/MovieScene/MovieRenderPipeline/Source/MovieRenderPipelineEditor/Private/MoviePipelineNewProcessExecutor.cpp
--- a/Engine/Plugins/MovieScene/MovieRenderPipeline/Source/MovieRenderPipelineEditor/Private/MoviePipelineNewProcessExecutor.cpp
+++ b/Engine/Plugins/MovieScene/MovieRenderPipeline/Source/MovieRenderPipelineEditor/Private/MoviePipelineNewProcessExecutor.cpp
@@ -115,7 +115,7 @@ void UMoviePipelineNewProcessExecutor::Execute_Implementation(UMoviePipelineQueu
 
 			// Only nodes on the Globals branch can apply command line args. Get CDOs here as well, so all nodes which define command line args
 			// will be found, even if they were not instantiated in the graph.
-			TArray<UClass*> AllSettingsNodeClasses = UE::MovieRenderPipeline::FindMoviePipelineSettingClasses(UMovieGraphSettingNode::StaticClass());
+			TArray<UClass*> AllSettingsNodeClasses = UE::MovieRenderPipeline::FindMoviePipelineSettingClasses(UMovieGraphSettingNode::StaticClass(), false);
 			for (UClass* NodeClass : AllSettingsNodeClasses)
 			{
 				constexpr bool bIncludeCDOs = true;
