MovieGraph: Small validation improvements.

We log warnings (only once) in case users could have forgotten to disable the tone curve while using OCIO. Similarly, we log a warning in case no render layers are found, which previously would render and silently fail to produce any file.

#rnx

diff --git a/Engine/Plugins/MovieScene/MovieRenderPipeline/Source/MovieRenderPipelineCore/Private/Graph/MovieGraphConfig.cpp b/Engine/Plugins/MovieScene/MovieRenderPipeline/Source/MovieRenderPipelineCore/Private/Graph/MovieGraphConfig.cpp
--- a/Engine/Plugins/MovieScene/MovieRenderPipeline/Source/MovieRenderPipelineCore/Private/Graph/MovieGraphConfig.cpp
+++ b/Engine/Plugins/MovieScene/MovieRenderPipeline/Source/MovieRenderPipelineCore/Private/Graph/MovieGraphConfig.cpp
@@ -11,6 +11,7 @@
 #include "Graph/Nodes/MovieGraphOutputNode.h"
 #include "Graph/Nodes/MovieGraphGlobalOutputSettingNode.h"
 #include "Graph/Nodes/MovieGraphRemoveRenderSettingNode.h"
+#include "Graph/Nodes/MovieGraphRenderLayerNode.h"
 #include "Graph/Nodes/MovieGraphSubgraphNode.h"
 #include "Graph/Nodes/MovieGraphVariableNode.h"
 #include "Graph/Nodes/MovieGraphSelectNode.h"
@@ -1602,6 +1603,8 @@ UMovieGraphEvaluatedConfig* UMovieGraphConfig::CreateFlattenedGraph(const FMovie
 		}		
 	}
 
+	bool bHasRenderLayerNode = false;
+
 	for (const TPair<FName, FMovieGraphEvaluatedBranchConfig>& Pair : NewContext->BranchConfigMapping)
 	{
 		for (UMovieGraphNode* Node : Pair.Value.GetNodes())
@@ -1616,8 +1619,15 @@ UMovieGraphEvaluatedConfig* UMovieGraphConfig::CreateFlattenedGraph(const FMovie
 					CheckProperty->ExportText_InContainer(0, ExportText, Node, Node, Node, 0);
 				}
 			}
+			
+			bHasRenderLayerNode |= Node->IsA<UMovieGraphRenderLayerNode>();
 		}
+	}
 
+	if (!bHasRenderLayerNode)
+	{
+		// NOTE: While this doesn't cover all cases, we ensure the presence of at least one render layer node.
+		UE_CALL_ONCE([] { UE_LOG(LogMovieRenderPipeline, Error, TEXT("For render jobs to succeed, one or more render layer node(s) must be present.")); });
 	}
 
 	return NewContext;
diff --git a/Engine/Plugins/MovieScene/MovieRenderPipeline/Source/MovieRenderPipelineRenderPasses/Private/MovieGraphImageSequenceOutputNode.cpp b/Engine/Plugins/MovieScene/MovieRenderPipeline/Source/MovieRenderPipelineRenderPasses/Private/MovieGraphImageSequenceOutputNode.cpp
--- a/Engine/Plugins/MovieScene/MovieRenderPipeline/Source/MovieRenderPipelineRenderPasses/Private/MovieGraphImageSequenceOutputNode.cpp
+++ b/Engine/Plugins/MovieScene/MovieRenderPipeline/Source/MovieRenderPipelineRenderPasses/Private/MovieGraphImageSequenceOutputNode.cpp
@@ -173,6 +173,20 @@ namespace UE::MovieGraph::Private
 
 		return {};
 	}
+
+	/* Utility function to warn the user in case they forgot to check "Disable Tone Curve", which in turn controls the render's scene capture source. */
+	void ValidateDisableTonecurve(const UE::MovieGraph::FMovieGraphSampleState& InPayload)
+	{
+		if (InPayload.SceneCaptureSource != ESceneCaptureSource::SCS_FinalColorHDR)
+		{
+			UE_CALL_ONCE([]
+				{
+					UE_LOG(LogMovieRenderPipeline, Warning, TEXT(
+						"The OCIO transform did not receive scene-referred linear colors, which most standard workflows expect."
+						"You may wish to disable the tonecurve on your renderer node(s)."));
+				});
+		}
+	}
 #endif // WITH_OCIO
 } //end namespace UE::MovieGraph::Private
 
@@ -371,6 +385,8 @@ void UMovieGraphImageSequenceOutputNode::OnReceiveImageDataImpl(UMovieGraphPipel
 #if WITH_OCIO
 		if (ParentNode->OCIOConfiguration.bIsEnabled && Payload->bAllowOCIO)
 		{
+			UE::MovieGraph::Private::ValidateDisableTonecurve(*Payload);
+
 			TMap<FString, FString> ResolvedOCIOContext;
 
 			const TObjectPtr<UOpenColorIOConfiguration>& ConfigurationAsset = ParentNode->OCIOConfiguration.ColorConfiguration.ConfigurationSource;
@@ -496,6 +512,8 @@ void UMovieGraphImageSequenceOutputNode_EXR::UpdateTaskPerLayer(
 #if WITH_OCIO
 	if (InParentNode->OCIOConfiguration.bIsEnabled && Payload->bAllowOCIO)
 	{
+		UE::MovieGraph::Private::ValidateDisableTonecurve(*Payload);
+
 		FPixelPreProcessor OCIOPixelPreProcessor = UE::MovieGraph::Private::CreateOpenColorIOPixelPreProcessor(
 			InParentNode->OCIOConfiguration.ColorConfiguration,
 			InResolvedOCIOContext
