MoviePipeline: Made a large number of fixes/updates to the Command Line Encoder node and .wav Audio Output nodes.

In a nutshell:
* The Command Line Encoder node is now Globals-only.
* The Command Line Encoder node now uses the audio codec, video codec, and file extension settings from Project Settings for its default values.
* Both nodes are now capable of doing shot-level exports, either via the graph in the primary job (if it specifies {shot_name}) or via shot-level graphs.
* Fixed various crashes that could occur, like if the pipeline was canceled during an export stage.

#rnx

diff --git a/Engine/Plugins/MovieScene/MovieRenderPipeline/Source/MovieRenderPipelineCore/Private/Graph/MovieGraphPipeline.cpp b/Engine/Plugins/MovieScene/MovieRenderPipeline/Source/MovieRenderPipelineCore/Private/Graph/MovieGraphPipeline.cpp
--- a/Engine/Plugins/MovieScene/MovieRenderPipeline/Source/MovieRenderPipelineCore/Private/Graph/MovieGraphPipeline.cpp
+++ b/Engine/Plugins/MovieScene/MovieRenderPipeline/Source/MovieRenderPipelineCore/Private/Graph/MovieGraphPipeline.cpp
@@ -528,12 +528,7 @@ void UMovieGraphPipeline::TickPostFinalizeExport(const bool bInForceFinish)
 	// This step assumes you have produced data and filled the data structures.
 	check(PipelineState == EMovieRenderPipelineState::Export);
 	
-
-	// Temporary change to prevent a crash when rendering only one shot via shot mask.
-	TransitionToState(EMovieRenderPipelineState::Finished);
-	return;
-
-	/*UE_LOG(LogMovieRenderPipeline, Verbose, TEXT("[%d] PostFinalize Export (Start)."), GFrameCounter);
+	UE_LOG(LogMovieRenderPipeline, Verbose, TEXT("[%d] PostFinalize Export (Start)."), GFrameCounter);
 
 	// Loop through any extensions (such as XML export) and let them export using all of the
 	// data that was generated during this run such as containers, output names and lengths.
@@ -545,11 +540,11 @@ void UMovieGraphPipeline::TickPostFinalizeExport(const bool bInForceFinish)
 		bAllContainsFinishedProcessing = true;
 
 		// Ask the nodes if they're all done processing.
-		constexpr bool bIncludeCDOs = false;
-		constexpr bool bExactMatch = false;
-		for (const TPair<FName, UMovieGraphPostRenderNode*>& Pair : GetSettingForActiveRenderLayers<UMovieGraphPostRenderNode>(bIncludeCDOs, bExactMatch))
+		const TArray<IMovieGraphPostRenderNode*> PostRenderNodes =
+			PostRenderEvaluatedGraph->GetSettingsImplementing<IMovieGraphPostRenderNode>(UMovieGraphPostRenderNode::StaticClass(), UMovieGraphNode::GlobalsPinName);
+		for (IMovieGraphPostRenderNode* PostRenderNode : PostRenderNodes)
 		{
-			bAllContainsFinishedProcessing &= Pair.Value->HasFinishedExporting();
+			bAllContainsFinishedProcessing &= PostRenderNode->HasFinishedExporting();
 		}
 	
 		// If we aren't forcing a finish, early out after one loop to keep the editor/ui responsive.
@@ -572,7 +567,7 @@ void UMovieGraphPipeline::TickPostFinalizeExport(const bool bInForceFinish)
 		return;
 	}
 
-	TransitionToState(EMovieRenderPipelineState::Finished);*/
+	TransitionToState(EMovieRenderPipelineState::Finished);
 }
 
 void UMovieGraphPipeline::BeginFinalize()
@@ -581,25 +576,27 @@ void UMovieGraphPipeline::BeginFinalize()
 	// them to put fences into queues for file writes, etc.
 	for (const TObjectPtr<UMovieGraphFileOutputNode>& Node : GetOutputNodesUsed())
 	{
-		Node->OnAllFramesSubmitted();
+		Node->OnAllFramesSubmitted(this, PostRenderEvaluatedGraph);
 	}
 }
 
 void UMovieGraphPipeline::BeginExport()
 {
-	// Temporary change to prevent a crash when rendering only one shot via shot mask.
-	return;
+	if (!PostRenderEvaluatedGraph)
+	{
+		// The generation of the evaluated graph would have emitted a warning if it failed; no need to generate another warning here 
+		return;
+	}
 
-	/*constexpr bool bIncludeCDOs = false;
-	constexpr bool bExactMatch = false;
-	TArray<TPair<FName, UMovieGraphPostRenderNode*>> PostRenderNodes = GetSettingForActiveRenderLayers<UMovieGraphPostRenderNode>(bIncludeCDOs, bExactMatch);
-	for (TPair<FName, UMovieGraphPostRenderNode*>& Pair : PostRenderNodes)
+	// Exports are run per primary job, so it only makes sense to look for nodes in the primary job's graph in the Globals branch. Shot graphs are not
+	// relevant at this point.
+	const TArray<IMovieGraphPostRenderNode*> PostRenderNodes =
+		PostRenderEvaluatedGraph->GetSettingsImplementing<IMovieGraphPostRenderNode>(UMovieGraphPostRenderNode::StaticClass(), UMovieGraphNode::GlobalsPinName);
+	
+	for (IMovieGraphPostRenderNode* PostRenderNode : PostRenderNodes)
 	{
-		UMovieGraphPostRenderNode* PostRenderNode = Pair.Value;
-		FName& BranchName = Pair.Key;
-		
-		PostRenderNode->BeginExport(this, BranchName);
-	}*/
+		PostRenderNode->BeginExport(this, PostRenderEvaluatedGraph);
+	}
 }
 
 void UMovieGraphPipeline::SetupShot(const TObjectPtr<UMoviePipelineExecutorShot>& InShot)
@@ -667,13 +664,56 @@ void UMovieGraphPipeline::TeardownShot(const TObjectPtr<UMoviePipelineExecutorSh
 	const FMovieGraphTimeStepData& TimeStepData = GetTimeStepInstance()->GetCalculatedTimeData();
 	const UMovieGraphEvaluatedConfig* EvaluatedConfig = TimeStepData.EvaluatedConfig;
 
+	ProcessOutstandingFinishedFrames();
+
+	// Run any post-render file generation that the nodes need to do
+	bool bIncludeCDOs = false;
+	bool bExactMatch = false;
+	const TArray<UMovieGraphFileOutputNode*> FileOutputNodes =
+		EvaluatedConfig->GetSettingsForBranch<UMovieGraphFileOutputNode>(UMovieGraphNode::GlobalsPinName, bIncludeCDOs, bExactMatch);
+	for (UMovieGraphFileOutputNode* FileOutputNode : FileOutputNodes)
+	{
+		FileOutputNode->OnAllShotFramesSubmitted(this, InShot);
+	}
+
+	// Ensure all of our Futures have been converted to the GeneratedOutputData
+	ProcessOutstandingFutures();
+
+	// Run any shot-based exports now that file output nodes have had a chance to run
+	const TArray<IMovieGraphPostRenderNode*> PostRenderNodes =
+		EvaluatedConfig->GetSettingsImplementing<IMovieGraphPostRenderNode>(UMovieGraphPostRenderNode::StaticClass(), UMovieGraphNode::GlobalsPinName);
+	for (IMovieGraphPostRenderNode* PostRenderNode : PostRenderNodes)
+	{
+		PostRenderNode->BeginShotExport(this);
+	}
+
+	// Wait for the export nodes to finish before going on to the next shot
+	bool bFinishedExporting = true;
+	do
+	{
+		bFinishedExporting = true;
+		
+		for (IMovieGraphPostRenderNode* PostRenderNode : PostRenderNodes)
+		{
+			bFinishedExporting &= PostRenderNode->HasFinishedExporting();
+		}
+
+		if (bFinishedExporting)
+		{
+			break;
+		}
+		
+		// Sleep for a while to give the export processes time to do some work
+		FPlatformProcess::Sleep(1.f);
+	} while (true);
+
 	// Revert the cvar values that were initially applied for the shot
 	CVarManager->RevertAllCVars();
 
 	// Revert cvars set by the global game overrides. Needs to be done after the CVarManager reverts (since the global
 	// game overrides are applied first in SetupShot).
-	constexpr bool bIncludeCDOs = true;
-	constexpr bool bExactMatch = true;
+	bIncludeCDOs = true;
+	bExactMatch = true;
 	if (UMovieGraphGlobalGameOverridesNode* GlobalGameOverridesNode = EvaluatedConfig->GetSettingForBranch<UMovieGraphGlobalGameOverridesNode>(UMovieGraphNode::GlobalsPinName, bIncludeCDOs, bExactMatch))
 	{
 		constexpr bool bOverrideValues = false;
@@ -790,46 +830,6 @@ void UMovieGraphPipeline::ExpandShot(const TObjectPtr<UMoviePipelineExecutorShot
 	}
 }
 
-template<typename T>
-TArray<TPair<FName, T*>> UMovieGraphPipeline::GetSettingForActiveRenderLayers(const bool bIncludeCDOs, const bool bExactMatch)
-{
-	TArray<TPair<FName, T*>> FoundSettings;
-	
-	UMoviePipelineRenderLayerSubsystem* LayerSubsystem = GetWorld()->GetSubsystem<UMoviePipelineRenderLayerSubsystem>();
-
-
-	// ToDo: This is only called by the Export step for the CommandLineEncoderNode right now, which only
-	// works per-render layer. However, by the time the export step actually starts, there isn't actually
-	// an active shot anymore, and it's not clear what graph we should be actually evaluating now. Perhaps 
-	// if we find a UMovieGraphPostRenderNode on a Render Layer, we should pause and run it at the end of 
-	// that shot? This would cause issues with post-renderer nodes that wanted to do something with the whole
-	// sequence (such as turn all shots into one movie file). Maybe in the post-export step we should only evaluate
-	// the Globals pin on the Primary Config, because we don't know which layers actually got rendered (due to shot
-	// overrides, etc.)
-	//
-	// For now, we're going to fall back on the incorrect legacy behavior (which was relying on stale data from the last
-	// shot), but this is going to have to be revisited.
-	UMovieGraphTimeStepBase* TimeStepInstance = GetTimeStepInstance();
-	if (!TimeStepInstance)
-	{
-		TimeStepInstance = GraphTimeStepInstances.Last();
-	}
-
-	const FMovieGraphTimeStepData& TimeStepData = TimeStepInstance->GetCalculatedTimeData();
-	const TObjectPtr<UMovieGraphEvaluatedConfig> EvaluatedConfig = TimeStepData.EvaluatedConfig;
-
-	for (const UMoviePipelineRenderLayer* RenderLayer : LayerSubsystem->GetRenderLayers())
-	{
-		FName LayerName = RenderLayer->GetRenderLayerName();
-		for (T* SettingNode : EvaluatedConfig->GetSettingsForBranch<T>(LayerName, bIncludeCDOs, bExactMatch))
-		{
-			FoundSettings.Add({LayerName, SettingNode});
-		}
-	}
-
-	return FoundSettings;
-}
-
 UMovieGraphTimeStepBase* UMovieGraphPipeline::GetTimeStepInstance() const
 {
 	return GraphTimeStepInstance;
@@ -1025,7 +1025,7 @@ void UMovieGraphPipeline::TransitionToState(const EMovieRenderPipelineState InNe
 
 			// If we had naturally finished the last shot before doing this transition it will have
 			// already been torn down, so this only catches mid-shot transitions to ensure teardown.
-			if (CurrentShotIndex < ActiveShotList.Num())
+			if (bShutdownRequested && (CurrentShotIndex < ActiveShotList.Num()))
 			{
 				// Ensures all in-flight work for that shot is handled.
 				TeardownShot(ActiveShotList[CurrentShotIndex]);
@@ -1053,6 +1053,15 @@ void UMovieGraphPipeline::TransitionToState(const EMovieRenderPipelineState InNe
 			// Shut down our custom timestep which reqstores some world settings we modified.
 			GEngine->SetCustomTimeStep(PrevCustomEngineTimeStep);
 
+			// Generate a job-level (sequence) graph that post-render tasks can use. Shot graphs should not be used at this point.
+			constexpr bool bForShot = false;
+			FString FlattenError;
+			PostRenderEvaluatedGraph = GetCurrentJob()->GetGraphPreset()->CreateFlattenedGraph(GetCurrentTraversalContext(bForShot), FlattenError);
+			if (!PostRenderEvaluatedGraph)
+			{
+				UE_LOG(LogMovieRenderPipeline, Warning, TEXT("Generating a post-render evaluated graph failed for job '%s'. Post-render exports will not run. Evaluation error: %s"), *GetCurrentJob()->JobName, *FlattenError);
+			}
+
 			// This is called once notifying output containers that all frames that will be submitted have been submitted.
 			PipelineState = EMovieRenderPipelineState::Finalize;
 			BeginFinalize();
@@ -1066,8 +1075,6 @@ void UMovieGraphPipeline::TransitionToState(const EMovieRenderPipelineState InNe
 			// This is called once notifying our export step that they can begin the export.
 			PipelineState = EMovieRenderPipelineState::Export;
 
-
-
 			// Restore the sequence so that the export processes can operate on the original sequence. 
 			// This is also done in the finished state because it's not guaranteed that the Export state 
 			// will be set when the render is canceled early
@@ -1107,6 +1114,9 @@ void UMovieGraphPipeline::TransitionToState(const EMovieRenderPipelineState InNe
 				PreviewWidget = nullptr;
 			}
 
+			// Job-level evaluated graph should not be referenced after export has finished
+			PostRenderEvaluatedGraph = nullptr;
+
 			//TArray<UMoviePipelineOutputBase*> ContainerSettings = GetPipelinePrimaryConfig()->GetOutputContainers();
 			//Algo::SortBy(ContainerSettings, [](const UMoviePipelineOutputBase* Setting) { return Setting->GetPriority(); });
 			//for (UMoviePipelineOutputBase* Setting : ContainerSettings)
@@ -1366,14 +1376,14 @@ UMovieGraphConfig* UMovieGraphPipeline::GetRootGraphForShot(UMoviePipelineExecut
 	return nullptr;
 }
 
-FMovieGraphTraversalContext UMovieGraphPipeline::GetCurrentTraversalContext() const
+FMovieGraphTraversalContext UMovieGraphPipeline::GetCurrentTraversalContext(const bool bForShot) const
 {
 	FMovieGraphTraversalContext CurrentContext;
-	CurrentContext.ShotIndex = GetCurrentShotIndex();
+	CurrentContext.ShotIndex = bForShot ? GetCurrentShotIndex() : -1;
 	CurrentContext.ShotCount = GetActiveShotList().Num();
 	CurrentContext.Job = GetCurrentJob();
-	CurrentContext.RootGraph = GetRootGraphForShot(GetActiveShotList()[GetCurrentShotIndex()]);
-	CurrentContext.Time = GetTimeStepInstance()->GetCalculatedTimeData();
+	CurrentContext.RootGraph = bForShot ? GetRootGraphForShot(GetActiveShotList()[GetCurrentShotIndex()]) : CurrentContext.Job->GetGraphPreset();
+	CurrentContext.Time = bForShot ? GetTimeStepInstance()->GetCalculatedTimeData() : FMovieGraphTimeStepData();
 
 	return CurrentContext;
 }
diff --git a/Engine/Plugins/MovieScene/MovieRenderPipeline/Source/MovieRenderPipelineCore/Private/Graph/Nodes/MovieGraphAudioOutputNode.cpp b/Engine/Plugins/MovieScene/MovieRenderPipeline/Source/MovieRenderPipelineCore/Private/Graph/Nodes/MovieGraphAudioOutputNode.cpp
--- a/Engine/Plugins/MovieScene/MovieRenderPipeline/Source/MovieRenderPipelineCore/Private/Graph/Nodes/MovieGraphAudioOutputNode.cpp
+++ b/Engine/Plugins/MovieScene/MovieRenderPipeline/Source/MovieRenderPipelineCore/Private/Graph/Nodes/MovieGraphAudioOutputNode.cpp
@@ -58,71 +58,52 @@
 	return EMovieGraphBranchRestriction::Globals;
 }
 
-void UMovieGraphAudioOutputNode::OnReceiveImageDataImpl(UMovieGraphPipeline* InPipeline, UE::MovieGraph::FMovieGraphOutputMergerFrame* InRawFrameData, const TSet<FMovieGraphRenderDataIdentifier>& InMask)
+void UMovieGraphAudioOutputNode::OnAllFramesSubmittedImpl(UMovieGraphPipeline* InPipeline, TObjectPtr<UMovieGraphEvaluatedConfig>& InPrimaryJobEvaluatedGraph)
 {
 	CachedPipeline = InPipeline;
-}
+	EvaluatedGraph = InPrimaryJobEvaluatedGraph;
 
-void UMovieGraphAudioOutputNode::OnAllFramesSubmittedImpl()
-{
-	if (!UE::MoviePipeline::RenderGraph::Audio::IsMoviePipelineAudioOutputSupported(CachedPipeline.Get()))
+	// It's possible that the pipeline is being shut down early, in which case the shot index will be invalid. Do not begin generating audio
+	// if a shutdown is being requested.
+	const int32 ShotIndex = CachedPipeline->GetCurrentShotIndex();
+	if (!CachedPipeline->GetActiveShotList().IsValidIndex(ShotIndex))
 	{
 		return;
 	}
 
-	ActiveWriters.Reset();
-
-	const MoviePipeline::FAudioState& AudioState = CachedPipeline->GetAudioRendererInstance()->GetAudioState();
-		
-	// There should be no active submixes by the time we finalize - they should all have been converted to recorded samples.
-	check(AudioState.ActiveSubmixes.IsEmpty());
-
-	// If we didn't end up recording audio, don't try outputting any files.
-	if (AudioState.FinishedSegments.IsEmpty())
+	// No need to do work if per-shot audio has already been exported
+	if (NeedsPerShotFlushing())
 	{
 		return;
 	}
 
-	TArray<FFinalAudioData> FinalAudioData;
-	GenerateFinalAudioData(FinalAudioData);
+	StartAudioExport();
+}
 
-	// Write audio files to disk
-	for (FFinalAudioData& AudioData : FinalAudioData)
+void UMovieGraphAudioOutputNode::OnAllShotFramesSubmittedImpl(UMovieGraphPipeline* InPipeline, const UMoviePipelineExecutorShot* InShot)
+{
+	CachedPipeline = InPipeline;
+	EvaluatedGraph = InPipeline->GetTimeStepInstance()->GetCalculatedTimeData().EvaluatedConfig;
+
+	// Don't do per-shot exports if a sequence-wide export was requested
+	if (!NeedsPerShotFlushing())
 	{
-		UE::MovieGraph::FMovieGraphOutputFutureData OutputData;
-		OutputData.Shot = CachedPipeline->GetActiveShotList()[AudioData.ShotIndex];
-		OutputData.DataIdentifier = AudioData.RenderIdentifier;
-
-		// Do this before we start manipulating the filepath for the audio API
-		OutputData.FilePath = AudioData.FilePath;
-
-		const FString FileName = FPaths::GetBaseFilename(AudioData.FilePath);
-		FString FileFolder = FPaths::GetPath(AudioData.FilePath);
-
-		TPromise<bool> Completed;
-		CachedPipeline->AddOutputFuture(Completed.GetFuture(), OutputData);
-
-		TUniquePtr<Audio::FSoundWavePCMWriter> Writer = MakeUnique<Audio::FSoundWavePCMWriter>();
-		const bool bSuccess = Writer->BeginWriteToWavFile(AudioData.SampleBuffer, FileName, FileFolder);
-
-		Completed.SetValue(bSuccess);
-		ActiveWriters.Add(MoveTemp(Writer));
+		return;
 	}
 
-	// The FSoundWavePCMWriter is unfortunately async, and the completion callbacks don't work unless the main thread
-	// can be spun (as it enqueues a callback onto the main thread). We're going to just cheat here and stall the main thread
-	// for 0.5s to give it a chance to write to disk. It'll only potentially be an issue with command line encoding if it takes
-	// longer than 0.5s to write to disk.
-	if (const TConsoleVariableData<float>* CVar = IConsoleManager::Get().FindTConsoleVariableDataFloat(TEXT("MovieRenderPipeline.WaveOutput.WriteDelay")))
-	{
-		UE_LOG(LogMovieRenderPipeline, Log, TEXT("Delaying main thread for %f seconds while audio writes to disk."), CVar->GetValueOnGameThread());
-		FPlatformProcess::Sleep(CVar->GetValueOnGameThread());
-	}
+	StartAudioExport();
 }
 
 bool UMovieGraphAudioOutputNode::IsFinishedWritingToDiskImpl() const
 {
-	// Files are written in OnAllFramesSubmittedImpl(), so writing should be finished by the time this method is called
+	for (const TUniquePtr<Audio::FSoundWavePCMWriter>& ActiveWriter : ActiveWriters)
+	{
+		if (ActiveWriter && !ActiveWriter->IsDone())
+		{
+			return false;
+		}
+	}
+
 	return true;
 }
 
@@ -130,8 +111,7 @@
 {
 	constexpr bool bIncludeCDOs = true;
 	constexpr bool bExactMatch = true;
-	const TObjectPtr<UMovieGraphEvaluatedConfig> EvaluatedConfig = CachedPipeline->GetTimeStepInstance()->GetCalculatedTimeData().EvaluatedConfig;
-	const UMovieGraphGlobalOutputSettingNode* OutputNode = EvaluatedConfig->GetSettingForBranch<UMovieGraphGlobalOutputSettingNode>(GlobalsPinName, bIncludeCDOs, bExactMatch);
+	const UMovieGraphGlobalOutputSettingNode* OutputNode = EvaluatedGraph->GetSettingForBranch<UMovieGraphGlobalOutputSettingNode>(GlobalsPinName, bIncludeCDOs, bExactMatch);
 	FString FileNameFormatString = OutputNode->OutputDirectory.Path / FileNameFormat;
 
 	constexpr bool bIncludeRenderPass = false;
@@ -149,7 +129,7 @@
 		{TEXT("render_pass"), RendererName},
 		{TEXT("ext"), OutputExtension}
 	};
-	ResolveParams.EvaluatedConfig = EvaluatedConfig;
+	ResolveParams.EvaluatedConfig = EvaluatedGraph;
 	ResolveParams.RenderDataIdentifier = InRenderIdentifier;
 	ResolveParams.Version = InShot->ShotInfo.VersionNumber;
 
@@ -218,3 +198,77 @@
 		UE_LOG(LogMovieRenderPipeline, Log, TEXT("Audio Segment took %f seconds to convert to a sample buffer."), (FPlatformTime::Seconds() - StartTime));
 	}
 }
+
+void UMovieGraphAudioOutputNode::StartAudioExport()
+{
+	if (!UE::MoviePipeline::RenderGraph::Audio::IsMoviePipelineAudioOutputSupported(CachedPipeline.Get()))
+	{
+		return;
+	}
+
+	ActiveWriters.Reset();
+
+	const MoviePipeline::FAudioState& AudioState = CachedPipeline->GetAudioRendererInstance()->GetAudioState();
+		
+	// There should be no active submixes by the time we finalize - they should all have been converted to recorded samples.
+	check(AudioState.ActiveSubmixes.IsEmpty());
+
+	// If we didn't end up recording audio, don't try outputting any files.
+	if (AudioState.FinishedSegments.IsEmpty())
+	{
+		return;
+	}
+
+	TArray<FFinalAudioData> FinalAudioData;
+	GenerateFinalAudioData(FinalAudioData);
+
+	// Write audio files to disk
+	for (FFinalAudioData& AudioData : FinalAudioData)
+	{
+		UE::MovieGraph::FMovieGraphOutputFutureData OutputData;
+		OutputData.Shot = CachedPipeline->GetActiveShotList()[AudioData.ShotIndex];
+		OutputData.DataIdentifier = AudioData.RenderIdentifier;
+
+		// Do this before we start manipulating the filepath for the audio API
+		OutputData.FilePath = AudioData.FilePath;
+
+		const FString FileName = FPaths::GetBaseFilename(AudioData.FilePath);
+		FString FileFolder = FPaths::GetPath(AudioData.FilePath);
+
+		TPromise<bool> Completed;
+		CachedPipeline->AddOutputFuture(Completed.GetFuture(), OutputData);
+
+		TUniquePtr<Audio::FSoundWavePCMWriter> Writer = MakeUnique<Audio::FSoundWavePCMWriter>();
+		const bool bSuccess = Writer->BeginWriteToWavFile(AudioData.SampleBuffer, FileName, FileFolder);
+
+		Completed.SetValue(bSuccess);
+		ActiveWriters.Add(MoveTemp(Writer));
+	}
+
+	// The FSoundWavePCMWriter is unfortunately async, and the completion callbacks don't work unless the main thread
+	// can be spun (as it enqueues a callback onto the main thread). We're going to just cheat here and stall the main thread
+	// for 0.5s to give it a chance to write to disk. It'll only potentially be an issue with command line encoding if it takes
+	// longer than 0.5s to write to disk.
+	if (const TConsoleVariableData<float>* CVar = IConsoleManager::Get().FindTConsoleVariableDataFloat(TEXT("MovieRenderPipeline.WaveOutput.WriteDelay")))
+	{
+		UE_LOG(LogMovieRenderPipeline, Log, TEXT("Delaying main thread for %f seconds while audio writes to disk."), CVar->GetValueOnGameThread());
+		FPlatformProcess::Sleep(CVar->GetValueOnGameThread());
+	}
+}
+
+bool UMovieGraphAudioOutputNode::NeedsPerShotFlushing() const
+{
+	constexpr bool bIncludeCDOs = true;
+	constexpr bool bExactMatch = true;
+	const UMovieGraphGlobalOutputSettingNode* OutputSettingNode =
+		EvaluatedGraph->GetSettingForBranch<UMovieGraphGlobalOutputSettingNode>(GlobalsPinName, bIncludeCDOs, bExactMatch);
+	
+	const FString FullPath = OutputSettingNode->OutputDirectory.Path / FileNameFormat;
+	
+	if (FullPath.Contains(TEXT("{shot_name}")) || FullPath.Contains(TEXT("{camera_name}")))
+	{
+		return true;
+	}
+
+	return false;
+}
diff --git a/Engine/Plugins/MovieScene/MovieRenderPipeline/Source/MovieRenderPipelineCore/Private/Graph/Nodes/MovieGraphCommandLineEncoderNode.cpp b/Engine/Plugins/MovieScene/MovieRenderPipeline/Source/MovieRenderPipelineCore/Private/Graph/Nodes/MovieGraphCommandLineEncoderNode.cpp
--- a/Engine/Plugins/MovieScene/MovieRenderPipeline/Source/MovieRenderPipelineCore/Private/Graph/Nodes/MovieGraphCommandLineEncoderNode.cpp
+++ b/Engine/Plugins/MovieScene/MovieRenderPipeline/Source/MovieRenderPipelineCore/Private/Graph/Nodes/MovieGraphCommandLineEncoderNode.cpp
@@ -22,13 +22,21 @@ UMovieGraphCommandLineEncoderNode::UMovieGraphCommandLineEncoderNode()
 	, bSkipEncodeOnRenderCanceled(true)
 	, bRetainInputTextFiles(false)
 {
-	FileNameFormat = TEXT("{sequence_name}.{layer_name}");
+	FileNameFormat = TEXT("{sequence_name}");
 		
 	// Sensible defaults for cmdline arguments
 	CommandLineFormat = TEXT("-hide_banner -y -loglevel error {VideoInputs} {AudioInputs} -acodec {AudioCodec} -vcodec {VideoCodec} {Quality} \"{OutputPath}\"");
 	VideoInputStringFormat = TEXT("-f concat -safe 0 -i \"{InputFile}\" -r {FrameRate}");
 	AudioInputStringFormat = TEXT("-f concat -safe 0 -i \"{InputFile}\"");
 	EncodeSettings = TEXT("-crf 20");
+
+	// Fill in some defaults from the project settings
+	if (const UMoviePipelineCommandLineEncoderSettings* EncoderSettings = GetDefault<UMoviePipelineCommandLineEncoderSettings>())
+	{
+		AudioCodec = EncoderSettings->AudioCodec;
+		VideoCodec = EncoderSettings->VideoCodec;
+		OutputFileExtension = EncoderSettings->OutputFileExtension;
+	}
 }
 
 #if WITH_EDITOR
@@ -60,7 +68,7 @@ FSlateIcon UMovieGraphCommandLineEncoderNode::GetIconAndTint(FLinearColor& OutCo
 
 EMovieGraphBranchRestriction UMovieGraphCommandLineEncoderNode::GetBranchRestriction() const
 {
-	return EMovieGraphBranchRestriction::RenderLayer;
+	return EMovieGraphBranchRestriction::Globals;
 }
 
 void UMovieGraphCommandLineEncoderNode::StartEncodingProcess(TArray<FMovieGraphRenderOutputData>& InGeneratedData, const bool bInIsShotEncode)
@@ -106,7 +114,7 @@ void UMovieGraphCommandLineEncoderNode::StartEncodingProcess(TArray<FMovieGraphR
 		// because the encoding won't be finished by the time the scripting layer is called. Need to manipulate the original
 		// and not the copy that we're currently iterating through.
 		FMovieGraphRenderDataIdentifier CommandLineEncoderIdentifier;
-		CommandLineEncoderIdentifier.RootBranchName = CachedBranchName;
+		CommandLineEncoderIdentifier.RootBranchName = GlobalsPinName;
 		CommandLineEncoderIdentifier.RendererName = FString(TEXT("CommandLineEncoder"));
 		for (FMovieGraphRenderOutputData& OutputData : InGeneratedData)
 		{
@@ -124,20 +132,18 @@ void UMovieGraphCommandLineEncoderNode::StartEncodingProcess(TArray<FMovieGraphR
 	}
 }
 
-void UMovieGraphCommandLineEncoderNode::BeginExport(UMovieGraphPipeline* InMoviePipeline, const FName& InBranchName)
+void UMovieGraphCommandLineEncoderNode::BeginExport(UMovieGraphPipeline* InMoviePipeline, TObjectPtr<UMovieGraphEvaluatedConfig>& InPrimaryJobEvaluatedGraph)
 {
 	CachedPipeline = InMoviePipeline;
-	CachedBranchName = InBranchName;
-	
-	// When we start exporting, we remove the OnEndFrame delegate because if the user hit escape to cancel a movie render
-	// no frames will be ticked anymore. Instead we'll call OnTick from HasFinishedExporting() by hand as that will
-	// get called in a loop until it is finished.
-	FCoreDelegates::OnEndFrame.RemoveAll(this);
+	PrimaryJobEvaluatedGraph = InPrimaryJobEvaluatedGraph;
 
-	// This is called at the end of the movie render. If we were starting encode jobs per shot, then there's already
+	// This is called after all shots have finished rendering. If we were starting encode jobs per shot, then there's already
 	// an encode job going for all shots, so we early out.
 	if (NeedsPerShotFlushing())
 	{
+		// Don't warn here about a primary job graph wanting to generate a shot-level encode. It is not an error. If there are no shot-level
+		// graphs defined, it's valid that the primary graph requests that shot-level encodes are completed. The shot-level encodes will be picked
+		// up from the primary graph via BeginShotExport() and not started here.
 		return;
 	}
 
@@ -145,15 +151,49 @@ void UMovieGraphCommandLineEncoderNode::BeginExport(UMovieGraphPipeline* InMovie
 	TArray<FMovieGraphRenderOutputData>& OutputData = CachedPipeline->GetGeneratedOutputData();
 	constexpr bool bIsShotEncode = false;
 	StartEncodingProcess(OutputData, bIsShotEncode);
+
+	// Should not be used after a primary job export finishes
+	PrimaryJobEvaluatedGraph = nullptr;
+}
+
+void UMovieGraphCommandLineEncoderNode::BeginShotExport(UMovieGraphPipeline* InMoviePipeline)
+{
+	CachedPipeline = InMoviePipeline;
+	PrimaryJobEvaluatedGraph = nullptr;
+	
+	if (!NeedsPerShotFlushing())
+	{
+		UE_LOG(LogMovieRenderPipeline, Warning, TEXT("A shot-level Command Line Encoder node requested a non-shot export and will be ignored. Is it missing the {shot_name} token in the File Name Format?"));
+		return;
+	}
+	
+	TArray<FMovieGraphRenderOutputData>& OutputData = CachedPipeline->GetGeneratedOutputData();
+	TArray<FMovieGraphRenderOutputData> LatestShotData;
+	if (!OutputData.IsEmpty())
+	{
+		LatestShotData.Add(OutputData.Last());
+	}
+	
+	constexpr bool bIsShotEncode = true;
+	StartEncodingProcess(LatestShotData, bIsShotEncode);
 }
 
 bool UMovieGraphCommandLineEncoderNode::HasFinishedExporting()
 {
-	// Manually tick the output (which cleans up ActiveEncodeJobs). This is needed because manually canceling a job
-	// stops ticking the engine and repeatedly calls HasFinishedExporting().
+	// Manually tick the output (which cleans up ActiveEncodeJobs). Once a render finishes, the engine is not ticking frames, so the pipeline
+	// will call this method repeatedly until the encode has completed.
 	OnTick();
 	
-	return ActiveEncodeJobs.IsEmpty();
+	const bool bIsFinished = ActiveEncodeJobs.IsEmpty();
+
+	// Clean up persisted data if the encode is finished
+	if (bIsFinished)
+	{
+		CachedPipeline = nullptr;
+		PrimaryJobEvaluatedGraph = nullptr;
+	}
+	
+	return bIsFinished;
 }
 
 bool UMovieGraphCommandLineEncoderNode::AreSettingsValid(TArray<FText>& OutErrors) const
@@ -276,7 +316,7 @@ TMap<FMovieGraphRenderDataIdentifier, UMovieGraphCommandLineEncoderNode::FEncode
 	
 	for (FMovieGraphRenderOutputData& GeneratedRenderData : InGeneratedData)
 	{
-		// Search for audio to attach to every render layer
+		// Search for audio to attach
 		TMap<FString, TArray<FString>> AudioPathsByExtension;
 		for (const TPair<FMovieGraphRenderDataIdentifier, FMovieGraphRenderLayerOutputData>& InnerRenderLayer : GeneratedRenderData.RenderLayerData)
 		{
@@ -301,13 +341,7 @@ TMap<FMovieGraphRenderDataIdentifier, UMovieGraphCommandLineEncoderNode::FEncode
 			{
 				continue;
 			}
-
-			// This node should only be responsible for exporting movies for images generated on the node's branch
-			if (RenderIdentifier.RootBranchName != CachedBranchName)
-			{
-				continue;
-			}
-
+			
 			FEncoderParams& EncoderParams = RenderLayerEncoderParams.FindOrAdd(RenderIdentifier);
 			EncoderParams.Shot = GeneratedRenderData.Shot;
 			EncoderParams.RenderDataIdentifier = RenderIdentifier;
@@ -323,6 +357,9 @@ TMap<FMovieGraphRenderDataIdentifier, UMovieGraphCommandLineEncoderNode::FEncode
 			EncoderParams.FilesByExtensionType.Append(AudioPathsByExtension);
 
 			RenderDataToRemove.Add(RenderIdentifier);
+
+			// Data was found to include in the encode job; do not add additional data from other output types or branches
+			break;
 		}
 		
 		// If we're going to delete the source files, don't make them available to the scripting layer callback
@@ -360,7 +397,7 @@ FString UMovieGraphCommandLineEncoderNode::GetResolvedOutputFilename(const FMovi
 	ResolveParams.Shot = Shot.Get();
 	ResolveParams.FileNameFormatOverrides = FormatOverrides;
 	ResolveParams.FileNameOverride = FileNameFormat;
-	ResolveParams.EvaluatedConfig = CachedPipeline->GetTimeStepInstance()->GetCalculatedTimeData().EvaluatedConfig;
+	ResolveParams.EvaluatedConfig = GetEvaluatedConfig();
 	ResolveParams.RenderDataIdentifier = RenderIdentifier;
 	ResolveParams.Version = Shot.IsValid() ? Shot->ShotInfo.VersionNumber : UMovieGraphBlueprintLibrary::ResolveVersionNumber(ResolveParams);
 
@@ -411,7 +448,7 @@ void UMovieGraphCommandLineEncoderNode::GenerateTemporaryEncoderInputFiles(const
 		ResolveParams.Job = CachedPipeline->GetCurrentJob();
 		ResolveParams.Shot = InParams.Shot.Get();
 		ResolveParams.FileNameFormatOverrides = FormatOverrides;
-		ResolveParams.EvaluatedConfig = CachedPipeline->GetTimeStepInstance()->GetCalculatedTimeData().EvaluatedConfig;
+		ResolveParams.EvaluatedConfig = GetEvaluatedConfig();
 		ResolveParams.RenderDataIdentifier = InParams.RenderDataIdentifier;
 
 		FMovieGraphResolveArgs FinalFormatArgs;
@@ -483,6 +520,9 @@ FString UMovieGraphCommandLineEncoderNode::BuildEncoderCommand(const FEncoderPar
 
 void UMovieGraphCommandLineEncoderNode::LaunchEncoder(const FEncoderParams& InParams)
 {
+	// Clear out any stale jobs (shouldn't be needed, but just in case)
+	ActiveEncodeJobs.Reset();
+	
 	TArray<FString> VideoInputs;
 	TArray<FString> AudioInputs;
 	GenerateTemporaryEncoderInputFiles(InParams, VideoInputs, AudioInputs);
@@ -530,13 +570,18 @@ void UMovieGraphCommandLineEncoderNode::LaunchEncoder(const FEncoderParams& InPa
 	}
 }
 
+TObjectPtr<UMovieGraphEvaluatedConfig> UMovieGraphCommandLineEncoderNode::GetEvaluatedConfig() const
+{
+	// If a shot is being exported, the primary evaluated graph will be null. Use the pipeline's evaluated graph for the current shot otherwise.
+	return PrimaryJobEvaluatedGraph
+		? PrimaryJobEvaluatedGraph
+		: CachedPipeline->GetTimeStepInstance()->GetCalculatedTimeData().EvaluatedConfig;
+}
+
 template<typename T>
 T* UMovieGraphCommandLineEncoderNode::GetSettingOnBranch(const bool bIncludeCDOs, const bool bExactMatch) const
 {
-	const TObjectPtr<UMovieGraphEvaluatedConfig> EvaluatedConfig =
-		CachedPipeline->GetTimeStepInstance()->GetCalculatedTimeData().EvaluatedConfig;
-	
-	return EvaluatedConfig->GetSettingForBranch<T>(CachedBranchName, bIncludeCDOs, bExactMatch);
+	return GetEvaluatedConfig()->GetSettingForBranch<T>(GlobalsPinName, bIncludeCDOs, bExactMatch);
 }
 
 #undef LOCTEXT_NAMESPACE // "MovieGraph"
diff --git a/Engine/Plugins/MovieScene/MovieRenderPipeline/Source/MovieRenderPipelineCore/Public/Graph/MovieGraphConfig.h b/Engine/Plugins/MovieScene/MovieRenderPipeline/Source/MovieRenderPipelineCore/Public/Graph/MovieGraphConfig.h
--- a/Engine/Plugins/MovieScene/MovieRenderPipeline/Source/MovieRenderPipelineCore/Public/Graph/MovieGraphConfig.h
+++ b/Engine/Plugins/MovieScene/MovieRenderPipeline/Source/MovieRenderPipelineCore/Public/Graph/MovieGraphConfig.h
@@ -484,6 +484,31 @@ public:
 		return ResultNodes;
 	}
 
+	/** Gets settings that implement a specific interface. InInterfaceClass should be the U-prefixed class; InterfaceType should be I-prefixed. */
+	template<typename InterfaceType>
+	TArray<InterfaceType*> GetSettingsImplementing(const UClass* InInterfaceClass, const FName InBranchName) const
+	{
+		const FMovieGraphEvaluatedBranchConfig* BranchConfig = BranchConfigMapping.Find(InBranchName);
+		ensureMsgf(BranchConfig, TEXT("Failed to find branch mapping for Branch: %s"), *InBranchName.ToString());
+
+		TArray<InterfaceType*> ResultNodes;
+		if (BranchConfig)
+		{
+			for (const TObjectPtr<UMovieGraphNode>& Node : BranchConfig->GetNodes())
+			{
+				if (Node->GetClass()->ImplementsInterface(InInterfaceClass))
+				{
+					if (InterfaceType* CastNode = Cast<InterfaceType>(Node.Get()))
+					{
+						ResultNodes.Add(CastNode);
+					}
+				}
+			}
+		}
+
+		return ResultNodes;
+	}
+
 	template<typename NodeType>
 	NodeType* GetSettingForBranch(const FName InBranchName, bool bIncludeCDOs = true, bool bExactMatch = false) const
 	{
diff --git a/Engine/Plugins/MovieScene/MovieRenderPipeline/Source/MovieRenderPipelineCore/Public/Graph/MovieGraphNode.h b/Engine/Plugins/MovieScene/MovieRenderPipeline/Source/MovieRenderPipelineCore/Public/Graph/MovieGraphNode.h
--- a/Engine/Plugins/MovieScene/MovieRenderPipeline/Source/MovieRenderPipelineCore/Public/Graph/MovieGraphNode.h
+++ b/Engine/Plugins/MovieScene/MovieRenderPipeline/Source/MovieRenderPipelineCore/Public/Graph/MovieGraphNode.h
@@ -8,6 +8,7 @@
 #include "PropertyBag.h"
 #include "Graph/MovieGraphValueContainer.h"
 #include "Graph/MovieGraphFilenameResolveParams.h"
+#include "UObject/Interface.h"
 
 #if WITH_EDITOR
 #include "Textures/SlateIcon.h"
@@ -402,18 +403,26 @@ public:
 	virtual void BuildNewProcessCommandLineArgsImpl(TArray<FString>& InOutUnrealURLParams, TArray<FString>& InOutCommandLineArgs, TArray<FString>& InOutDeviceProfileCvars, TArray<FString>& InOutExecCmds) const { }
 };
 
+UINTERFACE()
+class MOVIERENDERPIPELINECORE_API UMovieGraphPostRenderNode : public UInterface
+{
+	GENERATED_BODY()
+};
+
 /**
- * A node which runs after all renders have completed.
+ * A node which runs after renders have completed. Can run after a single shot renders and/or after a sequence renders all of its shots.
  */
-UCLASS(Abstract)
-class MOVIERENDERPIPELINECORE_API UMovieGraphPostRenderNode : public UMovieGraphSettingNode
+class MOVIERENDERPIPELINECORE_API IMovieGraphPostRenderNode
 {
 	GENERATED_BODY()
 
 public:
-	/** Begins the export process for this node. */
-	virtual void BeginExport(UMovieGraphPipeline* InMoviePipeline, const FName& InBranchName) PURE_VIRTUAL(UMovieGraphPostRenderNode::BeginExport, );
+	/** Begins the export process for this node after all shots have completed. */
+	virtual void BeginExport(UMovieGraphPipeline* InMoviePipeline, TObjectPtr<UMovieGraphEvaluatedConfig>& InPrimaryJobEvaluatedGraph) = 0;
+
+	/** Begins the export process for this node after an individual shot completes. */
+	virtual void BeginShotExport(UMovieGraphPipeline* InMoviePipeline) = 0;
 
 	/** Returns true if this node has finished its export process, else false. */
-	virtual bool HasFinishedExporting() PURE_VIRTUAL(UMovieGraphPostRenderNode::HasFinishedExporting, return true; );
+	virtual bool HasFinishedExporting() = 0;
 };
diff --git a/Engine/Plugins/MovieScene/MovieRenderPipeline/Source/MovieRenderPipelineCore/Public/Graph/MovieGraphPipeline.h b/Engine/Plugins/MovieScene/MovieRenderPipeline/Source/MovieRenderPipelineCore/Public/Graph/MovieGraphPipeline.h
--- a/Engine/Plugins/MovieScene/MovieRenderPipeline/Source/MovieRenderPipelineCore/Public/Graph/MovieGraphPipeline.h
+++ b/Engine/Plugins/MovieScene/MovieRenderPipeline/Source/MovieRenderPipelineCore/Public/Graph/MovieGraphPipeline.h
@@ -57,10 +57,11 @@ public:
 	UFUNCTION(BlueprintCallable, Category = "Movie Graph")
 	void SetInitializationTime(const FDateTime& InDateTime) { GraphInitializationTime = InDateTime; }
 
+	/** Gets the graph config for the shot if one was specified for the shot. Otherwise, gets the graph config for the associated primary job. */
 	UMovieGraphConfig* GetRootGraphForShot(UMoviePipelineExecutorShot* InShot) const;
 
 	UFUNCTION(BlueprintCallable, Category = "Movie Graph")
-	FMovieGraphTraversalContext GetCurrentTraversalContext() const;
+	FMovieGraphTraversalContext GetCurrentTraversalContext(const bool bForShot = true) const;
 
 	/** Get the Active Shot list, which is the full shot list generated from the external data source, with disabled shots removed. */
 	const TArray<TObjectPtr<UMoviePipelineExecutorShot>>& GetActiveShotList() const { return ActiveShotList; }
@@ -135,7 +136,10 @@ protected:
 	virtual void TransitionToState(const EMovieRenderPipelineState InNewState);
 	virtual const TSet<TObjectPtr<UMovieGraphFileOutputNode>> GetOutputNodesUsed() const;
 	virtual void BeginFinalize();
+
+	/** Begins the export process for a primary job (not called for shot jobs). */
 	virtual void BeginExport();
+	
 	virtual void LoadPreviewWidget();
 	virtual void SetPreviewWidgetVisibleImpl(bool bInIsVisible);
 
@@ -149,10 +153,6 @@ protected:
 	virtual void ExpandShot(const TObjectPtr<UMoviePipelineExecutorShot>& InShot, const int32 InNumHandleFrames, const bool bInHasMultipleTemporalSamples, const bool bIsPrePass,
 		const FFrameRate& InDisplayRate, const FFrameRate& InTickResolution, const int32 InWarmUpFrames);
 
-	/** Gets a setting across all active render layers. Key = branch name, value = node object. */
-	template<typename T>
-	TArray<TPair<FName, T*>> GetSettingForActiveRenderLayers(const bool bIncludeCDOs, const bool bExactMatch);
-
 	/** Resolve the version number that should be used for the specified shot (in {version} tokens). */
 	int32 ResolveVersionForShot(const TObjectPtr<UMoviePipelineExecutorShot>& Shot, const TObjectPtr<UMovieGraphEvaluatedConfig>& EvaluatedConfig);
 
@@ -193,6 +193,13 @@ protected:
 	UPROPERTY(Transient)
 	TObjectPtr<UMovieGraphAudioRendererBase> GraphAudioRendererInstance;
 
+	/**
+	 * The evaluated graph that should be referenced after all shot rendering is complete. This graph will originate from a primary job (not shot) and
+	 * can be used in post-rendering tasks.
+	 */
+	UPROPERTY(Transient)
+	TObjectPtr<UMovieGraphEvaluatedConfig> PostRenderEvaluatedGraph;
+
 protected:
 	UPROPERTY(Transient)
 	TObjectPtr<UMoviePipelineExecutorJob> CurrentJob;
diff --git a/Engine/Plugins/MovieScene/MovieRenderPipeline/Source/MovieRenderPipelineCore/Public/Graph/Nodes/MovieGraphAudioOutputNode.h b/Engine/Plugins/MovieScene/MovieRenderPipeline/Source/MovieRenderPipelineCore/Public/Graph/Nodes/MovieGraphAudioOutputNode.h
--- a/Engine/Plugins/MovieScene/MovieRenderPipeline/Source/MovieRenderPipelineCore/Public/Graph/Nodes/MovieGraphAudioOutputNode.h
+++ b/Engine/Plugins/MovieScene/MovieRenderPipeline/Source/MovieRenderPipelineCore/Public/Graph/Nodes/MovieGraphAudioOutputNode.h
@@ -45,8 +45,8 @@ public:
 
 protected:
 	// UMovieGraphFileOutputNode Interface
-	virtual void OnReceiveImageDataImpl(UMovieGraphPipeline* InPipeline, UE::MovieGraph::FMovieGraphOutputMergerFrame* InRawFrameData, const TSet<FMovieGraphRenderDataIdentifier>& InMask) override;
-	virtual void OnAllFramesSubmittedImpl() override;
+	virtual void OnAllFramesSubmittedImpl(UMovieGraphPipeline* InPipeline, TObjectPtr<UMovieGraphEvaluatedConfig>& InPrimaryJobEvaluatedGraph) override;
+	virtual void OnAllShotFramesSubmittedImpl(UMovieGraphPipeline* InPipeline, const UMoviePipelineExecutorShot* InShot) override;
 	virtual bool IsFinishedWritingToDiskImpl() const override;
 	// ~UMovieGraphFileOutputNode Interface
 
@@ -66,13 +66,28 @@ private:
 	/** Generates final audio data from all of the audio segments that were collected during render time. */
 	void GenerateFinalAudioData(TArray<FFinalAudioData>& OutFinalAudioData);
 
+	/** Begins the audio export process. */
+	void StartAudioExport();
+	
+	/** Returns true if there should be one export per shot, else false. */
+	bool NeedsPerShotFlushing() const;
+
+// The pipeline generates many instances of the same node throughout its execution; however, some nodes need to have persistent data throughout the
+// pipeline's lifetime. This static data enables the node to have shared data across instances.
 private:
 	/** Kept alive during finalization because the writer writes async to disk but doesn't expect to fall out of scope. */
-	TArray<TUniquePtr<Audio::FSoundWavePCMWriter>> ActiveWriters;
+	inline static TArray<TUniquePtr<Audio::FSoundWavePCMWriter>> ActiveWriters;
 
 	/** Keep track of segments that have already been written to disk to avoid re-writing them (and generating new output futures). */
-	TSet<FGuid> AlreadyWrittenSegments;
+	inline static TSet<FGuid> AlreadyWrittenSegments;
 
+private:
 	/** The pipeline that is running this node. */
 	TWeakObjectPtr<UMovieGraphPipeline> CachedPipeline;
+
+	/**
+	 * The graph that should be accessed during execution of the node. Do not access the graph from the pipeline as it may be invalid depending on
+	 * if a shot-level or sequence-level export is being performed.
+	 */
+	TObjectPtr<UMovieGraphEvaluatedConfig> EvaluatedGraph;
 };
diff --git a/Engine/Plugins/MovieScene/MovieRenderPipeline/Source/MovieRenderPipelineCore/Public/Graph/Nodes/MovieGraphCommandLineEncoderNode.h b/Engine/Plugins/MovieScene/MovieRenderPipeline/Source/MovieRenderPipelineCore/Public/Graph/Nodes/MovieGraphCommandLineEncoderNode.h
--- a/Engine/Plugins/MovieScene/MovieRenderPipeline/Source/MovieRenderPipelineCore/Public/Graph/Nodes/MovieGraphCommandLineEncoderNode.h
+++ b/Engine/Plugins/MovieScene/MovieRenderPipeline/Source/MovieRenderPipelineCore/Public/Graph/Nodes/MovieGraphCommandLineEncoderNode.h
@@ -13,7 +13,7 @@ struct FMovieGraphRenderOutputData;
  * A node which kicks off an encode process after all renders have completed.
  */
 UCLASS()
-class MOVIERENDERPIPELINECORE_API UMovieGraphCommandLineEncoderNode : public UMovieGraphPostRenderNode
+class MOVIERENDERPIPELINECORE_API UMovieGraphCommandLineEncoderNode : public UMovieGraphSettingNode, public IMovieGraphPostRenderNode
 {
 	GENERATED_BODY()
 
@@ -34,17 +34,15 @@ public:
 	/** Begins the encode process (as long as there are no validation errors). */
 	void StartEncodingProcess(TArray<FMovieGraphRenderOutputData>& InGeneratedData, const bool bInIsShotEncode);
 	
-	// UMovieGraphPostRenderNode interface
-	virtual void BeginExport(UMovieGraphPipeline* InMoviePipeline, const FName& InBranchName) override;
+	// IMovieGraphPostRenderNode interface
+	virtual void BeginExport(UMovieGraphPipeline* InMoviePipeline, TObjectPtr<UMovieGraphEvaluatedConfig>& InPrimaryJobEvaluatedGraph) override;
+	virtual void BeginShotExport(UMovieGraphPipeline* InMoviePipeline) override;
 	virtual bool HasFinishedExporting() override;
-	// ~UMovieGraphPostRenderNode interface
+	// ~IMovieGraphPostRenderNode interface
 
 public:
 	UPROPERTY(EditAnywhere, BlueprintReadWrite, Category = Overrides, meta = (InlineEditConditionToggle))
 	uint8 bOverride_FileNameFormat : 1;
-	
-	UPROPERTY(EditAnywhere, BlueprintReadWrite, Category = Overrides, meta = (InlineEditConditionToggle))
-	uint8 bOverride_Quality : 1;
 
 	UPROPERTY(EditAnywhere, BlueprintReadWrite, Category = Overrides, meta = (InlineEditConditionToggle))
 	uint8 bOverride_bDeleteSourceFiles : 1;
@@ -183,13 +181,18 @@ private:
 	template<typename T>
 	T* GetSettingOnBranch(const bool bIncludeCDOs = true, const bool bExactMatch = true) const;
 
+	/** Determines the appropriate evaluated config to use for the current context (primary job encode or shot based encode). */
+	TObjectPtr<UMovieGraphEvaluatedConfig> GetEvaluatedConfig() const;
+
+// The pipeline generates many instances of the same node throughout its execution; however, some nodes need to have persistent data throughout the
+// pipeline's lifetime. This static data enables the node to have shared data across instances.
 private:
 	/** Encode jobs which are currently running. */
-	TArray<FActiveJob> ActiveEncodeJobs;
+	static inline TArray<FActiveJob> ActiveEncodeJobs;
 
 	/** The pipeline that started the export. */
-	TWeakObjectPtr<UMovieGraphPipeline> CachedPipeline;
+	static inline TWeakObjectPtr<UMovieGraphPipeline> CachedPipeline;
 
-	/** The branch that the export is occurring on. */
-	FName CachedBranchName;
+	/** The job-level graph that should be accessed during execution of the node. Do not access the graph from the pipeline as it will be invalid. */
+	static inline TObjectPtr<UMovieGraphEvaluatedConfig> PrimaryJobEvaluatedGraph;
 };
\ No newline at end of file
diff --git a/Engine/Plugins/MovieScene/MovieRenderPipeline/Source/MovieRenderPipelineCore/Public/Graph/Nodes/MovieGraphFileOutputNode.h b/Engine/Plugins/MovieScene/MovieRenderPipeline/Source/MovieRenderPipelineCore/Public/Graph/Nodes/MovieGraphFileOutputNode.h
--- a/Engine/Plugins/MovieScene/MovieRenderPipeline/Source/MovieRenderPipelineCore/Public/Graph/Nodes/MovieGraphFileOutputNode.h
+++ b/Engine/Plugins/MovieScene/MovieRenderPipeline/Source/MovieRenderPipelineCore/Public/Graph/Nodes/MovieGraphFileOutputNode.h
@@ -20,17 +20,29 @@ class MOVIERENDERPIPELINECORE_API UMovieGraphFileOutputNode : public UMovieGraph
 public:
 	UMovieGraphFileOutputNode() = default;
 	
-	// UMovieGraphFileOutputNode Interface
 	/** 
 	* This is called by the Game Thread containing the final output data for a given output frame. InRawFrameData contains
 	* a list of all of the data generated for this output frame, and InMask specifies which render layers within InRawFrameData
 	* should actually be written to disk by this output node.
 	*/
 	void OnReceiveImageData(UMovieGraphPipeline* InPipeline, UE::MovieGraph::FMovieGraphOutputMergerFrame* InRawFrameData, const TSet<FMovieGraphRenderDataIdentifier>& InMask) { OnReceiveImageDataImpl(InPipeline, InRawFrameData, InMask); }
-	void OnAllFramesSubmitted() { OnAllFramesSubmittedImpl(); }
+
+	/**
+	 * This is called when the pipeline has finished rendering all shots, and the sequence export process has not begun yet.
+	 * See OnAllShotFramesSubmitted() for a method that is called when an individual shot has finished rendering.
+	 */
+	void OnAllFramesSubmitted(UMovieGraphPipeline* InPipeline, TObjectPtr<UMovieGraphEvaluatedConfig>& InPrimaryJobEvaluatedGraph) { OnAllFramesSubmittedImpl(InPipeline, InPrimaryJobEvaluatedGraph); }
+
+	/**
+	 * This is called when a shot has finished rendering its last frame, and the shot export process has not begun yet.
+	 * See OnAllFramesSubmitted() for a method that is called after all shots have finished rendering.
+	 */
+	void OnAllShotFramesSubmitted(UMovieGraphPipeline* InPipeline, const UMoviePipelineExecutorShot* InShot) { OnAllShotFramesSubmittedImpl(InPipeline, InShot); }
+
+	/** Returns whether this node has finished writing all of its files to disk yet. */
 	bool IsFinishedWritingToDisk() const { return IsFinishedWritingToDiskImpl(); }
-	// ~UMovieGraphFileOutputNode Interface
 
+	//~ Begin UMovieGraphSettingNode interface
 	virtual TArray<FMovieGraphPinProperties> GetInputPinProperties() const override
 	{
 		TArray<FMovieGraphPinProperties> Properties;
@@ -44,6 +56,7 @@ public:
 		Properties.Add(FMovieGraphPinProperties::MakeBranchProperties());
 		return Properties;
 	}
+	//~ End UMovieGraphSettingNode interface
 
 #if WITH_EDITOR
 	virtual FText GetMenuCategory() const override
@@ -61,7 +74,8 @@ public:
 
 protected:
 	virtual void OnReceiveImageDataImpl(UMovieGraphPipeline* InPipeline, UE::MovieGraph::FMovieGraphOutputMergerFrame* InRawFrameData, const TSet<FMovieGraphRenderDataIdentifier>& InMask) {}
-	virtual void OnAllFramesSubmittedImpl() {}
+	virtual void OnAllFramesSubmittedImpl(UMovieGraphPipeline* InPipeline, TObjectPtr<UMovieGraphEvaluatedConfig>& InPrimaryJobEvaluatedGraph) {}
+	virtual void OnAllShotFramesSubmittedImpl(UMovieGraphPipeline* InPipeline, const UMoviePipelineExecutorShot* InShot) {}
 	virtual bool IsFinishedWritingToDiskImpl() const { return true; }
 	
 	/** Returns the number of evaluated (active) file nodes on the specified branch. */
diff --git a/Engine/Plugins/MovieScene/MovieRenderPipeline/Source/MovieRenderPipelineRenderPasses/Private/MovieGraphImageSequenceOutputNode.cpp b/Engine/Plugins/MovieScene/MovieRenderPipeline/Source/MovieRenderPipelineRenderPasses/Private/MovieGraphImageSequenceOutputNode.cpp
--- a/Engine/Plugins/MovieScene/MovieRenderPipeline/Source/MovieRenderPipelineRenderPasses/Private/MovieGraphImageSequenceOutputNode.cpp
+++ b/Engine/Plugins/MovieScene/MovieRenderPipeline/Source/MovieRenderPipelineRenderPasses/Private/MovieGraphImageSequenceOutputNode.cpp
@@ -195,7 +195,7 @@ UMovieGraphImageSequenceOutputNode::UMovieGraphImageSequenceOutputNode()
 	ImageWriteQueue = &FModuleManager::Get().LoadModuleChecked<IImageWriteQueueModule>("ImageWriteQueue").GetWriteQueue();
 }
 
-void UMovieGraphImageSequenceOutputNode::OnAllFramesSubmittedImpl()
+void UMovieGraphImageSequenceOutputNode::OnAllFramesSubmittedImpl(UMovieGraphPipeline* InPipeline, TObjectPtr<UMovieGraphEvaluatedConfig>& InPrimaryJobEvaluatedGraph)
 {
 	FinalizeFence = ImageWriteQueue->CreateFence();
 }
diff --git a/Engine/Plugins/MovieScene/MovieRenderPipeline/Source/MovieRenderPipelineRenderPasses/Public/MovieGraphImageSequenceOutputNode.h b/Engine/Plugins/MovieScene/MovieRenderPipeline/Source/MovieRenderPipelineRenderPasses/Public/MovieGraphImageSequenceOutputNode.h
--- a/Engine/Plugins/MovieScene/MovieRenderPipeline/Source/MovieRenderPipelineRenderPasses/Public/MovieGraphImageSequenceOutputNode.h
+++ b/Engine/Plugins/MovieScene/MovieRenderPipeline/Source/MovieRenderPipelineRenderPasses/Public/MovieGraphImageSequenceOutputNode.h
@@ -29,7 +29,7 @@ public:
 	
 	// UMovieGraphFileOutputNode Interface
 	virtual void OnReceiveImageDataImpl(UMovieGraphPipeline* InPipeline, UE::MovieGraph::FMovieGraphOutputMergerFrame* InRawFrameData, const TSet<FMovieGraphRenderDataIdentifier>& InMask) override;
-	virtual void OnAllFramesSubmittedImpl() override;
+	virtual void OnAllFramesSubmittedImpl(UMovieGraphPipeline* InPipeline, TObjectPtr<UMovieGraphEvaluatedConfig>& InPrimaryJobEvaluatedGraph) override;
 	virtual bool IsFinishedWritingToDiskImpl() const override;
 	// ~UMovieGraphFileOutputNode Interface
 
