Movie Graph:

-add nullptr check before adding nodes to graph

diff --git a/Engine/Plugins/MovieScene/MovieRenderPipeline/Source/MovieRenderPipelineCore/Private/Graph/MovieGraphConfig.cpp b/Engine/Plugins/MovieScene/MovieRenderPipeline/Source/MovieRenderPipelineCore/Private/Graph/MovieGraphConfig.cpp
--- a/Engine/Plugins/MovieScene/MovieRenderPipeline/Source/MovieRenderPipelineCore/Private/Graph/MovieGraphConfig.cpp
+++ b/Engine/Plugins/MovieScene/MovieRenderPipeline/Source/MovieRenderPipelineCore/Private/Graph/MovieGraphConfig.cpp
@@ -319,6 +319,18 @@ void UMovieGraphConfig::PostLoad()
 				OnGraphVariablesChangedDelegate.Broadcast();
 			});
 		}
+
+		// Remove all null nodes
+		AllNodes.RemoveAll(
+			[](UMovieGraphNode* Node)
+			{
+				if (!Node)
+				{
+					UE_LOG(LogMovieRenderPipeline, Warning, TEXT("Encountered invalid source node (nullptr) when building Movie Pipeline Editor graph, skipping creating an editor graph node for the invalid source node."))
+					return true;
+				}
+				return false;
+			});
 #endif
 	}
 }
