Movie Graph:

-Fix: Add support for Additional Post Process Materials

diff --git a/Engine/Plugins/MovieScene/MovieRenderPipeline/Source/MovieRenderPipelineCore/Private/Graph/MovieGraphDefaultRenderer.cpp b/Engine/Plugins/MovieScene/MovieRenderPipeline/Source/MovieRenderPipelineCore/Private/Graph/MovieGraphDefaultRenderer.cpp
--- a/Engine/Plugins/MovieScene/MovieRenderPipeline/Source/MovieRenderPipelineCore/Private/Graph/MovieGraphDefaultRenderer.cpp
+++ b/Engine/Plugins/MovieScene/MovieRenderPipeline/Source/MovieRenderPipelineCore/Private/Graph/MovieGraphDefaultRenderer.cpp
@@ -222,7 +222,7 @@ void UMovieGraphDefaultRenderer::Render(const FMovieGraphTimeStepData& InTimeSte
 
 		for (const TObjectPtr<UMovieGraphRenderPassNode>& RenderPass : RenderPassesInUse)
 		{
-			RenderPass->GatherOutputPasses(NewOutputFrame.ExpectedRenderPasses);
+			RenderPass->GatherOutputPasses(InTimeStepData.EvaluatedConfig, NewOutputFrame.ExpectedRenderPasses);
 		}
 
 		// Register the frame with our render statistics as being worked on
diff --git a/Engine/Plugins/MovieScene/MovieRenderPipeline/Source/MovieRenderPipelineCore/Private/Graph/MovieGraphOutputMerger.cpp b/Engine/Plugins/MovieScene/MovieRenderPipeline/Source/MovieRenderPipelineCore/Private/Graph/MovieGraphOutputMerger.cpp
--- a/Engine/Plugins/MovieScene/MovieRenderPipeline/Source/MovieRenderPipelineCore/Private/Graph/MovieGraphOutputMerger.cpp
+++ b/Engine/Plugins/MovieScene/MovieRenderPipeline/Source/MovieRenderPipelineCore/Private/Graph/MovieGraphOutputMerger.cpp
@@ -88,14 +88,6 @@ namespace UE::MovieGraph
 		// Put the new data inside this output frame.
 		OutputFrame->ImageOutputData.FindOrAdd(NewLayerId) = MoveTemp(InData);
 		
-		// ToDo:
-		// Merge the metadata from each output state. Metadata is part of the output state but gets forked when
-		// we submit different render passes, so we need to merge it again. Doesn't handle conflicts.
-		//for (const TPair<FString, FString>& KVP : Payload->SampleState.OutputState.FileMetadata)
-		//{
-		//	OutputFrame->FrameOutputState.FileMetadata.Add(KVP.Key, KVP.Value);
-		//}
-		
 		// Check to see if this was the last piece of data needed for this frame.
 		int32 TotalPasses = OutputFrame->ExpectedRenderPasses.Num();
 		int32 FinishedPasses = OutputFrame->ImageOutputData.Num();
diff --git a/Engine/Plugins/MovieScene/MovieRenderPipeline/Source/MovieRenderPipelineCore/Private/Graph/Nodes/MovieGraphWidgetRendererBaseNode.cpp b/Engine/Plugins/MovieScene/MovieRenderPipeline/Source/MovieRenderPipelineCore/Private/Graph/Nodes/MovieGraphWidgetRendererBaseNode.cpp
--- a/Engine/Plugins/MovieScene/MovieRenderPipeline/Source/MovieRenderPipelineCore/Private/Graph/Nodes/MovieGraphWidgetRendererBaseNode.cpp
+++ b/Engine/Plugins/MovieScene/MovieRenderPipeline/Source/MovieRenderPipelineCore/Private/Graph/Nodes/MovieGraphWidgetRendererBaseNode.cpp
@@ -61,7 +61,7 @@ void UMovieGraphWidgetRendererBaseNode::RenderImpl(const FMovieGraphTraversalCon
 	}
 }
 
-void UMovieGraphWidgetRendererBaseNode::GatherOutputPassesImpl(TArray<FMovieGraphRenderDataIdentifier>& OutExpectedPasses) const
+void UMovieGraphWidgetRendererBaseNode::GatherOutputPassesImpl(UMovieGraphEvaluatedConfig* InConfig, TArray<FMovieGraphRenderDataIdentifier>& OutExpectedPasses) const
 {
 	for (const TUniquePtr<FMovieGraphWidgetPass>& Instance : CurrentInstances)
     {
diff --git a/Engine/Plugins/MovieScene/MovieRenderPipeline/Source/MovieRenderPipelineCore/Private/Graph/Nodes/MovieGraphWidgetRendererBaseNode.h b/Engine/Plugins/MovieScene/MovieRenderPipeline/Source/MovieRenderPipelineCore/Private/Graph/Nodes/MovieGraphWidgetRendererBaseNode.h
--- a/Engine/Plugins/MovieScene/MovieRenderPipeline/Source/MovieRenderPipelineCore/Private/Graph/Nodes/MovieGraphWidgetRendererBaseNode.h
+++ b/Engine/Plugins/MovieScene/MovieRenderPipeline/Source/MovieRenderPipelineCore/Private/Graph/Nodes/MovieGraphWidgetRendererBaseNode.h
@@ -70,7 +70,7 @@ protected:
 	virtual void SetupImpl(const FMovieGraphRenderPassSetupData& InSetupData) override;
 	virtual void TeardownImpl() override;
 	virtual void RenderImpl(const FMovieGraphTraversalContext& InFrameTraversalContext, const FMovieGraphTimeStepData& InTimeData) override;
-	virtual void GatherOutputPassesImpl(TArray<FMovieGraphRenderDataIdentifier>& OutExpectedPasses) const override;
+	virtual void GatherOutputPassesImpl(UMovieGraphEvaluatedConfig* InConfig, TArray<FMovieGraphRenderDataIdentifier>& OutExpectedPasses) const override;
 	// ~UMovieGraphRenderPassNode Interface
 
 private:
diff --git a/Engine/Plugins/MovieScene/MovieRenderPipeline/Source/MovieRenderPipelineCore/Private/MoviePipelineUtils.cpp b/Engine/Plugins/MovieScene/MovieRenderPipeline/Source/MovieRenderPipelineCore/Private/MoviePipelineUtils.cpp
--- a/Engine/Plugins/MovieScene/MovieRenderPipeline/Source/MovieRenderPipelineCore/Private/MoviePipelineUtils.cpp
+++ b/Engine/Plugins/MovieScene/MovieRenderPipeline/Source/MovieRenderPipelineCore/Private/MoviePipelineUtils.cpp
@@ -695,7 +695,7 @@ namespace UE
 			}
 		}
 
-		void ValidateOutputFormatString(FString& InOutFilenameFormatString, const bool bTestRenderPass, const bool bTestFrameNumber, const bool bTestCameraName)
+		void ValidateOutputFormatString(FString& InOutFilenameFormatString, const bool bTestRenderPass, const bool bTestFrameNumber, const bool bTestCameraName, const bool bIncludeSubRendererName)
 		{
 			const FString FrameNumberIdentifiers[] = { TEXT("{frame_number}"), TEXT("{frame_number_shot}"), TEXT("{frame_number_rel}"), TEXT("{frame_number_shot_rel}") };
 
@@ -793,6 +793,14 @@ namespace UE
 				}
 			}
 
+			if (bIncludeSubRendererName)
+			{
+				if (!InOutFilenameFormatString.Contains(TEXT("{renderer_sub_name}"), ESearchCase::IgnoreCase))
+				{
+					InOutFilenameFormatString.Append(TEXT(".{renderer_sub_name}"));
+				}
+			}
+
 			if (!InOutFilenameFormatString.Contains(TEXT("{file_dup}"), ESearchCase::IgnoreCase))
 			{
 				InOutFilenameFormatString.Append(TEXT("{file_dup}"));
diff --git a/Engine/Plugins/MovieScene/MovieRenderPipeline/Source/MovieRenderPipelineCore/Public/Graph/Nodes/MovieGraphRenderPassNode.h b/Engine/Plugins/MovieScene/MovieRenderPipeline/Source/MovieRenderPipelineCore/Public/Graph/Nodes/MovieGraphRenderPassNode.h
--- a/Engine/Plugins/MovieScene/MovieRenderPipeline/Source/MovieRenderPipelineCore/Public/Graph/Nodes/MovieGraphRenderPassNode.h
+++ b/Engine/Plugins/MovieScene/MovieRenderPipeline/Source/MovieRenderPipelineCore/Public/Graph/Nodes/MovieGraphRenderPassNode.h
@@ -28,7 +28,6 @@ public:
 	{
 	}
 
-	// UMovieGraphRenderPassNode Interface
 	/** Get the name of this renderer. Deferred, Path Tracer, Panoramic, etc. Called on the CDO. */
 	FString GetRendererName() const { return GetRendererNameImpl(); }
 	
@@ -47,8 +46,7 @@ public:
 	* Output Merger to ensure all of the render data for a given frame has been generated before passing it on to write to disk.
 	* Called on the CDO.
 	*/
-	void GatherOutputPasses(TArray<FMovieGraphRenderDataIdentifier>& OutExpectedPasses) const { GatherOutputPassesImpl(OutExpectedPasses); }
-	// ~UMovieGraphRenderPassNode Interface
+	void GatherOutputPasses(UMovieGraphEvaluatedConfig* InConfig, TArray<FMovieGraphRenderDataIdentifier>& OutExpectedPasses) const { GatherOutputPassesImpl(InConfig, OutExpectedPasses); }
 
 #if WITH_EDITOR
 	virtual FText GetMenuCategory() const override
@@ -74,5 +72,5 @@ protected:
 	virtual void SetupImpl(const FMovieGraphRenderPassSetupData& InSetupData) {}
 	virtual void TeardownImpl() {}
 	virtual void RenderImpl(const FMovieGraphTraversalContext& InFrameTraversalContext, const FMovieGraphTimeStepData& InTimeData) {}
-	virtual void GatherOutputPassesImpl(TArray<FMovieGraphRenderDataIdentifier>& OutExpectedPasses) const {}
+	virtual void GatherOutputPassesImpl(UMovieGraphEvaluatedConfig* InConfig, TArray<FMovieGraphRenderDataIdentifier>& OutExpectedPasses) const {}
 };
\ No newline at end of file
diff --git a/Engine/Plugins/MovieScene/MovieRenderPipeline/Source/MovieRenderPipelineCore/Public/MoviePipelineUtils.h b/Engine/Plugins/MovieScene/MovieRenderPipeline/Source/MovieRenderPipelineCore/Public/MoviePipelineUtils.h
--- a/Engine/Plugins/MovieScene/MovieRenderPipeline/Source/MovieRenderPipelineCore/Public/MoviePipelineUtils.h
+++ b/Engine/Plugins/MovieScene/MovieRenderPipeline/Source/MovieRenderPipelineCore/Public/MoviePipelineUtils.h
@@ -132,7 +132,7 @@ namespace UE
 	namespace MoviePipeline
 	{
 		MOVIERENDERPIPELINECORE_API void ConformOutputFormatStringToken(FString& InOutFilenameFormatString, const FStringView InToken, const FName& InNodeName, const FName& InBranchName);
-		MOVIERENDERPIPELINECORE_API void ValidateOutputFormatString(FString& InOutFilenameFormatString, const bool bTestRenderPass, const bool bTestFrameNumber, const bool bIncludeCameraName = false);
+		MOVIERENDERPIPELINECORE_API void ValidateOutputFormatString(FString& InOutFilenameFormatString, const bool bTestRenderPass, const bool bTestFrameNumber, const bool bIncludeCameraName = false, const bool bIncludeSubRendererName = false);
 		MOVIERENDERPIPELINECORE_API void RemoveFrameNumberFormatStrings(FString& InOutFilenameFormatString, const bool bIncludeShots);
 		/** De-duplicates the provided array of strings by appending (1), (2), etc to the end of duplicates. */
 		MOVIERENDERPIPELINECORE_API void DeduplicateNameArray(TArray<FString>& InOutNames);
diff --git a/Engine/Plugins/MovieScene/MovieRenderPipeline/Source/MovieRenderPipelineRenderPasses/Private/Graph/Nodes/MovieGraphDeferredPassNode.cpp b/Engine/Plugins/MovieScene/MovieRenderPipeline/Source/MovieRenderPipelineRenderPasses/Private/Graph/Nodes/MovieGraphDeferredPassNode.cpp
--- a/Engine/Plugins/MovieScene/MovieRenderPipeline/Source/MovieRenderPipelineRenderPasses/Private/Graph/Nodes/MovieGraphDeferredPassNode.cpp
+++ b/Engine/Plugins/MovieScene/MovieRenderPipeline/Source/MovieRenderPipelineRenderPasses/Private/Graph/Nodes/MovieGraphDeferredPassNode.cpp
@@ -11,9 +11,9 @@ TUniquePtr<UE::MovieGraph::Rendering::FMovieGraphImagePassBase> UMovieGraphDefer
 UMovieGraphDeferredRenderPassNode::UMovieGraphDeferredRenderPassNode()
 	: SpatialSampleCount(1)
 	, AntiAliasingMethod(EAntiAliasingMethod::AAM_TSR)
+	, bWriteAllSamples(false)
 	, bDisableToneCurve(false)
 	, bAllowOCIO(true)
-	, bWriteAllSamples(false)
 	, ViewModeIndex(VMI_Lit)
 {
 }
@@ -47,6 +47,11 @@ bool UMovieGraphDeferredRenderPassNode::GetWriteAllSamples() const
 	return bWriteAllSamples;
 }
 
+TArray<FMoviePipelinePostProcessPass> UMovieGraphDeferredRenderPassNode::GetAdditionalPostProcessMaterials() const
+{
+	return AdditionalPostProcessMaterials;
+}
+
 int32 UMovieGraphDeferredRenderPassNode::GetNumSpatialSamples() const
 {
 	return SpatialSampleCount;
diff --git a/Engine/Plugins/MovieScene/MovieRenderPipeline/Source/MovieRenderPipelineRenderPasses/Private/Graph/Nodes/MovieGraphImagePassBaseNode.cpp b/Engine/Plugins/MovieScene/MovieRenderPipeline/Source/MovieRenderPipelineRenderPasses/Private/Graph/Nodes/MovieGraphImagePassBaseNode.cpp
--- a/Engine/Plugins/MovieScene/MovieRenderPipeline/Source/MovieRenderPipelineRenderPasses/Private/Graph/Nodes/MovieGraphImagePassBaseNode.cpp
+++ b/Engine/Plugins/MovieScene/MovieRenderPipeline/Source/MovieRenderPipelineRenderPasses/Private/Graph/Nodes/MovieGraphImagePassBaseNode.cpp
@@ -87,11 +87,11 @@ void UMovieGraphImagePassBaseNode::RenderImpl(const FMovieGraphTraversalContext&
 	}
 }
 
-void UMovieGraphImagePassBaseNode::GatherOutputPassesImpl(TArray<FMovieGraphRenderDataIdentifier>& OutExpectedPasses) const
+void UMovieGraphImagePassBaseNode::GatherOutputPassesImpl(UMovieGraphEvaluatedConfig* InConfig, TArray<FMovieGraphRenderDataIdentifier>& OutExpectedPasses) const
 {
 	for (const TUniquePtr<UE::MovieGraph::Rendering::FMovieGraphImagePassBase>& Instance : CurrentInstances)
 	{
-		Instance->GatherOutputPasses(OutExpectedPasses);
+		Instance->GatherOutputPasses(InConfig, OutExpectedPasses);
 	}
 }
 
diff --git a/Engine/Plugins/MovieScene/MovieRenderPipeline/Source/MovieRenderPipelineRenderPasses/Private/Graph/Nodes/MovieGraphPathTracedRenderPassNode.cpp b/Engine/Plugins/MovieScene/MovieRenderPipeline/Source/MovieRenderPipelineRenderPasses/Private/Graph/Nodes/MovieGraphPathTracedRenderPassNode.cpp
--- a/Engine/Plugins/MovieScene/MovieRenderPipeline/Source/MovieRenderPipelineRenderPasses/Private/Graph/Nodes/MovieGraphPathTracedRenderPassNode.cpp
+++ b/Engine/Plugins/MovieScene/MovieRenderPipeline/Source/MovieRenderPipelineRenderPasses/Private/Graph/Nodes/MovieGraphPathTracedRenderPassNode.cpp
@@ -66,6 +66,11 @@ bool UMovieGraphPathTracedRenderPassNode::GetWriteAllSamples() const
 	return bWriteAllSamples;
 }
 
+TArray<FMoviePipelinePostProcessPass> UMovieGraphPathTracedRenderPassNode::GetAdditionalPostProcessMaterials() const
+{
+	return AdditionalPostProcessMaterials;
+}
+
 int32 UMovieGraphPathTracedRenderPassNode::GetNumSpatialSamples() const
 {
 	return SpatialSampleCount;
diff --git a/Engine/Plugins/MovieScene/MovieRenderPipeline/Source/MovieRenderPipelineRenderPasses/Private/Graph/Renderers/MovieGraphDeferredPass.cpp b/Engine/Plugins/MovieScene/MovieRenderPipeline/Source/MovieRenderPipelineRenderPasses/Private/Graph/Renderers/MovieGraphDeferredPass.cpp
--- a/Engine/Plugins/MovieScene/MovieRenderPipeline/Source/MovieRenderPipelineRenderPasses/Private/Graph/Renderers/MovieGraphDeferredPass.cpp
+++ b/Engine/Plugins/MovieScene/MovieRenderPipeline/Source/MovieRenderPipelineRenderPasses/Private/Graph/Renderers/MovieGraphDeferredPass.cpp
@@ -10,12 +10,14 @@
 #include "Graph/MovieGraphBlueprintLibrary.h"
 #include "MovieRenderOverlappedImage.h"
 #include "MoviePipelineSurfaceReader.h"
+
 #include "EngineModule.h"
 #include "SceneManagement.h"
 #include "CanvasTypes.h"
 #include "Engine/EngineBaseTypes.h"
 #include "Engine/TextureRenderTarget2D.h"
 #include "LegacyScreenPercentageDriver.h"
+#include "Materials/MaterialInterface.h"
 #include "OpenColorIODisplayExtension.h"
 #include "TextureResource.h"
 #include "Tasks/Task.h"
@@ -56,12 +58,29 @@ void FMovieGraphDeferredPass::Teardown()
 	SceneViewState.Destroy();
 }
 
-void FMovieGraphDeferredPass::GatherOutputPasses(TArray<FMovieGraphRenderDataIdentifier>& OutExpectedPasses) const
+void FMovieGraphDeferredPass::GatherOutputPasses(UMovieGraphEvaluatedConfig* InConfig, TArray<FMovieGraphRenderDataIdentifier>& OutExpectedPasses) const
 {
-	FMovieGraphImagePassBase::GatherOutputPasses(OutExpectedPasses);
+	FMovieGraphImagePassBase::GatherOutputPasses(InConfig,OutExpectedPasses);
 	
 	// Add our pre-calculated identifier
 	OutExpectedPasses.Add(RenderDataIdentifier);
+
+	if (const UMovieGraphImagePassBaseNode* ParentNode = GetParentNode(InConfig))
+	{
+		for (const FMoviePipelinePostProcessPass& AdditionalPass : ParentNode->GetAdditionalPostProcessMaterials())
+		{
+			if (AdditionalPass.bEnabled)
+			{
+				UMaterialInterface* Material = AdditionalPass.Material.LoadSynchronous();
+				if (Material)
+				{
+					FMovieGraphRenderDataIdentifier Identifier = RenderDataIdentifier;
+					Identifier.SubResourceName = Material->GetName();
+					OutExpectedPasses.Add(Identifier);
+				}
+			}
+		}
+	}
 }
 
 void FMovieGraphDeferredPass::AddReferencedObjects(FReferenceCollector& Collector)
@@ -229,9 +248,6 @@ void FMovieGraphDeferredPass::Render(const FMovieGraphTraversalContext& InFrameT
 		const float DPIScale = 1.0f;
 		FCanvas Canvas = FCanvas(RenderTargetResource, HitProxyConsumer, GraphRenderer->GetWorld(), GraphRenderer->GetWorld()->GetFeatureLevel(), FCanvas::CDM_DeferDrawing, DPIScale);
 		
-		// Submit the renderer to be rendered
-		GetRendererModule().BeginRenderingViewFamily(&Canvas, ViewFamily.ToSharedPtr().Get());
-
 		// Construct the sample state that reflects the current render sample
 		UE::MovieGraph::FMovieGraphSampleState SampleState;
 		{
@@ -254,13 +270,49 @@ void FMovieGraphDeferredPass::Render(const FMovieGraphTraversalContext& InFrameT
 			SampleState.SceneCaptureSource = SceneCaptureSource;
 		}
 
+		if (UMovieGraphImagePassBaseNode* ParentNode = GetParentNode(InFrameTraversalContext.Time.EvaluatedConfig))
+		{
+			for (const FMoviePipelinePostProcessPass& PostProcessPass : ParentNode->GetAdditionalPostProcessMaterials())
+			{
+				if (PostProcessPass.bEnabled)
+				{
+					UMaterialInterface* Material = PostProcessPass.Material.LoadSynchronous();
+					if (Material)
+					{
+						NewView->FinalPostProcessSettings.BufferVisualizationOverviewMaterials.Add(Material);
+					}
+				}
+			}
+
+			for (UMaterialInterface* VisMaterial : NewView->FinalPostProcessSettings.BufferVisualizationOverviewMaterials)
+			{
+				auto BufferPipe = MakeShared<FImagePixelPipe, ESPMode::ThreadSafe>();
+				
+				FMovieGraphRenderDataIdentifier Identifier = RenderDataIdentifier;
+				Identifier.SubResourceName = VisMaterial->GetName();
+				
+				UE::MovieGraph::FMovieGraphSampleState PassSampleState = SampleState;
+				PassSampleState.TraversalContext.RenderDataIdentifier = Identifier;
+				
+				BufferPipe->AddEndpoint(MakeForwardingEndpoint(PassSampleState, InTimeData));
+
+				NewView->FinalPostProcessSettings.BufferVisualizationPipes.Add(VisMaterial->GetFName(), BufferPipe);
+			}
+		}
+
+		int32 NumValidMaterials = NewView->FinalPostProcessSettings.BufferVisualizationPipes.Num();
+		NewView->FinalPostProcessSettings.bBufferVisualizationDumpRequired = NumValidMaterials > 0;
+
+		// Submit the renderer to be rendered
+		GetRendererModule().BeginRenderingViewFamily(&Canvas, ViewFamily.ToSharedPtr().Get());
+
 		// If this was just to contribute to the history buffer, no need to go any further.
 		bool bDiscardOutput = InTimeData.bDiscardOutput || ShouldDiscardOutput(ViewFamily, CameraInfo);
 		if (bDiscardOutput)
 		{
 			continue;
 		}
-
+		
 		// Readback + Accumulate.
 		PostRendererSubmission(SampleState, RenderTargetInitParams, Canvas, CameraInfo);
 	}
diff --git a/Engine/Plugins/MovieScene/MovieRenderPipeline/Source/MovieRenderPipelineRenderPasses/Private/Graph/Renderers/MovieGraphImagePassBase.cpp b/Engine/Plugins/MovieScene/MovieRenderPipeline/Source/MovieRenderPipelineRenderPasses/Private/Graph/Renderers/MovieGraphImagePassBase.cpp
--- a/Engine/Plugins/MovieScene/MovieRenderPipeline/Source/MovieRenderPipelineRenderPasses/Private/Graph/Renderers/MovieGraphImagePassBase.cpp
+++ b/Engine/Plugins/MovieScene/MovieRenderPipeline/Source/MovieRenderPipelineRenderPasses/Private/Graph/Renderers/MovieGraphImagePassBase.cpp
@@ -535,6 +535,83 @@ void FMovieGraphImagePassBase::PostRendererSubmission(
 		});
 }
 
+TFunction<void(TUniquePtr<FImagePixelData>&&)> FMovieGraphImagePassBase::MakeForwardingEndpoint(
+	const FMovieGraphSampleState& InSampleState, const FMovieGraphTimeStepData& InTimeData)
+{
+	// We have a pool of accumulators - we multi-thread the accumulation on the task graph, and for each frame,
+	// the task has the previous samples as pre-reqs to keep the accumulation in order. However, each accumulator
+	// can only work on one frame at a time, so we create a pool of them to work concurrently. This needs a limit
+	// as large accumulations (16k) can take a lot of system RAM.
+	TObjectPtr<UMovieGraphDefaultRenderer> GraphRenderer = GetRenderer().Get();
+	if (!GraphRenderer)
+	{
+		return nullptr;
+	}
+	
+	FMoviePipelineAccumulatorPoolPtr SampleAccumulator = GraphRenderer->GetOrCreateAccumulatorPool<FImageOverlappedAccumulator>();
+	UE::MovieGraph::DefaultRenderer::FSurfaceAccumulatorPool::FInstancePtr AccumulatorInstance =
+		SampleAccumulator->GetAccumulatorInstance_GameThread<FImageOverlappedAccumulator>(
+			InTimeData.RenderedFrameNumber, InSampleState.TraversalContext.RenderDataIdentifier);
+	
+	FMovieGraphRenderDataAccumulationArgs AccumulationArgs;
+	{
+		AccumulationArgs.OutputMerger = GraphRenderer->GetOwningGraph()->GetOutputMerger();
+		AccumulationArgs.ImageAccumulator = StaticCastSharedPtr<FImageOverlappedAccumulator>(AccumulatorInstance->Accumulator);
+		AccumulationArgs.bIsFirstSample = InTimeData.bIsFirstTemporalSampleForFrame;
+		AccumulationArgs.bIsLastSample = InTimeData.bIsLastTemporalSampleForFrame;
+	}
+
+	// The legacy surface reader takes the payload just so it can shuffle it into our callback, but we can just include the data
+	// directly in the callback, so this is just a dummy payload.
+	TSharedRef<FImagePixelDataPayload, ESPMode::ThreadSafe> FramePayload = MakeShared<FImagePixelDataPayload, ESPMode::ThreadSafe>();
+	auto Callback = [this, InSampleState, FramePayload, AccumulationArgs, AccumulatorInstance](TUniquePtr<FImagePixelData>&& InPixelData)
+	{
+		// Transfer the framePayload to the returned data
+		TUniquePtr<FImagePixelData> PixelDataWithPayload = nullptr;
+		switch (InPixelData->GetType())
+		{
+		case EImagePixelType::Color:
+		{
+			TImagePixelData<FColor>* SourceData = static_cast<TImagePixelData<FColor>*>(InPixelData.Get());
+			PixelDataWithPayload = MakeUnique<TImagePixelData<FColor>>(InPixelData->GetSize(), MoveTemp(SourceData->Pixels), FramePayload);
+			break;
+		}
+		case EImagePixelType::Float16:
+		{
+			TImagePixelData<FFloat16Color>* SourceData = static_cast<TImagePixelData<FFloat16Color>*>(InPixelData.Get());
+			PixelDataWithPayload = MakeUnique<TImagePixelData<FFloat16Color>>(InPixelData->GetSize(), MoveTemp(SourceData->Pixels), FramePayload);
+			break;
+		}
+		case EImagePixelType::Float32:
+		{
+			TImagePixelData<FLinearColor>* SourceData = static_cast<TImagePixelData<FLinearColor>*>(InPixelData.Get());
+			PixelDataWithPayload = MakeUnique<TImagePixelData<FLinearColor>>(InPixelData->GetSize(), MoveTemp(SourceData->Pixels), FramePayload);
+			break;
+		}
+		default:
+			checkNoEntry();
+		}
+
+		UE::Tasks::TTask<void> Task = UE::Tasks::Launch(UE_SOURCE_LOCATION, [PixelData = MoveTemp(PixelDataWithPayload), InSampleState, AccumulationArgs, AccumulatorInstance]() mutable
+		{
+			// Enqueue a encode for this frame onto our worker thread.
+			AccumulateSample_TaskThread(MoveTemp(PixelData), InSampleState, AccumulationArgs);
+
+			// We have to defer clearing the accumulator until after sample accumulation has finished
+			if (AccumulationArgs.bIsLastSample)
+			{
+				// Final sample has now been executed, free the accumulator for reuse.
+				AccumulatorInstance->bIsActive = false;
+			}
+		}, AccumulatorInstance->TaskPrereq);
+
+		// Make the next accumulation task that uses this accumulator use the task we just created as a pre-req.
+		AccumulatorInstance->TaskPrereq = Task;
+	};
+
+	return Callback;
+}
+
 void AccumulateSample_TaskThread(TUniquePtr<FImagePixelData>&& InPixelData, const ::UE::MovieGraph::FMovieGraphSampleState InSampleState, const FMovieGraphRenderDataAccumulationArgs& InAccumulationParams)
 {
 	TRACE_CPUPROFILER_EVENT_SCOPE(MoviePipeline_AccumulateSample);
diff --git a/Engine/Plugins/MovieScene/MovieRenderPipeline/Source/MovieRenderPipelineRenderPasses/Private/MovieGraphImageSequenceOutputNode.cpp b/Engine/Plugins/MovieScene/MovieRenderPipeline/Source/MovieRenderPipelineRenderPasses/Private/MovieGraphImageSequenceOutputNode.cpp
--- a/Engine/Plugins/MovieScene/MovieRenderPipeline/Source/MovieRenderPipelineRenderPasses/Private/MovieGraphImageSequenceOutputNode.cpp
+++ b/Engine/Plugins/MovieScene/MovieRenderPipeline/Source/MovieRenderPipelineRenderPasses/Private/MovieGraphImageSequenceOutputNode.cpp
@@ -70,6 +70,23 @@ namespace UE::MovieGraph::Private
 
 		return Params;
 	}
+
+	/** Determine whether this token should be added based on the number of unique SubResourceNames in ExpectedRenderPasses. */
+	bool ShouldIncludeSubRendererName(const TArray<FMovieGraphRenderDataIdentifier>& ExpectedRenderPasses)
+	{
+		TSet<FString> PassNames;
+		for (const FMovieGraphRenderDataIdentifier& Pass : ExpectedRenderPasses)
+		{
+			PassNames.Add(Pass.SubResourceName);
+
+			if (PassNames.Num() > 1)
+			{
+				return true;
+			}
+		}
+
+		return false;
+	}
 	
 #if WITH_OCIO
 	struct FOpenColorIOPixelPreProcessor
@@ -290,8 +307,10 @@ FString UMovieGraphImageSequenceOutputNode::CreateFileName(
 	// Previous method is preserved for output frame number validation.
 	constexpr bool bIncludeRenderPass = false;
 	constexpr bool bTestFrameNumber = true;
-	UE::MoviePipeline::ValidateOutputFormatString(FileNameFormatString, bIncludeRenderPass, bTestFrameNumber);
-
+	constexpr bool bIncludeCameraName = false;
+	const bool bIncludeSubRendererName = UE::MovieGraph::Private::ShouldIncludeSubRendererName(InRawFrameData->ExpectedRenderPasses);
+	UE::MoviePipeline::ValidateOutputFormatString(FileNameFormatString, bIncludeRenderPass, bTestFrameNumber,bIncludeCameraName, bIncludeSubRendererName);
+	
 	// Map the .ext to be specific to our output data.
 	TMap<FString, FString> AdditionalFormatArgs;
 	AdditionalFormatArgs.Add(TEXT("ext"), Extension);
@@ -831,8 +850,10 @@ FString UMovieGraphImageSequenceOutputNode_MultiLayerEXR::ResolveOutputFilename(
 	// overwrite the same file multiple times. Burn In overlays don't count because they get composited on top of an existing file.
 	constexpr bool bIncludeRenderPass = false;
 	constexpr bool bTestFrameNumber = true;
-	UE::MoviePipeline::ValidateOutputFormatString(FileNameFormatString, bIncludeRenderPass, bTestFrameNumber);
-
+	constexpr bool bIncludeCameraName = false;
+	const bool bIncludeSubRendererName = UE::MovieGraph::Private::ShouldIncludeSubRendererName(InRawFrameData->ExpectedRenderPasses);
+	UE::MoviePipeline::ValidateOutputFormatString(FileNameFormatString, bIncludeRenderPass, bTestFrameNumber, bIncludeCameraName, bIncludeSubRendererName);
+	
 	// Create specific data that needs to override 
 	TMap<FString, FString> FormatOverrides;
 	FormatOverrides.Add(TEXT("render_pass"), TEXT("")); // Render Passes are included inside the exr file by named layers.
diff --git a/Engine/Plugins/MovieScene/MovieRenderPipeline/Source/MovieRenderPipelineRenderPasses/Public/Graph/Nodes/MovieGraphDeferredPassNode.h b/Engine/Plugins/MovieScene/MovieRenderPipeline/Source/MovieRenderPipelineRenderPasses/Public/Graph/Nodes/MovieGraphDeferredPassNode.h
--- a/Engine/Plugins/MovieScene/MovieRenderPipeline/Source/MovieRenderPipelineRenderPasses/Public/Graph/Nodes/MovieGraphDeferredPassNode.h
+++ b/Engine/Plugins/MovieScene/MovieRenderPipeline/Source/MovieRenderPipelineRenderPasses/Public/Graph/Nodes/MovieGraphDeferredPassNode.h
@@ -2,6 +2,7 @@
 
 #pragma once
 #include "Graph/Nodes/MovieGraphImagePassBaseNode.h"
+#include "MoviePipelineDeferredPasses.h"
 #include "MovieGraphDeferredPassNode.generated.h"
 
 /**
@@ -26,6 +27,7 @@ protected:
 	virtual FString GetRendererNameImpl() const override;
 	virtual TUniquePtr<UE::MovieGraph::Rendering::FMovieGraphImagePassBase> CreateInstance() const override;
 	virtual bool GetWriteAllSamples() const override;
+	virtual TArray<FMoviePipelinePostProcessPass> GetAdditionalPostProcessMaterials() const override;
 	virtual int32 GetNumSpatialSamples() const override;
 	virtual bool GetDisableToneCurve() const override;
 	virtual bool GetAllowOCIO() const override;
@@ -55,24 +57,35 @@ public:
 	UPROPERTY(EditAnywhere, BlueprintReadWrite, Category = Overrides, meta = (InlineEditConditionToggle))
 	uint8 bOverride_bWriteAllSamples : 1;
 
+	UPROPERTY(EditAnywhere, BlueprintReadWrite, Category = Overrides, meta = (InlineEditConditionToggle))
+	uint8 bOverride_AdditionalPostProcessMaterials : 1;
+
 	/**
 	* How many sub-pixel jitter renders should we do per temporal sample? This can be used to achieve high
 	* sample counts without Temporal Sub-Sampling (allowing high sample counts without motion blur being enabled),
 	* but we generally recommend using Temporal Sub-Samples when possible. It can also be combined with
 	* temporal samples and you will get SpatialSampleCount many renders per temporal sample.
 	*/
-	UPROPERTY(EditAnywhere, BlueprintReadWrite, meta = (UIMin = 1, ClampMin = 1), Category = "Settings", meta=(EditCondition="bOverride_SpatialSampleCount"))
+	UPROPERTY(EditAnywhere, BlueprintReadWrite, meta = (UIMin = 1, ClampMin = 1), Category = "Sampling", meta=(EditCondition="bOverride_SpatialSampleCount"))
 	int32 SpatialSampleCount;
 
 	/**
 	* Which anti-aliasing method should this render use. If this is set to None, then Movie Render Graph
 	* will handle anti-aliasing by doing a sub-pixel jitter (one for each temporal/spatial sample). Some
-	* rendering effects rely on TSR or TAA to reduce noise so we reccomend leaving them enabled
+	* rendering effects rely on TSR or TAA to reduce noise so we recommend leaving them enabled
 	* where possible. All options work with Spatial and Temporal samples, but TSR/TAA may introduce minor
 	* visual artifacts (such as ghosting). MSAA is not supported in the deferred renderer.
 	*/
-	UPROPERTY(EditAnywhere, BlueprintReadWrite, Category = "Settings", meta = (EditCondition = "bOverride_AntiAliasingMethod"))
+	UPROPERTY(EditAnywhere, BlueprintReadWrite, Category = "Sampling", meta = (EditCondition = "bOverride_AntiAliasingMethod"))
 	TEnumAsByte<EAntiAliasingMethod> AntiAliasingMethod;
+	
+	/**
+	* Debug Feature. Can use this to write out each individual Temporal and Spatial sample rendered by this render pass,
+	* which allows you to see which images are being accumulated together. Can be useful for debugging incorrect looking
+	* frames to see which sub-frame evaluations were incorrect.
+	*/
+	UPROPERTY(EditAnywhere, BlueprintReadWrite, Category = "Sampling", AdvancedDisplay, DisplayName="Write All Samples (Debug)", meta = (EditCondition = "bOverride_bWriteAllsamples"))
+	bool bWriteAllSamples;
 
 	/**
 	* If true, the tone curve will be disabled for this render pass. This will result in values greater than 1.0 in final renders
@@ -87,15 +100,7 @@ public:
 	*/
 	UPROPERTY(EditAnywhere, BlueprintReadWrite, Category = "Settings", meta = (EditCondition = "bOverride_bAllowOCIO"))
 	bool bAllowOCIO;
-
-	/**
-	* Debug Feature. Can use this to write out each individual Temporal and Spatial sample rendered by this render pass,
-	* which allows you to see which images are being accumulated together. Can be useful for debugging incorrect looking
-	* frames to see which sub-frame evaluations were incorrect.
-	*/
-	UPROPERTY(EditAnywhere, BlueprintReadWrite, Category = "Debug", meta = (EditCondition = "bOverride_bWriteAllsamples"))
-	bool bWriteAllSamples;
-
+	
 	/** 
 	* The view mode index that will be applied to renders. These mirror the View Modes you find in the Viewport,
 	* but most view modes other than Lit are used for debugging so they may not do what you expect, or may
@@ -104,4 +109,10 @@ public:
 	*/
 	UPROPERTY(EditAnywhere, BlueprintReadWrite, Category = "Settings", meta=(EditCondition="bOverride_ViewModeIndex", InvalidEnumValues = "VMI_PathTracing"))
 	TEnumAsByte<EViewModeIndex> ViewModeIndex;
+
+	/**
+	* An array of additional post-processing materials to run after the frame is rendered. Using this feature may add a notable amount of render time.
+	*/
+	UPROPERTY(EditAnywhere, BlueprintReadWrite, Category = "Post Process Materials", meta=(EditCondition="bOverride_AdditionalPostProcessMaterials"))
+	TArray<FMoviePipelinePostProcessPass> AdditionalPostProcessMaterials;
 };
\ No newline at end of file
diff --git a/Engine/Plugins/MovieScene/MovieRenderPipeline/Source/MovieRenderPipelineRenderPasses/Public/Graph/Nodes/MovieGraphImagePassBaseNode.h b/Engine/Plugins/MovieScene/MovieRenderPipeline/Source/MovieRenderPipelineRenderPasses/Public/Graph/Nodes/MovieGraphImagePassBaseNode.h
--- a/Engine/Plugins/MovieScene/MovieRenderPipeline/Source/MovieRenderPipelineRenderPasses/Public/Graph/Nodes/MovieGraphImagePassBaseNode.h
+++ b/Engine/Plugins/MovieScene/MovieRenderPipeline/Source/MovieRenderPipelineRenderPasses/Public/Graph/Nodes/MovieGraphImagePassBaseNode.h
@@ -7,6 +7,7 @@
 #include "Graph/Renderers/MovieGraphImagePassBase.h"
 #include "Graph/Renderers/MovieGraphShowFlags.h"
 #include "Graph/Nodes/MovieGraphRenderPassNode.h"
+#include "MoviePipelineDeferredPasses.h"
 #include "MovieGraphImagePassBaseNode.generated.h"
 
 // Forward Declare
@@ -40,6 +41,9 @@ public:
 	/** Should each individual sample rendered be written out for debugging? */
 	virtual bool GetWriteAllSamples() const { return false; }
 
+	/** Get an array of user-added post-process materials for the render */
+	virtual TArray<FMoviePipelinePostProcessPass> GetAdditionalPostProcessMaterials() const { return {}; }
+
 	/** How many spatial samples should be rendered each frame? */
 	virtual int32 GetNumSpatialSamples() const { return 1; }
 
@@ -69,7 +73,7 @@ protected:
 	virtual void SetupImpl(const FMovieGraphRenderPassSetupData& InSetupData) override;
 	virtual void TeardownImpl() override;
 	virtual void RenderImpl(const FMovieGraphTraversalContext& InFrameTraversalContext, const FMovieGraphTimeStepData& InTimeData) override;
-	virtual void GatherOutputPassesImpl(TArray<FMovieGraphRenderDataIdentifier>& OutExpectedPasses) const override;
+	virtual void GatherOutputPassesImpl(UMovieGraphEvaluatedConfig* InConfig, TArray<FMovieGraphRenderDataIdentifier>& OutExpectedPasses) const override;
 	// ~UMovieGraphRenderPassNode Interface
 
 	virtual TUniquePtr<UE::MovieGraph::Rendering::FMovieGraphImagePassBase> CreateInstance() const { return nullptr; }
diff --git a/Engine/Plugins/MovieScene/MovieRenderPipeline/Source/MovieRenderPipelineRenderPasses/Public/Graph/Nodes/MovieGraphPathTracedRenderPassNode.h b/Engine/Plugins/MovieScene/MovieRenderPipeline/Source/MovieRenderPipelineRenderPasses/Public/Graph/Nodes/MovieGraphPathTracedRenderPassNode.h
--- a/Engine/Plugins/MovieScene/MovieRenderPipeline/Source/MovieRenderPipelineRenderPasses/Public/Graph/Nodes/MovieGraphPathTracedRenderPassNode.h
+++ b/Engine/Plugins/MovieScene/MovieRenderPipeline/Source/MovieRenderPipelineRenderPasses/Public/Graph/Nodes/MovieGraphPathTracedRenderPassNode.h
@@ -26,6 +26,7 @@ public:
 
 	// UMovieGraphImagePassBaseNode Interface
 	virtual bool GetWriteAllSamples() const override;
+	virtual TArray<FMoviePipelinePostProcessPass> GetAdditionalPostProcessMaterials() const override;
 	virtual int32 GetNumSpatialSamples() const override;
 	virtual bool GetDisableToneCurve() const override;
 	virtual bool GetAllowOCIO() const override;
@@ -56,13 +57,16 @@ public:
 	UPROPERTY(EditAnywhere, BlueprintReadWrite, Category = Overrides, meta = (InlineEditConditionToggle))
 	uint8 bOverride_bWriteAllSamples : 1;
 
+	UPROPERTY(EditAnywhere, BlueprintReadWrite, Category = Overrides, meta = (InlineEditConditionToggle))
+	uint8 bOverride_AdditionalPostProcessMaterials : 1;
+
 	/**
 	* How many sub-pixel jitter renders should we do per temporal sample? This can be used to achieve high
 	* sample counts without Temporal Sub-Sampling (allowing high sample counts without motion blur being enabled),
 	* but we generally recommend using Temporal Sub-Samples when possible. It can also be combined with
 	* temporal samples and you will get SpatialSampleCount many renders per temporal sample.
 	*/
-	UPROPERTY(EditAnywhere, BlueprintReadWrite, meta = (UIMin = 1, ClampMin = 1), Category = "Settings", meta = (EditCondition = "bOverride_SpatialSampleCount"))
+	UPROPERTY(EditAnywhere, BlueprintReadWrite, meta = (UIMin = 1, ClampMin = 1), Category = "Sampling", meta = (EditCondition = "bOverride_SpatialSampleCount"))
 	int32 SpatialSampleCount;
 
 	/**
@@ -84,9 +88,15 @@ public:
 	* which allows you to see which images are being accumulated together. Can be useful for debugging incorrect looking
 	* frames to see which sub-frame evaluations were incorrect.
 	*/
-	UPROPERTY(EditAnywhere, BlueprintReadWrite, Category = "Debug", meta = (EditCondition = "bOverride_bWriteAllsamples"))
+	UPROPERTY(EditAnywhere, BlueprintReadWrite, Category = "Sampling", meta = (EditCondition = "bOverride_bWriteAllsamples"))
 	bool bWriteAllSamples;
 
+	/**
+	* An array of additional post-processing materials to run after the frame is rendered. Using this feature may add a notable amount of render time.
+	*/
+	UPROPERTY(EditAnywhere, BlueprintReadWrite, Category = "Post Process Materials", meta=(EditCondition="bOverride_AdditionalPostProcessMaterials"))
+	TArray<FMoviePipelinePostProcessPass> AdditionalPostProcessMaterials;
+
 private:
 	/**
 	 * The original value of the "r.PathTracing.ProgressDisplay" cvar before the render starts. The progress display
diff --git a/Engine/Plugins/MovieScene/MovieRenderPipeline/Source/MovieRenderPipelineRenderPasses/Public/Graph/Renderers/MovieGraphDeferredPass.h b/Engine/Plugins/MovieScene/MovieRenderPipeline/Source/MovieRenderPipelineRenderPasses/Public/Graph/Renderers/MovieGraphDeferredPass.h
--- a/Engine/Plugins/MovieScene/MovieRenderPipeline/Source/MovieRenderPipelineRenderPasses/Public/Graph/Renderers/MovieGraphDeferredPass.h
+++ b/Engine/Plugins/MovieScene/MovieRenderPipeline/Source/MovieRenderPipelineRenderPasses/Public/Graph/Renderers/MovieGraphDeferredPass.h
@@ -1,19 +1,21 @@
 // Copyright Epic Games, Inc. All Rights Reserved.
 #pragma once
 
-#include "MovieGraphImagePassBase.h"
+#include "Graph/Renderers/MovieGraphImagePassBase.h"
 
 namespace UE::MovieGraph::Rendering
 {
 	struct FMovieGraphDeferredPass : public FMovieGraphImagePassBase
 	{
+		// FMovieGraphImagePassBase Interface
 		virtual void Setup(TWeakObjectPtr<UMovieGraphDefaultRenderer> InRenderer, TWeakObjectPtr<UMovieGraphImagePassBaseNode> InRenderPassNode, const FMovieGraphRenderPassLayerData& InLayer) override;
 		virtual void Teardown() override;
 		virtual void Render(const FMovieGraphTraversalContext& InFrameTraversalContext, const FMovieGraphTimeStepData& InTimeData) override;
-		virtual void GatherOutputPasses(TArray<FMovieGraphRenderDataIdentifier>& OutExpectedPasses) const override;
+		virtual void GatherOutputPasses(UMovieGraphEvaluatedConfig* InConfig, TArray<FMovieGraphRenderDataIdentifier>& OutExpectedPasses) const override;
 		virtual void AddReferencedObjects(FReferenceCollector& Collector) override;
 		virtual FName GetBranchName() const override;
 		virtual UMovieGraphImagePassBaseNode* GetParentNode(UMovieGraphEvaluatedConfig* InConfig) const override;
+		// End FMovieGraphImagePassBase
 			
 	protected:
 		virtual void PostRendererSubmission(const UE::MovieGraph::FMovieGraphSampleState& InSampleState, const UE::MovieGraph::DefaultRenderer::FRenderTargetInitParams& InRenderTargetInitParams, FCanvas& InCanvas, const UE::MovieGraph::DefaultRenderer::FCameraInfo& InCameraInfo) override;
diff --git a/Engine/Plugins/MovieScene/MovieRenderPipeline/Source/MovieRenderPipelineRenderPasses/Public/Graph/Renderers/MovieGraphImagePassBase.h b/Engine/Plugins/MovieScene/MovieRenderPipeline/Source/MovieRenderPipelineRenderPasses/Public/Graph/Renderers/MovieGraphImagePassBase.h
--- a/Engine/Plugins/MovieScene/MovieRenderPipeline/Source/MovieRenderPipelineRenderPasses/Public/Graph/Renderers/MovieGraphImagePassBase.h
+++ b/Engine/Plugins/MovieScene/MovieRenderPipeline/Source/MovieRenderPipelineRenderPasses/Public/Graph/Renderers/MovieGraphImagePassBase.h
@@ -60,7 +60,7 @@ namespace UE::MovieGraph::Rendering
 		virtual void Setup(TWeakObjectPtr<UMovieGraphDefaultRenderer> InRenderer, TWeakObjectPtr<UMovieGraphImagePassBaseNode> InRenderPassNode, const FMovieGraphRenderPassLayerData& InLayer);
 		virtual void Teardown() {}
 		virtual void Render(const FMovieGraphTraversalContext& InFrameTraversalContext, const FMovieGraphTimeStepData& InTimeData) {}
-		virtual void GatherOutputPasses(TArray<FMovieGraphRenderDataIdentifier>& OutExpectedPasses) const {}
+		virtual void GatherOutputPasses(UMovieGraphEvaluatedConfig* InConfig, TArray<FMovieGraphRenderDataIdentifier>& OutExpectedPasses) const {}
 		virtual void AddReferencedObjects(FReferenceCollector& Collector) {}
 		virtual FName GetBranchName() const { return NAME_None; }
 		virtual TWeakObjectPtr<UMovieGraphDefaultRenderer> GetRenderer() const { return WeakGraphRenderer; }
@@ -82,6 +82,7 @@ namespace UE::MovieGraph::Rendering
 		virtual void ApplyMovieGraphOverridesToSceneView(TSharedRef<FSceneViewFamilyContext> InOutFamily, const FViewFamilyInitData& InInitData, const UE::MovieGraph::DefaultRenderer::FCameraInfo& InCameraInfo) const;
 		virtual void ApplyMovieGraphOverridesToViewFamily(TSharedRef<FSceneViewFamilyContext> InOutFamily, const FViewFamilyInitData& InInitData) const;
 		virtual void PostRendererSubmission(const UE::MovieGraph::FMovieGraphSampleState& InSampleState, const UE::MovieGraph::DefaultRenderer::FRenderTargetInitParams& InRenderTargetInitParams, FCanvas& InCanvas, const UE::MovieGraph::DefaultRenderer::FCameraInfo& InCameraInfo);
+		virtual TFunction<void(TUniquePtr<FImagePixelData>&&)> MakeForwardingEndpoint(const FMovieGraphSampleState& InSampleState, const FMovieGraphTimeStepData& InTimeData);
 		virtual bool ShouldDiscardOutput(const TSharedRef<FSceneViewFamilyContext>& InFamily, const UE::MovieGraph::DefaultRenderer::FCameraInfo& InCameraInfo) const { return false; }
 		/** For this image pass, look up the associated node by type for the given config. */
 		virtual UMovieGraphImagePassBaseNode* GetParentNode(UMovieGraphEvaluatedConfig* InConfig) const { return nullptr; }
