MoviePipeline: Fixed an issue where shot variable overrides could be incorrect if a shot mask (disabling some shots in the queue) is used.

Also fixed a couple of other issues:
* Added more safeguards to ensure that variable overrides are updated properly. They cannot be updated in PostLoad (see comments), so the next-most reliable way of ensuring they are always updated is to update them when they're fetched. The update process is not very expensive, especially if no updates are required. This is also not a performance-critical codepath, so this paranoid way of keeping them up to date is excessive but should avoid refresh issues.
* Removing the variable update from PostLoad() surfaced a bug where the queue would indicate that a job was using a legacy config even if it was saved with a graph. The variable update was performing the only load of the graph preset's soft ptr, so once that was removed, all of the "is the job using a graph" checks were looking at a null soft ptr and assuming that a legacy config was being used. Code was updated to ensure that the graph is always loaded when it is being checked.

#rnx

diff --git a/Engine/Plugins/MovieScene/MovieRenderPipeline/Source/MovieRenderPipelineCore/Private/Graph/MovieGraphPipeline.cpp b/Engine/Plugins/MovieScene/MovieRenderPipeline/Source/MovieRenderPipelineCore/Private/Graph/MovieGraphPipeline.cpp
--- a/Engine/Plugins/MovieScene/MovieRenderPipeline/Source/MovieRenderPipelineCore/Private/Graph/MovieGraphPipeline.cpp
+++ b/Engine/Plugins/MovieScene/MovieRenderPipeline/Source/MovieRenderPipelineCore/Private/Graph/MovieGraphPipeline.cpp
@@ -309,6 +309,7 @@ void UMovieGraphPipeline::BuildShotListFromDataSource()
 		CurrentContext.ShotIndex = ShotIndex;
 		CurrentContext.ShotCount = GetActiveShotList().Num();
 		CurrentContext.Job = GetCurrentJob();
+		CurrentContext.Shot = Shot;
 		CurrentContext.RootGraph = GetRootGraphForShot(Shot);
 
 		// Frame Rate, Handle Frames, Warm-up Frames are all global settings so we provide an empty time context.
@@ -1378,11 +1379,14 @@ UMovieGraphConfig* UMovieGraphPipeline::GetRootGraphForShot(UMoviePipelineExecut
 
 FMovieGraphTraversalContext UMovieGraphPipeline::GetCurrentTraversalContext(const bool bForShot) const
 {
+	const TObjectPtr<UMoviePipelineExecutorShot> CurrentShot = bForShot ? GetActiveShotList()[GetCurrentShotIndex()] : nullptr;
+	
 	FMovieGraphTraversalContext CurrentContext;
 	CurrentContext.ShotIndex = bForShot ? GetCurrentShotIndex() : -1;
 	CurrentContext.ShotCount = GetActiveShotList().Num();
 	CurrentContext.Job = GetCurrentJob();
-	CurrentContext.RootGraph = bForShot ? GetRootGraphForShot(GetActiveShotList()[GetCurrentShotIndex()]) : CurrentContext.Job->GetGraphPreset();
+	CurrentContext.Shot = CurrentShot;
+	CurrentContext.RootGraph = bForShot ? GetRootGraphForShot(CurrentShot) : CurrentContext.Job->GetGraphPreset();
 	CurrentContext.Time = bForShot ? GetTimeStepInstance()->GetCalculatedTimeData() : FMovieGraphTimeStepData();
 
 	return CurrentContext;
diff --git a/Engine/Plugins/MovieScene/MovieRenderPipeline/Source/MovieRenderPipelineCore/Private/Graph/Nodes/MovieGraphVariableNode.cpp b/Engine/Plugins/MovieScene/MovieRenderPipeline/Source/MovieRenderPipelineCore/Private/Graph/Nodes/MovieGraphVariableNode.cpp
--- a/Engine/Plugins/MovieScene/MovieRenderPipeline/Source/MovieRenderPipelineCore/Private/Graph/Nodes/MovieGraphVariableNode.cpp
+++ b/Engine/Plugins/MovieScene/MovieRenderPipeline/Source/MovieRenderPipelineCore/Private/Graph/Nodes/MovieGraphVariableNode.cpp
@@ -140,21 +140,25 @@ bool UMovieGraphVariableNode::ContextHasEnabledAssignmentForVariable(const FMovi
 	TObjectPtr<UMovieJobVariableAssignmentContainer> ShotVariableAssignments = nullptr;
 	TObjectPtr<UMovieJobVariableAssignmentContainer> ShotVariableAssignments_PrimaryOverrides = nullptr;
 	TObjectPtr<UMovieJobVariableAssignmentContainer> JobVariableAssignments = nullptr;
+
+	const TObjectPtr<UMoviePipelineExecutorJob> PrimaryJob = InContext->Job;
+	const TObjectPtr<UMoviePipelineExecutorShot> Shot = InContext->Shot;
 	
-	if (InContext->Job->ShotInfo.IsValidIndex(InContext->ShotIndex))
+	if (Shot)
 	{
-		const TObjectPtr<UMoviePipelineExecutorShot> ShotJob = InContext->Job->ShotInfo[InContext->ShotIndex];
+		ShotVariableAssignments = Shot->GetOrCreateJobVariableAssignmentsForGraph(Shot->GetGraphPreset());
 		
-		ShotVariableAssignments = ShotJob->GetOrCreateJobVariableAssignmentsForGraph(ShotJob->GetGraphPreset());
-
-		// The shot can also override variables on the primary job's graph
-		constexpr bool bIsForPrimaryOverrides = true;
-		ShotVariableAssignments_PrimaryOverrides = ShotJob->GetOrCreateJobVariableAssignmentsForGraph(InContext->Job->GetGraphPreset(), bIsForPrimaryOverrides);
+		// The shot can also override variables on the primary job's graph (in addition to the shot-level ones fetched/created above)
+		if (PrimaryJob)
+		{
+			constexpr bool bIsForPrimaryOverrides = true;
+			ShotVariableAssignments_PrimaryOverrides = Shot->GetOrCreateJobVariableAssignmentsForGraph(PrimaryJob->GetGraphPreset(), bIsForPrimaryOverrides);
+		}
 	}
 
-	if (InContext->Job)
+	if (PrimaryJob)
 	{
-		JobVariableAssignments = InContext->Job->GetOrCreateJobVariableAssignmentsForGraph(InContext->Job->GetGraphPreset());
+		JobVariableAssignments = PrimaryJob->GetOrCreateJobVariableAssignmentsForGraph(PrimaryJob->GetGraphPreset());
 	}
 	
 	// Check the shot job first for an enabled job variable assignment for this variable. Shot jobs take precedence over primary jobs.
diff --git a/Engine/Plugins/MovieScene/MovieRenderPipeline/Source/MovieRenderPipelineCore/Private/MoviePipelineQueue.cpp b/Engine/Plugins/MovieScene/MovieRenderPipeline/Source/MovieRenderPipelineCore/Private/MoviePipelineQueue.cpp
--- a/Engine/Plugins/MovieScene/MovieRenderPipeline/Source/MovieRenderPipelineCore/Private/MoviePipelineQueue.cpp
+++ b/Engine/Plugins/MovieScene/MovieRenderPipeline/Source/MovieRenderPipelineCore/Private/MoviePipelineQueue.cpp
@@ -264,27 +264,33 @@ void UMoviePipelineExecutorJob::PreSave(FObjectPreSaveContext ObjectSaveContext)
 void UMoviePipelineExecutorJob::PostLoad()
 {
 	Super::PostLoad();
-	
-	RefreshAllVariableAssignments();
 
 #if WITH_EDITOR
 	// Listen for any graph/subgraph saves. It's difficult/error-prone to track which graphs to monitor because the subgraph hierarchy can
 	// constantly change. Instead, just run any logic needed when any graph is saved.
 	FCoreUObjectDelegates::OnObjectPreSave.AddUObject(this, &UMoviePipelineExecutorJob::OnGraphPreSave);
 #endif
+
+	// Ideally the variable assignments would be updated here, but currently GraphPreset.LoadSynchronous() does not provide a fully-loaded
+	// object (GraphPreset will only be partially loaded at this point and its variables will be empty).
 }
 
 void UMoviePipelineExecutorJob::BeginDestroy()
 {
-	Super::BeginDestroy();
-
 #if WITH_EDITOR
 	FCoreUObjectDelegates::OnObjectPreSave.RemoveAll(this);
 #endif
+
+	Super::BeginDestroy();
 }
 
-TArray<TObjectPtr<UMovieJobVariableAssignmentContainer>>& UMoviePipelineExecutorJob::GetGraphVariableAssignments()
+TArray<TObjectPtr<UMovieJobVariableAssignmentContainer>>& UMoviePipelineExecutorJob::GetGraphVariableAssignments(const bool bUpdateAssignments)
 {
+	if (bUpdateAssignments)
+	{
+		RefreshAllVariableAssignments();
+	}
+	
 	return GraphVariableAssignments;
 }
 
@@ -456,8 +462,9 @@ void UMoviePipelineExecutorShot::SetShotOverridePresetOrigin(UMoviePipelineShotC
 bool UMoviePipelineExecutorShot::IsUsingGraphConfiguration() const
 {
 	const UMoviePipelineExecutorJob* PrimaryJob = GetTypedOuter<UMoviePipelineExecutorJob>();
-	
-	return GraphPreset.IsValid() || (PrimaryJob && PrimaryJob->IsUsingGraphConfiguration());
+
+	// Calling GetGraphPreset() here is important, rather than just referencing GraphPreset (to ensure that the soft ptr has a chance to load) 
+	return (GetGraphPreset() != nullptr) || (PrimaryJob && PrimaryJob->IsUsingGraphConfiguration());
 }
 
 void UMoviePipelineExecutorShot::PreSave(FObjectPreSaveContext ObjectSaveContext)
@@ -471,31 +478,42 @@ void UMoviePipelineExecutorShot::PostLoad()
 {
 	Super::PostLoad();
 
-	RefreshAllVariableAssignments();
-
 #if WITH_EDITOR
 	// Listen for any graph/subgraph saves. It's difficult/error-prone to track which graphs to monitor because the subgraph hierarchy can
 	// constantly change. Instead, just run any logic needed when any graph is saved.
 	FCoreUObjectDelegates::OnObjectPreSave.AddUObject(this, &UMoviePipelineExecutorShot::OnGraphPreSave);
 #endif
+
+	// Ideally the variable assignments would be updated here, but currently GraphPreset.LoadSynchronous() does not provide a fully-loaded
+	// object (GraphPreset will only be partially loaded at this point and its variables will be empty).
 }
 
 void UMoviePipelineExecutorShot::BeginDestroy()
 {
-	Super::BeginDestroy();
-
 #if WITH_EDITOR
 	FCoreUObjectDelegates::OnObjectPreSave.RemoveAll(this);
 #endif
+
+	Super::BeginDestroy();
 }
 
-TArray<TObjectPtr<UMovieJobVariableAssignmentContainer>>& UMoviePipelineExecutorShot::GetGraphVariableAssignments()
+TArray<TObjectPtr<UMovieJobVariableAssignmentContainer>>& UMoviePipelineExecutorShot::GetGraphVariableAssignments(const bool bUpdateAssignments)
 {
+	if (bUpdateAssignments)
+	{
+		RefreshAllVariableAssignments();
+	}
+	
 	return GraphVariableAssignments;
 }
 
-TArray<TObjectPtr<UMovieJobVariableAssignmentContainer>>& UMoviePipelineExecutorShot::GetPrimaryGraphVariableAssignments()
+TArray<TObjectPtr<UMovieJobVariableAssignmentContainer>>& UMoviePipelineExecutorShot::GetPrimaryGraphVariableAssignments(const bool bUpdateAssignments)
 {
+	if (bUpdateAssignments)
+	{
+		RefreshAllVariableAssignments();
+	}
+	
 	return PrimaryGraphVariableAssignments;
 }
 
diff --git a/Engine/Plugins/MovieScene/MovieRenderPipeline/Source/MovieRenderPipelineCore/Public/Graph/MovieGraphTraversalContext.h b/Engine/Plugins/MovieScene/MovieRenderPipeline/Source/MovieRenderPipelineCore/Public/Graph/MovieGraphTraversalContext.h
--- a/Engine/Plugins/MovieScene/MovieRenderPipeline/Source/MovieRenderPipelineCore/Public/Graph/MovieGraphTraversalContext.h
+++ b/Engine/Plugins/MovieScene/MovieRenderPipeline/Source/MovieRenderPipelineCore/Public/Graph/MovieGraphTraversalContext.h
@@ -9,6 +9,7 @@
 // Forward Declares
 class UMovieGraphConfig;
 class UMoviePipelineExecutorJob;
+class UMoviePipelineExecutorShot;
 
 
 USTRUCT(BlueprintType)
@@ -45,6 +46,10 @@ public:
 	UPROPERTY(BlueprintReadWrite, EditAnywhere, Category = "Movie Graph")
 	TObjectPtr<UMoviePipelineExecutorJob> Job;
 
+	/** The shot in the queue this traversal context is for (if any). */
+	UPROPERTY(BlueprintReadWrite, EditAnywhere, Category = "Movie Graph")
+	TObjectPtr<UMoviePipelineExecutorShot> Shot;
+
 	/** The root graph to start our traversal from. This could be a shared config for the whole job, or a shot-specific override. */
 	UPROPERTY(BlueprintReadWrite, EditAnywhere, Category = "Movie Graph")
 	TObjectPtr<UMovieGraphConfig> RootGraph;
diff --git a/Engine/Plugins/MovieScene/MovieRenderPipeline/Source/MovieRenderPipelineCore/Public/MoviePipelineQueue.h b/Engine/Plugins/MovieScene/MovieRenderPipeline/Source/MovieRenderPipelineCore/Public/MoviePipelineQueue.h
--- a/Engine/Plugins/MovieScene/MovieRenderPipeline/Source/MovieRenderPipelineCore/Public/MoviePipelineQueue.h
+++ b/Engine/Plugins/MovieScene/MovieRenderPipeline/Source/MovieRenderPipelineCore/Public/MoviePipelineQueue.h
@@ -186,13 +186,19 @@ public:
 	/**
 	 * Gets overrides on the variables in graph presets associated with this job. A job can have multiple graphs associated with it if the job's
 	 * assigned graph contains subgraphs. The job's graph and each subgraph will have an entry in the returned array.
+	 * 
+	 * The assignments are updated before they are returned. To opt-out of this, set bUpdateAssignments to false (generally not recommended unless
+	 * there is very specific behavior that is needed).
 	 */
-	TArray<TObjectPtr<UMovieJobVariableAssignmentContainer>>& GetGraphVariableAssignments();
+	TArray<TObjectPtr<UMovieJobVariableAssignmentContainer>>& GetGraphVariableAssignments(const bool bUpdateAssignments = true);
 
 	/**
 	 * Gets overrides on the variables in the primary graph (and its subgraphs) associated with this job.
+	 *
+	 * The assignments are updated before they are returned. To opt-out of this, set bUpdateAssignments to false (generally not recommended unless
+	 * there is very specific behavior that is needed).
 	 */
-	TArray<TObjectPtr<UMovieJobVariableAssignmentContainer>>& GetPrimaryGraphVariableAssignments();
+	TArray<TObjectPtr<UMovieJobVariableAssignmentContainer>>& GetPrimaryGraphVariableAssignments(const bool bUpdateAssignments = true);
 
 	/** Refreshes the variable assignments associated with this shot, both for the shot's own graph preset and the associated primary graph. */
 	void RefreshAllVariableAssignments();
@@ -440,7 +446,8 @@ public:
 	UFUNCTION(BlueprintPure, Category = "Movie Render Pipeline")
 	bool IsUsingGraphConfiguration() const
 	{
-		return GraphPreset.IsValid();
+		// Calling GetGraphPreset() here is important, rather than just referencing GraphPreset (to ensure that the soft ptr has a chance to load)
+		return GetGraphPreset() != nullptr;
 	}
 
 	/**
@@ -479,8 +486,11 @@ public:
 	/**
 	 * Gets overrides on the variables in graph presets associated with this job. A job can have multiple graphs associated with it if the job's
 	 * assigned graph contains subgraphs. The job's graph and each subgraph will have an entry in the returned array.
+	 *
+	 * The assignments are updated before they are returned. To opt-out of this, set bUpdateAssignments to false (generally not recommended unless
+	 * there is very specific behavior that is needed).
 	 */
-	TArray<TObjectPtr<UMovieJobVariableAssignmentContainer>>& GetGraphVariableAssignments();
+	TArray<TObjectPtr<UMovieJobVariableAssignmentContainer>>& GetGraphVariableAssignments(const bool bUpdateAssignments = true);
 
 	/** Refreshes the variable assignments associated with this job, both for the jobs's own graph preset and the associated shot graphs. */
 	void RefreshAllVariableAssignments();
diff --git a/Engine/Plugins/MovieScene/MovieRenderPipeline/Source/MovieRenderPipelineEditor/Private/Graph/SMovieGraphActiveRenderSettingsTabContent.cpp b/Engine/Plugins/MovieScene/MovieRenderPipeline/Source/MovieRenderPipelineEditor/Private/Graph/SMovieGraphActiveRenderSettingsTabContent.cpp
--- a/Engine/Plugins/MovieScene/MovieRenderPipeline/Source/MovieRenderPipelineEditor/Private/Graph/SMovieGraphActiveRenderSettingsTabContent.cpp
+++ b/Engine/Plugins/MovieScene/MovieRenderPipeline/Source/MovieRenderPipelineEditor/Private/Graph/SMovieGraphActiveRenderSettingsTabContent.cpp
@@ -459,6 +459,7 @@ void SMovieGraphActiveRenderSettingsTabContent::TraverseGraph()
 
 	FMovieGraphTraversalContext Context;
 	Context.Job = TraversalJob.Get();
+	// Context.Shot = ?	// TODO: The shot should be exposed in the UI at some point
 
 	// Traverse the graph, and update the root elements
 	FlattenedGraph = TStrongObjectPtr(CurrentGraph->CreateFlattenedGraph(Context, TraversalError));
