MovieGraph: Add support for ScriptExecutionNodes which can run code during a render. To use this feature, implement a class that inherits from UMovieGraphScriptBase (in either Python or C++), and then add a UMovieGraphExecuteScriptNode node to your Graph's configuration. Within the Script Node, specify the script that you want run.

At runtime, Movie Render Graph will create an instance of your script before the job starts. It will then persist this instance throughout the job until the end, at which point the instance will be destroyed. During the runtime, various callbacks will be called on your script and it will be provided data from the current graph.

The graph duplicates the Job, and Graph Configuration Assets before passing them to the scripting environment. This allows your script to make changes to these settings (such as overriding a variable on the job, or adding nodes to the graph) without having to worry about restoring these changes at the end of the job.

diff --git a/Engine/Plugins/MovieScene/MovieRenderPipeline/Source/MovieRenderPipelineCore/Private/Graph/MovieGraphPipeline.cpp b/Engine/Plugins/MovieScene/MovieRenderPipeline/Source/MovieRenderPipelineCore/Private/Graph/MovieGraphPipeline.cpp
--- a/Engine/Plugins/MovieScene/MovieRenderPipeline/Source/MovieRenderPipelineCore/Private/Graph/MovieGraphPipeline.cpp
+++ b/Engine/Plugins/MovieScene/MovieRenderPipeline/Source/MovieRenderPipelineCore/Private/Graph/MovieGraphPipeline.cpp
@@ -15,6 +15,7 @@
 #include "Graph/Nodes/MovieGraphSamplingMethodNode.h"
 #include "Graph/Nodes/MovieGraphGlobalOutputSettingNode.h"
 #include "Graph/Nodes/MovieGraphWarmUpSettingNode.h"
+#include "Graph/Nodes/MovieGraphExecuteScriptNode.h"
 #include "Graph/MovieGraphBlueprintLibrary.h"
 #include "MovieRenderPipelineCoreModule.h"
 #include "Misc/CoreDelegates.h"
@@ -23,7 +24,7 @@
 #include "RenderingThread.h"
 #include "ImageWriteQueue.h"
 #include "Modules/ModuleManager.h"
-
+#include "UObject/Package.h"
 
 FString UMovieGraphPipeline::DefaultPreviewWidgetAsset = TEXT("/MovieRenderPipeline/Blueprints/UI_MovieGraphPipelineScreenOverlay.UI_MovieGraphPipelineScreenOverlay_C");
 
@@ -80,16 +81,19 @@ void UMovieGraphPipeline::Initialize(UMoviePipelineExecutorJob* InJob, const FMo
 		// Called at the end of the frame after everything has been ticked and rendered for the frame.
 		FCoreDelegates::OnEndFrame.AddUObject(this, &UMovieGraphPipeline::OnEngineTickEndFrame);
 	}
+	CurrentJob = InJob;
+	CurrentShotIndex = 0;
+	GraphInitializationTime = FDateTime::UtcNow();
+	CVarManager = MakeShared<UE::MovieGraph::Private::FMovieGraphCVarManager>();
+
+	DuplicateJobAndConfiguration();
+	ExecutePreJobScripts();
 
 	// Create instances of our different classes from the InitConfig
 	GraphRendererInstance = NewObject<UMovieGraphRendererBase>(this, InitConfig.RendererClass);
 	GraphDataSourceInstance = NewObject<UMovieGraphDataSourceBase>(this, InitConfig.DataSourceClass);
 	GraphAudioRendererInstance = NewObject<UMovieGraphAudioRendererBase>(this, InitConfig.AudioRendererClass);
-	
-	CurrentJob = InJob;
-	CurrentShotIndex = 0;
-	GraphInitializationTime = FDateTime::UtcNow();
-	CVarManager = MakeShared<UE::MovieGraph::Private::FMovieGraphCVarManager>();
+
 
 	// Now that we've created our various systems, we will start using them. First thing we do is cache data about
 	// the world, job, player viewport, etc, before we make any modifications. These will be restored at the end
@@ -101,6 +105,7 @@ void UMovieGraphPipeline::Initialize(UMoviePipelineExecutorJob* InJob, const FMo
 
 	GraphAudioRendererInstance->SetupAudioRendering();
 
+
 	// Update our list of shots from our data source, and then
 	// create our list of active shots, so we don't try to render
 	// a shot the user has deactivated.
@@ -126,7 +131,7 @@ void UMovieGraphPipeline::Initialize(UMoviePipelineExecutorJob* InJob, const FMo
 	if (ActiveShotList.Num() == 0)
 	{
 		// We have to transition twice as Uninitialized -> n state is a no-op, so the second tick will take us to Finished which shuts down.
-		UE_LOG(LogMovieRenderPipeline, Warning, TEXT("MovieGraph No shots detected to render. Either all outside playback range, or disabled via shot mask, bailing."));
+		UE_LOG(LogMovieRenderPipeline, Warning, TEXT("MovieGraph No shots detected to render. Either all outside playback range, or disabled via shot mask. Shutting down."));
 	
 		TransitionToState(EMovieRenderPipelineState::Export);
 		TransitionToState(EMovieRenderPipelineState::Finished);
@@ -137,6 +142,70 @@ void UMovieGraphPipeline::Initialize(UMoviePipelineExecutorJob* InJob, const FMo
 	}
 }
 
+void UMovieGraphPipeline::DuplicateJobAndConfiguration()
+{
+	// Scripting is likely to want to modify both the job (to set variable assignments) and 
+	// the configuration itself (to add nodes, or override an output directory, etc. If scripts
+	// directly modified the job/configuration it would lead to a lot of unintentional mutation
+	// of assets and queues, so we instead choose to duplicate the job and configurations for
+	// the duration of a render, and all of the Graph Pipeline code should look at the duplicates.
+	FObjectDuplicationParameters JobDuplicationParms = FObjectDuplicationParameters(CurrentJob, GetTransientPackage());
+	CurrentJobDuplicate = Cast<UMoviePipelineExecutorJob>(StaticDuplicateObjectEx(JobDuplicationParms));
+
+	// The duplicate job is a mix of duplicated objects and non-duplicated objects. Objects that 
+	// don't have the job as their outer will not have been duplicate (which is good for the World/Sequence),
+	// but this also means that the graph configurations were not duplicated (as they are assets), so we need
+	// to manually duplicate them and update the pointers in the duplicated job.
+	FObjectDuplicationParameters PrimaryConfigDuplicationParams(CurrentJob->GetGraphPreset(), GetTransientPackage());
+	UMovieGraphConfig* DuplicatePrimaryConfig = Cast<UMovieGraphConfig>(StaticDuplicateObjectEx(PrimaryConfigDuplicationParams));
+	CurrentJobDuplicate->SetGraphPreset(DuplicatePrimaryConfig);
+	for (int32 Index = 0; Index < CurrentJob->ShotInfo.Num(); Index++)
+	{
+		// Now for each shot we need to duplicate its config (if any)
+		if (UMovieGraphConfig* ShotGraphPreset = CurrentJob->ShotInfo[Index]->GetGraphPreset())
+		{
+			// We use the transient package here and above because the _configs_ don't belong to the executor job usually (they belong
+			// to an asset package)
+			FObjectDuplicationParameters ShotConfigDuplicationParams(ShotGraphPreset, GetTransientPackage());
+			UMovieGraphConfig* DuplicateShotConfig = Cast<UMovieGraphConfig>(StaticDuplicateObjectEx(ShotConfigDuplicationParams));
+
+			CurrentJobDuplicate->ShotInfo[Index]->SetGraphPreset(DuplicateShotConfig);
+		}
+	}
+
+	// We only look in the primary configuration for the job for script nodes (and not shot specific overrides). If we looked
+	// in the shot-specific overrides, we would end up creating instances of the scripts for those shots. This can become really
+	// confusing if a shot specifies an Execute Script node with the same script as the Primary Graph, now you'll have two separate
+	// instances of your script, one which recieves 4 callbacks (pre/post job, pre/post shot) and one that only recieves pre/post shot.
+	FMovieGraphTraversalContext Context;
+	Context.Job = CurrentJobDuplicate;
+	Context.RootGraph = DuplicatePrimaryConfig;
+	FString OutError;
+
+	if (UMovieGraphEvaluatedConfig* FlattenedConfig = DuplicatePrimaryConfig->CreateFlattenedGraph(Context, OutError))
+	{
+		constexpr bool bIncludeCDOs = false;
+		TArray<UMovieGraphExecuteScriptNode*> ScriptNodes = FlattenedConfig->GetSettingsForBranch<UMovieGraphExecuteScriptNode>(UMovieGraphNode::GlobalsPinName, bIncludeCDOs);
+
+		// We now have the duplicated job, shots, and configs and we've fixed up the pointers, so we can now pass the duplicates
+		// to the user-defined callbacks, allowing them to modify most things without worrying about accidental mutation. If they
+		// choose to modify shared state (ie: the level sequence) then there's not much we can do.
+		for (const UMovieGraphExecuteScriptNode* Node : ScriptNodes)
+		{
+			// We instantiate an instance of each script and store it, and we'll call all subsequent callbacks on these (ignoring the nodes)
+			// so that we ensure that only one copy of the scripts exist and that they can store state during a render.
+			UMovieGraphScriptBase* ScriptInstance = Node->AllocateScriptInstance();
+		
+			// It's valid for a node to return a nullptr script instance (invalid class or no class set).
+			if (ScriptInstance)
+			{
+				CurrentScriptInstances.Add(ScriptInstance);
+			}
+		}
+	}
+
+}
+
 void UMovieGraphPipeline::LoadPreviewWidget()
 {
 	// ToDo: Allow overriding this widget so that users can style it to match their project.
@@ -602,6 +671,8 @@ void UMovieGraphPipeline::BeginExport()
 
 void UMovieGraphPipeline::SetupShot(const TObjectPtr<UMoviePipelineExecutorShot>& InShot)
 {
+	ExecutePreShotScripts(InShot);
+
 	// Set the new shot as the active shot. This enables the specified shot section and disables all other shot sections.
 	SetSoloShot(InShot);
 
@@ -721,6 +792,17 @@ void UMovieGraphPipeline::TeardownShot(const TObjectPtr<UMoviePipelineExecutorSh
 		GlobalGameOverridesNode->ApplySettings(bOverrideValues, GetWorld());
 	}
 
+	if (IsPostShotCallbacksNeeded())
+	{
+		ProcessOutstandingFinishedFrames();
+		ProcessOutstandingFutures();
+
+		// ToDo: Allow the command line encoder to modify the file list
+		// since it may add or remove files...
+
+		ExecutePostShotScripts();
+	}
+
 	// Check to see if this was the last shot in the Pipeline, otherwise on the next
 	// tick the new shot will be initialized and processed.
 	if (CurrentShotIndex >= (ActiveShotList.Num() - 1))
@@ -990,14 +1072,16 @@ void UMovieGraphPipeline::ShutdownImpl(bool bIsError)
 		// We were either in the middle of writing frames to disk, or we have moved to Finalize as a result of the above block.
 		// Tick output containers until they report they have finished writing to disk. This is a blocking operation. 
 		// Finalize automatically switches our state to Export so no need to manually transition afterwards.
-		//TickFinalizeOutputContainers(true);
+		constexpr bool bForceFinish = true;
+		TickFinalizeOutputContainers(bForceFinish);
 	}
 
 	if (PipelineState == EMovieRenderPipelineState::Export)
 	{
 		// All frames have been written to disk but we're doing a post-export step (such as encoding). Flush this operation as well.
 		// Export automatically switches our state to Finished so no need to manually transition afterwards.
-		TickPostFinalizeExport(true);
+		constexpr bool bForceFinish = true;
+		TickPostFinalizeExport(bForceFinish);
 	}
 }
 
@@ -1103,7 +1187,7 @@ void UMovieGraphPipeline::TransitionToState(const EMovieRenderPipelineState InNe
 			//	Setting->OnMoviePipelineShutdown(this);
 			//}
 
-
+			ExecutePostJobScripts();
 
 			// Ensure our delegates don't get called anymore as we're going to become null soon.
 			FCoreDelegates::OnBeginFrame.RemoveAll(this);
@@ -1176,18 +1260,23 @@ void UMovieGraphPipeline::OnMoviePipelineFinishedImpl()
 	// Broadcast to both Native and Python/BP
 	//
 	// Generate a params struct containing the data generated by this job.
-	FMoviePipelineOutputData Params;
-	//Params.Pipeline = this;
-	//Params.Job = GetCurrentJob();
-	//Params.bSuccess = !bFatalError;
-	//Params.ShotData = GeneratedShotOutputData;
+	FMovieGraphPipelineOutputData Params;
+	Params.Pipeline = this;
+	Params.Job = GetCurrentJob();
+	Params.bSuccess = !bShutdownSetErrorFlag;
+	Params.ShotData = GeneratedOutputData;
 	//
 	//UE_LOG(LogMovieRenderPipelineIO, Verbose, TEXT("Files written to disk for entire sequence:"));
 	//PrintVerboseLogForFiles(GeneratedShotOutputData);
 	//UE_LOG(LogMovieRenderPipelineIO, Verbose, TEXT("Completed outputting files written to disk."));
 	//
-	OnMoviePipelineWorkFinishedDelegateNative.Broadcast(Params);
-	// OnMoviePipelineWorkFinishedDelegate.Broadcast(Params);
+
+	// The callbacks are shared with the legacy system, but they have a different type due to render data identifiers.
+	// For now we need to at least call the callback (otherwise the executor never knows that the job is done).
+	FMoviePipelineOutputData DummyLegacyParams;
+	DummyLegacyParams.bSuccess = !bShutdownSetErrorFlag;
+	OnMoviePipelineWorkFinishedDelegateNative.Broadcast(DummyLegacyParams);
+	// OnMoviePipelineWorkFinishedDelegate.Broadcast(DummyLegacyParams);
 }
 
 //void UMovieGraphPipeline::PrintVerboseLogForFiles(const TArray<FMoviePipelineShotOutputData>& InOutputData) const
@@ -1396,3 +1485,74 @@ const TSet<TObjectPtr<UMovieGraphFileOutputNode>> UMovieGraphPipeline::GetOutput
 {
 	return OutputNodesDataSentTo;
 }
+
+
+bool UMovieGraphPipeline::IsPostShotCallbacksNeeded() const
+{
+	bool bAnyScriptNeedsCallbacks = false;
+	for (UMovieGraphScriptBase* Script : CurrentScriptInstances)
+	{
+		if (Script->IsPerShotCallbackNeeded())
+		{
+			bAnyScriptNeedsCallbacks = true;
+			break;
+		}
+	}
+
+	return bAnyScriptNeedsCallbacks;
+}
+
+void UMovieGraphPipeline::ExecutePreJobScripts()
+{
+	for (UMovieGraphScriptBase* Script : CurrentScriptInstances)
+	{
+		// GetCurrentJob returns the duplicated job and not the original.
+		Script->OnJobStart(GetCurrentJob());
+	}
+}
+
+void UMovieGraphPipeline::ExecutePreShotScripts(const TObjectPtr<UMoviePipelineExecutorShot>& InShot)
+{
+	for (UMovieGraphScriptBase* Script : CurrentScriptInstances)
+	{
+		if (Script->IsPerShotCallbackNeeded())
+		{
+			Script->OnShotStart(GetCurrentJob(), InShot);
+		}
+	}
+}
+
+void UMovieGraphPipeline::ExecutePostShotScripts()
+{
+	FMovieGraphPipelineOutputData Params;
+	Params.Pipeline = this;
+	Params.Job = GetCurrentJob();
+	Params.bSuccess = !bShutdownSetErrorFlag;
+
+	// We only provide data that the current shot generated during post-shot callbacks.
+	TArray< FMovieGraphRenderOutputData> SingleOutputData;
+	SingleOutputData.Add(GeneratedOutputData[CurrentShotIndex]);
+	Params.ShotData = SingleOutputData;
+
+	for (UMovieGraphScriptBase* Script : CurrentScriptInstances)
+	{
+		if (Script->IsPerShotCallbackNeeded())
+		{
+			Script->OnShotFinished(GetCurrentJob(), ActiveShotList[CurrentShotIndex], Params);
+		}
+	}
+}
+
+void UMovieGraphPipeline::ExecutePostJobScripts()
+{
+	FMovieGraphPipelineOutputData Params;
+	Params.Pipeline = this;
+	Params.Job = GetCurrentJob();
+	Params.bSuccess = !bShutdownSetErrorFlag;
+	Params.ShotData = GeneratedOutputData;
+
+	for (UMovieGraphScriptBase* Script : CurrentScriptInstances)
+	{
+		Script->OnJobFinished(GetCurrentJob(), Params);
+	}
+}
\ No newline at end of file
diff --git a/Engine/Plugins/MovieScene/MovieRenderPipeline/Source/MovieRenderPipelineCore/Private/Graph/Nodes/MovieGraphExecuteScriptNode.cpp b/Engine/Plugins/MovieScene/MovieRenderPipeline/Source/MovieRenderPipelineCore/Private/Graph/Nodes/MovieGraphExecuteScriptNode.cpp
new file mode 100644
--- /dev/null
+++ b/Engine/Plugins/MovieScene/MovieRenderPipeline/Source/MovieRenderPipelineCore/Private/Graph/Nodes/MovieGraphExecuteScriptNode.cpp
@@ -0,0 +1,43 @@
+// Copyright Epic Games, Inc. All Rights Reserved.
+
+#include "Graph/Nodes/MovieGraphExecuteScriptNode.h"
+#include "Styling/AppStyle.h"
+#include "UObject/Package.h"
+
+#if WITH_EDITOR
+FText UMovieGraphExecuteScriptNode::GetNodeTitle(const bool bGetDescriptive) const
+{
+	return NSLOCTEXT("MovieGraphNodes", "ExecuteScriptNode_Description", "Execute Script");
+}
+
+FText UMovieGraphExecuteScriptNode::GetMenuCategory() const
+{
+	static const FText NodeCategory_Globals = NSLOCTEXT("MovieGraphNodes", "NodeCategory_Globals", "Globals");
+	return NodeCategory_Globals;
+}
+
+FLinearColor UMovieGraphExecuteScriptNode::GetNodeTitleColor() const
+{
+	return FLinearColor(0.1f, 0.1f, 0.85f);
+}
+
+FSlateIcon UMovieGraphExecuteScriptNode::GetIconAndTint(FLinearColor& OutColor) const
+{
+	static const FSlateIcon DeferredRendererIcon = FSlateIcon(FAppStyle::GetAppStyleSetName(), "SequenceRecorder.TabIcon");
+
+	OutColor = FLinearColor::White;
+	return DeferredRendererIcon;
+}
+#endif	// WITH_EDITOR
+
+UMovieGraphScriptBase* UMovieGraphExecuteScriptNode::AllocateScriptInstance() const
+{
+	UClass* NewInstanceType = Script.TryLoadClass<UMovieGraphScriptBase>();
+	if (NewInstanceType)
+	{
+		UMovieGraphScriptBase* NewInstance = NewObject<UMovieGraphScriptBase>(GetTransientPackage(), NewInstanceType);
+		return NewInstance;
+	}
+
+	return nullptr;
+}
\ No newline at end of file
diff --git a/Engine/Plugins/MovieScene/MovieRenderPipeline/Source/MovieRenderPipelineCore/Public/Graph/MovieGraphDataTypes.h b/Engine/Plugins/MovieScene/MovieRenderPipeline/Source/MovieRenderPipelineCore/Public/Graph/MovieGraphDataTypes.h
--- a/Engine/Plugins/MovieScene/MovieRenderPipeline/Source/MovieRenderPipelineCore/Public/Graph/MovieGraphDataTypes.h
+++ b/Engine/Plugins/MovieScene/MovieRenderPipeline/Source/MovieRenderPipelineCore/Public/Graph/MovieGraphDataTypes.h
@@ -118,6 +118,44 @@ public:
 	TMap<FMovieGraphRenderDataIdentifier, FMovieGraphRenderLayerOutputData> RenderLayerData;
 };
 
+USTRUCT(BlueprintType)
+struct MOVIERENDERPIPELINECORE_API FMovieGraphPipelineOutputData
+{
+	GENERATED_BODY()
+
+		FMovieGraphPipelineOutputData()
+		: Pipeline(nullptr)
+		, Job(nullptr)
+		, bSuccess(false)
+	{}
+
+	/**
+	* The UMoviePipeline instance that generated this data. This is only provided as an id (in the event you were the one who created
+	* the UMoviePipeline instance.)
+	*
+	* DO NOT CALL FUNCTIONS ON THIS (unless you know what you're doing)
+	*
+	* Provided here for backwards compatibility.
+	*/
+	UPROPERTY(VisibleAnywhere, BlueprintReadOnly, Transient, Category = "Movie Pipeline")
+	TObjectPtr<UMovieGraphPipeline> Pipeline;
+
+	/** Job the data is for. Job may still be in progress (if a shot callback) so be careful about modifying properties on it */
+	UPROPERTY(VisibleAnywhere, BlueprintReadOnly, Transient, Category = "Movie Pipeline")
+	TObjectPtr<UMoviePipelineExecutorJob> Job;
+
+	/** Did the job succeed, or was it canceled early due to an error (such as failure to write file to disk)? */
+	UPROPERTY(VisibleAnywhere, BlueprintReadOnly, Category = "Movie Pipeline")
+	bool bSuccess;
+
+	/**
+	* The file data for each shot that was rendered. If no files were written this will be empty. If this is from the per-shot work
+	* finished callback it will only have one entry (for the just finished shot). Will not include shots that did not get rendered
+	* due to the pipeline encountering an error.
+	*/
+	UPROPERTY(VisibleAnywhere, BlueprintReadOnly, Category = "Movie Pipeline")
+	TArray<FMovieGraphRenderOutputData> ShotData;
+};
 
 UCLASS(BlueprintType, Abstract)
 class MOVIERENDERPIPELINECORE_API UMovieGraphTimeStepBase : public UObject
diff --git a/Engine/Plugins/MovieScene/MovieRenderPipeline/Source/MovieRenderPipelineCore/Public/Graph/MovieGraphPipeline.h b/Engine/Plugins/MovieScene/MovieRenderPipeline/Source/MovieRenderPipelineCore/Public/Graph/MovieGraphPipeline.h
--- a/Engine/Plugins/MovieScene/MovieRenderPipeline/Source/MovieRenderPipelineCore/Public/Graph/MovieGraphPipeline.h
+++ b/Engine/Plugins/MovieScene/MovieRenderPipeline/Source/MovieRenderPipelineCore/Public/Graph/MovieGraphPipeline.h
@@ -13,6 +13,7 @@
 class UMoviePipelineExecutorJob;
 class UMoviePipelineExecutorShot;
 class UMovieGraphFileOutputNode;
+class UMovieGraphScriptBase;
 class IImageWriteQueue;
 struct FMovieGraphRenderOutputData;
 
@@ -37,8 +38,13 @@ public:
 	UFUNCTION(BlueprintCallable, Category = "Movie Graph")
 	void Initialize(UMoviePipelineExecutorJob* InJob, const FMovieGraphInitConfig& InitConfig);
 
+	/**
+	* Returns the internal job being used by the Movie Graph Pipeline, which is a duplicate of
+	* the job provided originaly by Initialize. This allows scripting to mutate the job/configuration
+	* without leaking changes into assets or the original user-defined queue entry.
+	*/
 	UFUNCTION(BlueprintCallable, Category = "Movie Graph")
-	UMoviePipelineExecutorJob* GetCurrentJob() const { return CurrentJob; }
+	UMoviePipelineExecutorJob* GetCurrentJob() const { return CurrentJobDuplicate; }
 
 	/**
 	* Returns the time this movie pipeline was initialized at.
@@ -143,6 +149,13 @@ protected:
 	virtual void LoadPreviewWidget();
 	virtual void SetPreviewWidgetVisibleImpl(bool bInIsVisible);
 
+	virtual void DuplicateJobAndConfiguration();
+	virtual void ExecutePreJobScripts();
+	virtual void ExecutePostJobScripts();
+	virtual void ExecutePreShotScripts(const TObjectPtr<UMoviePipelineExecutorShot>& InShot);
+	virtual void ExecutePostShotScripts();
+	virtual bool IsPostShotCallbacksNeeded() const;
+
 	// Update our data source to isolate the shot we're currently working on, so that expanded shots don't interfere with each other.
 	virtual void SetSoloShot(const TObjectPtr<UMoviePipelineExecutorShot>& InShot);
 	
@@ -201,9 +214,23 @@ protected:
 	TObjectPtr<UMovieGraphEvaluatedConfig> PostRenderEvaluatedGraph;
 
 protected:
+	/**
+	* This is the job as was provided to the ::Initialize() call, which we duplicate
+	* to allow modifications to scripting without leaking changes into assets.
+	*/
 	UPROPERTY(Transient)
 	TObjectPtr<UMoviePipelineExecutorJob> CurrentJob;
 
+	/**
+	* This is the duplicated job, parented to the Transient package. The shots inside
+	* have been duplicated as well, and their graph configurations duplicated. Graph
+	* configurations are assets and scripting may want to modify them, or it may want
+	* to modify the variables in a job, so we have to duplicate both to allow a cleanly
+	* mutable version for scripting.
+	*/
+	UPROPERTY(Transient)
+	TObjectPtr<UMoviePipelineExecutorJob> CurrentJobDuplicate;
+
 	UPROPERTY(Transient)
 	TArray<TObjectPtr<UMoviePipelineExecutorShot>> ActiveShotList;
 
@@ -213,6 +240,10 @@ protected:
 	UPROPERTY(Transient)
 	TObjectPtr<UMovieGraphRenderPreviewWidget> PreviewWidget;
 
+
+	UPROPERTY(Transient)
+	TArray<UMovieGraphScriptBase*> CurrentScriptInstances;
+
 	/**
 	* An array of Node CDOs that we sent data through to write data to disk.
 	* This is the list of output nodes that will have OnAllFramesSubmitted/IsFinishedWritingToDisk
diff --git a/Engine/Plugins/MovieScene/MovieRenderPipeline/Source/MovieRenderPipelineCore/Public/Graph/Nodes/MovieGraphExecuteScriptNode.h b/Engine/Plugins/MovieScene/MovieRenderPipeline/Source/MovieRenderPipelineCore/Public/Graph/Nodes/MovieGraphExecuteScriptNode.h
new file mode 100644
--- /dev/null
+++ b/Engine/Plugins/MovieScene/MovieRenderPipeline/Source/MovieRenderPipelineCore/Public/Graph/Nodes/MovieGraphExecuteScriptNode.h
@@ -0,0 +1,210 @@
+// Copyright Epic Games, Inc. All Rights Reserved.
+
+#pragma once
+
+#include "Graph/MovieGraphNode.h"
+#include "MovieRenderPipelineDataTypes.h"
+#include "MovieGraphExecuteScriptNode.generated.h"
+
+/**
+* This is the base class used for writing scripts that can be executed as part
+* of a Movie Graph Pipeline. After creating a UMovieGraphExecuteScriptNode and
+* choosing the type of your script, at runtime, an instance will be created when
+* the job starts. This instance will persist until the end of the job. The instance
+* will recieve callbacks during the render which contain data that you can 
+* manipulate with your script, or data that was generated by the pipeline.
+*
+* This can be implemented in either C++ or Python. For C++ you can use standard
+* inheritance but for Python you need to write a UClass implemented in Python.
+* This has slightly different syntax than standard Python;
+* 
+* In this example, create a "Python" folder in your Content folder. Then create
+* a file named "MovieGraphPipelineScriptExample.py", and a file named "init_unreal.py"
+* 
+* In "init_unreal.py" add "import MovieGraphPipelineScriptExample" as the contents.
+* In "MovieGraphPipelineScriptExample.py" use the following:
+*
+*import unreal
+*
[REDACTED_EMAIL]()
+*class MovieGraphPipelineScriptExample(unreal.MovieGraphScriptBase):
+
+*    @unreal.ufunction(override=True)
+*    def on_job_start(self, inJobCopy):
+*        super().on_job_start(inJobCopy)
+*        unreal.log_warning("OnJobStart")
+* 
+* Save and restart Unreal Editor, then create a Movie Graph Configuration
+* with a Script Node, and set it to run MovieGraphPipelineScriptExample. When
+* the job is run, the output console should print "OnJobStart" as a warning.
+* See /Engine/Plugins/MovieScene/MovieRenderPipeline/Content/Python for more 
+* python examples.
+*/
+UCLASS(BlueprintType, Abstract)
+class MOVIERENDERPIPELINECORE_API UMovieGraphScriptBase : public UObject
+{
+	GENERATED_BODY()
+
+public:
+	/**
+	* This callback is called very early on in the Movie Graph Pipeline, before
+	* data is read from the job. The provided job is a duplicate of the job defined
+	* in the UMoviePipelineQueue, allowing you to make changes to the job without
+	* worrying about restoring them at the end of a render. Additionally, the 
+	* Graph Configuration assets (for both the Primary configuration, and shot
+	* overrides) have been duplicated to allow mutation of the graph configuration
+	* without worrying about leaking these changes to the assets on disk. Other
+	* pointers to external packages (such as the Level Sequence on the job) have
+	* not been duplicated though, so be aware that changes to these assets will
+	* leak, unless you restore the changes in OnJobFinished. OnJobFinished will
+	* be called even if the job is canceled.
+	* 
+	* For C++ implementations override `virtual void OnJobStart_Implementation(UMoviePipelineExecutorJob* InJobCopy) override`
+	* For Python implementations override
+	*	@unreal.ufunction(override=True)
+	*	def on_job_start(self, inJobCopy):
+	*		super().on_job_start(inJobCopy)
+	* 
+	* 
+	* @param InJobCopy - A copy of the job being rendered with this graph config,
+	*					 with its Graph Preset configurations also copied to avoid
+	*					 accidental mutation of assests on disk.
+	*/
+	UFUNCTION(BlueprintNativeEvent, Category = "Movie Graph")
+	void OnJobStart(UMoviePipelineExecutorJob* InJobCopy);
+	virtual void OnJobStart_Implementation(UMoviePipelineExecutorJob* InJobCopy) {}
+
+	/**
+	* This callback is called at the very end of the Movie Graph Pipeline, right before
+	* shutting down. The files specified by the OutputData should have been written
+	* to disk by this point. This function will get called both in the case of a successful
+	* job, or one canceled by the user. You can use OutputData.bSuccess to determine if
+	* the render was successfuly completed.
+	* 
+	* For C++ implementations override `virtual void OnJobFinished_Implementation(UMoviePipelineExecutorJob* InJobCopy, const FMovieGraphPipelineOutputData& OutputData) override`
+	* For Python implementations override
+	*	@unreal.ufunction(override=True)
+	*	def on_job_finished(self, inJobCopy, inOutputData):
+	*		super().on_job_finished(inJobCopy, inOutputData)
+	* 
+	* 
+	* @param InJobCopy - A copy of the job being rendered with this graph config,
+	*					 with its Graph Preset configurations also copied to avoid
+	*					 accidental mutation of assests on disk.
+	* @param InOutputData - Information about the files generated during this render.
+	* 						Contains a flag (bSuccess) for checking if the job was 
+	*						successful or not, and contains a list of render layers,
+	* 						and the files written to disk for that render layer. Contains
+	*						data for all shots rendered up to this point.
+	*/
+	UFUNCTION(BlueprintNativeEvent, Category = "Movie Graph")
+	void OnJobFinished(UMoviePipelineExecutorJob* InJobCopy, const FMovieGraphPipelineOutputData& InOutputData);
+	virtual void OnJobFinished_Implementation(UMoviePipelineExecutorJob* InJobCopy, const FMovieGraphPipelineOutputData& InOutputData ) {}
+
+	/**
+	* This is similar to OnJobStart but called before a shot is set up. This is only called
+	* if IsPerShotCallbackNeeded() returns true.
+	* 
+	* For C++ implementations override `virtual void OnShotStart_Implementation(UMoviePipelineExecutorJob* InJobCopy,  UMoviePipelineExecutorShot* InShotCopy) override`
+	* For Python implementations override
+	*	@unreal.ufunction(override=True)
+	*	def on_shot_start(self, inJobCopy, inShotCopy):
+	*		super().on_shot_start(inJobCopy, inShotCopy)
+	* 
+	* 
+	* @param InJobCopy - A copy of the job being rendered with this graph config,
+	*					 with its Graph Preset configurations also copied to avoid
+	*					 accidental mutation of assests on disk.
+	* @param InShotCopy - A copy of the shot that is about to be started.
+	*/
+	UFUNCTION(BlueprintNativeEvent, Category = "Movie Graph")
+	void OnShotStart(UMoviePipelineExecutorJob* InJobCopy, UMoviePipelineExecutorShot* InShotCopy);
+	virtual void OnShotStart_Implementation(UMoviePipelineExecutorJob* InJobCopy, UMoviePipelineExecutorShot* InShotCopy) {}
+
+	/**
+	* This is similar to OnJobFinished but called after a shot is completed. This is only called
+	* if IsPerShotCallbackNeeded() returns true.
+	* 
+	* For C++ implementations override `virtual void OnShotFinished_Implementation(UMoviePipelineExecutorJob* InJobCopy,  UMoviePipelineExecutorShot* InShotCopy, const FMovieGraphPipelineOutputData& InOutputData) override`
+	* For Python implementations override
+	*	@unreal.ufunction(override=True)
+	*	def on_shot_finished(self, inJobCopy, inShotCopy, inOutputData):
+	*		super().on_shot_finished(inJobCopy, inShotCopy, inOutputData)
+	* 
+	* 
+	* @param InJobCopy - A copy of the job being rendered with this graph config,
+	*					 with its Graph Preset configurations also copied to avoid
+	*					 accidental mutation of assests on disk.
+	* @param InShotCopy - A copy of the shot that was just finished.
+	* @param InOutputData - The data generated by this shot. Only contains the files
+	* 						rendered for this particular shot, not all files generated
+	* 						by the movie render (use OnJobFinished for that).
+	*/
+	UFUNCTION(BlueprintNativeEvent, Category = "Movie Graph")
+	void OnShotFinished(UMoviePipelineExecutorJob* InJobCopy, UMoviePipelineExecutorShot* InShotCopy, const FMovieGraphPipelineOutputData& OutputData);
+	virtual void OnShotFinished_Implementation(UMoviePipelineExecutorJob* InJobCopy, UMoviePipelineExecutorShot* InShotCopy, const FMovieGraphPipelineOutputData& OutputData) {}
+
+
+	/**
+	* If you return true here the system will stall at the end of every shot and wait until
+	* all files have been written to disk before invoking this callback. This can increase 
+	* render times but is important if you need to operate on the files on disk as part of 
+	* the callback.
+	*
+	* For C++ implementations override `virtual bool IsPerShotCallbackNeeded_Implementation() const override`
+	* For Python implementations override
+	*	@unreal.ufunction(override=True)
+	*	def is_per_shot_callback_needed():
+	*		return False;
+	*/
+	UFUNCTION(BlueprintNativeEvent, Category = "Movie Graph")
+	bool IsPerShotCallbackNeeded() const;
+	virtual bool IsPerShotCallbackNeeded_Implementation() const
+	{
+		return false;
+	}
+};
+
+/** 
+ This node allows users to run code before and after a movie graph is being executed. Example use cases
+ would be to run custom logic that overrides the output directory, or to trigger an event after a shot finishes.
+ These bits of script can be implemented in either C++, or in Python (as a UClass) because we don't serialize the
+ object and only create it at runtime during the render. See UMovieGraphScriptBase for details.
+*/
+UCLASS()
+class MOVIERENDERPIPELINECORE_API UMovieGraphExecuteScriptNode : public UMovieGraphSettingNode
+{
+	GENERATED_BODY()
+
+public:
+	UMovieGraphExecuteScriptNode() = default;
+
+	virtual EMovieGraphBranchRestriction GetBranchRestriction() const override { return EMovieGraphBranchRestriction::Globals; }
+
+#if WITH_EDITOR
+	virtual FText GetNodeTitle(const bool bGetDescriptive = false) const override;
+	virtual FText GetMenuCategory() const override;
+	virtual FLinearColor GetNodeTitleColor() const override;
+	virtual FSlateIcon GetIconAndTint(FLinearColor& OutColor) const override;
+#endif
+
+	/** 
+	* Allocate an instance of the specifed script type.
+	* Can return nullptr.
+	*/
+	UMovieGraphScriptBase* AllocateScriptInstance() const;
+
+public:
+	UPROPERTY(EditAnywhere, BlueprintReadWrite, Category = Overrides, meta = (InlineEditConditionToggle))
+	uint8 bOverride_Script: 1;
+	
+
+	/**
+	* Which script should this node instantiate? The class will be instantiated when the job starts,
+	* and then different functions will be called on that instance throughout the life-time of the 
+	* movie render. Scripts can be implemented in Python or C++, see MovieGraphExecuteScriptNode.h
+	* for more details.
+	*/
+	UPROPERTY(EditAnywhere, BlueprintReadWrite, Category = "Settings", meta = (EditCondition = "bOverride_Script", MetaClass = "/Script/MovieRenderPipelineCore.MovieGraphScriptBase"))
+	FSoftClassPath Script;
+};
\ No newline at end of file
