Movie Graph:

-Fix: Reset to Default doesn't work on FileNameFormat property in FileOutput and CommandLineEncoder nodes

diff --git a/Engine/Plugins/MovieScene/MovieRenderPipeline/Source/MovieRenderPipelineEditor/Private/Customizations/Graph/MovieGraphFormatTokenCustomization.cpp b/Engine/Plugins/MovieScene/MovieRenderPipeline/Source/MovieRenderPipelineEditor/Private/Customizations/Graph/MovieGraphFormatTokenCustomization.cpp
--- a/Engine/Plugins/MovieScene/MovieRenderPipeline/Source/MovieRenderPipelineEditor/Private/Customizations/Graph/MovieGraphFormatTokenCustomization.cpp
+++ b/Engine/Plugins/MovieScene/MovieRenderPipeline/Source/MovieRenderPipelineEditor/Private/Customizations/Graph/MovieGraphFormatTokenCustomization.cpp
@@ -17,25 +17,41 @@ void FMovieGraphFormatTokenCustomization::CustomizeDetails(IDetailLayoutBuilder&
 
 	// This should work just fine for UMovieGraphFileOutputNode and UMovieGraphCommandLineEncoderNode as long as the property names stay in sync
 	OutputFormatPropertyHandle = DetailBuilder.GetProperty(GET_MEMBER_NAME_CHECKED(UMovieGraphFileOutputNode, FileNameFormat));
-	const bool bShowChildren = false;
-	FText InitialText;
-	OutputFormatPropertyHandle->GetValueAsDisplayText(InitialText);
 
-	DetailBuilder.EditDefaultProperty(OutputFormatPropertyHandle)->CustomWidget(bShowChildren)
+	FText StartingDisplayText;
+	OutputFormatPropertyHandle->GetValueAsDisplayText(StartingDisplayText);
+
+	// Update the text box when the handle value changes
+	OutputFormatPropertyHandle->SetOnPropertyValueChanged(
+		FSimpleDelegate::CreateSP(this, &FMovieGraphFormatTokenCustomization::OnPropertyChange));
+
+	DetailBuilder.EditDefaultProperty(OutputFormatPropertyHandle)->CustomWidget()
 		.NameContent()
 		[
 			OutputFormatPropertyHandle->CreatePropertyNameWidget()
 		]
 		.ValueContent()
 		.MinDesiredWidth(200.0f)
-		[
-			SNew(SMoviePipelineFormatTokenAutoCompleteBox)
-			.Text(this, &FMovieGraphFormatTokenCustomization::GetText)
+	[
+			// We choose not to bind the text box text here because simultaneously setting and getting 
+			// the handle value can cause the binding to stop working, so we manually update it when necessary 
+			SAssignNew(AutoCompleteBox, SMoviePipelineFormatTokenAutoCompleteBox)
+			.InitialText(StartingDisplayText)
 			.Suggestions(this, &FMovieGraphFormatTokenCustomization::GetSuggestions)
 			.OnTextChanged(this, &FMovieGraphFormatTokenCustomization::OnTextChanged)
 		];
 }
 
+void FMovieGraphFormatTokenCustomization::OnPropertyChange()
+{
+	// Sync the text box back to the handle value
+	FText DisplayText;
+	OutputFormatPropertyHandle->GetValueAsDisplayText(DisplayText);
+
+	AutoCompleteBox->SetText(DisplayText);
+	AutoCompleteBox->CloseMenuAndReset();
+}
+
 TArray<FString> FMovieGraphFormatTokenCustomization::GetSuggestions() const
 {
 	TArray<FString> Suggestions;
@@ -50,18 +66,11 @@ TArray<FString> FMovieGraphFormatTokenCustomization::GetSuggestions() const
 	return Suggestions;
 }
 
-void FMovieGraphFormatTokenCustomization::OnTextChanged(const FText& InValue) const
+void FMovieGraphFormatTokenCustomization::OnTextChanged(const FText& InValue)
 {
 	OutputFormatPropertyHandle->SetValue(InValue.ToString());
 }
 
-FText FMovieGraphFormatTokenCustomization::GetText() const
-{
-	FText DisplayText;
-	OutputFormatPropertyHandle->GetValueAsDisplayText(DisplayText);
-	return DisplayText;
-}
-
 void FMovieGraphFormatTokenCustomization::GetFormatArguments(FMoviePipelineFormatArgs& InOutFormatArgs)
 {
 	static const FString LevelName = TEXT("Level Name");
diff --git a/Engine/Plugins/MovieScene/MovieRenderPipeline/Source/MovieRenderPipelineEditor/Private/Customizations/Graph/MovieGraphFormatTokenCustomization.h b/Engine/Plugins/MovieScene/MovieRenderPipeline/Source/MovieRenderPipelineEditor/Private/Customizations/Graph/MovieGraphFormatTokenCustomization.h
--- a/Engine/Plugins/MovieScene/MovieRenderPipeline/Source/MovieRenderPipelineEditor/Private/Customizations/Graph/MovieGraphFormatTokenCustomization.h
+++ b/Engine/Plugins/MovieScene/MovieRenderPipeline/Source/MovieRenderPipelineEditor/Private/Customizations/Graph/MovieGraphFormatTokenCustomization.h
@@ -21,14 +21,15 @@ protected:
 	virtual void CustomizeDetails(IDetailLayoutBuilder& DetailBuilder) override;
 	//~ End IDetailCustomization interface
 
-	TArray<FString> GetSuggestions() const;
+	void OnPropertyChange();
 
-	void OnTextChanged(const FText& InValue) const;
+	TArray<FString> GetSuggestions() const;
 
-	FText GetText() const;
+	void OnTextChanged(const FText& InValue);
 
 	static void GetFormatArguments(FMoviePipelineFormatArgs& InOutFormatArgs);
 	
 	TWeakObjectPtr<UObject> CustomizedObject;
 	TSharedPtr<IPropertyHandle> OutputFormatPropertyHandle;
+	TSharedPtr<class SMoviePipelineFormatTokenAutoCompleteBox> AutoCompleteBox;
 };
diff --git a/Engine/Plugins/MovieScene/MovieRenderPipeline/Source/MovieRenderPipelineEditor/Private/Widgets/SMoviePipelineConfigEditor.cpp b/Engine/Plugins/MovieScene/MovieRenderPipeline/Source/MovieRenderPipelineEditor/Private/Widgets/SMoviePipelineConfigEditor.cpp
--- a/Engine/Plugins/MovieScene/MovieRenderPipeline/Source/MovieRenderPipelineEditor/Private/Widgets/SMoviePipelineConfigEditor.cpp
+++ b/Engine/Plugins/MovieScene/MovieRenderPipeline/Source/MovieRenderPipelineEditor/Private/Widgets/SMoviePipelineConfigEditor.cpp
@@ -71,6 +71,10 @@ protected:
 		FText InitialText;
 		OutputFormatPropertyHandle->GetValueAsDisplayText(InitialText);
 
+		// Update the text box when the handle value changes
+		OutputFormatPropertyHandle->SetOnPropertyValueChanged(
+			FSimpleDelegate::CreateSP(this, &FOutputFormatDetailsCustomization::OnPropertyChange));
+
 		DetailBuilder.EditDefaultProperty(OutputFormatPropertyHandle)->CustomWidget(bShowChildren)
 			.NameContent()
 			[
@@ -79,8 +83,8 @@ protected:
 			.ValueContent()
 			.MinDesiredWidth(200.0f)
 			[
-				SNew(SMoviePipelineFormatTokenAutoCompleteBox)
-				.Text(this, &FOutputFormatDetailsCustomization::GetText)
+				SAssignNew(AutoCompleteBox, SMoviePipelineFormatTokenAutoCompleteBox)
+				.InitialText(InitialText)
 				.Suggestions(this, &FOutputFormatDetailsCustomization::GetSuggestions)
 				.OnTextChanged(this, &FOutputFormatDetailsCustomization::OnTextChanged)
 			];
@@ -89,6 +93,16 @@ protected:
 	}
 	//~ End IDetailCustomization interface
 
+	void OnPropertyChange()
+	{
+		// Sync the text box back to the handle value
+		FText DisplayText;
+		OutputFormatPropertyHandle->GetValueAsDisplayText(DisplayText);
+
+		AutoCompleteBox->SetText(DisplayText);
+		AutoCompleteBox->CloseMenuAndReset();
+	}
+
 	TArray<FString> GetSuggestions() const
 	{
 		TSharedPtr<SMoviePipelineConfigEditor> ConfigEditor = OwningEditor.Pin();
@@ -125,15 +139,9 @@ protected:
 		OutputFormatPropertyHandle->SetValue(InValue.ToString());
 	}
 
-	FText GetText() const
-	{
-		FText DisplayText;
-		OutputFormatPropertyHandle->GetValueAsDisplayText(DisplayText);
-		return DisplayText;
-	}
-
 	TSharedPtr<IPropertyHandle> OutputFormatPropertyHandle;
 	TWeakPtr<SMoviePipelineConfigEditor> OwningEditor;
+	TSharedPtr<SMoviePipelineFormatTokenAutoCompleteBox> AutoCompleteBox;
 };
 
 /**
diff --git a/Engine/Plugins/MovieScene/MovieRenderPipeline/Source/MovieRenderPipelineEditor/Private/Widgets/SMoviePipelineFormatTokenAutoCompleteBox.cpp b/Engine/Plugins/MovieScene/MovieRenderPipeline/Source/MovieRenderPipelineEditor/Private/Widgets/SMoviePipelineFormatTokenAutoCompleteBox.cpp
--- a/Engine/Plugins/MovieScene/MovieRenderPipeline/Source/MovieRenderPipelineEditor/Private/Widgets/SMoviePipelineFormatTokenAutoCompleteBox.cpp
+++ b/Engine/Plugins/MovieScene/MovieRenderPipeline/Source/MovieRenderPipelineEditor/Private/Widgets/SMoviePipelineFormatTokenAutoCompleteBox.cpp
@@ -13,7 +13,7 @@ void SMoviePipelineFormatTokenAutoCompleteBox::Construct(const FArguments& InArg
 		.Placement(MenuPlacement_ComboBox)
 		[
 			SAssignNew(TextBox, SMultiLineEditableTextBox)
-			.Text(InArgs._Text)
+			.Text(InArgs._InitialText)
 			.Font(IDetailLayoutBuilder::GetDetailFont())
 			.OnKeyDownHandler(this, &SMoviePipelineFormatTokenAutoCompleteBox::OnKeyDown)
 			.OnTextChanged(this, &SMoviePipelineFormatTokenAutoCompleteBox::HandleTextBoxTextChanged)
@@ -103,6 +103,11 @@ FReply SMoviePipelineFormatTokenAutoCompleteBox::OnKeyDown(const FGeometry& MyGe
 	return FReply::Unhandled();
 }
 
+void SMoviePipelineFormatTokenAutoCompleteBox::SetText(const FText& InText)
+{
+	TextBox->SetText(InText);
+}
+
 void SMoviePipelineFormatTokenAutoCompleteBox::OnItemClicked(TSharedPtr<FString> Item) const
 {
 	ReplaceRelevantTextWithSuggestion(*Item);
diff --git a/Engine/Plugins/MovieScene/MovieRenderPipeline/Source/MovieRenderPipelineEditor/Private/Widgets/SMoviePipelineFormatTokenAutoCompleteBox.h b/Engine/Plugins/MovieScene/MovieRenderPipeline/Source/MovieRenderPipelineEditor/Private/Widgets/SMoviePipelineFormatTokenAutoCompleteBox.h
--- a/Engine/Plugins/MovieScene/MovieRenderPipeline/Source/MovieRenderPipelineEditor/Private/Widgets/SMoviePipelineFormatTokenAutoCompleteBox.h
+++ b/Engine/Plugins/MovieScene/MovieRenderPipeline/Source/MovieRenderPipelineEditor/Private/Widgets/SMoviePipelineFormatTokenAutoCompleteBox.h
@@ -15,12 +15,9 @@ class SMoviePipelineFormatTokenAutoCompleteBox : public SCompoundWidget
 {
 public:
 
-	SLATE_BEGIN_ARGS(SMoviePipelineFormatTokenAutoCompleteBox)
-		: _Text()
-		, _Suggestions()
-	{}
+	SLATE_BEGIN_ARGS(SMoviePipelineFormatTokenAutoCompleteBox){}
 
-	SLATE_ATTRIBUTE(FText, Text)
+	SLATE_ARGUMENT(FText, InitialText)
 	SLATE_ATTRIBUTE(TArray<FString>, Suggestions)
 	/** Called whenever the text is changed programmatically or interactively by the user. */
 	SLATE_EVENT(FOnTextChanged, OnTextChanged)
@@ -30,9 +27,10 @@ public:
 
 	// SWidget interface
 	virtual void OnFocusChanging(const FWeakWidgetPath& PreviousFocusPath, const FWidgetPath& NewWidgetPath, const FFocusEvent& InFocusEvent) override;
+	virtual FReply OnKeyDown(const FGeometry& MyGeometry, const FKeyEvent& KeyEvent) override;
 	// End of SWidget interface
 
-	FReply OnKeyDown(const FGeometry& MyGeometry, const FKeyEvent& KeyEvent);
+	void SetText(const FText& InText);
 
 	void OnItemClicked(TSharedPtr<FString>) const;
 
