MoviePipeline: Fixed automated test failures due to the Command Line Encoder node being moved to Globals-only.

There is no longer a shipping node which is restricted to only render layer branches, so a testing-only dummy node was created as a standin.

#rnx

diff --git a/Engine/Plugins/MovieScene/MovieRenderPipeline/Source/MovieRenderPipelineEditor/Private/Graph/MovieGraphSchema.cpp b/Engine/Plugins/MovieScene/MovieRenderPipeline/Source/MovieRenderPipelineEditor/Private/Graph/MovieGraphSchema.cpp
--- a/Engine/Plugins/MovieScene/MovieRenderPipeline/Source/MovieRenderPipelineEditor/Private/Graph/MovieGraphSchema.cpp
+++ b/Engine/Plugins/MovieScene/MovieRenderPipeline/Source/MovieRenderPipelineEditor/Private/Graph/MovieGraphSchema.cpp
@@ -125,7 +125,7 @@ void UMovieGraphSchema::InitMoviePipelineNodeClasses()
 	for (TObjectIterator<UClass> It; It; ++It)
 	{
 		if (It->IsChildOf(UMovieGraphNode::StaticClass())
-			&& !It->HasAnyClassFlags(CLASS_Abstract | CLASS_Deprecated | CLASS_NewerVersionExists))
+			&& !It->HasAnyClassFlags(CLASS_Abstract | CLASS_Deprecated | CLASS_NewerVersionExists | CLASS_Hidden | CLASS_HideDropDown))
 		{
 			MoviePipelineNodeClasses.Add(*It);
 		}
diff --git a/Engine/Plugins/MovieScene/MovieRenderPipeline/Source/MovieRenderPipelineEditor/Private/Tests/MovieGraphEditorTestUtilities.h b/Engine/Plugins/MovieScene/MovieRenderPipeline/Source/MovieRenderPipelineEditor/Private/Tests/MovieGraphEditorTestUtilities.h
--- a/Engine/Plugins/MovieScene/MovieRenderPipeline/Source/MovieRenderPipelineEditor/Private/Tests/MovieGraphEditorTestUtilities.h
+++ b/Engine/Plugins/MovieScene/MovieRenderPipeline/Source/MovieRenderPipelineEditor/Private/Tests/MovieGraphEditorTestUtilities.h
@@ -3,11 +3,49 @@
 #pragma once
 
 #include "CoreMinimal.h"
+#include "Graph/MovieGraphNode.h"
+
+#include "MovieGraphEditorTestUtilities.generated.h"
 
 class FAutomationTestBase;
 class UMovieGraphConfig;
 class UMovieGraphNode;
 
+/**
+ * A dummy render layer only node. Currently MRG does not ship with any nodes which are restricted to just render layer branches, so this node
+ * exists solely to test functionality of restricting nodes to just render layer branches (EMovieGraphBranchRestriction::RenderLayer). Can be removed
+ * once a shipped node with this restriction exists.
+ */
+UCLASS(Hidden, HideDropdown, Deprecated)
+class UDEPRECATED_DummyRenderLayerOnlyNode : public UMovieGraphSettingNode
+{
+	GENERATED_BODY()
+
+public:
+	UDEPRECATED_DummyRenderLayerOnlyNode() = default;
+
+	//~ Begin UMovieGraphNode interface
+#if WITH_EDITOR
+	virtual FText GetNodeTitle(const bool bGetDescriptive = false) const override
+	{
+		static const FText DummyTitle = FText::FromString(TEXT("Dummy Render Layer Only Node"));
+		return DummyTitle;
+	}
+
+	virtual FText GetMenuCategory() const override
+	{
+		static const FText DummyCategory = FText::FromString(TEXT("Dummy Nodes"));
+		return DummyCategory;
+	}
+#endif
+
+	virtual EMovieGraphBranchRestriction GetBranchRestriction() const override
+	{
+		return EMovieGraphBranchRestriction::RenderLayer;
+	}
+	//~ UMovieGraphNode interface
+};
+
 namespace UE::MovieGraph::Private::Tests
 {
 	UMovieGraphConfig* CreateNewMovieGraphConfig(const FName InName = "GraphTestConfig");
diff --git a/Engine/Plugins/MovieScene/MovieRenderPipeline/Source/MovieRenderPipelineEditor/Private/Tests/MovieGraphTests.cpp b/Engine/Plugins/MovieScene/MovieRenderPipeline/Source/MovieRenderPipelineEditor/Private/Tests/MovieGraphTests.cpp
--- a/Engine/Plugins/MovieScene/MovieRenderPipeline/Source/MovieRenderPipelineEditor/Private/Tests/MovieGraphTests.cpp
+++ b/Engine/Plugins/MovieScene/MovieRenderPipeline/Source/MovieRenderPipelineEditor/Private/Tests/MovieGraphTests.cpp
@@ -6,11 +6,9 @@
 #include "Graph/MovieGraphNode.h"
 #include "Graph/Nodes/MovieGraphApplyCVarPresetNode.h"
 #include "Graph/Nodes/MovieGraphBranchNode.h"
-#include "Graph/Nodes/MovieGraphCommandLineEncoderNode.h"
 #include "Graph/Nodes/MovieGraphDeferredPassNode.h"
 #include "Graph/Nodes/MovieGraphInputNode.h"
 #include "Graph/Nodes/MovieGraphOutputNode.h"
-#include "Graph/Nodes/MovieGraphSamplingMethodNode.h"
 #include "Graph/Nodes/MovieGraphSetCVarValueNode.h"
 #include "Graph/Nodes/MovieGraphSubgraphNode.h"
 #include "Graph/Nodes/MovieGraphVariableNode.h"
@@ -456,7 +454,7 @@ namespace UE::MovieGraph::Private::Tests
 		
 		// Globals: UMovieGraphSetCVarValueNode, UMovieGraphApplyCVarPresetNode, UMovieGraphSamplingMethodNode,
 		// UMovieGraphAudioOutputNode, UMovieGraphWarmUpSettingNode
-		// RenderLayer: UMovieGraphCommandLineEncoderNode
+		// RenderLayer: UDummyRenderLayerOnlyNode
 		// Any: All other node types as of October 6, 2023
 		
 		// Create Config
@@ -510,8 +508,8 @@ namespace UE::MovieGraph::Private::Tests
 		// Add middle nodes of BranchRestriction type 'RenderLayer'
 		// Ideally we'd want two different types but only one exists at present
 		UMovieGraphNode* RenderLayerNodeA =
-			GraphConfig->ConstructRuntimeNode<UMovieGraphCommandLineEncoderNode>(
-				UMovieGraphCommandLineEncoderNode::StaticClass());
+			GraphConfig->ConstructRuntimeNode<UDEPRECATED_DummyRenderLayerOnlyNode>(
+				UDEPRECATED_DummyRenderLayerOnlyNode::StaticClass());
 		TestTrue(
 			TEXT("RenderLayerNodeA successfully added"), GraphConfig->GetNodes().Contains(RenderLayerNodeA));
 		if (!RenderLayerNodeA)
@@ -519,8 +517,8 @@ namespace UE::MovieGraph::Private::Tests
 			return false;
 		}
 		UMovieGraphNode* RenderLayerNodeB =
-			GraphConfig->ConstructRuntimeNode<UMovieGraphCommandLineEncoderNode>(
-				UMovieGraphCommandLineEncoderNode::StaticClass());
+			GraphConfig->ConstructRuntimeNode<UDEPRECATED_DummyRenderLayerOnlyNode>(
+				UDEPRECATED_DummyRenderLayerOnlyNode::StaticClass());
 		TestTrue(
 			TEXT("RenderLayerNodeB successfully added"), GraphConfig->GetNodes().Contains(RenderLayerNodeB));
 		if (!RenderLayerNodeB)
@@ -644,7 +642,7 @@ namespace UE::MovieGraph::Private::Tests
 		
 		// Globals: UMovieGraphSetCVarValueNode, UMovieGraphApplyCVarPresetNode, UMovieGraphSamplingMethodNode,
 		// UMovieGraphAudioOutputNode, UMovieGraphWarmUpSettingNode
-		// RenderLayer: UMovieGraphCommandLineEncoderNode
+		// RenderLayer: UDummyRenderLayerOnlyNode
 		// Any: All other node types as of October 5, 2023
 		
 		// Create Config
@@ -690,8 +688,8 @@ namespace UE::MovieGraph::Private::Tests
 		}
 
 		UMovieGraphNode* RenderLayerNode =
-			GraphConfig->ConstructRuntimeNode<UMovieGraphCommandLineEncoderNode>(
-				UMovieGraphCommandLineEncoderNode::StaticClass());
+			GraphConfig->ConstructRuntimeNode<UDEPRECATED_DummyRenderLayerOnlyNode>(
+				UDEPRECATED_DummyRenderLayerOnlyNode::StaticClass());
 		TestTrue(
 			TEXT("RenderLayerNode successfully added"), GraphConfig->GetNodes().Contains(RenderLayerNode));
 		if (!RenderLayerNode)
