MovieGraph: Fixed completion percentages being [0-100] when passed to the UMoviePpipelineExecutorJob when they should have been [0-1] range. This caused the Status column in the Queue UI to always show full.

diff --git a/Engine/Plugins/MovieScene/MovieRenderPipeline/Source/MovieRenderPipelineCore/Private/MoviePipelineLinearExecutor.cpp b/Engine/Plugins/MovieScene/MovieRenderPipeline/Source/MovieRenderPipelineCore/Private/MoviePipelineLinearExecutor.cpp
--- a/Engine/Plugins/MovieScene/MovieRenderPipeline/Source/MovieRenderPipelineCore/Private/MoviePipelineLinearExecutor.cpp
+++ b/Engine/Plugins/MovieScene/MovieRenderPipeline/Source/MovieRenderPipelineCore/Private/MoviePipelineLinearExecutor.cpp
@@ -112,7 +112,7 @@ FText UMoviePipelineLinearExecutorBase::GetWindowTitle()
 	PercentFormatOptions.MaximumFractionalDigits = 0;
 	PercentFormatOptions.RoundingMode = ERoundingMode::HalfFromZero;
 
-	const float CompletionPercentage = GetCompletionPercentageFromActivePipeline();
+	const float CompletionPercentage = GetCompletionPercentageFromActivePipeline() * 100.f;
 
 	const FText TitleFormatString = LOCTEXT("MoviePreviewWindowTitleFormat", "Movie Pipeline Render (Preview) [Job {CurrentCount}/{TotalCount} Total] Current Job: {PercentComplete}% Completed.");
 	return FText::FormatNamed(TitleFormatString, TEXT("CurrentCount"), FText::AsNumber(CurrentPipelineIndex + 1), TEXT("TotalCount"), FText::AsNumber(Queue->GetJobs().Num()), TEXT("PercentComplete"), FText::AsNumber(CompletionPercentage, &PercentFormatOptions));
@@ -123,11 +123,11 @@ float UMoviePipelineLinearExecutorBase::GetCompletionPercentageFromActivePipelin
 	float CompletionPercentage = 0.f;
 	if (const UMoviePipeline* AsLegacyPipeline = Cast<UMoviePipeline>(ActiveMoviePipeline))
 	{
-		CompletionPercentage = UMoviePipelineBlueprintLibrary::GetCompletionPercentage(AsLegacyPipeline) * 100.f;
+		CompletionPercentage = UMoviePipelineBlueprintLibrary::GetCompletionPercentage(AsLegacyPipeline);
 	}
 	else if (const UMovieGraphPipeline* AsGraphPipeline = Cast<UMovieGraphPipeline>(ActiveMoviePipeline))
 	{
-		CompletionPercentage = UMovieGraphBlueprintLibrary::GetCompletionPercentage(AsGraphPipeline) * 100.f;
+		CompletionPercentage = UMovieGraphBlueprintLibrary::GetCompletionPercentage(AsGraphPipeline);
 	}
 
 	return CompletionPercentage;
