MoviePipeline: Fixed an issue where visibility modifier properties would apply even if they were not set as overridden.

#rnx

diff --git a/Engine/Plugins/MovieScene/MovieRenderPipeline/Source/MovieRenderPipelineCore/Private/Graph/MoviePipelineRenderLayerSubsystem.cpp b/Engine/Plugins/MovieScene/MovieRenderPipeline/Source/MovieRenderPipelineCore/Private/Graph/MoviePipelineRenderLayerSubsystem.cpp
--- a/Engine/Plugins/MovieScene/MovieRenderPipeline/Source/MovieRenderPipelineCore/Private/Graph/MoviePipelineRenderLayerSubsystem.cpp
+++ b/Engine/Plugins/MovieScene/MovieRenderPipeline/Source/MovieRenderPipelineCore/Private/Graph/MoviePipelineRenderLayerSubsystem.cpp
@@ -154,7 +154,8 @@ UMoviePipelineVisibilityModifier::UMoviePipelineVisibilityModifier()
 	, bAffectIndirectLightingWhileHidden(false)
 	, bHoldout(false)
 {
-	
+	// Note: The default modifier values here reflect the defaults on the scene component. If a modifier property is marked as overridden, the
+	// override will initially be a no-op due to the defaults being the same.
 }
 
 void UMoviePipelineVisibilityModifier::ApplyModifier(const UWorld* World)
@@ -252,12 +253,15 @@ void UMoviePipelineVisibilityModifier::SetActorVisibilityState(const FActorVisib
 	{
 		return;
 	}
-	
-	Actor->SetActorHiddenInGame(NewVisibilityState.bIsHidden);
+
+	if (bOverride_bIsHidden)
+	{
+		Actor->SetActorHiddenInGame(NewVisibilityState.bIsHidden);
 
 #if WITH_EDITOR
-	Actor->SetIsTemporarilyHiddenInEditor(NewVisibilityState.bIsHidden);
+		Actor->SetIsTemporarilyHiddenInEditor(NewVisibilityState.bIsHidden);
 #endif
+	}
 
 	for (const FActorVisibilityState::FComponentState& ComponentState : NewVisibilityState.Components)
 	{
@@ -270,23 +274,47 @@ void UMoviePipelineVisibilityModifier::SetActorVisibilityState(const FActorVisib
 
 		if (UPrimitiveComponent* AsPrimitiveComponent = Cast<UPrimitiveComponent>(ComponentState.Component.Get()))
 		{
-			AsPrimitiveComponent->SetCastShadow(ComponentState.bCastsShadows);
-			AsPrimitiveComponent->SetCastHiddenShadow(ComponentState.bCastShadowWhileHidden);
-			AsPrimitiveComponent->SetAffectIndirectLightingWhileHidden(ComponentState.bAffectIndirectLightingWhileHidden);
-			AsPrimitiveComponent->SetHoldout(ComponentState.bHoldout);
+			if (bOverride_bCastsShadows)
+			{
+				AsPrimitiveComponent->SetCastShadow(ComponentState.bCastsShadows);
+			}
+
+			if (bOverride_bCastShadowWhileHidden)
+			{
+				AsPrimitiveComponent->SetCastHiddenShadow(ComponentState.bCastShadowWhileHidden);
+			}
+
+			if (bOverride_bAffectIndirectLightingWhileHidden)
+			{
+				AsPrimitiveComponent->SetAffectIndirectLightingWhileHidden(ComponentState.bAffectIndirectLightingWhileHidden);
+			}
+
+			if (bOverride_bHoldout)
+			{
+				AsPrimitiveComponent->SetHoldout(ComponentState.bHoldout);
+			}
 		}
 		// Volumetrics are special cases as they don't inherit from UPrimitiveComponent, and don't support all of the flags.
 		else if (UVolumetricCloudComponent* AsVolumetricCloudComponent = Cast<UVolumetricCloudComponent>(ComponentState.Component.Get()))
 		{
-			AsVolumetricCloudComponent->SetHoldout(ComponentState.bHoldout);
+			if (bOverride_bHoldout)
+			{
+				AsVolumetricCloudComponent->SetHoldout(ComponentState.bHoldout);
+			}
 		}
 		else if (USkyAtmosphereComponent* AsSkyAtmosphereComponent = Cast<USkyAtmosphereComponent>(ComponentState.Component.Get()))
 		{
-			AsSkyAtmosphereComponent->SetHoldout(ComponentState.bHoldout);
+			if (bOverride_bHoldout)
+			{
+				AsSkyAtmosphereComponent->SetHoldout(ComponentState.bHoldout);
+			}
 		}
 		else if (UExponentialHeightFogComponent* AsExponentialHeightFogComponent = Cast<UExponentialHeightFogComponent>(ComponentState.Component.Get()))
 		{
-			AsExponentialHeightFogComponent->SetHoldout(ComponentState.bHoldout);
+			if (bOverride_bHoldout)
+			{
+				AsExponentialHeightFogComponent->SetHoldout(ComponentState.bHoldout);
+			}
 		}
 	}
 }
diff --git a/Engine/Plugins/MovieScene/MovieRenderPipeline/Source/MovieRenderPipelineCore/Public/Graph/MoviePipelineRenderLayerSubsystem.h b/Engine/Plugins/MovieScene/MovieRenderPipeline/Source/MovieRenderPipelineCore/Public/Graph/MoviePipelineRenderLayerSubsystem.h
--- a/Engine/Plugins/MovieScene/MovieRenderPipeline/Source/MovieRenderPipelineCore/Public/Graph/MoviePipelineRenderLayerSubsystem.h
+++ b/Engine/Plugins/MovieScene/MovieRenderPipeline/Source/MovieRenderPipelineCore/Public/Graph/MoviePipelineRenderLayerSubsystem.h
@@ -683,6 +683,8 @@ private:
 		{
 			TSoftObjectPtr<USceneComponent> Component = nullptr;
 			
+			// Note: The default values specified here reflect the defaults on the scene component. If a modifier property is marked as overridden, the
+			// override will initially be a no-op due to the defaults being the same.
 			uint8 bCastsShadows : 1 = true;
 			uint8 bCastShadowWhileHidden : 1 = false;
 			uint8 bAffectIndirectLightingWhileHidden : 1 = false;
