Movie Graph:

-Fix: MRQ isn't compositing burn in onto renders with the pathtracer

diff --git a/Engine/Plugins/MovieScene/MovieRenderPipeline/Source/MovieRenderPipelineRenderPasses/Private/MoviePipelineImageSequenceOutput.cpp b/Engine/Plugins/MovieScene/MovieRenderPipeline/Source/MovieRenderPipelineRenderPasses/Private/MoviePipelineImageSequenceOutput.cpp
--- a/Engine/Plugins/MovieScene/MovieRenderPipeline/Source/MovieRenderPipelineRenderPasses/Private/MoviePipelineImageSequenceOutput.cpp
+++ b/Engine/Plugins/MovieScene/MovieRenderPipeline/Source/MovieRenderPipelineRenderPasses/Private/MoviePipelineImageSequenceOutput.cpp
@@ -99,6 +99,11 @@ void UMoviePipelineImageSequenceOutputBase::OnReceiveImageDataImpl(FMoviePipelin
 
 	FString OutputDirectory = OutputSettings->OutputDirectory.Path;
 
+	// The InMergedOutputFrame->ImageOutputData map contains both RenderPasses and CompositePasses.
+	// We determine how we gather pixel data based on the number of RenderPasses we have done, not counting the CompositePasses.
+	// This is the reason for using a separate RenderPassIteration counter with a foreach loop, only incrementing it for RenderPasses.
+	int32 RenderPassIteration = 0;
+	const int32 RenderPassCount = InMergedOutputFrame->ImageOutputData.Num() - CompositedPasses.Num();
 	for (TPair<FMoviePipelinePassIdentifier, TUniquePtr<FImagePixelData>>& RenderPassData : InMergedOutputFrame->ImageOutputData)
 	{
 		// Don't write out a composited pass in this loop, as it will be merged with the Final Image and not written separately. 
@@ -227,7 +232,7 @@ void UMoviePipelineImageSequenceOutputBase::OnReceiveImageDataImpl(FMoviePipelin
 
 
 		// We composite before flipping the alpha so that it is consistent for all formats.
-		if (RenderPassData.Key.Name == TEXT("FinalImage"))
+		if (RenderPassData.Key.Name == TEXT("FinalImage") || RenderPassData.Key.Name == TEXT("PathTracer")) 
 		{
 			for (const MoviePipeline::FCompositePassInfo& CompositePass : CompositedPasses)
 			{
@@ -237,19 +242,23 @@ void UMoviePipelineImageSequenceOutputBase::OnReceiveImageDataImpl(FMoviePipelin
 					continue;
 				}
 
+				// If there's more than one render pass, we need to copy the composite passes for the first render pass then move for the remaining ones
+				const bool bShouldCopyImageData = RenderPassCount > 1 && RenderPassIteration == 0;
+				TUniquePtr<FImagePixelData> PixelData = bShouldCopyImageData ? CompositePass.PixelData->CopyImageData() : CompositePass.PixelData->MoveImageDataToNew();
+				
 				// We don't need to copy the data here (even though it's being passed to a async system) because we already made a unique copy of the
 				// burn in/widget data when we decided to composite it.
 				switch (QuantizedPixelType)
 				{
-				case EImagePixelType::Color:
-					TileImageTask->PixelPreProcessors.Add(TAsyncCompositeImage<FColor>(CompositePass.PixelData->MoveImageDataToNew()));
-					break;
-				case EImagePixelType::Float16:
-					TileImageTask->PixelPreProcessors.Add(TAsyncCompositeImage<FFloat16Color>(CompositePass.PixelData->MoveImageDataToNew()));
-					break;
-				case EImagePixelType::Float32:
-					TileImageTask->PixelPreProcessors.Add(TAsyncCompositeImage<FLinearColor>(CompositePass.PixelData->MoveImageDataToNew()));
-					break;
+					case EImagePixelType::Color:
+						TileImageTask->PixelPreProcessors.Add(TAsyncCompositeImage<FColor>(MoveTemp(PixelData)));
+						break;
+					case EImagePixelType::Float16:
+						TileImageTask->PixelPreProcessors.Add(TAsyncCompositeImage<FFloat16Color>(MoveTemp(PixelData)));
+						break;
+					case EImagePixelType::Float32:
+						TileImageTask->PixelPreProcessors.Add(TAsyncCompositeImage<FLinearColor>(MoveTemp(PixelData)));
+						break;
 				}
 			}
 		}
@@ -269,6 +278,8 @@ void UMoviePipelineImageSequenceOutputBase::OnReceiveImageDataImpl(FMoviePipelin
 #endif
 
 		GetPipeline()->AddOutputFuture(ImageWriteQueue->Enqueue(MoveTemp(TileImageTask)), OutputData);
+
+		RenderPassIteration++;
 	}
 }
 
