MoviePipeline: Fixed issue which caused value resolution to be incorrect when variable nodes are disabled.

#rnx

diff --git a/Engine/Plugins/MovieScene/MovieRenderPipeline/Source/MovieRenderPipelineCore/Private/Graph/MovieGraphConfig.cpp b/Engine/Plugins/MovieScene/MovieRenderPipeline/Source/MovieRenderPipelineCore/Private/Graph/MovieGraphConfig.cpp
--- a/Engine/Plugins/MovieScene/MovieRenderPipeline/Source/MovieRenderPipelineCore/Private/Graph/MovieGraphConfig.cpp
+++ b/Engine/Plugins/MovieScene/MovieRenderPipeline/Source/MovieRenderPipelineCore/Private/Graph/MovieGraphConfig.cpp
@@ -1226,8 +1226,9 @@ void UMovieGraphConfig::CopyOverriddenProperties(UMovieGraphNode* FromNode, UMov
 		{
 			if (const UMovieGraphPin* InputPin = FromNode->GetInputPin(PropertyName))
 			{
+				// For the connected value to be used, the type must match and the node the value is originating from must be enabled
 				const UMovieGraphPin* ConnectedPin = InputPin->GetFirstConnectedPin();
-				if (ConnectedPin && (ConnectedPin->Properties.Type == InputPin->Properties.Type) && ConnectedPin->Node)
+				if (ConnectedPin && (ConnectedPin->Properties.Type == InputPin->Properties.Type) && ConnectedPin->Node && !ConnectedPin->Node->IsDisabled())
 				{
 					// There was a valid connection to the input pin; resolve the value from the connected output and set the
 					// value on this property
