MoviePipeline: Fixed issue that caused duplicated variable nodes to not update when the variable is renamed.

The main fix was implementing PostEditImport(). It also turns out that the RegisterDelegates() call in the constructor was redundant because it is called from PostLoad() in the node's base class.

#rnx

diff --git a/Engine/Plugins/MovieScene/MovieRenderPipeline/Source/MovieRenderPipelineCore/Private/Graph/Nodes/MovieGraphVariableNode.cpp b/Engine/Plugins/MovieScene/MovieRenderPipeline/Source/MovieRenderPipelineCore/Private/Graph/Nodes/MovieGraphVariableNode.cpp
--- a/Engine/Plugins/MovieScene/MovieRenderPipeline/Source/MovieRenderPipelineCore/Private/Graph/Nodes/MovieGraphVariableNode.cpp
+++ b/Engine/Plugins/MovieScene/MovieRenderPipeline/Source/MovieRenderPipelineCore/Private/Graph/Nodes/MovieGraphVariableNode.cpp
@@ -6,12 +6,12 @@
 #include "MoviePipelineQueue.h"
 #include "Styling/AppStyle.h"
 
-UMovieGraphVariableNode::UMovieGraphVariableNode()
+void UMovieGraphVariableNode::PostEditImport()
 {
-	if (!HasAnyFlags(RF_ClassDefaultObject))
-	{
-		RegisterDelegates();
-	}
+	Super::PostEditImport();
+
+	// Allow pasted/duplicated nodes to register delegates
+	RegisterDelegates();
 }
 
 TArray<FMovieGraphPinProperties> UMovieGraphVariableNode::GetOutputPinProperties() const
diff --git a/Engine/Plugins/MovieScene/MovieRenderPipeline/Source/MovieRenderPipelineCore/Public/Graph/Nodes/MovieGraphVariableNode.h b/Engine/Plugins/MovieScene/MovieRenderPipeline/Source/MovieRenderPipelineCore/Public/Graph/Nodes/MovieGraphVariableNode.h
--- a/Engine/Plugins/MovieScene/MovieRenderPipeline/Source/MovieRenderPipelineCore/Public/Graph/Nodes/MovieGraphVariableNode.h
+++ b/Engine/Plugins/MovieScene/MovieRenderPipeline/Source/MovieRenderPipelineCore/Public/Graph/Nodes/MovieGraphVariableNode.h
@@ -16,7 +16,11 @@ class MOVIERENDERPIPELINECORE_API UMovieGraphVariableNode : public UMovieGraphNo
 	GENERATED_BODY()
 
 public:
-	UMovieGraphVariableNode();
+	UMovieGraphVariableNode() = default;
+
+	//~ Begin UObject interface
+	virtual void PostEditImport() override;
+	//~ End UObject interface
 
 	virtual TArray<FMovieGraphPinProperties> GetOutputPinProperties() const override;
 	virtual FString GetResolvedValueForOutputPin(const FName& InPinName, const FMovieGraphTraversalContext* InContext) const override;
