Movie Graph:

-Fix: When you first create a Movie Graph Named Resolution Variable it marks 1920x1080 as "Custom"

diff --git a/Engine/Plugins/MovieScene/MovieRenderPipeline/Source/MovieRenderPipelineCore/Public/Graph/MovieGraphNamedResolution.h b/Engine/Plugins/MovieScene/MovieRenderPipeline/Source/MovieRenderPipelineCore/Public/Graph/MovieGraphNamedResolution.h
--- a/Engine/Plugins/MovieScene/MovieRenderPipeline/Source/MovieRenderPipelineCore/Public/Graph/MovieGraphNamedResolution.h
+++ b/Engine/Plugins/MovieScene/MovieRenderPipeline/Source/MovieRenderPipelineCore/Public/Graph/MovieGraphNamedResolution.h
@@ -27,11 +27,11 @@ public:
 	
 	/** The name of the resolution this links to */
 	UPROPERTY(config, EditAnywhere, BlueprintReadWrite, Category = "Movie Graph|Resolution")
-	FName ProfileName = NAME_None;
+	FName ProfileName = DefaultResolutionName;
 
 	/** The screen resolution (in pixels). */
 	UPROPERTY(config, EditAnywhere, BlueprintReadWrite, Category = "Movie Graph|Resolution")
-	FIntPoint Resolution = FIntPoint(1920, 1080);
+	FIntPoint Resolution = DefaultResolution;
 	
 	/** The description text for this screen resolution. */
 	UPROPERTY(config, EditAnywhere, BlueprintReadWrite, Category = "Movie Graph|Resolution")
@@ -47,6 +47,11 @@ public:
 	 */
 	inline static FName DefaultResolutionName = TEXT("1080p (FHD)");
 
+	/**
+	 * The default resolution name to use when one is not defined.
+	 */
+	inline static FIntPoint DefaultResolution = FIntPoint(1920, 1080);
+
 	/**
 	 * Predefined name for the 'custom' resolution option in the combobox
 	 */
