MoviePipeline: The search field for filtering actors in the collection node's Add menu should be cleared out after adding an actor.

This is a problem for any searchable/filterable widget within the Add menu (Actor, Component Type, Actor Type, etc). There doesn't appear to be an API for clearing out the filter boxes, so instead the widget in the Add menu is re-created every time the Add menu is shown.

#rnx

diff --git a/Engine/Plugins/MovieScene/MovieRenderPipeline/Source/MovieRenderPipelineCore/Private/Graph/MoviePipelineRenderLayerSubsystem.cpp b/Engine/Plugins/MovieScene/MovieRenderPipeline/Source/MovieRenderPipelineCore/Private/Graph/MoviePipelineRenderLayerSubsystem.cpp
--- a/Engine/Plugins/MovieScene/MovieRenderPipeline/Source/MovieRenderPipelineCore/Private/Graph/MoviePipelineRenderLayerSubsystem.cpp
+++ b/Engine/Plugins/MovieScene/MovieRenderPipeline/Source/MovieRenderPipelineCore/Private/Graph/MoviePipelineRenderLayerSubsystem.cpp
@@ -532,6 +532,11 @@ TArray<TSharedRef<SWidget>> UMovieGraphConditionGroupQueryBase::GetWidgets()
 	return TArray<TSharedRef<SWidget>>();
 }
 
+bool UMovieGraphConditionGroupQueryBase::HasAddMenu() const
+{
+	return false;
+}
+
 TSharedRef<SWidget> UMovieGraphConditionGroupQueryBase::GetAddMenuContents(const FMovieGraphConditionGroupQueryContentsChanged& OnAddFinished)
 {
 	return SNullWidget::NullWidget;
@@ -682,6 +687,11 @@ TArray<TSharedRef<SWidget>> UMovieGraphConditionGroupQuery_Actor::GetWidgets()
 	return Widgets;
 }
 
+bool UMovieGraphConditionGroupQuery_Actor::HasAddMenu() const
+{
+	return true;
+}
+
 TSharedRef<SWidget> UMovieGraphConditionGroupQuery_Actor::GetAddMenuContents(const FMovieGraphConditionGroupQueryContentsChanged& OnAddFinished)
 {	
 	FSceneOutlinerInitializationOptions SceneOutlinerInitOptions;
@@ -1005,6 +1015,11 @@ TArray<TSharedRef<SWidget>> UMovieGraphConditionGroupQuery_ActorType::GetWidgets
 	return Widgets;
 }
 
+bool UMovieGraphConditionGroupQuery_ActorType::HasAddMenu() const
+{
+	return true;
+}
+
 TSharedRef<SWidget> UMovieGraphConditionGroupQuery_ActorType::GetAddMenuContents(const FMovieGraphConditionGroupQueryContentsChanged& OnAddFinished)
 {
 	FClassViewerModule& ClassViewerModule = FModuleManager::LoadModuleChecked<FClassViewerModule>("ClassViewer");
@@ -1211,6 +1226,11 @@ TArray<TSharedRef<SWidget>> UMovieGraphConditionGroupQuery_ComponentType::GetWid
 	return Widgets;
 }
 
+bool UMovieGraphConditionGroupQuery_ComponentType::HasAddMenu() const
+{
+	return true;
+}
+
 TSharedRef<SWidget> UMovieGraphConditionGroupQuery_ComponentType::GetAddMenuContents(const FMovieGraphConditionGroupQueryContentsChanged& OnAddFinished)
 {
 	FClassViewerModule& ClassViewerModule = FModuleManager::LoadModuleChecked<FClassViewerModule>("ClassViewer");
@@ -1321,6 +1341,11 @@ TArray<TSharedRef<SWidget>> UMovieGraphConditionGroupQuery_EditorFolder::GetWidg
     return Widgets;
 }
 
+bool UMovieGraphConditionGroupQuery_EditorFolder::HasAddMenu() const
+{
+	return true;
+}
+
 TSharedRef<SWidget> UMovieGraphConditionGroupQuery_EditorFolder::GetAddMenuContents(const FMovieGraphConditionGroupQueryContentsChanged& OnAddFinished)
 {
 	auto OnItemPicked = FOnSceneOutlinerItemPicked::CreateLambda([this, OnAddFinished](TSharedRef<ISceneOutlinerTreeItem> Item)
@@ -1468,6 +1493,11 @@ TArray<TSharedRef<SWidget>> UMovieGraphConditionGroupQuery_Sublevel::GetWidgets(
 	return Widgets;
 }
 
+bool UMovieGraphConditionGroupQuery_Sublevel::HasAddMenu() const
+{
+	return true;
+}
+
 TSharedRef<SWidget> UMovieGraphConditionGroupQuery_Sublevel::GetAddMenuContents(const FMovieGraphConditionGroupQueryContentsChanged& OnAddFinished)
 {
 	FAssetPickerConfig SublevelPickerConfig;
diff --git a/Engine/Plugins/MovieScene/MovieRenderPipeline/Source/MovieRenderPipelineCore/Public/Graph/MoviePipelineRenderLayerSubsystem.h b/Engine/Plugins/MovieScene/MovieRenderPipeline/Source/MovieRenderPipelineCore/Public/Graph/MoviePipelineRenderLayerSubsystem.h
--- a/Engine/Plugins/MovieScene/MovieRenderPipeline/Source/MovieRenderPipelineCore/Public/Graph/MoviePipelineRenderLayerSubsystem.h
+++ b/Engine/Plugins/MovieScene/MovieRenderPipeline/Source/MovieRenderPipelineCore/Public/Graph/MoviePipelineRenderLayerSubsystem.h
@@ -126,8 +126,18 @@ public:
 	virtual TArray<TSharedRef<SWidget>> GetWidgets();
 
 	/**
-	 * Gets the contents of the "Add" menu in the UI, if any. The Add menu will not display if SNullWidget::NullWidget is returned. When the Add
-	 * menu updates properties within the query, OnAddFinished should be called in order to give the UI a chance to update itself.
+	 * Returns true if this query should expose an Add menu, or false if no Add menu is visible.
+	 *
+	 * @see GetAddMenuContents()
+	 */
+	virtual bool HasAddMenu() const;
+
+	/**
+	 * Gets the contents of the "Add" menu in the UI, if any. When the Add menu updates properties within the query, OnAddFinished should be called
+	 * in order to give the UI a chance to update itself. Note that HasAddMenu() must return true in order for the contents returned from this method
+	 * to be displayed in the UI.
+	 *
+	 * @see HasAddMenu()
 	 */
 	virtual TSharedRef<SWidget> GetAddMenuContents(const FMovieGraphConditionGroupQueryContentsChanged& OnAddFinished);
 #endif
@@ -171,6 +181,7 @@ public:
 
 #if WITH_EDITOR
 	virtual TArray<TSharedRef<SWidget>> GetWidgets() override;
+	virtual bool HasAddMenu() const override;
 	virtual TSharedRef<SWidget> GetAddMenuContents(const FMovieGraphConditionGroupQueryContentsChanged& OnAddFinished) override;
 #endif
 
@@ -279,6 +290,7 @@ public:
 
 #if WITH_EDITOR
 	virtual TArray<TSharedRef<SWidget>> GetWidgets() override;
+	virtual bool HasAddMenu() const override;
 	virtual TSharedRef<SWidget> GetAddMenuContents(const FMovieGraphConditionGroupQueryContentsChanged& OnAddFinished) override;
 #endif
 
@@ -346,6 +358,7 @@ public:
 
 #if WITH_EDITOR
 	virtual TArray<TSharedRef<SWidget>> GetWidgets() override;
+	virtual bool HasAddMenu() const override;
 	virtual TSharedRef<SWidget> GetAddMenuContents(const FMovieGraphConditionGroupQueryContentsChanged& OnAddFinished) override;
 #endif
 
@@ -381,6 +394,7 @@ public:
 
 #if WITH_EDITOR
 	virtual TArray<TSharedRef<SWidget>> GetWidgets() override;
+	virtual bool HasAddMenu() const override;
 	virtual TSharedRef<SWidget> GetAddMenuContents(const FMovieGraphConditionGroupQueryContentsChanged& OnAddFinished) override;
 #endif
 
@@ -415,6 +429,7 @@ public:
 
 #if WITH_EDITOR
 	virtual TArray<TSharedRef<SWidget>> GetWidgets() override;
+	virtual bool HasAddMenu() const override;
 	virtual TSharedRef<SWidget> GetAddMenuContents(const FMovieGraphConditionGroupQueryContentsChanged& OnAddFinished) override;
 #endif
 
diff --git a/Engine/Plugins/MovieScene/MovieRenderPipeline/Source/MovieRenderPipelineEditor/Private/Customizations/Graph/MovieGraphCollectionsCustomization.cpp b/Engine/Plugins/MovieScene/MovieRenderPipeline/Source/MovieRenderPipelineEditor/Private/Customizations/Graph/MovieGraphCollectionsCustomization.cpp
--- a/Engine/Plugins/MovieScene/MovieRenderPipeline/Source/MovieRenderPipelineEditor/Private/Customizations/Graph/MovieGraphCollectionsCustomization.cpp
+++ b/Engine/Plugins/MovieScene/MovieRenderPipeline/Source/MovieRenderPipelineEditor/Private/Customizations/Graph/MovieGraphCollectionsCustomization.cpp
@@ -752,14 +752,6 @@ void SMovieGraphCollectionTreeAddQueryContentWidget::Construct(const FArguments&
 		return;
 	}
 
-	// Call the underlying query to get the contents of the menu. The query will call back to the widget if something was added so the UI has
-	// an opportunity to refresh itself.
-	TSharedRef<SWidget> AddMenuContents = WeakQuery->GetAddMenuContents(
-		UMovieGraphConditionGroupQueryBase::FMovieGraphConditionGroupQueryContentsChanged::CreateLambda([]()
-    {
-    	FSlateApplication::Get().DismissAllMenus();
-    }));
-
     // Generate a (multi-layered) icon for the "Add" menu
     const TSharedRef<SLayeredImage> AddIcon =
     	SNew(SLayeredImage)
@@ -778,14 +770,29 @@ void SMovieGraphCollectionTreeAddQueryContentWidget::Construct(const FArguments&
 		{
 			return WeakQuery.IsValid() && WeakQuery->IsEnabled();
 		})
-		.OnGetMenuContent_Lambda([AddMenuContents]()
+		.OnGetMenuContent_Lambda([WeakQuery]()
 		{
-			return AddMenuContents;
+			if (WeakQuery.IsValid())
+			{
+				// Call the underlying query to get the contents of the menu. The query will call back to the widget if something was added so the UI has
+				// an opportunity to refresh itself.
+				return WeakQuery->GetAddMenuContents(
+					UMovieGraphConditionGroupQueryBase::FMovieGraphConditionGroupQueryContentsChanged::CreateLambda([]()
+					{
+						FSlateApplication::Get().DismissAllMenus();
+					}));
+			}
+
+			return SNullWidget::NullWidget;
 		})
-		.Visibility_Lambda([AddMenuContents]()
+		.Visibility_Lambda([WeakQuery]()
 		{
-			// Queries can opt to not show an Add menu by returning NullWidget
-			return (AddMenuContents != SNullWidget::NullWidget) ? EVisibility::Visible : EVisibility::Collapsed;
+			if (WeakQuery.IsValid())
+			{
+				return WeakQuery->HasAddMenu() ? EVisibility::Visible : EVisibility::Collapsed;
+			}
+
+			return EVisibility::Collapsed;
 		})
 		.ButtonContent()
 		[
