MovieGraph: Removed "MoviePipeline.EnableRenderGraph" cvar, we now always show Graph related UI elements.

diff --git a/Engine/Plugins/MovieScene/MovieRenderPipeline/Source/MovieRenderPipelineEditor/Private/Graph/MovieGraphAssetToolkit.cpp b/Engine/Plugins/MovieScene/MovieRenderPipeline/Source/MovieRenderPipelineEditor/Private/Graph/MovieGraphAssetToolkit.cpp
--- a/Engine/Plugins/MovieScene/MovieRenderPipeline/Source/MovieRenderPipelineEditor/Private/Graph/MovieGraphAssetToolkit.cpp
+++ b/Engine/Plugins/MovieScene/MovieRenderPipeline/Source/MovieRenderPipelineEditor/Private/Graph/MovieGraphAssetToolkit.cpp
@@ -160,13 +160,6 @@ private:
 	TWeakPtr<IDetailsView> DetailsView;
 };
 
-// Temporary cvar to enable/disable upgrading to a graph-based configuration
-static TAutoConsoleVariable<bool> CVarMoviePipelineEnableRenderGraph(
-	TEXT("MoviePipeline.EnableRenderGraph"),
-	false,
-	TEXT("Determines if the Render Graph feature is enabled in the UI. This is a highly experimental feature and is not ready for use.")
-);
-
 FMovieGraphAssetToolkit::FMovieGraphAssetToolkit()
 	: bIsInternalSelectionChange(false)
 {
diff --git a/Engine/Plugins/MovieScene/MovieRenderPipeline/Source/MovieRenderPipelineEditor/Private/Graph/MovieGraphConfigFactory.cpp b/Engine/Plugins/MovieScene/MovieRenderPipeline/Source/MovieRenderPipelineEditor/Private/Graph/MovieGraphConfigFactory.cpp
--- a/Engine/Plugins/MovieScene/MovieRenderPipeline/Source/MovieRenderPipelineEditor/Private/Graph/MovieGraphConfigFactory.cpp
+++ b/Engine/Plugins/MovieScene/MovieRenderPipeline/Source/MovieRenderPipelineEditor/Private/Graph/MovieGraphConfigFactory.cpp
@@ -52,12 +52,6 @@ UObject* UMovieGraphConfigFactory::FactoryCreateNew(UClass* Class, UObject* InPa
 	return NewObject<UMovieGraphConfig>(InParent, Class, Name, Flags);
 }
 
-bool UMovieGraphConfigFactory::ShouldShowInNewMenu() const
-{
-	IConsoleVariable* RenderGraphCVar = IConsoleManager::Get().FindConsoleVariable(TEXT("MoviePipeline.EnableRenderGraph"));
-	return RenderGraphCVar && RenderGraphCVar->GetBool();
-}
-
 void UMovieGraphConfigFactory::AddSubgraphNodeToGraph(UMovieGraphConfig* InTargetGraph) const
 {
 	auto ExposeSubgraphInputOutputOnTargetGraph = [InTargetGraph](
diff --git a/Engine/Plugins/MovieScene/MovieRenderPipeline/Source/MovieRenderPipelineEditor/Private/Widgets/SMoviePipelineQueueEditor.cpp b/Engine/Plugins/MovieScene/MovieRenderPipeline/Source/MovieRenderPipelineEditor/Private/Widgets/SMoviePipelineQueueEditor.cpp
--- a/Engine/Plugins/MovieScene/MovieRenderPipeline/Source/MovieRenderPipelineEditor/Private/Widgets/SMoviePipelineQueueEditor.cpp
+++ b/Engine/Plugins/MovieScene/MovieRenderPipeline/Source/MovieRenderPipelineEditor/Private/Widgets/SMoviePipelineQueueEditor.cpp
@@ -487,9 +487,7 @@ public:
 			AssetPickerConfig.OnAssetSelected = InOnAssetSelected;
 		}
 
-		const bool bIsGraphFeatureEnabled = IConsoleManager::Get().FindConsoleVariable(TEXT("MoviePipeline.EnableRenderGraph"))->GetBool();
-
-		if (bIsGraphFeatureEnabled && TargetJob->IsUsingGraphConfiguration())
+		if (TargetJob->IsUsingGraphConfiguration())
 		{
 			MenuBuilder.BeginSection(NAME_None, LOCTEXT("NewConfig_MenuSection", "New Configuration"));
 			{
@@ -515,10 +513,10 @@ public:
 				EUserInterfaceActionType::Button
 			);
 
-			if (bIsGraphFeatureEnabled && !TargetJob->IsUsingGraphConfiguration())
+			if (!TargetJob->IsUsingGraphConfiguration())
 			{
 				MenuBuilder.AddMenuEntry(
-					LOCTEXT("ReplaceWithGraph_Label", "Replace with Graph"),
+					LOCTEXT("ReplaceWithGraph_Label", "Replace with Graph (Experimental)"),
 					LOCTEXT("ReplaceWithGraph_Tooltip", "Replaces the current configuration with a new graph representation."),
 					FSlateIcon(),
 					FUIAction(InNewRenderGraph),
diff --git a/Engine/Plugins/MovieScene/MovieRenderPipeline/Source/MovieRenderPipelineEditor/Public/Graph/MovieGraphConfigFactory.h b/Engine/Plugins/MovieScene/MovieRenderPipeline/Source/MovieRenderPipelineEditor/Public/Graph/MovieGraphConfigFactory.h
--- a/Engine/Plugins/MovieScene/MovieRenderPipeline/Source/MovieRenderPipelineEditor/Public/Graph/MovieGraphConfigFactory.h
+++ b/Engine/Plugins/MovieScene/MovieRenderPipeline/Source/MovieRenderPipelineEditor/Public/Graph/MovieGraphConfigFactory.h
@@ -16,7 +16,6 @@ public:
 	
 	//~ Begin UFactory Interface
 	virtual UObject* FactoryCreateNew(UClass* Class, UObject* InParent, FName Name, EObjectFlags Flags, UObject* Context, FFeedbackContext* Warn) override;
-	virtual bool ShouldShowInNewMenu() const override;
 	//~ End UFactory Interface
 
 private:
