MoviePipeline: Fixed graphs appearing unsaved in the queue UI, even if they have been saved.

Also fixed an issue that caused parent subgraphs to become dirtied when subgraphs were modified, even if the parent graph needed no changes as a result of the modification. The only subgraph changes that should cause the parent to become dirtied are input/output changes.

#rnx

diff --git a/Engine/Plugins/MovieScene/MovieRenderPipeline/Source/MovieRenderPipelineCore/Private/Graph/MovieGraphNode.cpp b/Engine/Plugins/MovieScene/MovieRenderPipeline/Source/MovieRenderPipelineCore/Private/Graph/MovieGraphNode.cpp
--- a/Engine/Plugins/MovieScene/MovieRenderPipeline/Source/MovieRenderPipelineCore/Private/Graph/MovieGraphNode.cpp
+++ b/Engine/Plugins/MovieScene/MovieRenderPipeline/Source/MovieRenderPipelineCore/Private/Graph/MovieGraphNode.cpp
@@ -199,10 +199,11 @@ void UMovieGraphNode::UpdatePins()
 	// Include the exposed dynamic properties in the input pins
 	InputPinProperties.Append(GetExposedPinProperties());
 
-	auto UpdatePins = [this](TArray<UMovieGraphPin*>& Pins, const TArray<FMovieGraphPinProperties>& PinProperties)
+	bool bChangedPins = false;
+
+	auto UpdatePins = [this, &bChangedPins](TArray<UMovieGraphPin*>& Pins, const TArray<FMovieGraphPinProperties>& PinProperties)
 	{
 		bool bAppliedEdgeChanges = false;
-		bool bChangedPins = false;
 
 		// Find unmatched pins vs. properties (via name matching)
 		TArray<UMovieGraphPin*> UnmatchedPins;
@@ -266,7 +267,10 @@ void UMovieGraphNode::UpdatePins()
 	UpdatePins(MutableView(InputPins), InputPinProperties);
 	UpdatePins(MutableView(OutputPins), OutputPinProperties);
 
-	OnNodeChangedDelegate.Broadcast(this);
+	if (bChangedPins)
+	{
+		OnNodeChangedDelegate.Broadcast(this);
+	}
 }
 
 void UMovieGraphNode::UpdateDynamicProperties()
diff --git a/Engine/Plugins/MovieScene/MovieRenderPipeline/Source/MovieRenderPipelineEditor/Private/Widgets/SMoviePipelineQueueEditor.cpp b/Engine/Plugins/MovieScene/MovieRenderPipeline/Source/MovieRenderPipelineEditor/Private/Widgets/SMoviePipelineQueueEditor.cpp
--- a/Engine/Plugins/MovieScene/MovieRenderPipeline/Source/MovieRenderPipelineEditor/Private/Widgets/SMoviePipelineQueueEditor.cpp
+++ b/Engine/Plugins/MovieScene/MovieRenderPipeline/Source/MovieRenderPipelineEditor/Private/Widgets/SMoviePipelineQueueEditor.cpp
@@ -346,6 +346,11 @@ public:
 		UMoviePipelineExecutorJob* Job = WeakJob.Get();
 		if (Job)
 		{
+			if (const UMovieGraphConfig* GraphPreset = Job->GetGraphPreset())
+			{
+				return GraphPreset->GetPackage()->IsDirty() ? EVisibility::Visible : EVisibility::Collapsed;
+			}
+			
 			return (Job->GetPresetOrigin() == nullptr) ? EVisibility::Visible : EVisibility::Collapsed;
 		}
 		
@@ -1045,10 +1050,17 @@ struct FMoviePipelineShotItem : IMoviePipelineQueueTreeItem
 
 	EVisibility GetShotConfigModifiedVisibility() const
 	{
-		UMoviePipelineExecutorShot* Shot = WeakShot.Get();
-		if (Shot && Shot->GetShotOverrideConfiguration() && (Shot->GetShotOverridePresetOrigin() == nullptr))
+		if (const UMoviePipelineExecutorShot* Shot = WeakShot.Get())
 		{
-			return EVisibility::Visible;
+			if (const UMovieGraphConfig* GraphPreset = Shot->GetGraphPreset())
+			{
+				return GraphPreset->GetPackage()->IsDirty() ? EVisibility::Visible : EVisibility::Collapsed;
+			}
+		
+			if (Shot->GetShotOverrideConfiguration() && (Shot->GetShotOverridePresetOrigin() == nullptr))
+			{
+				return EVisibility::Visible;
+			}
 		}
 
 		return EVisibility::Collapsed;
