MovieGraph: Fixed an issue where separate layers could accumulate onto the wrong target, effectively making layer outputs blended results between entirely separate layers.

diff --git a/Engine/Plugins/MovieScene/MovieRenderPipeline/Source/MovieRenderPipelineCore/Public/Graph/MovieGraphDefaultRenderer.h b/Engine/Plugins/MovieScene/MovieRenderPipeline/Source/MovieRenderPipelineCore/Public/Graph/MovieGraphDefaultRenderer.h
--- a/Engine/Plugins/MovieScene/MovieRenderPipeline/Source/MovieRenderPipelineCore/Public/Graph/MovieGraphDefaultRenderer.h
+++ b/Engine/Plugins/MovieScene/MovieRenderPipeline/Source/MovieRenderPipelineCore/Public/Graph/MovieGraphDefaultRenderer.h
@@ -138,19 +138,25 @@ namespace UE::MovieGraph::DefaultRenderer
 			FInstance(TSharedPtr<MoviePipelineOverlappedAccumulatorInterface, ESPMode::ThreadSafe> InAccumulator)
 			{
 				Accumulator = InAccumulator;
-				ActiveFrameNumber = INDEX_NONE;
-				bIsActive = false;
+				SetIsActive(false);
 			}
 
-
 			bool IsActive() const { return bIsActive; }
-			void SetIsActive(const bool bInIsActive) { bIsActive = bInIsActive; }
+			
+			/** Changing the active state resets the internal state, so set it to active before configuring. */
+			void SetIsActive(const bool bInIsActive) 
+			{ 
+				bIsActive = bInIsActive;
+				ActiveFrameNumber = INDEX_NONE;
+				TaskPrereq = UE::Tasks::FTask();
+			}
 
 			TSharedPtr<MoviePipelineOverlappedAccumulatorInterface, ESPMode::ThreadSafe> Accumulator;
 			int32 ActiveFrameNumber;
 			FMovieGraphRenderDataIdentifier ActivePassIdentifier;
-			FThreadSafeBool bIsActive;
 			UE::Tasks::FTask TaskPrereq;
+		private:
+			FThreadSafeBool bIsActive;
 		};
 
 		typedef TSharedPtr<FInstance, ESPMode::ThreadSafe> FInstancePtr;
@@ -164,32 +170,23 @@ namespace UE::MovieGraph::DefaultRenderer
 			FScopeLock ScopeLock(&CriticalSection);
 
 			// Search for an existing accumulator for the given frame number and render data
-			int32 AvailableIndex = INDEX_NONE;
 			for (int32 Index = 0; Index < Accumulators.Num(); Index++)
 			{
 				if (InFrameNumber == Accumulators[Index]->ActiveFrameNumber && InPassIdentifier == Accumulators[Index]->ActivePassIdentifier)
 				{
-					AvailableIndex = Index;
-					break;
+					return Accumulators[Index];
 				}
 			}
 
 			// If we didn't find one already in use for this frame, look to see if there's a previously
 			// allocated one which is no longer being used.
-			if (AvailableIndex == INDEX_NONE)
+			int32 AvailableIndex = INDEX_NONE;
+			for (int32 Index = 0; Index < Accumulators.Num(); Index++)
 			{
-				for (int32 Index = 0; Index < Accumulators.Num(); Index++)
+				if (!Accumulators[Index]->IsActive())
 				{
-					if (!Accumulators[Index]->IsActive())
-					{
-						// Found a free one, tie it to this output frame.
-						Accumulators[Index]->ActiveFrameNumber = InFrameNumber;
-						Accumulators[Index]->ActivePassIdentifier = InPassIdentifier;
-						Accumulators[Index]->bIsActive = true;
-						Accumulators[Index]->TaskPrereq = UE::Tasks::FTask();
-						AvailableIndex = Index;
-						break;
-					}
+					AvailableIndex = Index;
+					break;
 				}
 			}
 
@@ -201,10 +198,16 @@ namespace UE::MovieGraph::DefaultRenderer
 
 				TSharedPtr<AccumulatorType> NewAccumulatorInstance = MakeShared<AccumulatorType>();
 				Accumulators.Add(MakeShared<UE::MovieGraph::DefaultRenderer::FSurfaceAccumulatorPool::FInstance>(NewAccumulatorInstance));
-
 				UE_LOG(LogMovieRenderPipeline, Log, TEXT("Allocated a Accumulator for Pool %s, New Pool Count: %d"), *AccumulatorType::GetName().ToString(), Accumulators.Num());
 			}
 
+			// Ensure we've either updated the reused accumulator to our new data
+			// or configured our new accumulator for the first time.
+			Accumulators[AvailableIndex]->SetIsActive(true);
+			Accumulators[AvailableIndex]->ActiveFrameNumber = InFrameNumber;
+			Accumulators[AvailableIndex]->ActivePassIdentifier = InPassIdentifier;
+			Accumulators[AvailableIndex]->TaskPrereq = UE::Tasks::FTask();
+
 			return Accumulators[AvailableIndex];
 		}
 	};
diff --git a/Engine/Plugins/MovieScene/MovieRenderPipeline/Source/MovieRenderPipelineRenderPasses/Private/Graph/Renderers/MovieGraphImagePassBase.cpp b/Engine/Plugins/MovieScene/MovieRenderPipeline/Source/MovieRenderPipelineRenderPasses/Private/Graph/Renderers/MovieGraphImagePassBase.cpp
--- a/Engine/Plugins/MovieScene/MovieRenderPipeline/Source/MovieRenderPipelineRenderPasses/Private/Graph/Renderers/MovieGraphImagePassBase.cpp
+++ b/Engine/Plugins/MovieScene/MovieRenderPipeline/Source/MovieRenderPipelineRenderPasses/Private/Graph/Renderers/MovieGraphImagePassBase.cpp
@@ -478,7 +478,7 @@ void FMovieGraphImagePassBase::PostRendererSubmission(
 	}
 
 	FMoviePipelineAccumulatorPoolPtr SampleAccumulatorPool = GraphRenderer->GetOrCreateAccumulatorPool<FImageOverlappedAccumulator>();
-	UE::MovieGraph::DefaultRenderer::FSurfaceAccumulatorPool::FInstancePtr AccumulatorInstance = SampleAccumulatorPool->GetAccumulatorInstance_GameThread<FImageOverlappedAccumulator>(InSampleState.TraversalContext.Time.RenderedFrameNumber, InSampleState.TraversalContext.RenderDataIdentifier);
+	UE::MovieGraph::DefaultRenderer::FSurfaceAccumulatorPool::FInstancePtr AccumulatorInstance = SampleAccumulatorPool->GetAccumulatorInstance_GameThread<FImageOverlappedAccumulator>(InSampleState.TraversalContext.Time.OutputFrameNumber, InSampleState.TraversalContext.RenderDataIdentifier);
 	
 	FMoviePipelineSurfaceQueuePtr LocalSurfaceQueue = GraphRenderer->GetOrCreateSurfaceQueue(InRenderTargetInitParams);
 	LocalSurfaceQueue->BlockUntilAnyAvailable();
@@ -502,7 +502,7 @@ void FMovieGraphImagePassBase::PostRendererSubmission(
 			if (AccumulationArgs.bIsLastSample)
 			{
 				// Final sample has now been executed, free the accumulator for reuse.
-				AccumulatorInstance->bIsActive = false;
+				AccumulatorInstance->SetIsActive(false);
 			}
 		}, AccumulatorInstance->TaskPrereq);
 
@@ -601,7 +601,7 @@ TFunction<void(TUniquePtr<FImagePixelData>&&)> FMovieGraphImagePassBase::MakeFor
 			if (AccumulationArgs.bIsLastSample)
 			{
 				// Final sample has now been executed, free the accumulator for reuse.
-				AccumulatorInstance->bIsActive = false;
+				AccumulatorInstance->SetIsActive(false);
 			}
 		}, AccumulatorInstance->TaskPrereq);
 
