MovieGraph: Fixed an issue where in some sequences warm-up frames could be evaluated on the wrong frame (such as an out of bounds frame that caused the warm-up view to be different than the first frame of your sequence).

diff --git a/Engine/Plugins/MovieScene/MovieRenderPipeline/Source/MovieRenderPipelineCore/Private/Graph/MovieGraphSequenceDataSource.cpp b/Engine/Plugins/MovieScene/MovieRenderPipeline/Source/MovieRenderPipelineCore/Private/Graph/MovieGraphSequenceDataSource.cpp
--- a/Engine/Plugins/MovieScene/MovieRenderPipeline/Source/MovieRenderPipelineCore/Private/Graph/MovieGraphSequenceDataSource.cpp
+++ b/Engine/Plugins/MovieScene/MovieRenderPipeline/Source/MovieRenderPipelineCore/Private/Graph/MovieGraphSequenceDataSource.cpp
@@ -309,23 +309,42 @@ void UMovieGraphSequenceDataSource::ExpandShot(const TObjectPtr<UMoviePipelineEx
 
 void UMovieGraphSequenceDataSource::SyncDataSourceTime(const FFrameTime& InTime)
 {
-	FFrameRate TickResolution = LevelSequenceActor->GetSequence()->GetMovieScene()->GetTickResolution();
-	CustomSequenceTimeController->SetCachedFrameTiming(FQualifiedFrameTime(InTime, TickResolution));
+	if(LevelSequenceActor && LevelSequenceActor->GetSequencePlayer())
+	{
+		FFrameRate TickResolution = LevelSequenceActor->GetSequence()->GetMovieScene()->GetTickResolution();
+		CustomSequenceTimeController->SetCachedFrameTiming(FQualifiedFrameTime(InTime, TickResolution));
+	}
 }
 void UMovieGraphSequenceDataSource::PlayDataSource()
 {
-	LevelSequenceActor->GetSequencePlayer()->Play();
+	if(LevelSequenceActor && LevelSequenceActor->GetSequencePlayer())
+	{
+		LevelSequenceActor->GetSequencePlayer()->Play();
+	}
 }
 
 void UMovieGraphSequenceDataSource::PauseDataSource() 
 {
-	LevelSequenceActor->GetSequencePlayer()->Pause();
+	if(LevelSequenceActor && LevelSequenceActor->GetSequencePlayer())
+	{
+		LevelSequenceActor->GetSequencePlayer()->Pause();
+	}
 }
 
 void UMovieGraphSequenceDataSource::JumpDataSource(const FFrameTime& InTimeToJumpTo) 
 {
-	LevelSequenceActor->GetSequencePlayer()->SetPlaybackPosition(FMovieSceneSequencePlaybackParams(InTimeToJumpTo, EUpdatePositionMethod::Jump));
+	if(LevelSequenceActor && LevelSequenceActor->GetSequencePlayer())
+	{
 
+		// SetPlaybackPosition takes time in display rate and not tick resolution, so we convert
+		// it from tick resolution to display rate. JumpDataSource provides it in tick resolution
+		// to be consistent with other functions in this class.
+		FFrameRate TickResolution = LevelSequenceActor->GetSequence()->GetMovieScene()->GetTickResolution();
+		FFrameRate DisplayRate = LevelSequenceActor->GetSequence()->GetMovieScene()->GetDisplayRate();
+		FFrameTime RequestTime = FFrameRate::TransformTime(InTimeToJumpTo, TickResolution, DisplayRate);
+
+		LevelSequenceActor->GetSequencePlayer()->SetPlaybackPosition(FMovieSceneSequencePlaybackParams(RequestTime, EUpdatePositionMethod::Jump));
+	}
 }
 
 namespace UE::MovieGraph
