MoviePipeline: Fixed issue that caused the input/output nodes in a newly-created graph to not reflect renaming until the graph is loaded again after an editor restart.

Graphs are now created via (duplicated) templates, so the duplicated nodes needed to register delegates just like non-duplicated nodes.

#rnx

diff --git a/Engine/Plugins/MovieScene/MovieRenderPipeline/Source/MovieRenderPipelineCore/Private/Graph/MovieGraphNode.cpp b/Engine/Plugins/MovieScene/MovieRenderPipeline/Source/MovieRenderPipelineCore/Private/Graph/MovieGraphNode.cpp
--- a/Engine/Plugins/MovieScene/MovieRenderPipeline/Source/MovieRenderPipelineCore/Private/Graph/MovieGraphNode.cpp
+++ b/Engine/Plugins/MovieScene/MovieRenderPipeline/Source/MovieRenderPipelineCore/Private/Graph/MovieGraphNode.cpp
@@ -401,6 +401,14 @@ void UMovieGraphNode::PostLoad()
 	RegisterDelegates();
 }
 
+void UMovieGraphNode::PostDuplicate(EDuplicateMode::Type DuplicateMode)
+{
+	Super::PostDuplicate(DuplicateMode);
+
+	// New graphs are typically created by duplicating a template. Make sure the nodes in the template have their delegates registered.
+	RegisterDelegates();
+}
+
 bool UMovieGraphNode::SetDynamicPropertyValue(const FName PropertyName, const FString& InNewValue)
 {
 	return DynamicProperties.SetValueSerializedString(PropertyName, InNewValue) == EPropertyBagResult::Success;
diff --git a/Engine/Plugins/MovieScene/MovieRenderPipeline/Source/MovieRenderPipelineCore/Private/Graph/Nodes/MovieGraphInputNode.cpp b/Engine/Plugins/MovieScene/MovieRenderPipeline/Source/MovieRenderPipelineCore/Private/Graph/Nodes/MovieGraphInputNode.cpp
--- a/Engine/Plugins/MovieScene/MovieRenderPipeline/Source/MovieRenderPipelineCore/Private/Graph/Nodes/MovieGraphInputNode.cpp
+++ b/Engine/Plugins/MovieScene/MovieRenderPipeline/Source/MovieRenderPipelineCore/Private/Graph/Nodes/MovieGraphInputNode.cpp
@@ -110,7 +110,7 @@ FSlateIcon UMovieGraphInputNode::GetIconAndTint(FLinearColor& OutColor) const
 }
 #endif // WITH_EDITOR
 
-void UMovieGraphInputNode::RegisterDelegates() const
+void UMovieGraphInputNode::RegisterDelegates()
 {
 	Super::RegisterDelegates();
 	
@@ -123,11 +123,12 @@ void UMovieGraphInputNode::RegisterDelegates() const
 	}
 }
 
-void UMovieGraphInputNode::RegisterDelegates(UMovieGraphInput* Input) const
+void UMovieGraphInputNode::RegisterDelegates(UMovieGraphInput* Input)
 {
 #if WITH_EDITOR
 	if (Input)
 	{
+		Input->OnMovieGraphInputChangedDelegate.RemoveAll(this);
 		Input->OnMovieGraphInputChangedDelegate.AddUObject(this, &UMovieGraphInputNode::UpdateExistingPins);
 	}
 #endif
diff --git a/Engine/Plugins/MovieScene/MovieRenderPipeline/Source/MovieRenderPipelineCore/Private/Graph/Nodes/MovieGraphOutputNode.cpp b/Engine/Plugins/MovieScene/MovieRenderPipeline/Source/MovieRenderPipelineCore/Private/Graph/Nodes/MovieGraphOutputNode.cpp
--- a/Engine/Plugins/MovieScene/MovieRenderPipeline/Source/MovieRenderPipelineCore/Private/Graph/Nodes/MovieGraphOutputNode.cpp
+++ b/Engine/Plugins/MovieScene/MovieRenderPipeline/Source/MovieRenderPipelineCore/Private/Graph/Nodes/MovieGraphOutputNode.cpp
@@ -66,7 +66,7 @@ FSlateIcon UMovieGraphOutputNode::GetIconAndTint(FLinearColor& OutColor) const
 }
 #endif // WITH_EDITOR
 
-void UMovieGraphOutputNode::RegisterDelegates() const
+void UMovieGraphOutputNode::RegisterDelegates()
 {
 	Super::RegisterDelegates();
 	
@@ -79,11 +79,12 @@ void UMovieGraphOutputNode::RegisterDelegates() const
 	}
 }
 
-void UMovieGraphOutputNode::RegisterDelegates(UMovieGraphOutput* Output) const
+void UMovieGraphOutputNode::RegisterDelegates(UMovieGraphOutput* Output)
 {
 #if WITH_EDITOR
 	if (Output)
 	{
+		Output->OnMovieGraphOutputChangedDelegate.RemoveAll(this);
 		Output->OnMovieGraphOutputChangedDelegate.AddUObject(this, &UMovieGraphOutputNode::UpdateExistingPins);
 	}
 #endif
diff --git a/Engine/Plugins/MovieScene/MovieRenderPipeline/Source/MovieRenderPipelineCore/Private/Graph/Nodes/MovieGraphVariableNode.cpp b/Engine/Plugins/MovieScene/MovieRenderPipeline/Source/MovieRenderPipelineCore/Private/Graph/Nodes/MovieGraphVariableNode.cpp
--- a/Engine/Plugins/MovieScene/MovieRenderPipeline/Source/MovieRenderPipelineCore/Private/Graph/Nodes/MovieGraphVariableNode.cpp
+++ b/Engine/Plugins/MovieScene/MovieRenderPipeline/Source/MovieRenderPipelineCore/Private/Graph/Nodes/MovieGraphVariableNode.cpp
@@ -105,13 +105,14 @@ FSlateIcon UMovieGraphVariableNode::GetIconAndTint(FLinearColor& OutColor) const
 }
 #endif // WITH_EDITOR
 
-void UMovieGraphVariableNode::RegisterDelegates() const
+void UMovieGraphVariableNode::RegisterDelegates()
 {
 	Super::RegisterDelegates();
 
 #if WITH_EDITOR
 	if (GraphVariable)
 	{
+		GraphVariable->OnMovieGraphVariableChangedDelegate.RemoveAll(this);
 		GraphVariable->OnMovieGraphVariableChangedDelegate.AddUObject(this, &UMovieGraphVariableNode::UpdateOutputPin);
 	}
 #endif
diff --git a/Engine/Plugins/MovieScene/MovieRenderPipeline/Source/MovieRenderPipelineCore/Public/Graph/MovieGraphNode.h b/Engine/Plugins/MovieScene/MovieRenderPipeline/Source/MovieRenderPipelineCore/Public/Graph/MovieGraphNode.h
--- a/Engine/Plugins/MovieScene/MovieRenderPipeline/Source/MovieRenderPipelineCore/Public/Graph/MovieGraphNode.h
+++ b/Engine/Plugins/MovieScene/MovieRenderPipeline/Source/MovieRenderPipelineCore/Public/Graph/MovieGraphNode.h
@@ -292,6 +292,7 @@ public:
 
 	//~ Begin UObject Interface
 	virtual void PostLoad() override;
+	virtual void PostDuplicate(EDuplicateMode::Type DuplicateMode) override;
 	//~ End UObject Interface
 
 public:
@@ -324,7 +325,7 @@ protected:
 	virtual TArray<FMovieGraphPinProperties> GetExposedPinProperties() const;
 
 	/** Register any delegates that need to be set up on the node. Called in PostLoad(). */
-	virtual void RegisterDelegates() const { }
+	virtual void RegisterDelegates() { }
 
 protected:
 	UPROPERTY()
diff --git a/Engine/Plugins/MovieScene/MovieRenderPipeline/Source/MovieRenderPipelineCore/Public/Graph/Nodes/MovieGraphInputNode.h b/Engine/Plugins/MovieScene/MovieRenderPipeline/Source/MovieRenderPipelineCore/Public/Graph/Nodes/MovieGraphInputNode.h
--- a/Engine/Plugins/MovieScene/MovieRenderPipeline/Source/MovieRenderPipelineCore/Public/Graph/Nodes/MovieGraphInputNode.h
+++ b/Engine/Plugins/MovieScene/MovieRenderPipeline/Source/MovieRenderPipelineCore/Public/Graph/Nodes/MovieGraphInputNode.h
@@ -31,10 +31,10 @@ public:
 	//~ End UMovieGraphNode interface
 
 private:
-	virtual void RegisterDelegates() const override;
+	virtual void RegisterDelegates() override;
 
 	/** Register delegates for the provided input member. */
-	void RegisterDelegates(UMovieGraphInput* Input) const;
+	void RegisterDelegates(UMovieGraphInput* Input);
 
 	/** Update data (name, etc) on all existing output pins on this node to reflect the input members on the graph. */
 	void UpdateExistingPins(UMovieGraphMember* ChangedVariable) const;
diff --git a/Engine/Plugins/MovieScene/MovieRenderPipeline/Source/MovieRenderPipelineCore/Public/Graph/Nodes/MovieGraphOutputNode.h b/Engine/Plugins/MovieScene/MovieRenderPipeline/Source/MovieRenderPipelineCore/Public/Graph/Nodes/MovieGraphOutputNode.h
--- a/Engine/Plugins/MovieScene/MovieRenderPipeline/Source/MovieRenderPipelineCore/Public/Graph/Nodes/MovieGraphOutputNode.h
+++ b/Engine/Plugins/MovieScene/MovieRenderPipeline/Source/MovieRenderPipelineCore/Public/Graph/Nodes/MovieGraphOutputNode.h
@@ -28,10 +28,10 @@ public:
 #endif
 
 private:
-	virtual void RegisterDelegates() const override;
+	virtual void RegisterDelegates() override;
 
 	/** Register delegates for the provided output member. */
-	void RegisterDelegates(UMovieGraphOutput* Output) const;
+	void RegisterDelegates(UMovieGraphOutput* Output);
 
 	/** Update data (name, etc) on all existing input pins on this node to reflect the output members on the graph. */
 	void UpdateExistingPins(UMovieGraphMember* ChangedVariable) const;
diff --git a/Engine/Plugins/MovieScene/MovieRenderPipeline/Source/MovieRenderPipelineCore/Public/Graph/Nodes/MovieGraphVariableNode.h b/Engine/Plugins/MovieScene/MovieRenderPipeline/Source/MovieRenderPipelineCore/Public/Graph/Nodes/MovieGraphVariableNode.h
--- a/Engine/Plugins/MovieScene/MovieRenderPipeline/Source/MovieRenderPipelineCore/Public/Graph/Nodes/MovieGraphVariableNode.h
+++ b/Engine/Plugins/MovieScene/MovieRenderPipeline/Source/MovieRenderPipelineCore/Public/Graph/Nodes/MovieGraphVariableNode.h
@@ -45,7 +45,7 @@ public:
 #endif
 
 private:
-	virtual void RegisterDelegates() const override;
+	virtual void RegisterDelegates() override;
 	
 	/** Updates the output pin on the node to match the provided variable. */
 	void UpdateOutputPin(UMovieGraphMember* ChangedVariable) const;
