nDisplay, MRQ:
- Fixed an issue with camera settings' rules that conflicted with nDisplay.

diff --git a/Engine/Plugins/MovieScene/MovieRenderPipeline/Source/MovieRenderPipelineRenderPasses/Private/MoviePipelineDeferredPasses.cpp b/Engine/Plugins/MovieScene/MovieRenderPipeline/Source/MovieRenderPipelineRenderPasses/Private/MoviePipelineDeferredPasses.cpp
--- a/Engine/Plugins/MovieScene/MovieRenderPipeline/Source/MovieRenderPipelineRenderPasses/Private/MoviePipelineDeferredPasses.cpp
+++ b/Engine/Plugins/MovieScene/MovieRenderPipeline/Source/MovieRenderPipelineRenderPasses/Private/MoviePipelineDeferredPasses.cpp
@@ -98,11 +98,24 @@ int32 UMoviePipelineDeferredPassBase::GetNumCamerasToRender() const
 	return CameraSettings->bRenderAllCameras ? CurrentShot->SidecarCameras.Num() : 1;
 }
 
+int32 UMoviePipelineDeferredPassBase::GetCameraIndexForRenderPass(const int32 InCameraIndex) const
+{
+	UMoviePipelineExecutorShot* CurrentShot = GetPipeline()->GetActiveShotList()[GetPipeline()->GetCurrentShotIndex()];
+	UMoviePipelineCameraSetting* CameraSettings = GetPipeline()->FindOrAddSettingForShot<UMoviePipelineCameraSetting>(CurrentShot);
+
+	// If we're not rendering all cameras, we need to pass -1 so we pick up the real camera name.
+	return CameraSettings->bRenderAllCameras ? InCameraIndex : -1;
+}
+
 FString UMoviePipelineDeferredPassBase::GetCameraName(const int32 InCameraIndex) const
 {
 	UMoviePipelineExecutorShot* CurrentShot = GetPipeline()->GetActiveShotList()[GetPipeline()->GetCurrentShotIndex()];
+	UMoviePipelineCameraSetting* CameraSettings = GetPipeline()->FindOrAddSettingForShot<UMoviePipelineCameraSetting>(CurrentShot);
+
+	// If we're not rendering all cameras, we need to pass -1 so we pick up the real camera name.
+	const int32 LocalCameraIndex = CameraSettings->bRenderAllCameras ? InCameraIndex : -1;
 
-	return CurrentShot->GetCameraName(InCameraIndex);
+	return CurrentShot->GetCameraName(LocalCameraIndex);
 }
 
 FString UMoviePipelineDeferredPassBase::GetCameraNameOverride(const int32 InCameraIndex) const
@@ -423,17 +436,13 @@ void UMoviePipelineDeferredPassBase::GatherOutputPassesImpl(TArray<FMoviePipelin
 	// Super::GatherOutputPassesImpl(ExpectedRenderPasses);
 
 	UMoviePipelineExecutorShot* CurrentShot = GetPipeline()->GetActiveShotList()[GetPipeline()->GetCurrentShotIndex()];
-	UMoviePipelineCameraSetting* CameraSettings = GetPipeline()->FindOrAddSettingForShot<UMoviePipelineCameraSetting>(CurrentShot);
 
 	const int32 NumCameras = GetNumCamerasToRender();
 	for (int32 CameraIndex = 0; CameraIndex < NumCameras; CameraIndex++)
 	{
 		FMoviePipelinePassIdentifier PassIdentifierForCurrentCamera;
 		PassIdentifierForCurrentCamera.Name = PassIdentifier.Name;
-
-		// If we're not rendering all cameras, we need to pass -1 so we pick up the real camera name.
-		int32 LocalCameraIndex = CameraSettings->bRenderAllCameras ? CameraIndex : -1;
-		PassIdentifierForCurrentCamera.CameraName = GetCameraName(LocalCameraIndex);
+		PassIdentifierForCurrentCamera.CameraName = GetCameraName(CameraIndex);
 
 		// Add the default backbuffer
 		if (bRenderMainPass)
@@ -494,7 +503,6 @@ void UMoviePipelineDeferredPassBase::RenderSample_GameThreadImpl(const FMoviePip
 
 	const int32 NumCameras = GetNumCamerasToRender();
 	UMoviePipelineExecutorShot* CurrentShot = GetPipeline()->GetActiveShotList()[GetPipeline()->GetCurrentShotIndex()];
-	UMoviePipelineCameraSetting* CameraSettings = GetPipeline()->FindOrAddSettingForShot<UMoviePipelineCameraSetting>(CurrentShot);
 
 	for (int32 CameraIndex = 0; CameraIndex < NumCameras; CameraIndex++)
 	{
@@ -502,19 +510,19 @@ void UMoviePipelineDeferredPassBase::RenderSample_GameThreadImpl(const FMoviePip
 		PassIdentifierForCurrentCamera.Name = PassIdentifier.Name;
 
 		// If we're not rendering all cameras, we need to pass -1 so we pick up the real camera name.
-		int32 LocalCameraIndex = CameraSettings->bRenderAllCameras ? CameraIndex : -1;
-		PassIdentifierForCurrentCamera.CameraName = GetCameraName(LocalCameraIndex);
+		const int32 CameraIndexForRenderPass = GetCameraIndexForRenderPass(CameraIndex);
+		PassIdentifierForCurrentCamera.CameraName = GetCameraName(CameraIndex);
 
 		// Main Render Pass
 		if (bRenderMainPass)
 		{
-			FMoviePipelineRenderPassMetrics InOutSampleState = GetRenderPassMetricsForCamera(LocalCameraIndex, InSampleState);
+			FMoviePipelineRenderPassMetrics InOutSampleState = GetRenderPassMetricsForCamera(CameraIndexForRenderPass, InSampleState);
 			// InOutSampleState.OutputState.CameraCount = NumCameras;
-			InOutSampleState.OutputState.CameraIndex = LocalCameraIndex;
-			InOutSampleState.OutputState.CameraNameOverride = GetCameraNameOverride(LocalCameraIndex);
+			InOutSampleState.OutputState.CameraIndex = CameraIndexForRenderPass;
+			InOutSampleState.OutputState.CameraNameOverride = GetCameraNameOverride(CameraIndex);
 
 			UE::MoviePipeline::FDeferredPassRenderStatePayload Payload;
-			Payload.CameraIndex = LocalCameraIndex;
+			Payload.CameraIndex = CameraIndexForRenderPass;
 			Payload.TileIndex = InOutSampleState.TileIndexes;
 
 			// Main renders use index 0.
@@ -571,9 +579,9 @@ void UMoviePipelineDeferredPassBase::RenderSample_GameThreadImpl(const FMoviePip
 
 		// Now do the stencil layer submission (which doesn't support additional post processing materials)
 		{
-			FMoviePipelineRenderPassMetrics InOutSampleState = GetRenderPassMetricsForCamera(LocalCameraIndex, InSampleState);
-			InOutSampleState.OutputState.CameraIndex = LocalCameraIndex;
-			InOutSampleState.OutputState.CameraNameOverride = GetCameraNameOverride(LocalCameraIndex);
+			FMoviePipelineRenderPassMetrics InOutSampleState = GetRenderPassMetricsForCamera(CameraIndexForRenderPass, InSampleState);
+			InOutSampleState.OutputState.CameraIndex = CameraIndexForRenderPass;
+			InOutSampleState.OutputState.CameraNameOverride = GetCameraNameOverride(CameraIndex);
 
 			struct FStencilValues
 			{
@@ -670,7 +678,7 @@ void UMoviePipelineDeferredPassBase::RenderSample_GameThreadImpl(const FMoviePip
 				if (StencilLayerMaterial)
 				{
 					UE::MoviePipeline::FDeferredPassRenderStatePayload Payload;
-					Payload.CameraIndex = LocalCameraIndex;
+					Payload.CameraIndex = CameraIndexForRenderPass;
 					Payload.TileIndex = InOutSampleState.TileIndexes;
 					Payload.SceneViewIndex = StencilLayerIndex + (bRenderMainPass ? 1 : 0);
 					TSharedPtr<FSceneViewFamilyContext> ViewFamily = CalculateViewFamily(InOutSampleState, &Payload);
diff --git a/Engine/Plugins/MovieScene/MovieRenderPipeline/Source/MovieRenderPipelineRenderPasses/Public/MoviePipelineDeferredPasses.h b/Engine/Plugins/MovieScene/MovieRenderPipeline/Source/MovieRenderPipelineRenderPasses/Public/MoviePipelineDeferredPasses.h
--- a/Engine/Plugins/MovieScene/MovieRenderPipeline/Source/MovieRenderPipelineRenderPasses/Public/MoviePipelineDeferredPasses.h
+++ b/Engine/Plugins/MovieScene/MovieRenderPipeline/Source/MovieRenderPipelineRenderPasses/Public/MoviePipelineDeferredPasses.h
@@ -66,6 +66,7 @@ protected:
 	virtual void PostRendererSubmission(const FMoviePipelineRenderPassMetrics& InSampleState, const FMoviePipelinePassIdentifier InPassIdentifier, const int32 InSortingOrder, FCanvas& InCanvas);
 
 	virtual int32 GetNumCamerasToRender() const;
+	virtual int32 GetCameraIndexForRenderPass(const int32 InCameraIndex) const;
 	virtual FString GetCameraName(const int32 InCameraIndex) const;
 	virtual FString GetCameraNameOverride(const int32 InCameraIndex) const;
 
diff --git a/Engine/Plugins/Runtime/nDisplay/Source/DisplayClusterMoviePipeline/Private/DisplayClusterMoviePipelineViewportPass.cpp b/Engine/Plugins/Runtime/nDisplay/Source/DisplayClusterMoviePipeline/Private/DisplayClusterMoviePipelineViewportPass.cpp
--- a/Engine/Plugins/Runtime/nDisplay/Source/DisplayClusterMoviePipeline/Private/DisplayClusterMoviePipelineViewportPass.cpp
+++ b/Engine/Plugins/Runtime/nDisplay/Source/DisplayClusterMoviePipeline/Private/DisplayClusterMoviePipelineViewportPass.cpp
@@ -215,27 +215,24 @@ void UDisplayClusterMoviePipelineViewportPassBase::BlendPostProcessSettings(FSce
 
 int32 UDisplayClusterMoviePipelineViewportPassBase::GetNumCamerasToRender() const
 {
+	// nDisplay always uses viewports as a set of cameras.
 	return DisplayClusterViewports.Num();
 }
 
-FString UDisplayClusterMoviePipelineViewportPassBase::GetCameraName(const int32 InCameraIndex) const
+int32 UDisplayClusterMoviePipelineViewportPassBase::GetCameraIndexForRenderPass(const int32 InCameraIndex) const
 {
-	if (InCameraIndex >= 0 && InCameraIndex < DisplayClusterViewports.Num())
-	{
-		return DisplayClusterViewports[InCameraIndex];
-	}
+	// nDisplay always uses viewports as a set of cameras.
+	return InCameraIndex;
+}
 
-	return TEXT("UndefinedCameraName");
+FString UDisplayClusterMoviePipelineViewportPassBase::GetCameraName(const int32 InCameraIndex) const
+{
+	return DisplayClusterViewports.IsValidIndex(InCameraIndex) ? DisplayClusterViewports[InCameraIndex] : TEXT("UndefinedCameraName");
 }
 
 FString UDisplayClusterMoviePipelineViewportPassBase::GetCameraNameOverride(const int32 InCameraIndex) const
 {
-	if (DisplayClusterViewports.IsValidIndex(InCameraIndex))
-	{
-		return DisplayClusterViewports[InCameraIndex];
-	}
-
-	return TEXT("");
+	return DisplayClusterViewports.IsValidIndex(InCameraIndex) ? DisplayClusterViewports[InCameraIndex] : TEXT("");
 }
 
 void UDisplayClusterMoviePipelineViewportPassBase::ModifyProjectionMatrixForTiling(const FMoviePipelineRenderPassMetrics& InSampleState, const bool bInOrthographic, FMatrix& InOutProjectionMatrix, float& OutDoFSensorScale) const
diff --git a/Engine/Plugins/Runtime/nDisplay/Source/DisplayClusterMoviePipeline/Public/DisplayClusterMoviePipelineViewportPass.h b/Engine/Plugins/Runtime/nDisplay/Source/DisplayClusterMoviePipeline/Public/DisplayClusterMoviePipelineViewportPass.h
--- a/Engine/Plugins/Runtime/nDisplay/Source/DisplayClusterMoviePipeline/Public/DisplayClusterMoviePipelineViewportPass.h
+++ b/Engine/Plugins/Runtime/nDisplay/Source/DisplayClusterMoviePipeline/Public/DisplayClusterMoviePipelineViewportPass.h
@@ -37,6 +37,7 @@ public:
 	virtual void PostRendererSubmission(const FMoviePipelineRenderPassMetrics& InSampleState, const FMoviePipelinePassIdentifier InPassIdentifier, const int32 InSortingOrder, FCanvas& InCanvas) override;
 
 	virtual int32 GetNumCamerasToRender() const override;
+	virtual int32 GetCameraIndexForRenderPass(const int32 InCameraIndex) const override;
 	virtual FString GetCameraName(const int32 InCameraIndex) const override;
 	virtual FString GetCameraNameOverride(const int32 InCameraIndex) const override;
 
