Movie Graph:

-Fix: Unable to Set Graph as Default Config

diff --git a/Engine/Plugins/MovieScene/MovieRenderPipeline/Source/MovieRenderPipelineEditor/Private/Widgets/SMoviePipelineQueueEditor.cpp b/Engine/Plugins/MovieScene/MovieRenderPipeline/Source/MovieRenderPipelineEditor/Private/Widgets/SMoviePipelineQueueEditor.cpp
--- a/Engine/Plugins/MovieScene/MovieRenderPipeline/Source/MovieRenderPipelineEditor/Private/Widgets/SMoviePipelineQueueEditor.cpp
+++ b/Engine/Plugins/MovieScene/MovieRenderPipeline/Source/MovieRenderPipelineEditor/Private/Widgets/SMoviePipelineQueueEditor.cpp
@@ -48,6 +48,7 @@
 #include "Engine/EngineTypes.h"
 #include "Framework/Application/SlateApplication.h"
 #include "SMoviePipelineConfigPanel.h"
+#include "Graph/MovieGraphPipeline.h"
 #include "Widgets/SWindow.h"
 #include "HAL/FileManager.h"
 #include "Widgets/Layout/SBox.h"
@@ -319,14 +320,7 @@ public:
 			return;
 		}
 
-		// Note: Setting the graph preset will transition the job to use a graph-based configuration
-		// Use the default graph specified in Project Settings.
-		const UMovieRenderPipelineProjectSettings* ProjectSettings = GetDefault<UMovieRenderPipelineProjectSettings>();
-		const TSoftObjectPtr<UMovieGraphConfig> ProjectDefaultGraph = ProjectSettings->DefaultGraph;
-		if (const UMovieGraphConfig* DefaultGraph = ProjectDefaultGraph.LoadSynchronous())
-		{
-			Job->SetGraphPreset(DefaultGraph);
-		}
+		SMoviePipelineQueueEditor::AssignDefaultGraphPresetToJob(Job);
 	}
 
 	void OnCreateNewGraphAndAssign() const
@@ -509,17 +503,16 @@ public:
 
 		MenuBuilder.BeginSection(NAME_None, LOCTEXT("CurrentConfig_MenuSection", "Current Configuration"));
 		{
-			MenuBuilder.AddMenuEntry(
-				LOCTEXT("ClearConfig_Label", "Clear Config"),
-				LOCTEXT("ClearConfig_Tooltip", "Resets the changes to the config and goes back to the defaults."),
-				FSlateIcon(),
-				FUIAction(InNewConfig),
-				NAME_None,
-				EUserInterfaceActionType::Button
-			);
-
 			if (!TargetJob->IsUsingGraphConfiguration())
 			{
+				MenuBuilder.AddMenuEntry(
+					LOCTEXT("ClearConfig_Label", "Clear Config"),
+					LOCTEXT("ClearConfig_Tooltip", "Resets the changes to the config and goes back to the defaults."),
+					FSlateIcon(),
+					FUIAction(InNewConfig),
+					NAME_None,
+					EUserInterfaceActionType::Button
+				);
 				MenuBuilder.AddMenuEntry(
 					LOCTEXT("ReplaceWithGraph_Label", "Replace with Graph (Experimental)"),
 					LOCTEXT("ReplaceWithGraph_Tooltip", "Replaces the current configuration with a new graph representation."),
@@ -529,6 +522,17 @@ public:
 					EUserInterfaceActionType::Button
 				);
 			}
+			else
+			{
+				MenuBuilder.AddMenuEntry(
+					LOCTEXT("ReplaceWithPreset_Label", "Replace with Preset"),
+					LOCTEXT("ReplaceWithPreset_Tooltip", "Replaces the current configuration with a new default (non-graph) config."),
+					FSlateIcon(),
+					FUIAction(InNewConfig),
+					NAME_None,
+					EUserInterfaceActionType::Button
+				);
+			}
 		}
 
 		MenuBuilder.BeginSection(NAME_None, LOCTEXT("ImportConfig_MenuSection", "Import Configuration"));
@@ -1381,6 +1385,16 @@ TSharedRef<SWidget> SMoviePipelineQueueEditor::OnGenerateNewJobFromAssetMenu()
 	return MenuBuilder.MakeWidget();
 }
 
+void SMoviePipelineQueueEditor::AssignDefaultGraphPresetToJob(UMoviePipelineExecutorJob* InJob)
+{
+	const UMovieRenderPipelineProjectSettings* ProjectSettings = GetDefault<UMovieRenderPipelineProjectSettings>();
+	const TSoftObjectPtr<UMovieGraphConfig> ProjectDefaultGraph = ProjectSettings->DefaultGraph;
+	if (const UMovieGraphConfig* DefaultGraph = ProjectDefaultGraph.LoadSynchronous())
+	{
+		InJob->SetGraphPreset(DefaultGraph);
+	}
+}
+
 UE_ENABLE_OPTIMIZATION_SHIP
 
 void SMoviePipelineQueueEditor::OnCreateJobFromAsset(const FAssetData& InAsset)
@@ -1418,6 +1432,7 @@ void SMoviePipelineQueueEditor::OnCreateJobFromAsset(const FAssetData& InAsset)
 	}
 
 	const UMovieRenderPipelineProjectSettings* ProjectSettings = GetDefault<UMovieRenderPipelineProjectSettings>();
+	const UClass* DefaultPipeline = Cast<UClass>(ProjectSettings->DefaultPipeline.TryLoad());
 	for (UMoviePipelineExecutorJob* NewJob : NewJobs)
 	{
 		PendingJobsToSelect.Add(NewJob);
@@ -1434,6 +1449,12 @@ void SMoviePipelineQueueEditor::OnCreateJobFromAsset(const FAssetData& InAsset)
 		// Ensure the job has the settings specified by the project settings added. If they're already added
 		// we don't modify the object so that we don't make it confused about whether or not you've modified the preset.
 		UMoviePipelineEditorBlueprintLibrary::EnsureJobHasDefaultSettings(NewJob);
+
+		// If the default class is a movie graph, assign the default graph
+		if (DefaultPipeline && DefaultPipeline == UMovieGraphPipeline::StaticClass())
+		{
+			AssignDefaultGraphPresetToJob(NewJob);
+		}
 	}
 }
 
@@ -1519,7 +1540,6 @@ void SMoviePipelineQueueEditor::ReconstructTree()
 	TreeView->RequestTreeRefresh();
 }
 
-
 FReply SMoviePipelineQueueEditor::OnKeyDown(const FGeometry& MyGeometry, const FKeyEvent& InKeyEvent)
 {
 	if (CommandList->ProcessCommandBindings(InKeyEvent))
diff --git a/Engine/Plugins/MovieScene/MovieRenderPipeline/Source/MovieRenderPipelineEditor/Private/Widgets/SMoviePipelineQueueEditor.h b/Engine/Plugins/MovieScene/MovieRenderPipeline/Source/MovieRenderPipelineEditor/Private/Widgets/SMoviePipelineQueueEditor.h
--- a/Engine/Plugins/MovieScene/MovieRenderPipeline/Source/MovieRenderPipelineEditor/Private/Widgets/SMoviePipelineQueueEditor.h
+++ b/Engine/Plugins/MovieScene/MovieRenderPipeline/Source/MovieRenderPipelineEditor/Private/Widgets/SMoviePipelineQueueEditor.h
@@ -51,6 +51,9 @@ public:
 	TArray<TSharedPtr<IMoviePipelineQueueTreeItem>> GetSelectedItems() const { return TreeView->GetSelectedItems(); }
 	void SetSelectedJobs(const TArray<UMoviePipelineExecutorJob*>& InJobs) { PendingJobsToSelect = InJobs; }
 
+	/** Get the default graph preset from the project settings and assign it to the given job */
+	static void AssignDefaultGraphPresetToJob(UMoviePipelineExecutorJob* InJob);
+
 private:
 	// SWidget Interface
 	FReply OnKeyDown(const FGeometry& MyGeometry, const FKeyEvent& InKeyEvent);
@@ -63,6 +66,7 @@ private:
 	// ~FEditorUndoClient
 
 private:
+	
 	void OnCreateJobFromAsset(const FAssetData& InAsset);
 
 	TSharedRef<ITableRow> OnGenerateRow(TSharedPtr<IMoviePipelineQueueTreeItem> Item, const TSharedRef<STableViewBase>& Tree);
