Movie Graph:

-Fix: Hide Evaluation Context

diff --git a/Engine/Plugins/MovieScene/MovieRenderPipeline/Source/MovieRenderPipelineEditor/Private/Graph/SMovieGraphActiveRenderSettingsTabContent.cpp b/Engine/Plugins/MovieScene/MovieRenderPipeline/Source/MovieRenderPipelineEditor/Private/Graph/SMovieGraphActiveRenderSettingsTabContent.cpp
--- a/Engine/Plugins/MovieScene/MovieRenderPipeline/Source/MovieRenderPipelineEditor/Private/Graph/SMovieGraphActiveRenderSettingsTabContent.cpp
+++ b/Engine/Plugins/MovieScene/MovieRenderPipeline/Source/MovieRenderPipelineEditor/Private/Graph/SMovieGraphActiveRenderSettingsTabContent.cpp
@@ -349,15 +349,15 @@ void SMovieGraphActiveRenderSettingsTabContent::Construct(const FArguments& InAr
 				.OnClicked(this, &SMovieGraphActiveRenderSettingsTabContent::OnEvaluateGraphClicked)
 			]
 
-			+SHorizontalBox::Slot()
-			.Padding(5, 0, 0, 0)
-			.AutoWidth()
-			[
-				SNew(SPositiveActionButton)
-				.Text(LOCTEXT("Button_EvaluationContext", "Evaluation Context..."))
-				.Icon(FAppStyle::GetBrush("EditorPreferences.TabIcon"))
-				.OnGetMenuContent(this, &SMovieGraphActiveRenderSettingsTabContent::GenerateEvaluationContextMenu)
-			]
+			// +SHorizontalBox::Slot()
+			// .Padding(5, 0, 0, 0)
+			// .AutoWidth()
+			// [
+			// 	SNew(SPositiveActionButton)
+			// 	.Text(LOCTEXT("Button_EvaluationContext", "Evaluation Context..."))
+			// 	.Icon(FAppStyle::GetBrush("EditorPreferences.TabIcon"))
+			// 	.OnGetMenuContent(this, &SMovieGraphActiveRenderSettingsTabContent::GenerateEvaluationContextMenu)
+			// ]
 		]
 
 		+ SVerticalBox::Slot()
