MovieGraph: Path Tracer no longer takes an extremely long time to warm-up by default with high spatial sample counts - the denoiser and spatial samples are now disabled during warm-ups.

Added a flag for the denoiser to the UI, and cleaned up the Sampling Method node for now since the adaptive sampling method is not yet implemented.

diff --git a/Engine/Plugins/MovieScene/MovieRenderPipeline/Source/MovieRenderPipelineCore/Private/Graph/MovieGraphStochasticTimeStep.cpp b/Engine/Plugins/MovieScene/MovieRenderPipeline/Source/MovieRenderPipelineCore/Private/Graph/MovieGraphStochasticTimeStep.cpp
deleted file mode 100644
diff --git a/Engine/Plugins/MovieScene/MovieRenderPipeline/Source/MovieRenderPipelineCore/Private/Graph/Nodes/MovieGraphSamplingMethodNode.cpp b/Engine/Plugins/MovieScene/MovieRenderPipeline/Source/MovieRenderPipelineCore/Private/Graph/Nodes/MovieGraphSamplingMethodNode.cpp
--- a/Engine/Plugins/MovieScene/MovieRenderPipeline/Source/MovieRenderPipelineCore/Private/Graph/Nodes/MovieGraphSamplingMethodNode.cpp
+++ b/Engine/Plugins/MovieScene/MovieRenderPipeline/Source/MovieRenderPipelineCore/Private/Graph/Nodes/MovieGraphSamplingMethodNode.cpp
@@ -1,4 +1,4 @@
-﻿// Copyright Epic Games, Inc. All Rights Reserved.
+// Copyright Epic Games, Inc. All Rights Reserved.
 
 #include "Graph/Nodes/MovieGraphSamplingMethodNode.h"
 
@@ -9,7 +9,6 @@
 UMovieGraphSamplingMethodNode::UMovieGraphSamplingMethodNode()
 	: SamplingMethodClass(UMovieGraphLinearTimeStep::StaticClass())
 	, TemporalSampleCount(1)
-	, SampleCount(1)
 {
 }
 
diff --git a/Engine/Plugins/MovieScene/MovieRenderPipeline/Source/MovieRenderPipelineCore/Public/Graph/MovieGraphStochasticTimeStep.h b/Engine/Plugins/MovieScene/MovieRenderPipeline/Source/MovieRenderPipelineCore/Public/Graph/MovieGraphStochasticTimeStep.h
deleted file mode 100644
diff --git a/Engine/Plugins/MovieScene/MovieRenderPipeline/Source/MovieRenderPipelineCore/Public/Graph/Nodes/MovieGraphSamplingMethodNode.h b/Engine/Plugins/MovieScene/MovieRenderPipeline/Source/MovieRenderPipelineCore/Public/Graph/Nodes/MovieGraphSamplingMethodNode.h
--- a/Engine/Plugins/MovieScene/MovieRenderPipeline/Source/MovieRenderPipelineCore/Public/Graph/Nodes/MovieGraphSamplingMethodNode.h
+++ b/Engine/Plugins/MovieScene/MovieRenderPipeline/Source/MovieRenderPipelineCore/Public/Graph/Nodes/MovieGraphSamplingMethodNode.h
@@ -1,4 +1,4 @@
-﻿// Copyright Epic Games, Inc. All Rights Reserved.
+// Copyright Epic Games, Inc. All Rights Reserved.
 
 #pragma once
 
@@ -26,48 +26,19 @@ public:
 #endif
 
 public:
-	UPROPERTY(EditAnywhere, BlueprintReadWrite, Category = "Settings", meta = (InlineEditConditionToggle))
-	uint8 bOverride_SamplingMethodClass : 1;
 
-	UPROPERTY(EditAnywhere, BlueprintReadWrite, Category = "Settings", meta = (InlineEditConditionToggle))
-	uint8 bOverride_TemporalSampleCount : 1;
 
 	UPROPERTY(EditAnywhere, BlueprintReadWrite, Category = "Settings", meta = (InlineEditConditionToggle))
-	uint8 bOverride_SampleCount : 1;
-	
-	UPROPERTY(EditAnywhere, BlueprintReadWrite, Category = "Settings", meta = (InlineEditConditionToggle))
-	uint8 bOverride_NoiseThreshold : 1;
+	uint8 bOverride_SamplingMethodClass : 1;
 
 	UPROPERTY(EditAnywhere, BlueprintReadWrite, Category = "Settings", meta = (InlineEditConditionToggle))
-	uint8 bOverride_MaxWallTime : 1;
+	uint8 bOverride_TemporalSampleCount : 1;
 
-	/** The type of sampling the render should use. */
-	UPROPERTY(BlueprintReadWrite, EditAnywhere, Category = "Settings", NoClear, meta = (EditCondition = "bOverride_SamplingMethodClass", DisplayName = "Sampling Method", MetaClass = "/Script/MovieRenderPipelineCore.MovieGraphTimeStepBase", ShowDisplayNames))
+	/** The type of sampling the render should use. Not currently marked as BlueprintReadWrite/EditAnywhere as there is only one class anyways.*/
+	UPROPERTY(NoClear, meta = (EditCondition = "bOverride_SamplingMethodClass", DisplayName = "Sampling Method", MetaClass = "/Script/MovieRenderPipelineCore.MovieGraphTimeStepBase", ShowDisplayNames))
 	FSoftClassPath SamplingMethodClass;
 
 	/** The number of temporal samples which should be taken on one frame. Applies only to linear sampling. */
-	UPROPERTY(BlueprintReadWrite, EditAnywhere, Category = "Settings|Linear Sampling", meta = (UIMin = 1, ClampMin = 1, EditCondition = "bOverride_TemporalSampleCount"))
+	UPROPERTY(BlueprintReadWrite, EditAnywhere, Category = "Settings", meta = (UIMin = 1, ClampMin = 1, EditCondition = "bOverride_TemporalSampleCount"))
 	int32 TemporalSampleCount;
-
-	/**
-	 * The number of samples which should be taken on one frame. The number of samples taken may be less if the max
-	 * wall time or noise threshold is reached. Applies only to stochastic sampling.
-	 */
-	UPROPERTY(BlueprintReadWrite, EditAnywhere, Category = "Settings|Stochastic Sampling", meta = (UIMin = 1, ClampMin = 1, EditCondition = "bOverride_SampleCount"))
-	int32 SampleCount;
-
-	/**
-	 * [UNIMPLEMENTED] The noise threshold that must be reached on the current frame before samples stop being collected.
-	 * Samples may stop being collected before the noise threshold is reached if the max wall time is met, or the sample
-	 * count is set too low. Applies only to stochastic sampling.
-	 */
-	UPROPERTY(BlueprintReadWrite, EditAnywhere, Category = "Settings|Stochastic Sampling", meta = (UIMin = 0, ClampMin = 0, EditCondition = "bOverride_NoiseThreshold"))
-	float NoiseThreshold;
-
-	/**
-	 * [UNIMPLEMENTED] The maximum amount of aggregate time (seconds) that can be spent rendering all samples for one frame.
-	 * Applies only to stochastic sampling.
-	 */
-	UPROPERTY(BlueprintReadWrite, EditAnywhere, Category = "Settings|Stochastic Sampling", meta = (UIMin = 0, ClampMin = 0, Unit = "s", EditCondition = "bOverride_MaxWallTime"))
-	float MaxWallTime;
 };
\ No newline at end of file
diff --git a/Engine/Plugins/MovieScene/MovieRenderPipeline/Source/MovieRenderPipelineRenderPasses/Private/Graph/Nodes/MovieGraphPathTracedRenderPassNode.cpp b/Engine/Plugins/MovieScene/MovieRenderPipeline/Source/MovieRenderPipelineRenderPasses/Private/Graph/Nodes/MovieGraphPathTracedRenderPassNode.cpp
--- a/Engine/Plugins/MovieScene/MovieRenderPipeline/Source/MovieRenderPipelineRenderPasses/Private/Graph/Nodes/MovieGraphPathTracedRenderPassNode.cpp
+++ b/Engine/Plugins/MovieScene/MovieRenderPipeline/Source/MovieRenderPipelineRenderPasses/Private/Graph/Nodes/MovieGraphPathTracedRenderPassNode.cpp
@@ -12,9 +12,10 @@ TUniquePtr<UE::MovieGraph::Rendering::FMovieGraphImagePassBase> UMovieGraphPathT
 
 UMovieGraphPathTracedRenderPassNode::UMovieGraphPathTracedRenderPassNode()
 	: SpatialSampleCount(1)
+	, bDenoiser(true)
+	, bWriteAllSamples(false)
 	, bDisableToneCurve(false)
 	, bAllowOCIO(true)
-	, bWriteAllSamples(false)
 {
 	ShowFlags->ApplyDefaultShowFlagValue(VMI_PathTracing, true);
 	// TODO: Showflag for SetMotionBlur()?
@@ -76,6 +77,14 @@ int32 UMovieGraphPathTracedRenderPassNode::GetNumSpatialSamples() const
 	return SpatialSampleCount;
 }
 
+int32 UMovieGraphPathTracedRenderPassNode::GetNumSpatialSamplesDuringWarmUp() const
+{
+	// Path Tracer doesn't have an image history like the deferred renderer, so it doesn't need
+	// to run all the spatial samples.
+	return 1;
+}
+
+
 bool UMovieGraphPathTracedRenderPassNode::GetDisableToneCurve() const
 {
 	return bDisableToneCurve;
@@ -85,3 +94,9 @@ bool UMovieGraphPathTracedRenderPassNode::GetAllowOCIO() const
 {
 	return bAllowOCIO;
 }
+
+bool UMovieGraphPathTracedRenderPassNode::GetAllowDenoiser() const
+{
+	return bDenoiser;
+}
+
diff --git a/Engine/Plugins/MovieScene/MovieRenderPipeline/Source/MovieRenderPipelineRenderPasses/Private/Graph/Renderers/MovieGraphDeferredPass.cpp b/Engine/Plugins/MovieScene/MovieRenderPipeline/Source/MovieRenderPipelineRenderPasses/Private/Graph/Renderers/MovieGraphDeferredPass.cpp
--- a/Engine/Plugins/MovieScene/MovieRenderPipeline/Source/MovieRenderPipelineRenderPasses/Private/Graph/Renderers/MovieGraphDeferredPass.cpp
+++ b/Engine/Plugins/MovieScene/MovieRenderPipeline/Source/MovieRenderPipelineRenderPasses/Private/Graph/Renderers/MovieGraphDeferredPass.cpp
@@ -124,7 +124,9 @@ void FMovieGraphDeferredPass::Render(const FMovieGraphTraversalContext& InFrameT
 
 	UMovieGraphImagePassBaseNode* ParentNodeThisFrame = GetParentNode(InTimeData.EvaluatedConfig);
 	const bool bWriteAllSamples = ParentNodeThisFrame->GetWriteAllSamples();
-	int32 NumSpatialSamples = FMath::Max(1, ParentNodeThisFrame->GetNumSpatialSamples());
+	const bool bIsRenderingState = InFrameTraversalContext.Shot->ShotInfo.State == EMovieRenderShotState::Rendering;
+	int32 NumSpatialSamples = FMath::Max(1, bIsRenderingState ? ParentNodeThisFrame->GetNumSpatialSamples() : ParentNodeThisFrame->GetNumSpatialSamplesDuringWarmUp());
+
 	const ESceneCaptureSource SceneCaptureSource = ParentNodeThisFrame->GetDisableToneCurve() ? ESceneCaptureSource::SCS_FinalColorHDR : ESceneCaptureSource::SCS_FinalToneCurveHDR;
 	const EAntiAliasingMethod AntiAliasingMethod = ParentNodeThisFrame->GetAntiAliasingMethod();
 	float OverscanFraction = 0.f;
@@ -302,6 +304,10 @@ void FMovieGraphDeferredPass::Render(const FMovieGraphTraversalContext& InFrameT
 
 		int32 NumValidMaterials = NewView->FinalPostProcessSettings.BufferVisualizationPipes.Num();
 		NewView->FinalPostProcessSettings.bBufferVisualizationDumpRequired = NumValidMaterials > 0;
+		NewView->FinalPostProcessSettings.bOverride_PathTracingEnableDenoiser = true;
+
+		// The denoiser is disabled during warm-up frames.
+		NewView->FinalPostProcessSettings.PathTracingEnableDenoiser = bIsRenderingState && ParentNodeThisFrame->GetAllowDenoiser();
 
 		// Submit the renderer to be rendered
 		GetRendererModule().BeginRenderingViewFamily(&Canvas, ViewFamily.ToSharedPtr().Get());
diff --git a/Engine/Plugins/MovieScene/MovieRenderPipeline/Source/MovieRenderPipelineRenderPasses/Public/Graph/Nodes/MovieGraphImagePassBaseNode.h b/Engine/Plugins/MovieScene/MovieRenderPipeline/Source/MovieRenderPipelineRenderPasses/Public/Graph/Nodes/MovieGraphImagePassBaseNode.h
--- a/Engine/Plugins/MovieScene/MovieRenderPipeline/Source/MovieRenderPipelineRenderPasses/Public/Graph/Nodes/MovieGraphImagePassBaseNode.h
+++ b/Engine/Plugins/MovieScene/MovieRenderPipeline/Source/MovieRenderPipelineRenderPasses/Public/Graph/Nodes/MovieGraphImagePassBaseNode.h
@@ -47,12 +47,18 @@ public:
 	/** How many spatial samples should be rendered each frame? */
 	virtual int32 GetNumSpatialSamples() const { return 1; }
 
+	/** How many spatial samples should be used during warm-up frames? */
+	virtual int32 GetNumSpatialSamplesDuringWarmUp() const { return GetNumSpatialSamples(); }
+
 	/** Should the tone curve be disabled while rendering? Allows for linear values in exrs but changes the look of the final image. */
 	virtual bool GetDisableToneCurve() const { return false; }
 
 	/** Should the output file be allowed to apply an OCIO transform on this render? */
 	virtual bool GetAllowOCIO() const { return true; }
 
+	/** Should the denoiser be run on the resulting image (only has any effect with the Path Tracer) */
+	virtual bool GetAllowDenoiser() const { return true; }
+
 	/** Which AA Method should be used? */
 	virtual EAntiAliasingMethod GetAntiAliasingMethod() const { return EAntiAliasingMethod::AAM_None; }
 
diff --git a/Engine/Plugins/MovieScene/MovieRenderPipeline/Source/MovieRenderPipelineRenderPasses/Public/Graph/Nodes/MovieGraphPathTracedRenderPassNode.h b/Engine/Plugins/MovieScene/MovieRenderPipeline/Source/MovieRenderPipelineRenderPasses/Public/Graph/Nodes/MovieGraphPathTracedRenderPassNode.h
--- a/Engine/Plugins/MovieScene/MovieRenderPipeline/Source/MovieRenderPipelineRenderPasses/Public/Graph/Nodes/MovieGraphPathTracedRenderPassNode.h
+++ b/Engine/Plugins/MovieScene/MovieRenderPipeline/Source/MovieRenderPipelineRenderPasses/Public/Graph/Nodes/MovieGraphPathTracedRenderPassNode.h
@@ -28,8 +28,10 @@ public:
 	virtual bool GetWriteAllSamples() const override;
 	virtual TArray<FMoviePipelinePostProcessPass> GetAdditionalPostProcessMaterials() const override;
 	virtual int32 GetNumSpatialSamples() const override;
+	virtual int32 GetNumSpatialSamplesDuringWarmUp() const override;
 	virtual bool GetDisableToneCurve() const override;
 	virtual bool GetAllowOCIO() const override;
+	virtual bool GetAllowDenoiser() const override;
 	virtual TUniquePtr<UE::MovieGraph::Rendering::FMovieGraphImagePassBase> CreateInstance() const;
 	// ~UMovieGraphImagePassBaseNode Interface
 
@@ -48,6 +50,9 @@ public:
 	UPROPERTY(EditAnywhere, BlueprintReadWrite, Category = Overrides, meta = (InlineEditConditionToggle))
 	uint8 bOverride_SpatialSampleCount : 1;
 
+	UPROPERTY(EditAnywhere, BlueprintReadWrite, Category = Overrides, meta = (InlineEditConditionToggle))
+	uint8 bOverride_bDenoiser : 1;
+
 	UPROPERTY(EditAnywhere, BlueprintReadWrite, Category = Overrides, meta = (InlineEditConditionToggle))
 	uint8 bOverride_bDisableToneCurve : 1;
 
@@ -69,6 +74,16 @@ public:
 	UPROPERTY(EditAnywhere, BlueprintReadWrite, meta = (UIMin = 1, ClampMin = 1), Category = "Sampling", meta = (EditCondition = "bOverride_SpatialSampleCount"))
 	int32 SpatialSampleCount;
 
+	/** If true the resulting image will be denoised at the end of each set of Spatial Samples. */
+	UPROPERTY(EditAnywhere, BlueprintReadWrite, meta = (UIMin = 1, ClampMin = 1), Category = "Sampling", meta = (EditCondition = "bOverride_bDenoiser"))
+	bool bDenoiser;
+
+	/**
+	* Debug Feature. Not currently marked BlueprintReadWrite/EditAnywhere as it's not totally implemented on the Path Tracer right now.
+	*/
+	UPROPERTY(DisplayName = "Write All Samples (Debug)", meta = (EditCondition = "bOverride_bWriteAllsamples"))
+	bool bWriteAllSamples;
+
 	/**
 	* If true, the tone curve will be disabled for this render pass. This will result in values greater than 1.0 in final renders
 	* and can optionally be combined with OCIO profiles on the file output nodes to convert from Linear Values in Working Color Space
@@ -83,13 +98,7 @@ public:
 	UPROPERTY(EditAnywhere, BlueprintReadWrite, Category = "Settings", meta = (EditCondition = "bOverride_bAllowOCIO"))
 	bool bAllowOCIO;
 
-	/**
-	* Debug Feature. Can use this to write out each individual Temporal and Spatial sample rendered by this render pass,
-	* which allows you to see which images are being accumulated together. Can be useful for debugging incorrect looking
-	* frames to see which sub-frame evaluations were incorrect.
-	*/
-	UPROPERTY(EditAnywhere, BlueprintReadWrite, Category = "Sampling", meta = (EditCondition = "bOverride_bWriteAllsamples"))
-	bool bWriteAllSamples;
+
 
 	/**
 	* An array of additional post-processing materials to run after the frame is rendered. Using this feature may add a notable amount of render time.
