MovieGraph: Fixed post-render callback not containing any data when rendered via the graph.

diff --git a/Engine/Plugins/MovieScene/MovieRenderPipeline/Source/MovieRenderPipelineCore/Private/Graph/MovieGraphPipeline.cpp b/Engine/Plugins/MovieScene/MovieRenderPipeline/Source/MovieRenderPipelineCore/Private/Graph/MovieGraphPipeline.cpp
--- a/Engine/Plugins/MovieScene/MovieRenderPipeline/Source/MovieRenderPipelineCore/Private/Graph/MovieGraphPipeline.cpp
+++ b/Engine/Plugins/MovieScene/MovieRenderPipeline/Source/MovieRenderPipelineCore/Private/Graph/MovieGraphPipeline.cpp
@@ -792,7 +792,7 @@ void UMovieGraphPipeline::TeardownShot(const TObjectPtr<UMoviePipelineExecutorSh
 		GlobalGameOverridesNode->ApplySettings(bOverrideValues, GetWorld());
 	}
 
-	if (IsPostShotCallbacksNeeded())
+	if (IsPostShotCallbackNeeded())
 	{
 		ProcessOutstandingFinishedFrames();
 		ProcessOutstandingFutures();
@@ -800,7 +800,22 @@ void UMovieGraphPipeline::TeardownShot(const TObjectPtr<UMoviePipelineExecutorSh
 		// ToDo: Allow the command line encoder to modify the file list
 		// since it may add or remove files...
 
-		ExecutePostShotScripts();
+		FMoviePipelineOutputData Params;
+		Params.Pipeline = this;
+		Params.Job = GetCurrentJob();
+		Params.bSuccess = !bShutdownSetErrorFlag;
+
+		// We only provide data that the current shot generated during post-shot callbacks.
+		TArray< FMovieGraphRenderOutputData> SingleOutputData;
+		if (CurrentShotIndex < GeneratedOutputData.Num())
+		{
+			// If a shot has been started but canceled during warm-up, nothing will have been pushed to the GeneratedOutputData for that shot,
+			// so we have to put a range check on this.
+			SingleOutputData.Add(GeneratedOutputData[CurrentShotIndex]);
+		}
+		Params.GraphData = SingleOutputData;
+
+		ExecutePostShotScripts(Params);
 	}
 
 	// Check to see if this was the last shot in the Pipeline, otherwise on the next
@@ -1187,7 +1202,13 @@ void UMovieGraphPipeline::TransitionToState(const EMovieRenderPipelineState InNe
 			//	Setting->OnMoviePipelineShutdown(this);
 			//}
 
-			ExecutePostJobScripts();
+			FMoviePipelineOutputData Params;
+			Params.Pipeline = this;
+			Params.Job = GetCurrentJob();
+			Params.bSuccess = !bShutdownSetErrorFlag;
+			Params.GraphData = GeneratedOutputData;
+
+			ExecutePostJobScripts(Params);
 
 			// Ensure our delegates don't get called anymore as we're going to become null soon.
 			FCoreDelegates::OnBeginFrame.RemoveAll(this);
@@ -1260,23 +1281,19 @@ void UMovieGraphPipeline::OnMoviePipelineFinishedImpl()
 	// Broadcast to both Native and Python/BP
 	//
 	// Generate a params struct containing the data generated by this job.
-	FMovieGraphPipelineOutputData Params;
+	FMoviePipelineOutputData Params;
 	Params.Pipeline = this;
 	Params.Job = GetCurrentJob();
 	Params.bSuccess = !bShutdownSetErrorFlag;
-	Params.ShotData = GeneratedOutputData;
+	Params.GraphData = GeneratedOutputData;
 	//
 	//UE_LOG(LogMovieRenderPipelineIO, Verbose, TEXT("Files written to disk for entire sequence:"));
 	//PrintVerboseLogForFiles(GeneratedShotOutputData);
 	//UE_LOG(LogMovieRenderPipelineIO, Verbose, TEXT("Completed outputting files written to disk."));
 	//
 
-	// The callbacks are shared with the legacy system, but they have a different type due to render data identifiers.
-	// For now we need to at least call the callback (otherwise the executor never knows that the job is done).
-	FMoviePipelineOutputData DummyLegacyParams;
-	DummyLegacyParams.bSuccess = !bShutdownSetErrorFlag;
-	OnMoviePipelineWorkFinishedDelegateNative.Broadcast(DummyLegacyParams);
-	// OnMoviePipelineWorkFinishedDelegate.Broadcast(DummyLegacyParams);
+	OnMoviePipelineWorkFinishedDelegateNative.Broadcast(Params);
+	OnMoviePipelineWorkFinishedDelegate.Broadcast(Params);
 }
 
 //void UMovieGraphPipeline::PrintVerboseLogForFiles(const TArray<FMoviePipelineShotOutputData>& InOutputData) const
@@ -1487,7 +1504,7 @@ const TSet<TObjectPtr<UMovieGraphFileOutputNode>> UMovieGraphPipeline::GetOutput
 }
 
 
-bool UMovieGraphPipeline::IsPostShotCallbacksNeeded() const
+bool UMovieGraphPipeline::IsPostShotCallbackNeeded() const
 {
 	bool bAnyScriptNeedsCallbacks = false;
 	for (UMovieGraphScriptBase* Script : CurrentScriptInstances)
@@ -1522,37 +1539,24 @@ void UMovieGraphPipeline::ExecutePreShotScripts(const TObjectPtr<UMoviePipelineE
 	}
 }
 
-void UMovieGraphPipeline::ExecutePostShotScripts()
+void UMovieGraphPipeline::ExecutePostShotScripts(const FMoviePipelineOutputData& InData)
 {
-	FMovieGraphPipelineOutputData Params;
-	Params.Pipeline = this;
-	Params.Job = GetCurrentJob();
-	Params.bSuccess = !bShutdownSetErrorFlag;
-
-	// We only provide data that the current shot generated during post-shot callbacks.
-	TArray< FMovieGraphRenderOutputData> SingleOutputData;
-	SingleOutputData.Add(GeneratedOutputData[CurrentShotIndex]);
-	Params.ShotData = SingleOutputData;
-
 	for (UMovieGraphScriptBase* Script : CurrentScriptInstances)
 	{
 		if (Script->IsPerShotCallbackNeeded())
 		{
-			Script->OnShotFinished(GetCurrentJob(), ActiveShotList[CurrentShotIndex], Params);
+			Script->OnShotFinished(GetCurrentJob(), ActiveShotList[CurrentShotIndex], InData);
 		}
 	}
+
+	OnMoviePipelineShotWorkFinishedDelegate.Broadcast(InData);
+	OnMoviePipelineShotWorkFinishedImpl().Broadcast(InData);
 }
 
-void UMovieGraphPipeline::ExecutePostJobScripts()
+void UMovieGraphPipeline::ExecutePostJobScripts(const FMoviePipelineOutputData& InData)
 {
-	FMovieGraphPipelineOutputData Params;
-	Params.Pipeline = this;
-	Params.Job = GetCurrentJob();
-	Params.bSuccess = !bShutdownSetErrorFlag;
-	Params.ShotData = GeneratedOutputData;
-
 	for (UMovieGraphScriptBase* Script : CurrentScriptInstances)
 	{
-		Script->OnJobFinished(GetCurrentJob(), Params);
+		Script->OnJobFinished(GetCurrentJob(), InData);
 	}
 }
\ No newline at end of file
diff --git a/Engine/Plugins/MovieScene/MovieRenderPipeline/Source/MovieRenderPipelineCore/Private/MoviePipelineInProcessExecutor.cpp b/Engine/Plugins/MovieScene/MovieRenderPipeline/Source/MovieRenderPipelineCore/Private/MoviePipelineInProcessExecutor.cpp
--- a/Engine/Plugins/MovieScene/MovieRenderPipeline/Source/MovieRenderPipelineCore/Private/MoviePipelineInProcessExecutor.cpp
+++ b/Engine/Plugins/MovieScene/MovieRenderPipeline/Source/MovieRenderPipelineCore/Private/MoviePipelineInProcessExecutor.cpp
@@ -220,7 +220,10 @@ void UMoviePipelineInProcessExecutor::OnMoviePipelineFinished(FMoviePipelineOutp
 	RestoreState();
 
 	// Now that another frame has passed and we should be OK to start another PIE session, notify our owner.
-	OnIndividualPipelineFinished(InOutputData.Pipeline);
+	// We do the cast to UMoviePipeline (which will return nullptr during a graph render) because the functions
+	// all take a pointer, but it's never actually used, and you can only have one pipeline at once so the API
+	// doesn't really need to pass back the pointer to a soon-to-be-destroyed object.
+	OnIndividualPipelineFinished(Cast<UMoviePipeline>(InOutputData.Pipeline));
 }
 
 void UMoviePipelineInProcessExecutor::BackupState()
diff --git a/Engine/Plugins/MovieScene/MovieRenderPipeline/Source/MovieRenderPipelineCore/Public/Graph/MovieGraphDataTypes.h b/Engine/Plugins/MovieScene/MovieRenderPipeline/Source/MovieRenderPipelineCore/Public/Graph/MovieGraphDataTypes.h
--- a/Engine/Plugins/MovieScene/MovieRenderPipeline/Source/MovieRenderPipelineCore/Public/Graph/MovieGraphDataTypes.h
+++ b/Engine/Plugins/MovieScene/MovieRenderPipeline/Source/MovieRenderPipelineCore/Public/Graph/MovieGraphDataTypes.h
@@ -91,72 +91,6 @@ struct MOVIERENDERPIPELINECORE_API FMovieGraphInitConfig
 	bool bRenderViewport;
 };
 
-USTRUCT(BlueprintType)
-struct FMovieGraphRenderLayerOutputData
-{
-	GENERATED_BODY()
-public:
-	/** A list of file paths on disk (in order) that were generated for this particular render pass. */
-	UPROPERTY(VisibleAnywhere, BlueprintReadOnly, Category = "Movie Pipeline")
-	TArray<FString> FilePaths;
-};
-
-USTRUCT(BlueprintType)
-struct MOVIERENDERPIPELINECORE_API FMovieGraphRenderOutputData
-{
-	GENERATED_BODY()
-public:
-	/** Which shot is this output data for. */
-	UPROPERTY(VisibleAnywhere, BlueprintReadOnly, Category = "Movie Graph")
-	TWeakObjectPtr<UMoviePipelineExecutorShot> Shot;
-
-	/**
-	* A mapping between render layers (such as "beauty") and an array containing the files written for that shot.
-	* Will be multiple files if using image sequences.
-	*/
-	UPROPERTY(VisibleAnywhere, BlueprintReadOnly, Category = "Movie Graph")
-	TMap<FMovieGraphRenderDataIdentifier, FMovieGraphRenderLayerOutputData> RenderLayerData;
-};
-
-USTRUCT(BlueprintType)
-struct MOVIERENDERPIPELINECORE_API FMovieGraphPipelineOutputData
-{
-	GENERATED_BODY()
-
-		FMovieGraphPipelineOutputData()
-		: Pipeline(nullptr)
-		, Job(nullptr)
-		, bSuccess(false)
-	{}
-
-	/**
-	* The UMoviePipeline instance that generated this data. This is only provided as an id (in the event you were the one who created
-	* the UMoviePipeline instance.)
-	*
-	* DO NOT CALL FUNCTIONS ON THIS (unless you know what you're doing)
-	*
-	* Provided here for backwards compatibility.
-	*/
-	UPROPERTY(VisibleAnywhere, BlueprintReadOnly, Transient, Category = "Movie Pipeline")
-	TObjectPtr<UMovieGraphPipeline> Pipeline;
-
-	/** Job the data is for. Job may still be in progress (if a shot callback) so be careful about modifying properties on it */
-	UPROPERTY(VisibleAnywhere, BlueprintReadOnly, Transient, Category = "Movie Pipeline")
-	TObjectPtr<UMoviePipelineExecutorJob> Job;
-
-	/** Did the job succeed, or was it canceled early due to an error (such as failure to write file to disk)? */
-	UPROPERTY(VisibleAnywhere, BlueprintReadOnly, Category = "Movie Pipeline")
-	bool bSuccess;
-
-	/**
-	* The file data for each shot that was rendered. If no files were written this will be empty. If this is from the per-shot work
-	* finished callback it will only have one entry (for the just finished shot). Will not include shots that did not get rendered
-	* due to the pipeline encountering an error.
-	*/
-	UPROPERTY(VisibleAnywhere, BlueprintReadOnly, Category = "Movie Pipeline")
-	TArray<FMovieGraphRenderOutputData> ShotData;
-};
-
 UCLASS(BlueprintType, Abstract)
 class MOVIERENDERPIPELINECORE_API UMovieGraphTimeStepBase : public UObject
 {
diff --git a/Engine/Plugins/MovieScene/MovieRenderPipeline/Source/MovieRenderPipelineCore/Public/Graph/MovieGraphPipeline.h b/Engine/Plugins/MovieScene/MovieRenderPipeline/Source/MovieRenderPipelineCore/Public/Graph/MovieGraphPipeline.h
--- a/Engine/Plugins/MovieScene/MovieRenderPipeline/Source/MovieRenderPipelineCore/Public/Graph/MovieGraphPipeline.h
+++ b/Engine/Plugins/MovieScene/MovieRenderPipeline/Source/MovieRenderPipelineCore/Public/Graph/MovieGraphPipeline.h
@@ -151,10 +151,9 @@ protected:
 
 	virtual void DuplicateJobAndConfiguration();
 	virtual void ExecutePreJobScripts();
-	virtual void ExecutePostJobScripts();
+	virtual void ExecutePostJobScripts(const FMoviePipelineOutputData& InData);
 	virtual void ExecutePreShotScripts(const TObjectPtr<UMoviePipelineExecutorShot>& InShot);
-	virtual void ExecutePostShotScripts();
-	virtual bool IsPostShotCallbacksNeeded() const;
+	virtual void ExecutePostShotScripts(const FMoviePipelineOutputData& InData);
 
 	// Update our data source to isolate the shot we're currently working on, so that expanded shots don't interfere with each other.
 	virtual void SetSoloShot(const TObjectPtr<UMoviePipelineExecutorShot>& InShot);
@@ -174,6 +173,7 @@ protected:
 	virtual void ShutdownImpl(bool bIsError) override;
 	virtual bool IsShutdownRequestedImpl() const override { return bShutdownRequested; }
 	virtual EMovieRenderPipelineState GetPipelineStateImpl() const override { return PipelineState; }
+	virtual bool IsPostShotCallbackNeeded() const override;
 	// ~UMoviePipelineBase Interface
 
 protected:
diff --git a/Engine/Plugins/MovieScene/MovieRenderPipeline/Source/MovieRenderPipelineCore/Public/Graph/MovieGraphRenderDataIdentifier.h b/Engine/Plugins/MovieScene/MovieRenderPipeline/Source/MovieRenderPipelineCore/Public/Graph/MovieGraphRenderDataIdentifier.h
--- a/Engine/Plugins/MovieScene/MovieRenderPipeline/Source/MovieRenderPipelineCore/Public/Graph/MovieGraphRenderDataIdentifier.h
+++ b/Engine/Plugins/MovieScene/MovieRenderPipeline/Source/MovieRenderPipelineCore/Public/Graph/MovieGraphRenderDataIdentifier.h
@@ -80,4 +80,31 @@ public:
 	/** The name of the camera being used for this render. */
 	UPROPERTY(EditAnywhere, BlueprintReadWrite, Category = "Movie Pipeline")
 	FString CameraName;
-};
\ No newline at end of file
+};
+
+USTRUCT(BlueprintType)
+struct MOVIERENDERPIPELINECORE_API FMovieGraphRenderLayerOutputData
+{
+	GENERATED_BODY()
+public:
+	/** A list of file paths on disk (in order) that were generated for this particular render pass. */
+	UPROPERTY(VisibleAnywhere, BlueprintReadOnly, Category = "Movie Pipeline")
+	TArray<FString> FilePaths;
+};
+
+USTRUCT(BlueprintType)
+struct MOVIERENDERPIPELINECORE_API FMovieGraphRenderOutputData
+{
+	GENERATED_BODY()
+public:
+	/** Which shot is this output data for. */
+	UPROPERTY(VisibleAnywhere, BlueprintReadOnly, Category = "Movie Graph")
+	TWeakObjectPtr<class UMoviePipelineExecutorShot> Shot;
+
+	/**
+	* A mapping between render layers (such as "beauty") and an array containing the files written for that shot.
+	* Will be multiple files if using image sequences.
+	*/
+	UPROPERTY(VisibleAnywhere, BlueprintReadOnly, Category = "Movie Graph")
+	TMap<FMovieGraphRenderDataIdentifier, FMovieGraphRenderLayerOutputData> RenderLayerData;
+};
diff --git a/Engine/Plugins/MovieScene/MovieRenderPipeline/Source/MovieRenderPipelineCore/Public/Graph/Nodes/MovieGraphExecuteScriptNode.h b/Engine/Plugins/MovieScene/MovieRenderPipeline/Source/MovieRenderPipelineCore/Public/Graph/Nodes/MovieGraphExecuteScriptNode.h
--- a/Engine/Plugins/MovieScene/MovieRenderPipeline/Source/MovieRenderPipelineCore/Public/Graph/Nodes/MovieGraphExecuteScriptNode.h
+++ b/Engine/Plugins/MovieScene/MovieRenderPipeline/Source/MovieRenderPipelineCore/Public/Graph/Nodes/MovieGraphExecuteScriptNode.h
@@ -8,8 +8,9 @@
 
 /**
 * This is the base class used for writing scripts that can be executed as part
-* of a Movie Graph Pipeline. After creating a UMovieGraphExecuteScriptNode and
-* choosing the type of your script, at runtime, an instance will be created when
+* of a Movie Graph Pipeline. After creating a UMovieGraphExecuteScriptNode 
+* (named 'Execute Script' in the UI) in your Movie Graph Configuration asset,
+* and choosing the type of your script, an instance will be created when
 * the job starts. This instance will persist until the end of the job. The instance
 * will recieve callbacks during the render which contain data that you can 
 * manipulate with your script, or data that was generated by the pipeline.
@@ -33,10 +34,19 @@
 *    def on_job_start(self, inJobCopy):
 *        super().on_job_start(inJobCopy)
 *        unreal.log_warning("OnJobStart")
+*        for shot in inOutputData.graph_data:
+*            for layerIdentifier in shot.render_layer_data:
+*                unreal.log("render layer: " + unreal.StringLibrary.conv_name_to_string(layerIdentifier.root_branch_name))
+*                for file in shot.render_layer_data[layerIdentifier].file_paths:
+*                    unreal.log("file: " + file)
 * 
 * Save and restart Unreal Editor, then create a Movie Graph Configuration
-* with a Script Node, and set it to run MovieGraphPipelineScriptExample. When
-* the job is run, the output console should print "OnJobStart" as a warning.
+* with an Execute Script Node, and set it to run MovieGraphPipelineScriptExample. 
+* When the job is run, the output console should print "OnJobStart" as a warning.
+* 
+* Use "help(unreal.MovieGraphScriptBase)" in the Python Console to see all available 
+* overridable functions.
+* 
 * See /Engine/Plugins/MovieScene/MovieRenderPipeline/Content/Python for more 
 * python examples.
 */
@@ -98,8 +108,8 @@ public:
 	*						data for all shots rendered up to this point.
 	*/
 	UFUNCTION(BlueprintNativeEvent, Category = "Movie Graph")
-	void OnJobFinished(UMoviePipelineExecutorJob* InJobCopy, const FMovieGraphPipelineOutputData& InOutputData);
-	virtual void OnJobFinished_Implementation(UMoviePipelineExecutorJob* InJobCopy, const FMovieGraphPipelineOutputData& InOutputData ) {}
+	void OnJobFinished(UMoviePipelineExecutorJob* InJobCopy, const FMoviePipelineOutputData& InOutputData);
+	virtual void OnJobFinished_Implementation(UMoviePipelineExecutorJob* InJobCopy, const FMoviePipelineOutputData& InOutputData ) {}
 
 	/**
 	* This is similar to OnJobStart but called before a shot is set up. This is only called
@@ -141,8 +151,8 @@ public:
 	* 						by the movie render (use OnJobFinished for that).
 	*/
 	UFUNCTION(BlueprintNativeEvent, Category = "Movie Graph")
-	void OnShotFinished(UMoviePipelineExecutorJob* InJobCopy, UMoviePipelineExecutorShot* InShotCopy, const FMovieGraphPipelineOutputData& OutputData);
-	virtual void OnShotFinished_Implementation(UMoviePipelineExecutorJob* InJobCopy, UMoviePipelineExecutorShot* InShotCopy, const FMovieGraphPipelineOutputData& OutputData) {}
+	void OnShotFinished(UMoviePipelineExecutorJob* InJobCopy, UMoviePipelineExecutorShot* InShotCopy, const FMoviePipelineOutputData& OutputData);
+	virtual void OnShotFinished_Implementation(UMoviePipelineExecutorJob* InJobCopy, UMoviePipelineExecutorShot* InShotCopy, const FMoviePipelineOutputData& OutputData) {}
 
 
 	/**
diff --git a/Engine/Plugins/MovieScene/MovieRenderPipeline/Source/MovieRenderPipelineCore/Public/MoviePipeline.h b/Engine/Plugins/MovieScene/MovieRenderPipeline/Source/MovieRenderPipelineCore/Public/MoviePipeline.h
--- a/Engine/Plugins/MovieScene/MovieRenderPipeline/Source/MovieRenderPipelineCore/Public/MoviePipeline.h
+++ b/Engine/Plugins/MovieScene/MovieRenderPipeline/Source/MovieRenderPipelineCore/Public/MoviePipeline.h
@@ -59,6 +59,7 @@ public:
 	virtual void ShutdownImpl(bool bError ) override;
 	virtual bool IsShutdownRequestedImpl() const override { return bShutdownRequested; }
 	virtual EMovieRenderPipelineState GetPipelineStateImpl() const override { return PipelineState; }
+	virtual bool IsPostShotCallbackNeeded() const override { return IsFlushDiskWritesPerShot(); }
 	// ~UMoviePipelineBase Interface
 
 	/**
@@ -90,30 +91,6 @@ public:
 	UPROPERTY(BlueprintAssignable, Category = "Movie Render Pipeline")
 	FMoviePipelineFinished OnMoviePipelineFinishedDelegate;
 
-
-	
-	/**
-	* Called when we have completely finished this pipeline. This means that all frames have been rendered,
-	* all files written to disk, and any post-finalize exports have finished. This Pipeline will call
-	* Shutdown() on itself before calling this delegate to ensure we've unregistered from all delegates
-	* and are no longer trying to do anything (even if we still exist).
-	*
-	* The params struct in the return will have metadata about files written to disk for each shot.
-	*/
-	UPROPERTY(BlueprintAssignable, Category = "Movie Render Pipeline")
-	FMoviePipelineWorkFinished OnMoviePipelineWorkFinishedDelegate;
-
-
-
-	/**
-	* Only called if `IsFlushDiskWritesPerShot()` is set!
-	* Called after each shot is finished and files have been flushed to disk. The returned data in
-	* the params struct will have only the per-shot metadata for the just finished shot. Use
-	* OnMoviePipelineFinished() if you need all ot the metadata.
-	*/
-	UPROPERTY(BlueprintAssignable, Category = "Movie Render Pipeline")
-	FMoviePipelineWorkFinished OnMoviePipelineShotWorkFinishedDelegate;
-
 	/**
 	* Get the Primary Configuration used to render this shot. This contains the global settings for the shot, as well as per-shot
 	* configurations which can contain their own settings.
diff --git a/Engine/Plugins/MovieScene/MovieRenderPipeline/Source/MovieRenderPipelineCore/Public/MoviePipelineBase.h b/Engine/Plugins/MovieScene/MovieRenderPipeline/Source/MovieRenderPipelineCore/Public/MoviePipelineBase.h
--- a/Engine/Plugins/MovieScene/MovieRenderPipeline/Source/MovieRenderPipelineCore/Public/MoviePipelineBase.h
+++ b/Engine/Plugins/MovieScene/MovieRenderPipeline/Source/MovieRenderPipelineCore/Public/MoviePipelineBase.h
@@ -78,6 +78,28 @@ public:
 	*/
 	FMoviePipelineWorkFinishedNative& OnMoviePipelineShotWorkFinished() { return OnMoviePipelineShotWorkFinishedImpl(); }
 
+	/**
+	* Called when we have completely finished this pipeline. This means that all frames have been rendered,
+	* all files written to disk, and any post-finalize exports have finished. This Pipeline will call
+	* Shutdown() on itself before calling this delegate to ensure we've unregistered from all delegates
+	* and are no longer trying to do anything (even if we still exist).
+	*
+	* The params struct in the return will have metadata about files written to disk for each shot.
+	*/
+	UPROPERTY(BlueprintAssignable, Category = "Movie Render Pipeline")
+	FMoviePipelineWorkFinished OnMoviePipelineWorkFinishedDelegate;
+
+
+
+	/**
+	* Only called if `IsPostShotCallbackNeeded()` returns true!
+	* Called after each shot is finished and files have been flushed to disk. The returned data in
+	* the params struct will have only the per-shot metadata for the just finished shot. Use
+	* OnMoviePipelineFinished() if you need all of the metadata.
+	*/
+	UPROPERTY(BlueprintAssignable, Category = "Movie Render Pipeline")
+	FMoviePipelineWorkFinished OnMoviePipelineShotWorkFinishedDelegate;
+
 protected:
 	virtual void RequestShutdownImpl(bool bIsError) {}
 	virtual void ShutdownImpl(bool bIsError) {}
@@ -85,6 +107,7 @@ protected:
 	virtual EMovieRenderPipelineState GetPipelineStateImpl() const { return EMovieRenderPipelineState::Uninitialized; }
 	virtual FMoviePipelineWorkFinishedNative& OnMoviePipelineWorkFinishedImpl() { return OnMoviePipelineWorkFinishedDelegateNative; }
 	virtual FMoviePipelineWorkFinishedNative& OnMoviePipelineShotWorkFinishedImpl() { return OnMoviePipelineShotWorkFinishedDelegateNative; }
+	virtual bool IsPostShotCallbackNeeded() const { return false; }
 
 	/** Called when we have completely finished. This object will call Shutdown before this and stop ticking. */
 	FMoviePipelineWorkFinishedNative OnMoviePipelineWorkFinishedDelegateNative;
diff --git a/Engine/Plugins/MovieScene/MovieRenderPipeline/Source/MovieRenderPipelineCore/Public/MoviePipelineLinearExecutor.h b/Engine/Plugins/MovieScene/MovieRenderPipeline/Source/MovieRenderPipelineCore/Public/MoviePipelineLinearExecutor.h
--- a/Engine/Plugins/MovieScene/MovieRenderPipeline/Source/MovieRenderPipelineCore/Public/MoviePipelineLinearExecutor.h
+++ b/Engine/Plugins/MovieScene/MovieRenderPipeline/Source/MovieRenderPipelineCore/Public/MoviePipelineLinearExecutor.h
@@ -39,8 +39,10 @@ protected:
 	virtual FText GetWindowTitle();
 	virtual float GetCompletionPercentageFromActivePipeline();
 public:
+	/** Note: When using a Movie Graph Pipeline this pointer will always be null. */
 	virtual void OnIndividualPipelineFinished(UMoviePipeline* /* FinishedPipeline */);
 	virtual void OnExecutorFinishedImpl();
+	/** Note: When using a Movie Graph Pipeline InPipeline will always be null. */
 	virtual void OnPipelineErrored(UMoviePipeline* InPipeline, bool bIsFatal, FText ErrorText);
 
 	virtual void CancelCurrentJob_Implementation();
diff --git a/Engine/Plugins/MovieScene/MovieRenderPipeline/Source/MovieRenderPipelineCore/Public/MovieRenderPipelineDataTypes.h b/Engine/Plugins/MovieScene/MovieRenderPipeline/Source/MovieRenderPipelineCore/Public/MovieRenderPipelineDataTypes.h
--- a/Engine/Plugins/MovieScene/MovieRenderPipeline/Source/MovieRenderPipelineCore/Public/MovieRenderPipelineDataTypes.h
+++ b/Engine/Plugins/MovieScene/MovieRenderPipeline/Source/MovieRenderPipelineCore/Public/MovieRenderPipelineDataTypes.h
@@ -13,6 +13,7 @@
 #include "Async/ParallelFor.h"
 #include "MovieSceneSequenceID.h"
 #include "MovieRenderDebugWidget.h"
+#include "Graph/MovieGraphRenderDataIdentifier.h"
 #include "MovieRenderPipelineDataTypes.generated.h"
 
 class UMovieSceneCinematicShotSection;
@@ -1311,9 +1312,12 @@ struct FMoviePipelineOutputData
 	* Provided here for backwards compatibility.
 	*/
 	UPROPERTY(VisibleAnywhere, BlueprintReadOnly, Transient, Category = "Movie Pipeline")
-	TObjectPtr<UMoviePipeline> Pipeline;
+	TObjectPtr<class UMoviePipelineBase> Pipeline;
 	
-	/** Job the data is for. Job may still be in progress (if a shot callback) so be careful about modifying properties on it */
+	/** 
+	* Job the data is for. Job may still be in progress (if a shot callback) so be careful about modifying properties on it 
+	* When using the Movie Render Graph this will point to the duplicated job created during execution.
+	*/
 	UPROPERTY(VisibleAnywhere, BlueprintReadOnly, Transient, Category = "Movie Pipeline")
 	TObjectPtr<UMoviePipelineExecutorJob> Job;
 	
@@ -1325,9 +1329,23 @@ struct FMoviePipelineOutputData
 	* The file data for each shot that was rendered. If no files were written this will be empty. If this is from the per-shot work
 	* finished callback it will only have one entry (for the just finished shot). Will not include shots that did not get rendered
 	* due to the pipeline encountering an error.
+	*
+	* This will be empty when using the Movie Render Graph. If a job is a Movie Render Graph job, use GraphData instead, which has
+	* a similar layout but with different data types.
 	*/
 	UPROPERTY(VisibleAnywhere, BlueprintReadOnly, Category = "Movie Pipeline")
 	TArray<FMoviePipelineShotOutputData> ShotData;
+
+	/** 
+	* The file data for each shot that was rendered. If no files were written this will be empty. If this is from the per-shot work
+	* finished callback it will only have one entry (for the just finished shot). Will not include shots that did not get rendered
+	* due to the pipeline encountering an error.
+	*
+	* This will be empty when using the legacy Movie Pipeline Configurations. If a job is a legacy configuration, use ShotData instead, which has
+	* a similar layout but with different data types.
+	*/
+	UPROPERTY(VisibleAnywhere, BlueprintReadOnly, Category = "Movie Pipeline")
+	TArray<FMovieGraphRenderOutputData> GraphData;
 };
 
 /**
diff --git a/Engine/Plugins/MovieScene/MovieRenderPipeline/Source/MovieRenderPipelineEditor/Private/MoviePipelinePIEExecutor.cpp b/Engine/Plugins/MovieScene/MovieRenderPipeline/Source/MovieRenderPipelineEditor/Private/MoviePipelinePIEExecutor.cpp
--- a/Engine/Plugins/MovieScene/MovieRenderPipeline/Source/MovieRenderPipelineEditor/Private/MoviePipelinePIEExecutor.cpp
+++ b/Engine/Plugins/MovieScene/MovieRenderPipeline/Source/MovieRenderPipelineEditor/Private/MoviePipelinePIEExecutor.cpp
@@ -316,7 +316,10 @@ void UMoviePipelinePIEExecutor::OnPIEMoviePipelineFinished(FMoviePipelineOutputD
 {
 	if (!InOutputData.bSuccess)
 	{
-		OnPipelineErrored(InOutputData.Pipeline, true, FText());
+		// We do the cast to UMoviePipeline (which will return nullptr during a graph render) because the functions
+		// all take a pointer, but it's never actually used, and you can only have one pipeline at once so the API
+		// doesn't really need to pass back the pointer to a soon-to-be-destroyed object.
+		OnPipelineErrored(Cast<UMoviePipeline>(InOutputData.Pipeline), true, FText());
 	}
 
 	// Unsubscribe to the EndPIE event so we don't think the user canceled it.
