MoviePipeline: Fixed performance degradation issues that originated from graph code in the MRQ queue UI when performing a migrate.

Also bumped up the default "Settings" column width by 5% in the queue UI to make the drop-down arrow more visible.

#rnx

diff --git a/Engine/Plugins/MovieScene/MovieRenderPipeline/Source/MovieRenderPipelineEditor/Private/Customizations/JobCustomization.h b/Engine/Plugins/MovieScene/MovieRenderPipeline/Source/MovieRenderPipelineEditor/Private/Customizations/JobCustomization.h
--- a/Engine/Plugins/MovieScene/MovieRenderPipeline/Source/MovieRenderPipelineEditor/Private/Customizations/JobCustomization.h
+++ b/Engine/Plugins/MovieScene/MovieRenderPipeline/Source/MovieRenderPipelineEditor/Private/Customizations/JobCustomization.h
@@ -23,7 +23,24 @@ public:
 
 	virtual ~FJobDetailsCustomization() override
 	{
+	}
+
+	virtual void PendingDelete() override
+	{
+		// Unregister delegates. It's important to do this in PendingDelete() vs the destructor because the destructor is not called before the next
+		// details panel is created (via ForceRefreshDetails()), leading to an exponential increase in the number of delegates registered.
+		
 		UPackage::PackageSavedWithContextEvent.RemoveAll(this);
+
+		if (SelectedJob.IsValid())
+		{
+			SelectedJob->OnJobGraphPresetChanged.RemoveAll(this);
+		}
+		
+		if (SelectedShot.IsValid())
+		{
+			SelectedShot->OnShotGraphPresetChanged.RemoveAll(this);
+		}
 	}
 
 protected:
@@ -63,13 +80,13 @@ protected:
 
 		for (const TWeakObjectPtr<UObject>& SelectedObject : ObjectsBeingCustomized)
 		{
-			if (UMoviePipelineExecutorJob* SelectedJob = Cast<UMoviePipelineExecutorJob>(SelectedObject.Get()))
+			if (UMoviePipelineExecutorJob* SelectedJobTemp = Cast<UMoviePipelineExecutorJob>(SelectedObject.Get()))
 			{
-				SelectedJobs.Add(SelectedJob);
+				SelectedJobs.Add(SelectedJobTemp);
 			}
-			else if (UMoviePipelineExecutorShot* SelectedShot = Cast<UMoviePipelineExecutorShot>(SelectedObject.Get()))
+			else if (UMoviePipelineExecutorShot* SelectedShotTemp = Cast<UMoviePipelineExecutorShot>(SelectedObject.Get()))
 			{
-				SelectedShots.Add(SelectedShot);
+				SelectedShots.Add(SelectedShotTemp);
 			}
 		}
 
@@ -95,17 +112,20 @@ protected:
 		// Refresh the UI if the graph preset changes (so the new variable assignments are displayed)
 		if (bIsShot)
 		{
-			SelectedShots[0]->OnShotGraphPresetChanged.AddSP(this, &FJobDetailsCustomization::RefreshLayout);
+			SelectedShot = SelectedShots[0];
+			SelectedShot->OnShotGraphPresetChanged.AddSP(this, &FJobDetailsCustomization::RefreshLayout);
 
 			// Also listen for changes to the primary job. Changes to the primary job can trigger an update to shot variable assignments.
-			if (UMoviePipelineExecutorJob* PrimaryJob = SelectedShots[0]->GetTypedOuter<UMoviePipelineExecutorJob>())
+			if (UMoviePipelineExecutorJob* PrimaryJob = SelectedShot->GetTypedOuter<UMoviePipelineExecutorJob>())
 			{
-				PrimaryJob->OnJobGraphPresetChanged.AddSP(this, &FJobDetailsCustomization::RefreshLayout);
+				SelectedJob = PrimaryJob;
+				SelectedJob->OnJobGraphPresetChanged.AddSP(this, &FJobDetailsCustomization::RefreshLayout);
 			}
 		}
 		else
 		{
-			SelectedJobs[0]->OnJobGraphPresetChanged.AddSP(this, &FJobDetailsCustomization::RefreshLayout);
+			SelectedJob = SelectedJobs[0];
+			SelectedJob->OnJobGraphPresetChanged.AddSP(this, &FJobDetailsCustomization::RefreshLayout);
 		}
 
 		// Set up the categories for variable assignments. Set as "Uncommon" priority to push variables down below the other properties.
@@ -123,12 +143,12 @@ protected:
 		
 		if (bIsShot)
 		{
-			AddVariableAssignments(SelectedShots[0]->GetGraphVariableAssignments(), ShotGraphVariablesCategory);
-			AddVariableAssignments(SelectedShots[0]->GetPrimaryGraphVariableAssignments(), PrimaryGraphVariablesShotOverridesCategory);
+			AddVariableAssignments(SelectedShot->GetGraphVariableAssignments(), ShotGraphVariablesCategory);
+			AddVariableAssignments(SelectedShot->GetPrimaryGraphVariableAssignments(), PrimaryGraphVariablesShotOverridesCategory);
 		}
 		else
 		{
-			AddVariableAssignments(SelectedJobs[0]->GetGraphVariableAssignments(), PrimaryGraphVariablesCategory);
+			AddVariableAssignments(SelectedJob->GetGraphVariableAssignments(), PrimaryGraphVariablesCategory);
 		}
 	}
 	//~ End IDetailCustomization interface
@@ -166,6 +186,12 @@ private:
 private:
 	/** The details builder associated with the customization. */
 	IDetailLayoutBuilder* DetailBuilder = nullptr;
+
+	/** The primary job that's selected in the UI. There should always be a selected primary job. */
+	TWeakObjectPtr<UMoviePipelineExecutorJob> SelectedJob;
+
+	/** The shot that's selected in the UI (may be null if only a primary is selected). */
+	TWeakObjectPtr<UMoviePipelineExecutorShot> SelectedShot;
 };
 
 #undef LOCTEXT_NAMESPACE
\ No newline at end of file
diff --git a/Engine/Plugins/MovieScene/MovieRenderPipeline/Source/MovieRenderPipelineEditor/Private/Widgets/SMoviePipelineQueueEditor.cpp b/Engine/Plugins/MovieScene/MovieRenderPipeline/Source/MovieRenderPipelineEditor/Private/Widgets/SMoviePipelineQueueEditor.cpp
--- a/Engine/Plugins/MovieScene/MovieRenderPipeline/Source/MovieRenderPipelineEditor/Private/Widgets/SMoviePipelineQueueEditor.cpp
+++ b/Engine/Plugins/MovieScene/MovieRenderPipeline/Source/MovieRenderPipelineEditor/Private/Widgets/SMoviePipelineQueueEditor.cpp
@@ -1256,11 +1256,11 @@ void SMoviePipelineQueueEditor::Construct(const FArguments& InArgs)
 			.DefaultLabel(LOCTEXT("QueueHeaderJobName_Text", "Job"))
 
 			+ SHeaderRow::Column(SQueueJobListRow::NAME_Settings)
-			.FillWidth(0.20f)
+			.FillWidth(0.25f)
 			.DefaultLabel(LOCTEXT("QueueHeaderSettings_Text", "Settings"))
 
 			+ SHeaderRow::Column(SQueueJobListRow::NAME_Output)
-			.FillWidth(0.5f)
+			.FillWidth(0.45f)
 			.DefaultLabel(LOCTEXT("QueueHeaderOutput_Text", "Output"))
 
 			+ SHeaderRow::Column(SQueueJobListRow::NAME_Status)
