Unshelved from pending changelist '30698804':

MovieGraph: Add workaround for motion blur not working on grooms when rendering with multiple layers.

diff --git a/Engine/Plugins/MovieScene/MovieRenderPipeline/Source/MovieRenderPipelineCore/Private/Graph/MovieGraphDefaultRenderer.cpp b/Engine/Plugins/MovieScene/MovieRenderPipeline/Source/MovieRenderPipelineCore/Private/Graph/MovieGraphDefaultRenderer.cpp
--- a/Engine/Plugins/MovieScene/MovieRenderPipeline/Source/MovieRenderPipelineCore/Private/Graph/MovieGraphDefaultRenderer.cpp
+++ b/Engine/Plugins/MovieScene/MovieRenderPipeline/Source/MovieRenderPipelineCore/Private/Graph/MovieGraphDefaultRenderer.cpp
@@ -250,6 +250,23 @@ void UMovieGraphDefaultRenderer::Render(const FMovieGraphTimeStepData& InTimeSte
 	}
 #endif
 
+	// Workaround for UE-202937, we need to detect when there will be multiple scene views rendered for a given frame
+	// to have grooms handle motion blur correctly when there are multiple views being rendered.
+	int32 NumSceneViewsRendered = 0;
+	for (const TObjectPtr<UMovieGraphRenderPassNode>& RenderPass : RenderPassesInUse)
+	{
+		NumSceneViewsRendered += RenderPass->GetNumSceneViewsRendered();
+	}
+
+	if (NumSceneViewsRendered > 1)
+	{
+		IConsoleVariable* HairStrandsCVar = IConsoleManager::Get().FindConsoleVariable(TEXT("r.HairStrands.Strands.MotionVectorCheckViewID"));
+		if (HairStrandsCVar)
+		{
+			HairStrandsCVar->SetWithCurrentPriority(0);
+		}
+	}
+
 	for (const TObjectPtr<UMovieGraphRenderPassNode>& RenderPass : RenderPassesInUse)
 	{
 		// Pass in a copy of the traversal context so the renderer can decide what to do with it.
@@ -257,6 +274,16 @@ void UMovieGraphDefaultRenderer::Render(const FMovieGraphTimeStepData& InTimeSte
 		RenderPass->Render(OutputFrame.TraversalContext, InTimeStepData);
 	}
 
+	if (NumSceneViewsRendered > 1)
+	{
+		// Restore the default value afterwards.
+		IConsoleVariable* HairStrandsCVar = IConsoleManager::Get().FindConsoleVariable(TEXT("r.HairStrands.Strands.MotionVectorCheckViewID"));
+		if (HairStrandsCVar)
+		{
+			HairStrandsCVar->SetWithCurrentPriority(1);
+		}
+	}
+
 	// Re-enable the progress widget so when the player viewport is drawn to the preview window, it shows.
 	GetOwningGraph()->SetPreviewWidgetVisible(true);
 }
diff --git a/Engine/Plugins/MovieScene/MovieRenderPipeline/Source/MovieRenderPipelineCore/Public/Graph/Nodes/MovieGraphRenderPassNode.h b/Engine/Plugins/MovieScene/MovieRenderPipeline/Source/MovieRenderPipelineCore/Public/Graph/Nodes/MovieGraphRenderPassNode.h
--- a/Engine/Plugins/MovieScene/MovieRenderPipeline/Source/MovieRenderPipelineCore/Public/Graph/Nodes/MovieGraphRenderPassNode.h
+++ b/Engine/Plugins/MovieScene/MovieRenderPipeline/Source/MovieRenderPipelineCore/Public/Graph/Nodes/MovieGraphRenderPassNode.h
@@ -48,6 +48,10 @@ public:
 	*/
 	void GatherOutputPasses(UMovieGraphEvaluatedConfig* InConfig, TArray<FMovieGraphRenderDataIdentifier>& OutExpectedPasses) const { GatherOutputPassesImpl(InConfig, OutExpectedPasses); }
 
+	/** Gets the number of Scene Views (that is, renders of the 3d scene) that this pass will produce. Can be zero for things like UI Renderers,
+	* or more than one for things like panoramic or tiling. */
+	int32 GetNumSceneViewsRendered() const { return GetNumSceneViewsRenderedImpl(); }
+
 #if WITH_EDITOR
 	virtual FText GetMenuCategory() const override
 	{
@@ -73,4 +77,5 @@ protected:
 	virtual void TeardownImpl() {}
 	virtual void RenderImpl(const FMovieGraphTraversalContext& InFrameTraversalContext, const FMovieGraphTimeStepData& InTimeData) {}
 	virtual void GatherOutputPassesImpl(UMovieGraphEvaluatedConfig* InConfig, TArray<FMovieGraphRenderDataIdentifier>& OutExpectedPasses) const {}
+	virtual int32 GetNumSceneViewsRenderedImpl() const { return 0; }
 };
\ No newline at end of file
diff --git a/Engine/Plugins/MovieScene/MovieRenderPipeline/Source/MovieRenderPipelineRenderPasses/Public/Graph/Nodes/MovieGraphImagePassBaseNode.h b/Engine/Plugins/MovieScene/MovieRenderPipeline/Source/MovieRenderPipelineRenderPasses/Public/Graph/Nodes/MovieGraphImagePassBaseNode.h
--- a/Engine/Plugins/MovieScene/MovieRenderPipeline/Source/MovieRenderPipelineRenderPasses/Public/Graph/Nodes/MovieGraphImagePassBaseNode.h
+++ b/Engine/Plugins/MovieScene/MovieRenderPipeline/Source/MovieRenderPipelineRenderPasses/Public/Graph/Nodes/MovieGraphImagePassBaseNode.h
@@ -80,7 +80,8 @@ protected:
 	virtual void TeardownImpl() override;
 	virtual void RenderImpl(const FMovieGraphTraversalContext& InFrameTraversalContext, const FMovieGraphTimeStepData& InTimeData) override;
 	virtual void GatherOutputPassesImpl(UMovieGraphEvaluatedConfig* InConfig, TArray<FMovieGraphRenderDataIdentifier>& OutExpectedPasses) const override;
-	// ~UMovieGraphRenderPassNode Interface
+	virtual int32 GetNumSceneViewsRenderedImpl() const { return CurrentInstances.Num(); }
+    // ~UMovieGraphRenderPassNode Interface
 
 	virtual TUniquePtr<UE::MovieGraph::Rendering::FMovieGraphImagePassBase> CreateInstance() const { return nullptr; }
 };
\ No newline at end of file
