Movie Graph:

-Fix: Crash with unconnected Subgraph Node

diff --git a/Engine/Plugins/MovieScene/MovieRenderPipeline/Source/MovieRenderPipelineCore/Private/Graph/MovieGraphConfig.cpp b/Engine/Plugins/MovieScene/MovieRenderPipeline/Source/MovieRenderPipelineCore/Private/Graph/MovieGraphConfig.cpp
--- a/Engine/Plugins/MovieScene/MovieRenderPipeline/Source/MovieRenderPipelineCore/Private/Graph/MovieGraphConfig.cpp
+++ b/Engine/Plugins/MovieScene/MovieRenderPipeline/Source/MovieRenderPipelineCore/Private/Graph/MovieGraphConfig.cpp
@@ -1045,6 +1045,11 @@ void UMovieGraphConfig::GetAllContainedSubgraphs(TSet<UMovieGraphConfig*>& OutSu
 		if (const UMovieGraphSubgraphNode* SubgraphNode = Cast<UMovieGraphSubgraphNode>(Node))
 		{
 			UMovieGraphConfig* SubgraphConfig = SubgraphNode->GetSubgraphAsset();
+
+			if (!SubgraphConfig) // A subgraph may not have been assigned yet
+			{
+				continue;
+			}
 			
 			// Don't recurse into this graph if it was already added (to prevent infinite recursion)
 			if (!OutSubgraphs.Contains(SubgraphConfig))
