MovieGraph: Unify namespaces for new files (UE::MovieGraph)

diff --git a/Engine/Plugins/MovieScene/MovieRenderPipeline/Source/MovieRenderPipelineCore/Private/Graph/MovieGraphConfig.cpp b/Engine/Plugins/MovieScene/MovieRenderPipeline/Source/MovieRenderPipelineCore/Private/Graph/MovieGraphConfig.cpp
--- a/Engine/Plugins/MovieScene/MovieRenderPipeline/Source/MovieRenderPipelineCore/Private/Graph/MovieGraphConfig.cpp
+++ b/Engine/Plugins/MovieScene/MovieRenderPipeline/Source/MovieRenderPipelineCore/Private/Graph/MovieGraphConfig.cpp
@@ -15,6 +15,7 @@
 #include "Graph/Nodes/MovieGraphSubgraphNode.h"
 #include "Graph/Nodes/MovieGraphVariableNode.h"
 #include "Graph/Nodes/MovieGraphSelectNode.h"
+#include "Graph/MovieGraphUtils.h"
 #include "MovieGraphUtils.h"
 #include "MoviePipelineQueue.h"
 #include "MovieRenderPipelineCoreModule.h"
@@ -667,8 +668,6 @@ RetType* UMovieGraphConfig::AddMember(TArray<TObjectPtr<ArrType>>& InMemberArray
 {
 	static_assert(std::is_base_of_v<UMovieGraphMember, RetType>, "RetType is not derived from UMovieGraphMember");
 	
-	using namespace UE::MoviePipeline::RenderGraph;
-
 	// TODO: This can be replaced with just CreateDefaultSubobject() when AddDefaultMembers() isn't called from PostLoad()
 	//
 	// This method will be called in two cases: 1) when default members are being added to a new graph when it is being
@@ -698,7 +697,7 @@ RetType* UMovieGraphConfig::AddMember(TArray<TObjectPtr<ArrType>>& InMemberArray
 	{
 		TArray<FString> ExistingMemberNames;
 		Algo::Transform(InMemberArray, ExistingMemberNames, [](const ArrType* Member) { return Member->GetMemberName(); });
-		NewMember->SetMemberName(GetUniqueName(ExistingMemberNames, InBaseName.ToString()));
+		NewMember->SetMemberName(UE::MovieGraph::GetUniqueName(ExistingMemberNames, InBaseName.ToString()));
 	}
 
 	return NewMember;
diff --git a/Engine/Plugins/MovieScene/MovieRenderPipeline/Source/MovieRenderPipelineCore/Private/Graph/MovieGraphDefaultAudioRenderer.cpp b/Engine/Plugins/MovieScene/MovieRenderPipeline/Source/MovieRenderPipelineCore/Private/Graph/MovieGraphDefaultAudioRenderer.cpp
--- a/Engine/Plugins/MovieScene/MovieRenderPipeline/Source/MovieRenderPipelineCore/Private/Graph/MovieGraphDefaultAudioRenderer.cpp
+++ b/Engine/Plugins/MovieScene/MovieRenderPipeline/Source/MovieRenderPipelineCore/Private/Graph/MovieGraphDefaultAudioRenderer.cpp
@@ -1,4 +1,4 @@
-﻿// Copyright Epic Games, Inc. All Rights Reserved.
+// Copyright Epic Games, Inc. All Rights Reserved.
 
 #include "Graph/MovieGraphDefaultAudioRenderer.h"
 
@@ -20,7 +20,7 @@ void UMovieGraphDefaultAudioRenderer::StartAudioRecording()
 {
 	AudioState.bIsRecordingAudio = true;
 
-	if (Audio::FMixerDevice* MixerDevice = UE::MoviePipeline::RenderGraph::Audio::GetAudioMixerDeviceFromWorldContext(this))
+	if (Audio::FMixerDevice* MixerDevice = UE::MovieGraph::Audio::GetAudioMixerDeviceFromWorldContext(this))
 	{
 		const TWeakPtr<Audio::FMixerSubmix> MasterSubmix = MixerDevice->GetMasterSubmix();
 		if (MasterSubmix.Pin())
@@ -58,7 +58,7 @@ void UMovieGraphDefaultAudioRenderer::StopAudioRecording()
 void UMovieGraphDefaultAudioRenderer::ProcessAudioTick()
 {
 	// Only supported on the new audio mixer (with the non-realtime device, windows only).
-	Audio::FMixerDevice* MixerDevice = UE::MoviePipeline::RenderGraph::Audio::GetAudioMixerDeviceFromWorldContext(this);
+	Audio::FMixerDevice* MixerDevice = UE::MovieGraph::Audio::GetAudioMixerDeviceFromWorldContext(this);
 	if (!MixerDevice)
 	{
 		return;
diff --git a/Engine/Plugins/MovieScene/MovieRenderPipeline/Source/MovieRenderPipelineCore/Private/Graph/MovieGraphUtils.cpp b/Engine/Plugins/MovieScene/MovieRenderPipeline/Source/MovieRenderPipelineCore/Private/Graph/MovieGraphUtils.cpp
--- a/Engine/Plugins/MovieScene/MovieRenderPipeline/Source/MovieRenderPipelineCore/Private/Graph/MovieGraphUtils.cpp
+++ b/Engine/Plugins/MovieScene/MovieRenderPipeline/Source/MovieRenderPipelineCore/Private/Graph/MovieGraphUtils.cpp
@@ -6,7 +6,7 @@
 #include "AudioThread.h"
 #include "Engine/Engine.h"
 
-namespace UE::MoviePipeline::RenderGraph
+namespace UE::MovieGraph
 {
 	FString GetUniqueName(const TArray<FString>& InExistingNames, const FString& InBaseName)
 	{
diff --git a/Engine/Plugins/MovieScene/MovieRenderPipeline/Source/MovieRenderPipelineCore/Private/Graph/MovieGraphUtils.h b/Engine/Plugins/MovieScene/MovieRenderPipeline/Source/MovieRenderPipelineCore/Private/Graph/MovieGraphUtils.h
--- a/Engine/Plugins/MovieScene/MovieRenderPipeline/Source/MovieRenderPipelineCore/Private/Graph/MovieGraphUtils.h
+++ b/Engine/Plugins/MovieScene/MovieRenderPipeline/Source/MovieRenderPipelineCore/Private/Graph/MovieGraphUtils.h
@@ -1,4 +1,4 @@
-﻿// Copyright Epic Games, Inc. All Rights Reserved.
+// Copyright Epic Games, Inc. All Rights Reserved.
 
 #pragma once
 
@@ -11,7 +11,7 @@ namespace Audio
 	class FMixerDevice;
 }
 
-namespace UE::MoviePipeline::RenderGraph
+namespace UE::MovieGraph
 {
 	/**
 	 * Generate a unique name given a set of existing names and the desired base name. The base name will
diff --git a/Engine/Plugins/MovieScene/MovieRenderPipeline/Source/MovieRenderPipelineCore/Private/Graph/Nodes/MovieGraphAudioOutputNode.cpp b/Engine/Plugins/MovieScene/MovieRenderPipeline/Source/MovieRenderPipelineCore/Private/Graph/Nodes/MovieGraphAudioOutputNode.cpp
--- a/Engine/Plugins/MovieScene/MovieRenderPipeline/Source/MovieRenderPipelineCore/Private/Graph/Nodes/MovieGraphAudioOutputNode.cpp
+++ b/Engine/Plugins/MovieScene/MovieRenderPipeline/Source/MovieRenderPipelineCore/Private/Graph/Nodes/MovieGraphAudioOutputNode.cpp
@@ -201,7 +201,7 @@ void UMovieGraphAudioOutputNode::GenerateFinalAudioData(TArray<FFinalAudioData>&
 
 void UMovieGraphAudioOutputNode::StartAudioExport()
 {
-	if (!UE::MoviePipeline::RenderGraph::Audio::IsMoviePipelineAudioOutputSupported(CachedPipeline.Get()))
+	if (!UE::MovieGraph::Audio::IsMoviePipelineAudioOutputSupported(CachedPipeline.Get()))
 	{
 		return;
 	}
