MoviePipeline: Fixed several issues in the graph that caused pins with data connections to not be resolved correctly.

The biggest issue was that traversal only looked one connection deep for resolving data connection values. For example, if an exposed node property was connected to a subgraph input, and that subgraph input was subsequently connected to a variable, the subgraph's input value would be used instead of the variable's.

Second, there were a number of places in the graph that default values for pins did not resolve correctly. This included output pins on input nodes, and output pins on subgraph nodes.

Lastly, UMovieGraphNode and UMovieGraphSubgraphNode now follow data connections properly.

Note that the Select and Branch nodes, due to their logic in EvaluatePinsToFollow(), still evaluate their data connections only one level deep like graph traversal previously did. This should continue to cover most use cases, but it's not ideal. These nodes should have their logic pins converted to overrideable properties, and those properties queried directly in EvaluatePinsToFollow() -- the new graph traversal logic should set their values correctly before EvaluatePinsToFollow() runs.

#rnx

diff --git a/Engine/Plugins/MovieScene/MovieRenderPipeline/Source/MovieRenderPipelineCore/Private/Graph/MovieGraphConfig.cpp b/Engine/Plugins/MovieScene/MovieRenderPipeline/Source/MovieRenderPipelineCore/Private/Graph/MovieGraphConfig.cpp
--- a/Engine/Plugins/MovieScene/MovieRenderPipeline/Source/MovieRenderPipelineCore/Private/Graph/MovieGraphConfig.cpp
+++ b/Engine/Plugins/MovieScene/MovieRenderPipeline/Source/MovieRenderPipelineCore/Private/Graph/MovieGraphConfig.cpp
@@ -1164,7 +1164,7 @@ void UMovieGraphConfig::InitializeFlattenedNode(UMovieGraphNode* InNode)
 	InNode->UpdateDynamicProperties();
 }
 
-void UMovieGraphConfig::CopyOverriddenProperties(UMovieGraphNode* FromNode, UMovieGraphNode* ToNode, const FMovieGraphTraversalContext* InContext)
+void UMovieGraphConfig::CopyOverriddenProperties(UMovieGraphNode* FromNode, UMovieGraphNode* ToNode, const FMovieGraphEvaluationContext& InEvaluationContext)
 {
 	if (!ensure(FromNode && ToNode))
 	{
@@ -1228,15 +1228,52 @@ void UMovieGraphConfig::CopyOverriddenProperties(UMovieGraphNode* FromNode, UMov
 		// If this property (dynamic or not) has been exposed, attempt to get its value via the connection to it (if any)
 		if (bIsExposed)
 		{
-			if (const UMovieGraphPin* InputPin = FromNode->GetInputPin(PropertyName))
+			if (UMovieGraphPin* InputPin = FromNode->GetInputPin(PropertyName))
 			{
-				// For the connected value to be used, the type must match and the node the value is originating from must be enabled
-				const UMovieGraphPin* ConnectedPin = InputPin->GetFirstConnectedPin();
-				if (ConnectedPin && (ConnectedPin->Properties.Type == InputPin->Properties.Type) && ConnectedPin->Node && !ConnectedPin->Node->IsDisabled())
+				TArray<UMovieGraphPin*> ConnectionPath;
+				
+				// Iterate up the connection chain and find all pins which might have a value that can be resolved.
+				FMovieGraphEvaluationContext ValueConnectionContext;
+				ValueConnectionContext.PinBeingFollowed = InputPin;
+				ValueConnectionContext.SubgraphStack = InEvaluationContext.SubgraphStack;
+				TArray<UMovieGraphPin*> ConnectedValuePins = InputPin->Node->EvaluatePinsToFollow(ValueConnectionContext);
+				while (!ConnectedValuePins.IsEmpty())
 				{
-					// There was a valid connection to the input pin; resolve the value from the connected output and set the
-					// value on this property
-					const FString ResolvedValue = ConnectedPin->Node->GetResolvedValueForOutputPin(ConnectedPin->Properties.Label, InContext);
+					UMovieGraphPin* ConnectedValuePin = ConnectedValuePins[0];
+					if (ConnectionPath.Contains(ConnectedValuePin))
+					{
+						// Recursive connection found
+						UE_LOG(LogMovieRenderPipeline, Warning, TEXT("Found a cycle when following the data connection on pin '%s' for node '%s'. Value will not be resolved."),
+							*ConnectedValuePin->Properties.Label.ToString(), *FromNode->GetName());
+						break;
+					}
+
+					// For the connected value to be used, the type must match and the node the value is originating from must be enabled
+					if (ConnectedValuePin && (ConnectedValuePin->Properties.Type == InputPin->Properties.Type) && ConnectedValuePin->Node && !ConnectedValuePin->Node->IsDisabled())
+					{
+						ConnectionPath.Add(ConnectedValuePin);
+					}
+
+					ValueConnectionContext.PinBeingFollowed = ConnectedValuePin;
+					ConnectedValuePins = ConnectedValuePin->Node->EvaluatePinsToFollow(ValueConnectionContext);
+				}
+
+				// Work backwards and use the most upstream value that can be resolved. The most upstream values wins. For example, if a node has an
+				// exposed pin, that pin is connected to a subgraph's input, and that input is then connected to a variable node in the parent graph.
+				// The variable node's value should be used if it can be resolved, not the subgraph's input value.
+				bool bFoundResolvedValue = false;
+				for (int32 Index = ConnectionPath.Num() - 1; Index >= 0; --Index)
+				{
+					const UMovieGraphPin* ConnectedPin = ConnectionPath[Index];
+					
+					const FString ResolvedValue = ConnectedPin->Node->GetResolvedValueForOutputPin(ConnectedPin->Properties.Label, &InEvaluationContext.UserContext);
+					if (ResolvedValue.IsEmpty())
+					{
+						continue;
+					}
+
+					bFoundResolvedValue = true;
+					
 					if (bIsDynamic)
 					{
 						ToNode->SetDynamicPropertyValue(PropertyName, ResolvedValue);
@@ -1248,7 +1285,13 @@ void UMovieGraphConfig::CopyOverriddenProperties(UMovieGraphNode* FromNode, UMov
 						EditConditionProperty->SetPropertyValue_InContainer(ToNode, true);
 					}
 
-					// The property value was set via a connected pin; move on to the next property
+					// Resolved a value for this pin; stop iterating over the connection chain
+					break;
+				}
+				
+				// The property value was set via a connected pin; move on to the next property
+				if (bFoundResolvedValue)
+				{
 					continue;
 				}
 			}
@@ -1391,7 +1434,7 @@ bool UMovieGraphConfig::CreateFlattenedGraph_Recursive(UMovieGraphEvaluatedConfi
 		// Now do a property-copy from this node onto our flattened one. We don't use the generic property
 		// copy routines in the engine because we have special handling (we want to check if the property
 		// is actually marked for override, and also skip if this has already been overridden).
-		CopyOverriddenProperties(Node, ExistingNode, &InEvaluationContext.UserContext);
+		CopyOverriddenProperties(Node, ExistingNode, InEvaluationContext);
 	}
 
 	// If this is a special "removal" node, keep track of the type that should be removed. Since this method is recursive,
diff --git a/Engine/Plugins/MovieScene/MovieRenderPipeline/Source/MovieRenderPipelineCore/Private/Graph/MovieGraphNode.cpp b/Engine/Plugins/MovieScene/MovieRenderPipeline/Source/MovieRenderPipelineCore/Private/Graph/MovieGraphNode.cpp
--- a/Engine/Plugins/MovieScene/MovieRenderPipeline/Source/MovieRenderPipelineCore/Private/Graph/MovieGraphNode.cpp
+++ b/Engine/Plugins/MovieScene/MovieRenderPipeline/Source/MovieRenderPipelineCore/Private/Graph/MovieGraphNode.cpp
@@ -499,9 +499,15 @@ TArray<FMovieGraphPropertyInfo> UMovieGraphNode::GetOverrideablePropertyInfo() c
 TArray<UMovieGraphPin*> UMovieGraphNode::EvaluatePinsToFollow(FMovieGraphEvaluationContext& InContext) const
 {
 	TArray<UMovieGraphPin*> PinsToFollow;
-
+	
+	if (!ensure(InContext.PinBeingFollowed))
+    {
+    	return PinsToFollow;
+    }
+    
 	// If the node is disabled, only follow the first connected pin.
-	if (IsDisabled())
+    // This is only important for branch connections. For data pins, just continue following the connection.
+	if (IsDisabled() && InContext.PinBeingFollowed->Properties.bIsBranch)
 	{
 		if (UMovieGraphPin* GraphPin = GetFirstConnectedInputPin())
 		{
@@ -511,6 +517,22 @@ TArray<UMovieGraphPin*> UMovieGraphNode::EvaluatePinsToFollow(FMovieGraphEvaluat
 		return PinsToFollow;
 	}
 
+	// If this is a data (non-branch) connection, just return the first connected pin.
+	if (!InContext.PinBeingFollowed->Properties.bIsBranch)
+	{
+		// Only do this if the pin is an input though. For data output pins, the generic case does not have enough information to determine what
+		// the upstream connected data pin is.
+		if (InContext.PinBeingFollowed->IsInputPin())
+		{
+			if (UMovieGraphPin* ConnectedInputPin = InContext.PinBeingFollowed->GetFirstConnectedPin())
+			{
+				PinsToFollow.Add(ConnectedInputPin);
+			}
+		}
+		
+		return PinsToFollow;
+	}
+
 	// By default we provide all Input Pins to this node that are the Branch Type.
 	// You should override this in downstream nodes that need custom logic, such
 	// as branch or switch nodes.
@@ -521,6 +543,7 @@ TArray<UMovieGraphPin*> UMovieGraphNode::EvaluatePinsToFollow(FMovieGraphEvaluat
 			PinsToFollow.Add(InputPin);
 		}
 	}
+	
 	return PinsToFollow;
 }
 
diff --git a/Engine/Plugins/MovieScene/MovieRenderPipeline/Source/MovieRenderPipelineCore/Private/Graph/Nodes/MovieGraphInputNode.cpp b/Engine/Plugins/MovieScene/MovieRenderPipeline/Source/MovieRenderPipelineCore/Private/Graph/Nodes/MovieGraphInputNode.cpp
--- a/Engine/Plugins/MovieScene/MovieRenderPipeline/Source/MovieRenderPipelineCore/Private/Graph/Nodes/MovieGraphInputNode.cpp
+++ b/Engine/Plugins/MovieScene/MovieRenderPipeline/Source/MovieRenderPipelineCore/Private/Graph/Nodes/MovieGraphInputNode.cpp
@@ -80,6 +80,22 @@ TArray<UMovieGraphPin*> UMovieGraphInputNode::EvaluatePinsToFollow(FMovieGraphEv
 	return PinsToFollow;
 }
 
+FString UMovieGraphInputNode::GetResolvedValueForOutputPin(const FName& InPinName, const FMovieGraphTraversalContext* InContext) const
+{
+	if (const UMovieGraphConfig* Graph = GetGraph())
+	{
+		for (UMovieGraphInput* Input : Graph->GetInputs())
+		{
+			if (Input && (FName(Input->GetMemberName()) == InPinName))
+			{
+				return Input->GetValueSerializedString();
+			}
+		}
+	}
+
+	return FString();
+}
+
 bool UMovieGraphInputNode::CanBeDisabled() const
 {
 	return false;
diff --git a/Engine/Plugins/MovieScene/MovieRenderPipeline/Source/MovieRenderPipelineCore/Private/Graph/Nodes/MovieGraphSubgraphNode.cpp b/Engine/Plugins/MovieScene/MovieRenderPipeline/Source/MovieRenderPipelineCore/Private/Graph/Nodes/MovieGraphSubgraphNode.cpp
--- a/Engine/Plugins/MovieScene/MovieRenderPipeline/Source/MovieRenderPipelineCore/Private/Graph/Nodes/MovieGraphSubgraphNode.cpp
+++ b/Engine/Plugins/MovieScene/MovieRenderPipeline/Source/MovieRenderPipelineCore/Private/Graph/Nodes/MovieGraphSubgraphNode.cpp
@@ -75,8 +75,14 @@ TArray<UMovieGraphPin*> UMovieGraphSubgraphNode::EvaluatePinsToFollow(FMovieGrap
 {
 	TArray<UMovieGraphPin*> PinsToFollow;
 
+	if (!ensure(InContext.PinBeingFollowed))
+	{
+		return PinsToFollow;
+	}
+
 	// If the node is disabled, follow the pin for the first option available.
-	if (IsDisabled())
+	// This is only important for branch connections. For data pins, just continue following the connection.
+	if (IsDisabled() && InContext.PinBeingFollowed->Properties.bIsBranch)
 	{
 		if (UMovieGraphPin* GraphPin = GetFirstConnectedInputPin())
 		{
@@ -86,11 +92,6 @@ TArray<UMovieGraphPin*> UMovieGraphSubgraphNode::EvaluatePinsToFollow(FMovieGrap
 		return PinsToFollow;
 	}
 	
-	if (!ensure(InContext.PinBeingFollowed))
-	{
-		return PinsToFollow;
-	}
-
 	const UMovieGraphConfig* SubgraphPtr = SubgraphAsset.LoadSynchronous();
 	if (!SubgraphPtr)
 	{
@@ -109,6 +110,17 @@ TArray<UMovieGraphPin*> UMovieGraphSubgraphNode::EvaluatePinsToFollow(FMovieGrap
 		return PinsToFollow;
 	}
 
+	// If an input on the node is being followed, just follow the connection on the pin, nothing fancy required.
+	if (InContext.PinBeingFollowed->IsInputPin())
+	{
+		if (UMovieGraphPin* ConnectedInputPin = InContext.PinBeingFollowed->GetFirstConnectedPin())
+		{
+			PinsToFollow.Add(ConnectedInputPin);
+		}
+		
+		return PinsToFollow;
+	}
+
 	// Find the input pin on the subgraph asset's outputs node that corresponds to the pin being followed
 	//
 	// Subgraph
@@ -139,6 +151,22 @@ TArray<UMovieGraphPin*> UMovieGraphSubgraphNode::EvaluatePinsToFollow(FMovieGrap
 	return PinsToFollow;
 }
 
+FString UMovieGraphSubgraphNode::GetResolvedValueForOutputPin(const FName& InPinName, const FMovieGraphTraversalContext* InContext) const
+{
+	if (const UMovieGraphConfig* Subgraph = SubgraphAsset.LoadSynchronous())
+	{
+		for (UMovieGraphOutput* Output : Subgraph->GetOutputs())
+		{
+			if (Output && (FName(Output->GetMemberName()) == InPinName))
+			{
+				return Output->GetValueSerializedString();
+			}
+		}
+	}
+
+	return FString();
+}
+
 #if WITH_EDITOR
 void UMovieGraphSubgraphNode::PostEditChangeProperty(FPropertyChangedEvent& PropertyChangedEvent)
 {
diff --git a/Engine/Plugins/MovieScene/MovieRenderPipeline/Source/MovieRenderPipelineCore/Public/Graph/MovieGraphConfig.h b/Engine/Plugins/MovieScene/MovieRenderPipeline/Source/MovieRenderPipelineCore/Public/Graph/MovieGraphConfig.h
--- a/Engine/Plugins/MovieScene/MovieRenderPipeline/Source/MovieRenderPipelineCore/Public/Graph/MovieGraphConfig.h
+++ b/Engine/Plugins/MovieScene/MovieRenderPipeline/Source/MovieRenderPipelineCore/Public/Graph/MovieGraphConfig.h
@@ -752,7 +752,7 @@ protected:
 	void RecurseUpGlobalsBranchToFindOutputDirectory(const UMovieGraphNode* InNode, FString& OutOutputDirectory, TArray<const UMovieGraphConfig*>& VisitedGraphStack) const;
 	
 	/** Copies properties in FromNode that are marked for override into ToNode, but only if ToNode doesn't already override that value. */
-	void CopyOverriddenProperties(UMovieGraphNode* FromNode, UMovieGraphNode* ToNode, const FMovieGraphTraversalContext* InContext);
+	void CopyOverriddenProperties(UMovieGraphNode* FromNode, UMovieGraphNode* ToNode, const FMovieGraphEvaluationContext& InEvaluationContext);
 	
 	/** Find all "Overrideable" marked properties, then find their edit condition properties, then set those to false. */
 	void InitializeFlattenedNode(UMovieGraphNode* InNode);
diff --git a/Engine/Plugins/MovieScene/MovieRenderPipeline/Source/MovieRenderPipelineCore/Public/Graph/MovieGraphNode.h b/Engine/Plugins/MovieScene/MovieRenderPipeline/Source/MovieRenderPipelineCore/Public/Graph/MovieGraphNode.h
--- a/Engine/Plugins/MovieScene/MovieRenderPipeline/Source/MovieRenderPipelineCore/Public/Graph/MovieGraphNode.h
+++ b/Engine/Plugins/MovieScene/MovieRenderPipeline/Source/MovieRenderPipelineCore/Public/Graph/MovieGraphNode.h
@@ -185,7 +185,7 @@ public:
 	}
 
 	/** 
-	* Used to determine which Branch type pins we should follow when trying to traverse the graph.
+	* Used to determine which pins we should follow when trying to traverse the graph.
 	* By default we will follow any input pin (with Branch type) on the node, but override this in
 	* inherited classes and change that if you need custom logic, such as boolean nodes that want 
 	* to choose one or the other based on the results of a conditional property.
diff --git a/Engine/Plugins/MovieScene/MovieRenderPipeline/Source/MovieRenderPipelineCore/Public/Graph/Nodes/MovieGraphInputNode.h b/Engine/Plugins/MovieScene/MovieRenderPipeline/Source/MovieRenderPipelineCore/Public/Graph/Nodes/MovieGraphInputNode.h
--- a/Engine/Plugins/MovieScene/MovieRenderPipeline/Source/MovieRenderPipelineCore/Public/Graph/Nodes/MovieGraphInputNode.h
+++ b/Engine/Plugins/MovieScene/MovieRenderPipeline/Source/MovieRenderPipelineCore/Public/Graph/Nodes/MovieGraphInputNode.h
@@ -18,6 +18,7 @@ public:
 	//~ Begin UMovieGraphNode interface
 	virtual TArray<FMovieGraphPinProperties> GetOutputPinProperties() const override;
 	virtual TArray<UMovieGraphPin*> EvaluatePinsToFollow(FMovieGraphEvaluationContext& InContext) const override;
+	virtual FString GetResolvedValueForOutputPin(const FName& InPinName, const FMovieGraphTraversalContext* InContext) const override;
 	virtual bool CanBeDisabled() const override;
 	
 	virtual bool CanBeAddedByUser() const override { return false; }
diff --git a/Engine/Plugins/MovieScene/MovieRenderPipeline/Source/MovieRenderPipelineCore/Public/Graph/Nodes/MovieGraphSubgraphNode.h b/Engine/Plugins/MovieScene/MovieRenderPipeline/Source/MovieRenderPipelineCore/Public/Graph/Nodes/MovieGraphSubgraphNode.h
--- a/Engine/Plugins/MovieScene/MovieRenderPipeline/Source/MovieRenderPipelineCore/Public/Graph/Nodes/MovieGraphSubgraphNode.h
+++ b/Engine/Plugins/MovieScene/MovieRenderPipeline/Source/MovieRenderPipelineCore/Public/Graph/Nodes/MovieGraphSubgraphNode.h
@@ -24,6 +24,7 @@ public:
 	virtual TArray<FMovieGraphPinProperties> GetInputPinProperties() const override;
 	virtual TArray<FMovieGraphPinProperties> GetOutputPinProperties() const override;
 	virtual TArray<UMovieGraphPin*> EvaluatePinsToFollow(FMovieGraphEvaluationContext& InContext) const override;
+	virtual FString GetResolvedValueForOutputPin(const FName& InPinName, const FMovieGraphTraversalContext* InContext) const override;
 
 #if WITH_EDITOR
 	virtual FText GetNodeTitle(const bool bGetDescriptive = false) const override;
