MovieGraph: Renamed Path Tracer node/class to be stylistically inline with the Deferred Renderer.

diff --git a/Engine/Plugins/MovieScene/MovieRenderPipeline/Config/BaseMovieRenderPipeline.ini b/Engine/Plugins/MovieScene/MovieRenderPipeline/Config/BaseMovieRenderPipeline.ini
--- a/Engine/Plugins/MovieScene/MovieRenderPipeline/Config/BaseMovieRenderPipeline.ini
+++ b/Engine/Plugins/MovieScene/MovieRenderPipeline/Config/BaseMovieRenderPipeline.ini
@@ -7,3 +7,4 @@
 +ClassRedirects=(OldName="/Script/MovieRenderPipelineCore.MovieGraphDeferredRenderPassNode",NewName="/Script/MovieRenderPipelineRenderPasses.MovieGraphDeferredRenderPassNode")
 +ClassRedirects=(OldName="/Script/MovieRenderPipelineCore.MovieGraphImagePassBaseNode",NewName="/Script/MovieRenderPipelineRenderPasses.MovieGraphImagePassBaseNode")
 +ClassRedirects=(OldName="/Script/MovieRenderPipelineCore.MovieGraphPathTracedRenderPassNode",NewName="/Script/MovieRenderPipelineRenderPasses.MovieGraphPathTracedRenderPassNode")
++ClassRedirects=(OldName="/Script/MovieRenderPipelineRenderPasses.MovieGraphPathTracedRenderPassNode",NewName="/Script/MovieRenderPipelineRenderPasses.MovieGraphPathTracerRenderPassNode")
\ No newline at end of file
diff --git a/Engine/Plugins/MovieScene/MovieRenderPipeline/Source/MovieRenderPipelineRenderPasses/Private/Graph/Nodes/MovieGraphPathTracedRenderPassNode.cpp b/Engine/Plugins/MovieScene/MovieRenderPipeline/Source/MovieRenderPipelineRenderPasses/Private/Graph/Nodes/MovieGraphPathTracedRenderPassNode.cpp
deleted file mode 100644
diff --git a/Engine/Plugins/MovieScene/MovieRenderPipeline/Source/MovieRenderPipelineRenderPasses/Private/Graph/Nodes/MovieGraphPathTracerPassNode.cpp b/Engine/Plugins/MovieScene/MovieRenderPipeline/Source/MovieRenderPipelineRenderPasses/Private/Graph/Nodes/MovieGraphPathTracerPassNode.cpp
new file mode 100644
--- /dev/null
+++ b/Engine/Plugins/MovieScene/MovieRenderPipeline/Source/MovieRenderPipelineRenderPasses/Private/Graph/Nodes/MovieGraphPathTracerPassNode.cpp
@@ -0,0 +1,102 @@
+// Copyright Epic Games, Inc. All Rights Reserved.
+
+#include "Graph/Nodes/MovieGraphPathTracerPassNode.h"
+#include "Graph/Renderers/MovieGraphPathTracerPass.h"
+#include "Engine/EngineBaseTypes.h"
+#include "ShowFlags.h"
+
+TUniquePtr<UE::MovieGraph::Rendering::FMovieGraphImagePassBase> UMovieGraphPathTracerRenderPassNode::CreateInstance() const
+{
+	return MakeUnique<UE::MovieGraph::Rendering::FMovieGraphPathTracerPass>();
+}
+
+UMovieGraphPathTracerRenderPassNode::UMovieGraphPathTracerRenderPassNode()
+	: SpatialSampleCount(1)
+	, bDenoiser(true)
+	, bWriteAllSamples(false)
+	, bDisableToneCurve(false)
+	, bAllowOCIO(true)
+{
+	ShowFlags->ApplyDefaultShowFlagValue(VMI_PathTracing, true);
+	// TODO: Showflag for SetMotionBlur()?
+}
+
+#if WITH_EDITOR
+FText UMovieGraphPathTracerRenderPassNode::GetNodeTitle(const bool bGetDescriptive) const
+{
+	return NSLOCTEXT("MovieGraphNodes", "PathTracedRenderPassGraphNode_Description", "Path Traced Renderer");
+}
+#endif
+
+void UMovieGraphPathTracerRenderPassNode::SetupImpl(const FMovieGraphRenderPassSetupData& InSetupData)
+{
+	Super::SetupImpl(InSetupData);
+
+	// Hide the progress display during the render
+	if (IConsoleVariable* ProgressDisplayCvar = IConsoleManager::Get().FindConsoleVariable(TEXT("r.PathTracing.ProgressDisplay")))
+	{
+		bOriginalProgressDisplayCvarValue = ProgressDisplayCvar->GetBool();
+		ProgressDisplayCvar->Set(false);
+	}
+}
+
+void UMovieGraphPathTracerRenderPassNode::TeardownImpl()
+{
+	Super::TeardownImpl();
+
+	// Restore the original setting for the progress display
+	if (IConsoleVariable* ProgressDisplayCvar = IConsoleManager::Get().FindConsoleVariable(TEXT("r.PathTracing.ProgressDisplay")))
+	{
+		ProgressDisplayCvar->Set(bOriginalProgressDisplayCvarValue);
+	}
+}
+
+FString UMovieGraphPathTracerRenderPassNode::GetRendererNameImpl() const
+{
+	static const FString RendererNameImpl(TEXT("PathTraced"));
+	return RendererNameImpl;
+}
+
+EViewModeIndex UMovieGraphPathTracerRenderPassNode::GetViewModeIndex() const
+{
+	return VMI_PathTracing;
+}
+
+bool UMovieGraphPathTracerRenderPassNode::GetWriteAllSamples() const
+{
+	return bWriteAllSamples;
+}
+
+TArray<FMoviePipelinePostProcessPass> UMovieGraphPathTracerRenderPassNode::GetAdditionalPostProcessMaterials() const
+{
+	return AdditionalPostProcessMaterials;
+}
+
+int32 UMovieGraphPathTracerRenderPassNode::GetNumSpatialSamples() const
+{
+	return SpatialSampleCount;
+}
+
+int32 UMovieGraphPathTracerRenderPassNode::GetNumSpatialSamplesDuringWarmUp() const
+{
+	// Path Tracer doesn't have an image history like the deferred renderer, so it doesn't need
+	// to run all the spatial samples.
+	return 1;
+}
+
+
+bool UMovieGraphPathTracerRenderPassNode::GetDisableToneCurve() const
+{
+	return bDisableToneCurve;
+}
+
+bool UMovieGraphPathTracerRenderPassNode::GetAllowOCIO() const
+{
+	return bAllowOCIO;
+}
+
+bool UMovieGraphPathTracerRenderPassNode::GetAllowDenoiser() const
+{
+	return bDenoiser;
+}
+
diff --git a/Engine/Plugins/MovieScene/MovieRenderPipeline/Source/MovieRenderPipelineRenderPasses/Private/Graph/Renderers/MovieGraphPathTracerPass.cpp b/Engine/Plugins/MovieScene/MovieRenderPipeline/Source/MovieRenderPipelineRenderPasses/Private/Graph/Renderers/MovieGraphPathTracerPass.cpp
--- a/Engine/Plugins/MovieScene/MovieRenderPipeline/Source/MovieRenderPipelineRenderPasses/Private/Graph/Renderers/MovieGraphPathTracerPass.cpp
+++ b/Engine/Plugins/MovieScene/MovieRenderPipeline/Source/MovieRenderPipelineRenderPasses/Private/Graph/Renderers/MovieGraphPathTracerPass.cpp
@@ -1,7 +1,7 @@
 // Copyright Epic Games, Inc. All Rights Reserved.
 
 #include "Graph/Renderers/MovieGraphPathTracerPass.h"
-#include "Graph/Nodes/MovieGraphPathTracedRenderPassNode.h"
+#include "Graph/Nodes/MovieGraphPathTracerPassNode.h"
 
 namespace UE::MovieGraph::Rendering
 {
@@ -12,7 +12,7 @@ UMovieGraphImagePassBaseNode* FMovieGraphPathTracerPass::GetParentNode(UMovieGra
 	// since that instance changes every frame. So instead we have a virtual function here so the node can look it up by type, and then we can
 	// call a bunch of virtual functions on the right instance to fetch values.
 	const bool bIncludeCDOs = true;
-	UMovieGraphPathTracedRenderPassNode* ParentNode = InConfig->GetSettingForBranch<UMovieGraphPathTracedRenderPassNode>(GetBranchName(), bIncludeCDOs);
+	UMovieGraphPathTracerRenderPassNode* ParentNode = InConfig->GetSettingForBranch<UMovieGraphPathTracerRenderPassNode>(GetBranchName(), bIncludeCDOs);
 	if (!ensureMsgf(ParentNode, TEXT("PathTracerPass should not exist without parent node in graph.")))
 	{
 		return nullptr;
diff --git a/Engine/Plugins/MovieScene/MovieRenderPipeline/Source/MovieRenderPipelineRenderPasses/Public/Graph/Nodes/MovieGraphPathTracedRenderPassNode.h b/Engine/Plugins/MovieScene/MovieRenderPipeline/Source/MovieRenderPipelineRenderPasses/Public/Graph/Nodes/MovieGraphPathTracedRenderPassNode.h
deleted file mode 100644
diff --git a/Engine/Plugins/MovieScene/MovieRenderPipeline/Source/MovieRenderPipelineRenderPasses/Public/Graph/Nodes/MovieGraphPathTracerPassNode.h b/Engine/Plugins/MovieScene/MovieRenderPipeline/Source/MovieRenderPipelineRenderPasses/Public/Graph/Nodes/MovieGraphPathTracerPassNode.h
new file mode 100644
--- /dev/null
+++ b/Engine/Plugins/MovieScene/MovieRenderPipeline/Source/MovieRenderPipelineRenderPasses/Public/Graph/Nodes/MovieGraphPathTracerPassNode.h
@@ -0,0 +1,115 @@
+// Copyright Epic Games, Inc. All Rights Reserved.
+
+#pragma once
+
+#include "Graph/Nodes/MovieGraphImagePassBaseNode.h"
+
+#include "MovieGraphPathTracerPassNode.generated.h"
+
+/** A render node which uses the path tracer. */
+UCLASS()
+class MOVIERENDERPIPELINERENDERPASSES_API UMovieGraphPathTracerRenderPassNode : public UMovieGraphImagePassBaseNode
+{
+	GENERATED_BODY()
+
+public:
+	UMovieGraphPathTracerRenderPassNode();
+
+#if WITH_EDITOR
+	virtual FText GetNodeTitle(const bool bGetDescriptive = false) const override;
+#endif
+
+	// UMovieGraphRenderPassNode Interface
+	virtual void SetupImpl(const FMovieGraphRenderPassSetupData& InSetupData) override;
+	virtual void TeardownImpl() override;
+	// ~UMovieGraphRenderPassNode Interface
+
+	// UMovieGraphImagePassBaseNode Interface
+	virtual bool GetWriteAllSamples() const override;
+	virtual TArray<FMoviePipelinePostProcessPass> GetAdditionalPostProcessMaterials() const override;
+	virtual int32 GetNumSpatialSamples() const override;
+	virtual int32 GetNumSpatialSamplesDuringWarmUp() const override;
+	virtual bool GetDisableToneCurve() const override;
+	virtual bool GetAllowOCIO() const override;
+	virtual bool GetAllowDenoiser() const override;
+	virtual TUniquePtr<UE::MovieGraph::Rendering::FMovieGraphImagePassBase> CreateInstance() const;
+	// ~UMovieGraphImagePassBaseNode Interface
+
+protected:
+	// UMovieGraphRenderPassNode Interface
+	virtual FString GetRendererNameImpl() const override;
+
+	// ~UMovieGraphRenderPassNode Interface
+
+	// UMovieGraphCoreRenderPassNode Interface
+	virtual EViewModeIndex GetViewModeIndex() const override;
+	// ~UMovieGraphCoreRenderPassNode Interface
+
+public:
+
+	UPROPERTY(EditAnywhere, BlueprintReadWrite, Category = Overrides, meta = (InlineEditConditionToggle))
+	uint8 bOverride_SpatialSampleCount : 1;
+
+	UPROPERTY(EditAnywhere, BlueprintReadWrite, Category = Overrides, meta = (InlineEditConditionToggle))
+	uint8 bOverride_bDenoiser : 1;
+
+	UPROPERTY(EditAnywhere, BlueprintReadWrite, Category = Overrides, meta = (InlineEditConditionToggle))
+	uint8 bOverride_bDisableToneCurve : 1;
+
+	UPROPERTY(EditAnywhere, BlueprintReadWrite, Category = Overrides, meta = (InlineEditConditionToggle))
+	uint8 bOverride_bAllowOCIO : 1;
+
+	UPROPERTY(EditAnywhere, BlueprintReadWrite, Category = Overrides, meta = (InlineEditConditionToggle))
+	uint8 bOverride_bWriteAllSamples : 1;
+
+	UPROPERTY(EditAnywhere, BlueprintReadWrite, Category = Overrides, meta = (InlineEditConditionToggle))
+	uint8 bOverride_AdditionalPostProcessMaterials : 1;
+
+	/**
+	* How many sub-pixel jitter renders should we do per temporal sample? This can be used to achieve high
+	* sample counts without Temporal Sub-Sampling (allowing high sample counts without motion blur being enabled),
+	* but we generally recommend using Temporal Sub-Samples when possible. It can also be combined with
+	* temporal samples and you will get SpatialSampleCount many renders per temporal sample.
+	*/
+	UPROPERTY(EditAnywhere, BlueprintReadWrite, meta = (UIMin = 1, ClampMin = 1), Category = "Sampling", meta = (EditCondition = "bOverride_SpatialSampleCount"))
+	int32 SpatialSampleCount;
+
+	/** If true the resulting image will be denoised at the end of each set of Spatial Samples. */
+	UPROPERTY(EditAnywhere, BlueprintReadWrite, meta = (UIMin = 1, ClampMin = 1), Category = "Sampling", meta = (EditCondition = "bOverride_bDenoiser"))
+	bool bDenoiser;
+
+	/**
+	* Debug Feature. Not currently marked BlueprintReadWrite/EditAnywhere as it's not totally implemented on the Path Tracer right now.
+	*/
+	UPROPERTY(DisplayName = "Write All Samples (Debug)", meta = (EditCondition = "bOverride_bWriteAllsamples"))
+	bool bWriteAllSamples;
+
+	/**
+	* If true, the tone curve will be disabled for this render pass. This will result in values greater than 1.0 in final renders
+	* and can optionally be combined with OCIO profiles on the file output nodes to convert from Linear Values in Working Color Space
+	* (which is sRGB  (Rec. 709) by default, unless changed in the project settings).
+	*/
+	UPROPERTY(EditAnywhere, BlueprintReadWrite, Category = "Settings", meta = (EditCondition = "bOverride_bDisableToneCurve"))
+	bool bDisableToneCurve;
+
+	/**
+	* Allow the output file OpenColorIO transform to be used on this render.
+	*/
+	UPROPERTY(EditAnywhere, BlueprintReadWrite, Category = "Settings", meta = (EditCondition = "bOverride_bAllowOCIO"))
+	bool bAllowOCIO;
+
+
+
+	/**
+	* An array of additional post-processing materials to run after the frame is rendered. Using this feature may add a notable amount of render time.
+	*/
+	UPROPERTY(EditAnywhere, BlueprintReadWrite, Category = "Post Process Materials", meta=(EditCondition="bOverride_AdditionalPostProcessMaterials"))
+	TArray<FMoviePipelinePostProcessPass> AdditionalPostProcessMaterials;
+
+private:
+	/**
+	 * The original value of the "r.PathTracing.ProgressDisplay" cvar before the render starts. The progress display
+	 * will be hidden during the render.
+	 */
+	bool bOriginalProgressDisplayCvarValue = false;
+};
\ No newline at end of file
