MovieGraph: Renamed the MoviePipelineRenderLayer (and related classes) to MovieGraphRenderPipeline to be in-line with the rest of the Movie Graph classes.

diff --git a/Engine/Plugins/MovieScene/MovieRenderPipeline/Config/BaseMovieRenderPipeline.ini b/Engine/Plugins/MovieScene/MovieRenderPipeline/Config/BaseMovieRenderPipeline.ini
--- a/Engine/Plugins/MovieScene/MovieRenderPipeline/Config/BaseMovieRenderPipeline.ini
+++ b/Engine/Plugins/MovieScene/MovieRenderPipeline/Config/BaseMovieRenderPipeline.ini
@@ -6,5 +6,12 @@
 
 +ClassRedirects=(OldName="/Script/MovieRenderPipelineCore.MovieGraphDeferredRenderPassNode",NewName="/Script/MovieRenderPipelineRenderPasses.MovieGraphDeferredRenderPassNode")
 +ClassRedirects=(OldName="/Script/MovieRenderPipelineCore.MovieGraphImagePassBaseNode",NewName="/Script/MovieRenderPipelineRenderPasses.MovieGraphImagePassBaseNode")
-+ClassRedirects=(OldName="/Script/MovieRenderPipelineCore.MovieGraphPathTracedRenderPassNode",NewName="/Script/MovieRenderPipelineRenderPasses.MovieGraphPathTracedRenderPassNode")
-+ClassRedirects=(OldName="/Script/MovieRenderPipelineRenderPasses.MovieGraphPathTracedRenderPassNode",NewName="/Script/MovieRenderPipelineRenderPasses.MovieGraphPathTracerRenderPassNode")
\ No newline at end of file
++ClassRedirects=(OldName="/Script/MovieRenderPipelineCore.MovieGraphPathTracedRenderPassNode",NewName="/Script/MovieRenderPipelineRenderPasses.MovieGraphPathTracerRenderPassNode")
++ClassRedirects=(OldName="/Script/MovieRenderPipelineRenderPasses.MovieGraphPathTracedRenderPassNode",NewName="/Script/MovieRenderPipelineRenderPasses.MovieGraphPathTracerRenderPassNode")
++ClassRedirects=(OldName="/Script/MovieRenderPipelineCore.MoviePipelineMaterialModifier",NewName="/Script/MovieRenderPipelineRenderPasses.MovieGraphMaterialModifier")
++ClassRedirects=(OldName="/Script/MovieRenderPipelineCore.MoviePipelineVisibilityModifier",NewName="/Script/MovieRenderPipelineRenderPasses.MovieGraphRenderPropertyModifier")
++ClassRedirects=(OldName="/Script/MovieRenderPipelineCore.MoviePipelineCollectionCommonQuery",NewName="/Script/MovieRenderPipelineRenderPasses.MovieGraphCollectionCommonQuery")
++ClassRedirects=(OldName="/Script/MovieRenderPipelineCore.MoviePipelineCollectionLightingQuery",NewName="/Script/MovieRenderPipelineRenderPasses.MovieGraphCollectionLightingQuery")
++ClassRedirects=(OldName="/Script/MovieRenderPipelineCore.MoviePipelineCollection",NewName="/Script/MovieRenderPipelineRenderPasses.MovieGraphCollection")
++ClassRedirects=(OldName="/Script/MovieRenderPipelineCore.MoviePipelineCollectionQuery",NewName="/Script/MovieRenderPipelineRenderPasses.MovieGraphCollectionQueryBase")
++EnumRedirects=(OldName="EMoviePipelineCollectionCommonQueryMode",NewName="EMovieGraphCollectionCommonQueryMode")
\ No newline at end of file
diff --git a/Engine/Plugins/MovieScene/MovieRenderPipeline/Source/MovieRenderPipelineCore/Private/Graph/MovieGraphPipeline.cpp b/Engine/Plugins/MovieScene/MovieRenderPipeline/Source/MovieRenderPipelineCore/Private/Graph/MovieGraphPipeline.cpp
--- a/Engine/Plugins/MovieScene/MovieRenderPipeline/Source/MovieRenderPipelineCore/Private/Graph/MovieGraphPipeline.cpp
+++ b/Engine/Plugins/MovieScene/MovieRenderPipeline/Source/MovieRenderPipelineCore/Private/Graph/MovieGraphPipeline.cpp
@@ -6,7 +6,7 @@
 #include "Graph/MovieGraphDataTypes.h"
 #include "Graph/MovieGraphLinearTimeStep.h"
 #include "Graph/MovieGraphOutputMerger.h"
-#include "Graph/MoviePipelineRenderLayerSubsystem.h"
+#include "Graph/MovieGraphRenderLayerSubsystem.h"
 #include "Graph/Nodes/MovieGraphCollectionNode.h"
 #include "Graph/Nodes/MovieGraphFileOutputNode.h"
 #include "Graph/Nodes/MovieGraphGlobalGameOverrides.h"
@@ -48,7 +48,7 @@ void UMovieGraphPipeline::Initialize(UMoviePipelineExecutorJob* InJob, const FMo
 		return;
 	}
 
-	if (!ensureAlwaysMsgf(InJob->GetGraphPreset(), TEXT("MoviePipeline cannot be initialized with a null job configuration. Make sure you've created a Graph Config for this job (or use the regular UMoviePipeline instead of UMovieGraphPipeline). Aborting.")))
+	if (!ensureAlwaysMsgf(InJob->GetGraphPreset(), TEXT("MoviePipeline cannot be initialized with a null job configuration. Make sure you've created a Graph Config for this job (or use the regular UMovieGraph instead of UMovieGraphPipeline). Aborting.")))
 	{
 		//Shutdown(true);
 		return;
@@ -248,7 +248,7 @@ TArray<FMovieGraphRenderOutputData>& UMovieGraphPipeline::GetGeneratedOutputData
 
 void UMovieGraphPipeline::CreateLayersInRenderLayerSubsystem(const UMovieGraphEvaluatedConfig* EvaluatedConfig) const
 {
-	UMoviePipelineRenderLayerSubsystem* LayerSubsystem = GetWorld()->GetSubsystem<UMoviePipelineRenderLayerSubsystem>();
+	UMovieGraphRenderLayerSubsystem* LayerSubsystem = GetWorld()->GetSubsystem<UMovieGraphRenderLayerSubsystem>();
 	if (!LayerSubsystem)
 	{
 		return;
@@ -267,7 +267,7 @@ void UMovieGraphPipeline::CreateLayersInRenderLayerSubsystem(const UMovieGraphEv
 		
 		if (RenderLayerNode && !RenderLayerNode->IsDisabled())
 		{
-			UMoviePipelineRenderLayer* RenderLayer = NewObject<UMoviePipelineRenderLayer>();
+			UMovieGraphRenderLayer* RenderLayer = NewObject<UMovieGraphRenderLayer>();
 			RenderLayer->SetRenderLayerName(BranchName);
 			LayerSubsystem->AddRenderLayer(RenderLayer);
 		}
@@ -276,13 +276,13 @@ void UMovieGraphPipeline::CreateLayersInRenderLayerSubsystem(const UMovieGraphEv
 
 void UMovieGraphPipeline::UpdateLayerContentsInRenderLayerSubsystem(const UMovieGraphEvaluatedConfig* EvaluatedConfig) const
 {
-	UMoviePipelineRenderLayerSubsystem* LayerSubsystem = GetWorld()->GetSubsystem<UMoviePipelineRenderLayerSubsystem>();
+	UMovieGraphRenderLayerSubsystem* LayerSubsystem = GetWorld()->GetSubsystem<UMovieGraphRenderLayerSubsystem>();
 	if (!LayerSubsystem)
 	{
 		return;
 	}
 
-	for (UMoviePipelineRenderLayer* RenderLayer : LayerSubsystem->GetRenderLayers())
+	for (UMovieGraphRenderLayer* RenderLayer : LayerSubsystem->GetRenderLayers())
 	{
 		const FName& LayerName = RenderLayer->GetRenderLayerName();
 		
@@ -309,7 +309,7 @@ void UMovieGraphPipeline::UpdateLayerContentsInRenderLayerSubsystem(const UMovie
 		for (const UMovieGraphModifierNode* ModifierNode : ModifierNodes)
 		{
 			// For each modifier instance, find the collection(s) that it is modifying
-			for (UMoviePipelineCollectionModifier* ModifierInstance : ModifierNode->GetModifiers())
+			for (UMovieGraphCollectionModifier* ModifierInstance : ModifierNode->GetModifiers())
 			{
 				ModifierInstance->SetCollections({});
 				TArray<UMovieGraphCollection*> ModifierCollections;
@@ -582,9 +582,9 @@ void UMovieGraphPipeline::TickFinalizeOutputContainers(const bool bInForceFinish
 		return;
 	}
 
-	//TArray<UMoviePipelineOutputBase*> Settings = GetPipelinePrimaryConfig()->GetOutputContainers();
-	//Algo::SortBy(Settings, [](const UMoviePipelineOutputBase* Setting) { return Setting->GetPriority(); });
-	//for (UMoviePipelineOutputBase* Container : Settings)
+	//TArray<UMovieGraphOutputBase*> Settings = GetPipelinePrimaryConfig()->GetOutputContainers();
+	//Algo::SortBy(Settings, [](const UMovieGraphOutputBase* Setting) { return Setting->GetPriority(); });
+	//for (UMovieGraphOutputBase* Container : Settings)
 	//{
 	//	// All containers have finished processing, final shutdown.
 	//	Container->Finalize();
@@ -677,9 +677,9 @@ void UMovieGraphPipeline::SetupShot(const TObjectPtr<UMoviePipelineExecutorShot>
 	SetSoloShot(InShot);
 
 	// Loop through just our primary settings and let them know which shot we're about to start.
-	//TArray<UMoviePipelineSetting*> Settings = GetPipelinePrimaryConfig()->GetAllSettings();
-	//Algo::SortBy(Settings, [](const UMoviePipelineSetting* Setting) { return Setting->GetPriority(); });
-	//for (UMoviePipelineSetting* Setting : Settings)
+	//TArray<UMovieGraphSetting*> Settings = GetPipelinePrimaryConfig()->GetAllSettings();
+	//Algo::SortBy(Settings, [](const UMovieGraphSetting* Setting) { return Setting->GetPriority(); });
+	//for (UMovieGraphSetting* Setting : Settings)
 	//{
 	//	Setting->OnSetupForShot(InShot);
 	//}
@@ -687,9 +687,9 @@ void UMovieGraphPipeline::SetupShot(const TObjectPtr<UMoviePipelineExecutorShot>
 	//if (InShot->GetShotOverrideConfiguration() != nullptr)
 	//{
 	//	// Any shot-specific overrides haven't had first time initialization. So we'll do that now.
-	//	TArray<UMoviePipelineSetting*> ShotSettings = InShot->GetShotOverrideConfiguration()->GetUserSettings();
-	//	Algo::SortBy(ShotSettings, [](const UMoviePipelineSetting* Setting) { return Setting->GetPriority(); });
-	//	for (UMoviePipelineSetting* Setting : ShotSettings)
+	//	TArray<UMovieGraphSetting*> ShotSettings = InShot->GetShotOverrideConfiguration()->GetUserSettings();
+	//	Algo::SortBy(ShotSettings, [](const UMovieGraphSetting* Setting) { return Setting->GetPriority(); });
+	//	for (UMovieGraphSetting* Setting : ShotSettings)
 	//	{
 	//		Setting->OnMoviePipelineInitialized(this);
 	//	}
@@ -1195,9 +1195,9 @@ void UMovieGraphPipeline::TransitionToState(const EMovieRenderPipelineState InNe
 			PipelineState = EMovieRenderPipelineState::Finished;
 
 			// Uninitialize our primary config settings. Reverse sorted so settings that cached values restore correctly.
-			//TArray<UMoviePipelineSetting*> Settings = GetPipelinePrimaryConfig()->GetAllSettings();
-			//Algo::SortBy(Settings, [](const UMoviePipelineSetting* Setting) { return Setting->GetPriority(); }, TLess<int32>());
-			//for (UMoviePipelineSetting* Setting : Settings)
+			//TArray<UMovieGraphSetting*> Settings = GetPipelinePrimaryConfig()->GetAllSettings();
+			//Algo::SortBy(Settings, [](const UMovieGraphSetting* Setting) { return Setting->GetPriority(); }, TLess<int32>());
+			//for (UMovieGraphSetting* Setting : Settings)
 			//{
 			//	Setting->OnMoviePipelineShutdown(this);
 			//}
@@ -1223,9 +1223,9 @@ void UMovieGraphPipeline::TransitionToState(const EMovieRenderPipelineState InNe
 			// Job-level evaluated graph should not be referenced after export has finished
 			PostRenderEvaluatedGraph = nullptr;
 
-			//TArray<UMoviePipelineOutputBase*> ContainerSettings = GetPipelinePrimaryConfig()->GetOutputContainers();
-			//Algo::SortBy(ContainerSettings, [](const UMoviePipelineOutputBase* Setting) { return Setting->GetPriority(); });
-			//for (UMoviePipelineOutputBase* Setting : ContainerSettings)
+			//TArray<UMovieGraphOutputBase*> ContainerSettings = GetPipelinePrimaryConfig()->GetOutputContainers();
+			//Algo::SortBy(ContainerSettings, [](const UMovieGraphOutputBase* Setting) { return Setting->GetPriority(); });
+			//for (UMovieGraphOutputBase* Setting : ContainerSettings)
 			//{
 			//	Setting->OnPipelineFinished();
 			//}
@@ -1250,7 +1250,7 @@ void UMovieGraphPipeline::TransitionToState(const EMovieRenderPipelineState InNe
 
 			//UE_LOG(LogMovieRenderPipeline, Log, TEXT("Movie Pipeline completed. Duration: %s"), *(FDateTime::UtcNow() - InitializationTime).ToString());
 
-			//UMoviePipelineDebugSettings* DebugSetting = FindOrAddSettingForShot<UMoviePipelineDebugSettings>(nullptr);
+			//UMovieGraphDebugSettings* DebugSetting = FindOrAddSettingForShot<UMovieGraphDebugSettings>(nullptr);
 			//if (DebugSetting)
 			//{
 			//	if (DebugSetting->bCaptureUnrealInsightsTrace)
diff --git a/Engine/Plugins/MovieScene/MovieRenderPipeline/Source/MovieRenderPipelineCore/Private/Graph/MovieGraphRenderLayerSubsystem.cpp b/Engine/Plugins/MovieScene/MovieRenderPipeline/Source/MovieRenderPipelineCore/Private/Graph/MovieGraphRenderLayerSubsystem.cpp
new file mode 100644
--- /dev/null
+++ b/Engine/Plugins/MovieScene/MovieRenderPipeline/Source/MovieRenderPipelineCore/Private/Graph/MovieGraphRenderLayerSubsystem.cpp
@@ -0,0 +1,2070 @@
+// Copyright Epic Games, Inc. All Rights Reserved.
+
+#include "Graph/MovieGraphRenderLayerSubsystem.h"
+
+#include "Components/PrimitiveComponent.h"
+#include "Components/SkyAtmosphereComponent.h"
+#include "Components/VolumetricCloudComponent.h"
+#include "Components/ExponentialHeightFogComponent.h"
+#include "Components/SkyLightComponent.h"
+#include "Components/SphereReflectionCaptureComponent.h"
+#include "EngineUtils.h"
+#include "Framework/Application/SlateApplication.h"
+#include "Materials/MaterialInterface.h"
+#include "Modules/ModuleManager.h"
+#include "MovieRenderPipelineCoreModule.h"
+#include "Styling/AppStyle.h"
+#include "Styling/SlateIconFinder.h"
+#include "UObject/Package.h"
+#include "Widgets/Input/SMultiLineEditableTextBox.h"
+
+#if WITH_EDITOR
+#include "ActorFolderPickingMode.h"
+#include "ActorFolderTreeItem.h"
+#include "ActorMode.h"
+#include "ActorTreeItem.h"
+#include "ClassViewerFilter.h"
+#include "ClassViewerModule.h"
+#include "ComponentTreeItem.h"
+#include "ContentBrowserDataSource.h"
+#include "ContentBrowserModule.h"
+#include "Editor.h"
+#include "EditorActorFolders.h"
+#include "Graph/MovieGraphSharedWidgets.h"
+#include "IContentBrowserSingleton.h"
+#include "ISceneOutliner.h"
+#include "SceneOutlinerModule.h"
+#include "SceneOutlinerPublicTypes.h"
+#include "SClassViewer.h"
+#include "Selection.h"
+#endif
+
+#define LOCTEXT_NAMESPACE "MovieGraph"
+
+namespace UE::MovieGraph::Private
+{
+#if WITH_EDITOR
+	/**
+	 * A filter that can be used in the class viewer that appears in the Add menu. Filters out specified classes, and optionally filters out classes
+	 * that do not have a specific base class.
+	 */
+	class FClassViewerTypeFilter final : public IClassViewerFilter
+	{
+	public:
+		explicit FClassViewerTypeFilter(TArray<UClass*>* InClassesToDisallow, UClass* InRequiredBaseClass = nullptr)
+			: ClassesToDisallow(InClassesToDisallow)
+			, RequiredBaseClass(InRequiredBaseClass)
+		{
+			
+		}
+
+		virtual bool IsClassAllowed(const FClassViewerInitializationOptions& InInitOptions, const UClass* InClass, TSharedRef<FClassViewerFilterFuncs> InFilterFuncs) override
+		{
+			return
+				InClass &&
+				!ClassesToDisallow->Contains(InClass) &&
+				(RequiredBaseClass ? InClass->IsChildOf(RequiredBaseClass) : true);
+		}
+		
+		virtual bool IsUnloadedClassAllowed(const FClassViewerInitializationOptions& InInitOptions, const TSharedRef<const IUnloadedBlueprintData> InUnloadedClassData, TSharedRef<FClassViewerFilterFuncs> InFilterFuncs) override
+		{
+			return false;
+		}
+
+	private:
+		/** Classes which should be prevented from showing up in the class viewer. */
+		TArray<UClass*>* ClassesToDisallow;
+
+		/** Classes must have this base class to pass the filter. */
+		UClass* RequiredBaseClass = nullptr;
+	};
+#endif
+}
+
+void UMovieGraphMaterialModifier::ApplyModifier(const UWorld* World)
+{
+	UMaterialInterface* NewMaterial = Material.LoadSynchronous();
+	if (!NewMaterial)
+	{
+		return;
+	}
+
+	ModifiedComponents.Empty();
+	
+	for (const UMovieGraphCollection* Collection : Collections)
+	{
+		if (!Collection)
+		{
+			continue;
+		}
+
+		for (const AActor* Actor : Collection->Evaluate(World))
+		{
+			const bool bIncludeFromChildActors = true;
+			TArray<UPrimitiveComponent*> PrimitiveComponents;
+			Actor->GetComponents<UPrimitiveComponent>(PrimitiveComponents, bIncludeFromChildActors);
+
+			for (UPrimitiveComponent* PrimitiveComponent : PrimitiveComponents)
+			{
+				TArray<FMaterialSlotAssignment>& ModifiedMaterials = ModifiedComponents.FindOrAdd(PrimitiveComponent);
+				
+				for (int32 Index = 0; Index < PrimitiveComponent->GetNumMaterials(); ++Index)
+				{
+					ModifiedMaterials.Add(FMaterialSlotAssignment(Index, PrimitiveComponent->GetMaterial(Index)));
+				
+					PrimitiveComponent->SetMaterial(Index, NewMaterial);
+				}
+			}
+		}
+	}
+}
+
+void UMovieGraphMaterialModifier::UndoModifier()
+{
+	for (const FComponentToMaterialMap::ElementType& ModifiedComponent : ModifiedComponents) 
+	{
+		UPrimitiveComponent* MeshComponent = ModifiedComponent.Key.LoadSynchronous();
+		const TArray<FMaterialSlotAssignment>& OldMaterials = ModifiedComponent.Value;
+
+		if (!MeshComponent)
+		{
+			continue;
+		}
+
+		for (const FMaterialSlotAssignment& MaterialPair : OldMaterials)
+		{
+			UMaterialInterface* MaterialInterface = MaterialPair.Value.LoadSynchronous();
+			if (!MaterialInterface)
+			{
+				continue;
+			}
+
+			const int32 ElementIndex = MaterialPair.Key;
+			MeshComponent->SetMaterial(ElementIndex, MaterialInterface);
+		}
+	}
+
+	ModifiedComponents.Empty();
+}
+
+UMovieGraphRenderPropertyModifier::UMovieGraphRenderPropertyModifier()
+	: bIsHidden(false)
+	, bCastsShadows(true)
+	, bCastShadowWhileHidden(false)
+	, bAffectIndirectLightingWhileHidden(false)
+	, bHoldout(false)
+{
+	// Note: The default modifier values here reflect the defaults on the scene component. If a modifier property is marked as overridden, the
+	// override will initially be a no-op due to the defaults being the same.
+}
+
+void UMovieGraphRenderPropertyModifier::ApplyModifier(const UWorld* World)
+{
+	ModifiedActors.Empty();
+
+	// SetActorVisibilityState() takes a FActorVisibilityState, but all calls here will contain the same settings (with
+	// a differing actor), so create one copy to prevent constantly re-creating structs.
+	FActorVisibilityState NewVisibilityState;
+	NewVisibilityState.bIsHidden = bIsHidden;
+
+	// Create a stub component representation since we're just using this to propagate to all children right now.
+	FActorVisibilityState::FComponentState& NewComponentState = NewVisibilityState.Components.AddDefaulted_GetRef();
+	NewComponentState.bCastsShadows = bCastsShadows;
+	NewComponentState.bCastShadowWhileHidden = bCastShadowWhileHidden;
+	NewComponentState.bAffectIndirectLightingWhileHidden = bAffectIndirectLightingWhileHidden;
+	NewComponentState.bHoldout = bHoldout;
+	
+	for (const UMovieGraphCollection* Collection : Collections)
+	{
+		if (!Collection)
+		{
+			continue;
+		}
+
+		const TSet<AActor*> MatchingActors = Collection->Evaluate(World);
+		ModifiedActors.Reserve(MatchingActors.Num());
+
+		for (const AActor* Actor : MatchingActors)
+		{
+			// Save out visibility state before the modifier is applied
+			FActorVisibilityState OriginalVisibilityState;
+			OriginalVisibilityState.Actor = Actor;
+			OriginalVisibilityState.bIsHidden = Actor->IsHidden();
+
+			const bool bIncludeFromChildActors = true;
+			TInlineComponentArray<USceneComponent*> Components;
+			Actor->GetComponents<USceneComponent>(Components, bIncludeFromChildActors);
+
+			OriginalVisibilityState.Components.Reserve(Components.Num());
+			for(const USceneComponent* SceneComponent : Components)
+			{
+				FActorVisibilityState::FComponentState& ComponentState = OriginalVisibilityState.Components.AddDefaulted_GetRef();
+				ComponentState.Component = SceneComponent;
+
+				// Cache the state
+				if (const UPrimitiveComponent* AsPrimitiveComponent = Cast<UPrimitiveComponent>(SceneComponent))
+				{
+					ComponentState.bCastsShadows = AsPrimitiveComponent->CastShadow;
+					ComponentState.bCastShadowWhileHidden = AsPrimitiveComponent->bCastHiddenShadow;
+					ComponentState.bAffectIndirectLightingWhileHidden = AsPrimitiveComponent->bAffectIndirectLightingWhileHidden;
+					ComponentState.bHoldout = AsPrimitiveComponent->bHoldout;
+				}
+				// Volumetrics are special cases as they don't inherit from UPrimitiveComponent, and don't support all of the flags.
+				else if (const UVolumetricCloudComponent* AsVolumetricCloudComponent = Cast<UVolumetricCloudComponent>(SceneComponent))
+				{
+					ComponentState.bHoldout = AsVolumetricCloudComponent->bHoldout;
+				}
+				else if (const USkyAtmosphereComponent* AsSkyAtmosphereComponent = Cast<USkyAtmosphereComponent>(SceneComponent))
+				{
+					ComponentState.bHoldout = AsSkyAtmosphereComponent->bHoldout;
+				}
+				else if (const UExponentialHeightFogComponent* AsExponentialHeightFogComponent = Cast<UExponentialHeightFogComponent>(SceneComponent))
+				{
+					ComponentState.bHoldout = AsExponentialHeightFogComponent->bHoldout;
+				}
+
+				// Then override it. We override it one component at a time so that we can avoid making a bunch of
+				// copies of NewVisibilityState for the varying number of components you might have.
+				NewVisibilityState.Actor = Actor;
+				NewVisibilityState.Components[0].Component = SceneComponent;
+				SetActorVisibilityState(NewVisibilityState);
+			}
+			
+			ModifiedActors.Add(OriginalVisibilityState);
+
+		}
+	}
+}
+
+void UMovieGraphRenderPropertyModifier::UndoModifier()
+{
+	for (const FActorVisibilityState& PrevVisibilityState : ModifiedActors)
+	{
+		SetActorVisibilityState(PrevVisibilityState);
+	}
+
+	ModifiedActors.Empty();
+}
+
+void UMovieGraphRenderPropertyModifier::SetActorVisibilityState(const FActorVisibilityState& NewVisibilityState)
+{
+	const TSoftObjectPtr<AActor> Actor = NewVisibilityState.Actor.LoadSynchronous();
+	if (!Actor)
+	{
+		return;
+	}
+
+	if (bOverride_bIsHidden)
+	{
+		Actor->SetActorHiddenInGame(NewVisibilityState.bIsHidden);
+
+#if WITH_EDITOR
+		Actor->SetIsTemporarilyHiddenInEditor(NewVisibilityState.bIsHidden);
+#endif
+	}
+
+	for (const FActorVisibilityState::FComponentState& ComponentState : NewVisibilityState.Components)
+	{
+		// TODO: These could potentially cause a large rendering penalty due to dirtying the render state; investigate potential
+		// ways to optimize this
+		if (!ComponentState.Component)
+		{
+			continue;
+		}
+
+		if (UPrimitiveComponent* AsPrimitiveComponent = Cast<UPrimitiveComponent>(ComponentState.Component.Get()))
+		{
+			if (bOverride_bCastsShadows)
+			{
+				AsPrimitiveComponent->SetCastShadow(ComponentState.bCastsShadows);
+			}
+
+			if (bOverride_bCastShadowWhileHidden)
+			{
+				AsPrimitiveComponent->SetCastHiddenShadow(ComponentState.bCastShadowWhileHidden);
+			}
+
+			if (bOverride_bAffectIndirectLightingWhileHidden)
+			{
+				AsPrimitiveComponent->SetAffectIndirectLightingWhileHidden(ComponentState.bAffectIndirectLightingWhileHidden);
+			}
+
+			if (bOverride_bHoldout)
+			{
+				AsPrimitiveComponent->SetHoldout(ComponentState.bHoldout);
+			}
+		}
+		// Volumetrics are special cases as they don't inherit from UPrimitiveComponent, and don't support all of the flags.
+		else if (UVolumetricCloudComponent* AsVolumetricCloudComponent = Cast<UVolumetricCloudComponent>(ComponentState.Component.Get()))
+		{
+			if (bOverride_bHoldout)
+			{
+				AsVolumetricCloudComponent->SetHoldout(ComponentState.bHoldout);
+			}
+		}
+		else if (USkyAtmosphereComponent* AsSkyAtmosphereComponent = Cast<USkyAtmosphereComponent>(ComponentState.Component.Get()))
+		{
+			if (bOverride_bHoldout)
+			{
+				AsSkyAtmosphereComponent->SetHoldout(ComponentState.bHoldout);
+			}
+		}
+		else if (UExponentialHeightFogComponent* AsExponentialHeightFogComponent = Cast<UExponentialHeightFogComponent>(ComponentState.Component.Get()))
+		{
+			if (bOverride_bHoldout)
+			{
+				AsExponentialHeightFogComponent->SetHoldout(ComponentState.bHoldout);
+			}
+		}
+	}
+}
+
+// TODO: This really should be "DoesComponentMatchQuery()"
+bool UMovieGraphCollectionCommonQuery::DoesActorMatchQuery(const AActor* Actor) const
+{
+	if (!Actor)
+	{
+		return false;
+	}
+
+	// TODO: This method should short-circuit so all query types aren't executed for every actor when the query mode is OR
+
+	bool bMatchesActorNames = false;
+	if (!ActorNames.IsEmpty())
+	{
+		bMatchesActorNames = ActorNames.Contains(Actor->GetName());
+	}
+	
+	bool bMatchesComponentTypes = false;
+	if (!ComponentTypes.IsEmpty())
+	{
+		const bool bIncludeFromChildActors = true;
+		TArray<UActorComponent*> ActorComponents;
+		Actor->GetComponents(ActorComponents, bIncludeFromChildActors);
+		for (const UActorComponent* Component : ActorComponents)
+		{
+			for (const UClass* ComponentType : ComponentTypes)
+			{
+				if (Component->IsA(ComponentType))
+				{
+					bMatchesComponentTypes = true;
+					break;
+				}
+			}
+		}
+	}
+
+	bool bMatchesTags = false;
+	if (!Tags.IsEmpty())
+	{
+		for (const FName& Tag : Tags)
+		{
+			if (Actor->Tags.Contains(Tag))
+			{
+				bMatchesTags = true;
+				break;
+			}
+		}
+	}
+
+	const bool bUsingActorNames = !ActorNames.IsEmpty();
+	const bool bUsingComponentTypes = !ComponentTypes.IsEmpty();
+	const bool bUsingTags = !Tags.IsEmpty();
+
+	if (QueryMode == EMovieGraphCollectionCommonQueryMode::And)
+	{
+		return (!bUsingActorNames || bMatchesActorNames) &&
+			   (!bUsingTags || bMatchesTags) &&
+			   (!bUsingComponentTypes || bMatchesComponentTypes);
+	}
+	
+	return bMatchesActorNames || bMatchesTags || bMatchesComponentTypes;
+}
+
+bool UMovieGraphCollectionLightingQuery::DoesActorMatchQuery(const AActor* Actor) const
+{
+	const TArray<UClass*> LightingComponentTypes = {
+		ULightComponentBase::StaticClass(),
+		UReflectionCaptureComponent::StaticClass(),
+		USkyAtmosphereComponent::StaticClass(),
+		USkyLightComponent::StaticClass()
+	};
+	
+	TArray<UActorComponent*> ActorComponents;
+	constexpr bool bIncludeFromChildActors = true;
+	Actor->GetComponents(ActorComponents, bIncludeFromChildActors);
+	
+	for (const UActorComponent* Component : ActorComponents)
+	{
+		for (const UClass* ComponentType : LightingComponentTypes)
+		{
+			if (Component->IsA(ComponentType))
+			{
+				return true;
+			}
+		}
+	}
+
+	return false;
+}
+
+void UMovieGraphCollectionModifier::AddCollection(UMovieGraphCollection* Collection)
+{
+	// Don't allow adding a duplicate collection
+	for (const UMovieGraphCollection* ExistingCollection : Collections)
+	{
+		if (Collection && ExistingCollection && Collection->GetCollectionName().Equals(ExistingCollection->GetCollectionName()))
+		{
+			return;
+		}
+	}
+	
+	Collections.Add(Collection);
+}
+
+UMovieGraphConditionGroupQueryBase::UMovieGraphConditionGroupQueryBase()
+	: OpType(EMovieGraphConditionGroupQueryOpType::Add)
+	, bIsEnabled(true)
+{
+}
+
+void UMovieGraphConditionGroupQueryBase::SetOperationType(const EMovieGraphConditionGroupQueryOpType OperationType)
+{
+	// Always allow setting the operation type to Union. If not setting to Union, only allow setting the operation type if this is not the first
+	// query in the condition group. The first query is always a Union.
+	if (OperationType == EMovieGraphConditionGroupQueryOpType::Add)
+	{
+		OpType = EMovieGraphConditionGroupQueryOpType::Add;
+		return;
+	}
+
+	const UMovieGraphConditionGroup* ParentConditionGroup = GetTypedOuter<UMovieGraphConditionGroup>();
+	if (ensureMsgf(ParentConditionGroup, TEXT("Cannot set the operation type on a condition group query that doesn't have a condition group outer")))
+	{
+		if (ParentConditionGroup->GetQueries().Find(this) != 0)
+		{
+			OpType = OperationType;
+		}
+	}
+}
+
+EMovieGraphConditionGroupQueryOpType UMovieGraphConditionGroupQueryBase::GetOperationType() const
+{
+	return OpType;
+}
+
+void UMovieGraphConditionGroupQueryBase::Evaluate(const TArray<AActor*>& InActorsToQuery, const UWorld* InWorld, TSet<AActor*>& OutMatchingActors) const
+{
+	// No implementation
+}
+
+bool UMovieGraphConditionGroupQueryBase::ShouldHidePropertyNames() const
+{
+	// Show property names by default; subclassed queries can opt-out if they want a cleaner UI 
+	return false;
+}
+
+const FSlateIcon& UMovieGraphConditionGroupQueryBase::GetIcon() const
+{
+	static const FSlateIcon EmptyIcon = FSlateIcon();
+	return EmptyIcon;
+}
+
+const FText& UMovieGraphConditionGroupQueryBase::GetDisplayName() const
+{
+	static const FText DisplayName = LOCTEXT("ConditionGroupQueryDisplayName", "Query Base");
+	return DisplayName;
+}
+
+#if WITH_EDITOR
+TArray<TSharedRef<SWidget>> UMovieGraphConditionGroupQueryBase::GetWidgets()
+{
+	return TArray<TSharedRef<SWidget>>();
+}
+
+bool UMovieGraphConditionGroupQueryBase::HasAddMenu() const
+{
+	return false;
+}
+
+TSharedRef<SWidget> UMovieGraphConditionGroupQueryBase::GetAddMenuContents(const FMovieGraphConditionGroupQueryContentsChanged& OnAddFinished)
+{
+	return SNullWidget::NullWidget;
+}
+#endif
+
+bool UMovieGraphConditionGroupQueryBase::IsEditorOnlyQuery() const
+{
+	return false;
+}
+
+void UMovieGraphConditionGroupQueryBase::SetEnabled(const bool bEnabled)
+{
+	bIsEnabled = bEnabled;
+}
+
+bool UMovieGraphConditionGroupQueryBase::IsEnabled() const
+{
+	return bIsEnabled;
+}
+
+bool UMovieGraphConditionGroupQueryBase::IsFirstConditionGroupQuery() const
+{
+	const UMovieGraphConditionGroup* ParentConditionGroup = GetTypedOuter<UMovieGraphConditionGroup>();
+	if (ensureMsgf(ParentConditionGroup, TEXT("Cannot determine if this is the first condition group query when no parent condition group is present")))
+	{
+		// GetQueries() returns an array of non-const pointers, so Find() doesn't like having a const pointer passed to it.
+		// Find() won't mutate the condition group query though, so the const_cast here is OK.
+		return ParentConditionGroup->GetQueries().Find(const_cast<UMovieGraphConditionGroupQueryBase*>(this)) == 0;
+	}
+	
+	return false;
+}
+
+void UMovieGraphConditionGroupQuery_Actor::Evaluate(const TArray<AActor*>& InActorsToQuery, const UWorld* InWorld, TSet<AActor*>& OutMatchingActors) const
+{
+	TRACE_CPUPROFILER_EVENT_SCOPE(UMovieGraphConditionGroupQuery_Actor::Evaluate);
+
+#if WITH_EDITOR
+	const bool bIsPIE = GEditor->IsPlaySessionInProgress();
+#else
+	const bool bIsPIE = false;
+#endif
+
+	// Convert the actors in the query to PIE (or editor) equivalents once, rather than constantly in the loop.
+	TArray<AActor*> ActorsToMatch_Converted;
+	ActorsToMatch_Converted.Reserve(ActorsToMatch.Num());
+	for (const TSoftObjectPtr<AActor>& SoftActorToMatch : ActorsToMatch)
+	{
+		if (!SoftActorToMatch.IsValid())
+		{
+			continue;
+		}
+		
+		AActor* ActorToMatch = SoftActorToMatch.Get();
+		const UWorld* ActorWorld = ActorToMatch->GetWorld();
+		const bool bIsEditorActor = ActorWorld && ActorWorld->IsEditorWorld();
+
+		// If a PIE session is NOT in progress, make sure that the actor is the editor equivalent
+		if (!bIsPIE)
+		{
+			// Only do PIE -> editor actor conversion when the actor is NOT from the editor
+			if (!bIsEditorActor)
+			{
+#if WITH_EDITOR
+				if (AActor* EditorActor = EditorUtilities::GetEditorWorldCounterpartActor(ActorToMatch))
+				{
+					ActorsToMatch_Converted.Add(EditorActor);
+				}
+#endif
+			}
+			else
+			{
+				// Just use ActorToMatch as-is if it's not from PIE
+				ActorsToMatch_Converted.Add(ActorToMatch);
+			}
+		}
+
+		// If a PIE session IS active, try to get the PIE equivalent of the editor actor
+		else
+		{
+			// Only do editor -> PIE actor conversion when the actor is from an editor world
+			if (bIsEditorActor)
+			{
+#if WITH_EDITOR
+				if (AActor* PieActor = EditorUtilities::GetSimWorldCounterpartActor(ActorToMatch))
+				{
+					ActorsToMatch_Converted.Add(PieActor);
+				}
+#endif
+			}
+			else
+			{
+				// Just use ActorToMatch as-is if it's not from an editor actor
+				ActorsToMatch_Converted.Add(ActorToMatch);
+			}	
+		}
+	}
+	
+	for (AActor* Actor : InActorsToQuery)
+	{
+		if (ActorsToMatch_Converted.Contains(Actor))
+		{
+			OutMatchingActors.Add(Actor);
+		}
+	}
+}
+
+const FSlateIcon& UMovieGraphConditionGroupQuery_Actor::GetIcon() const
+{
+	static const FSlateIcon ActorIcon = FSlateIcon(FAppStyle::GetAppStyleSetName(), "ClassIcon.Actor");
+	return ActorIcon;
+}
+
+const FText& UMovieGraphConditionGroupQuery_Actor::GetDisplayName() const
+{
+	static const FText DisplayName = LOCTEXT("ConditionGroupQueryDisplayName_Actor", "Actor");
+	return DisplayName;
+}
+
+#if WITH_EDITOR
+TArray<TSharedRef<SWidget>> UMovieGraphConditionGroupQuery_Actor::GetWidgets()
+{
+	TArray<TSharedRef<SWidget>> Widgets;
+
+	// Create the data source for the list view
+	ListDataSource.Empty();
+	for (TSoftObjectPtr<AActor>& Actor : ActorsToMatch)
+	{
+		ListDataSource.Add(MakeShared<TSoftObjectPtr<AActor>>(Actor));
+	}
+
+	Widgets.Add(
+		SAssignNew(ActorsList, SMovieGraphSimpleList<TSharedPtr<TSoftObjectPtr<AActor>>>)
+			.DataSource(&ListDataSource)
+			.DataType(FText::FromString("Actor"))
+			.DataTypePlural(FText::FromString("Actors"))
+			.OnGetRowText_Static(&GetRowText)
+			.OnGetRowIcon_Static(&GetRowIcon)
+			.OnDelete_Lambda([this](const TSharedPtr<TSoftObjectPtr<AActor>> InActor)
+			{
+				ListDataSource.Remove(InActor);
+				ActorsToMatch.Remove(*InActor.Get());
+				ActorsList->Refresh();
+			})
+	);
+
+	return Widgets;
+}
+
+bool UMovieGraphConditionGroupQuery_Actor::HasAddMenu() const
+{
+	return true;
+}
+
+TSharedRef<SWidget> UMovieGraphConditionGroupQuery_Actor::GetAddMenuContents(const FMovieGraphConditionGroupQueryContentsChanged& OnAddFinished)
+{	
+	FSceneOutlinerInitializationOptions SceneOutlinerInitOptions;
+	SceneOutlinerInitOptions.bShowHeaderRow = true;
+	SceneOutlinerInitOptions.bShowSearchBox = true;
+	SceneOutlinerInitOptions.bShowCreateNewFolder = false;
+	SceneOutlinerInitOptions.bFocusSearchBoxWhenOpened = true;
+
+	// Show the name/label column and the type column
+	SceneOutlinerInitOptions.ColumnMap.Add(
+		FSceneOutlinerBuiltInColumnTypes::Label(),
+		FSceneOutlinerColumnInfo(ESceneOutlinerColumnVisibility::Visible, 0, FCreateSceneOutlinerColumn(), false, TOptional<float>(), FSceneOutlinerBuiltInColumnTypes::Label_Localized()));
+	SceneOutlinerInitOptions.ColumnMap.Add(
+		FSceneOutlinerBuiltInColumnTypes::ActorInfo(),
+		FSceneOutlinerColumnInfo(ESceneOutlinerColumnVisibility::Visible, 10,  FCreateSceneOutlinerColumn(), false, TOptional<float>(), FSceneOutlinerBuiltInColumnTypes::ActorInfo_Localized()));
+
+	// Don't show actors which have already been picked
+	SceneOutlinerInitOptions.Filters->AddFilterPredicate<FActorTreeItem>(
+		FActorTreeItem::FFilterPredicate::CreateLambda([this](const AActor* InActor)
+		{
+			return !ActorsToMatch.Contains(InActor);
+		}));
+	
+	const FSceneOutlinerModule& SceneOutlinerModule = FModuleManager::LoadModuleChecked<FSceneOutlinerModule>("SceneOutliner");
+
+	FMenuBuilder MenuBuilder(false, MakeShared<FUICommandList>());
+
+	auto AddActorToList = [this](const AActor* InActor, const FMovieGraphConditionGroupQueryContentsChanged& OnAddFinished)
+	{
+		ActorsToMatch.Add(InActor);
+		ListDataSource.Add(MakeShared<TSoftObjectPtr<AActor>>(ActorsToMatch.Last()));
+		OnAddFinished.ExecuteIfBound();
+	};
+	
+	auto RefreshActorPickerFilterAndList = [this]()
+	{
+		// Ensure that the actor picker filter runs again so duplicate actors cannot be selected
+		if (ActorPickerWidget.IsValid())
+		{
+			ActorPickerWidget->FullRefresh();
+			ActorsList->Refresh();
+		}
+	};
+	
+	MenuBuilder.BeginSection("AddActor", LOCTEXT("AddActor", "Add Actor"));
+	{
+		MenuBuilder.AddMenuEntry(
+			LOCTEXT("AddSelectedInOutliner", "Add Selected In Outliner"),
+			LOCTEXT("AddSelectedInOutlinerTooltip", "Add actors currently selected in the level editor's scene outliner."),
+			FSlateIcon(FAppStyle::GetAppStyleSetName(),"FoliageEditMode.SetSelect"),
+			FUIAction(
+				FExecuteAction::CreateLambda([this, OnAddFinished, RefreshActorPickerFilterAndList, AddActorToList]()
+				{
+					USelection* CurrentSelection = GEditor->GetSelectedActors();
+					TArray<AActor*> SelectedActors;
+					CurrentSelection->GetSelectedObjects<AActor>(SelectedActors);
+					for (const AActor* Actor : SelectedActors)
+					{
+						if (Actor && !ActorsToMatch.Contains(Actor))
+						{
+							AddActorToList(Actor, OnAddFinished);
+						}
+					}
+
+					RefreshActorPickerFilterAndList();
+
+					FSlateApplication::Get().DismissAllMenus();
+				}
+				),
+				FCanExecuteAction::CreateLambda([]()
+				{
+					return GEditor->GetSelectedActors()->Num() > 0;
+					
+				})
+			)
+		);
+	}
+	MenuBuilder.EndSection();
+
+	MenuBuilder.BeginSection("Browse", LOCTEXT("Browse", "Browse"));
+	{
+		ActorPickerWidget = SceneOutlinerModule.CreateActorPicker(
+			SceneOutlinerInitOptions,
+			FOnActorPicked::CreateLambda([this, OnAddFinished, RefreshActorPickerFilterAndList, AddActorToList](AActor* InActor)
+			{
+				AddActorToList(InActor, OnAddFinished);
+
+				RefreshActorPickerFilterAndList();
+			}));
+
+		const TSharedRef<SBox> ActorPickerWidgetBox =
+			SNew(SBox)
+			.WidthOverride(400.f)
+			.HeightOverride(300.f)
+			[
+				ActorPickerWidget.ToSharedRef()
+			];
+
+		MenuBuilder.AddWidget(ActorPickerWidgetBox, FText());
+	}
+	MenuBuilder.EndSection();
+
+	return MenuBuilder.MakeWidget();
+}
+
+const FSlateBrush* UMovieGraphConditionGroupQuery_Actor::GetRowIcon(TSharedPtr<TSoftObjectPtr<AActor>> InActor)
+{
+	if (InActor.IsValid())
+	{
+		if (InActor.Get()->IsValid())
+		{
+			// The first Get() returns the TSoftObjectPtr, the second Get() dereferences the TSoftObjectPtr
+			return FSlateIconFinder::FindIconForClass(InActor.Get()->Get()->GetClass()).GetIcon();
+		}
+	}
+
+	return FSlateIconFinder::FindIconForClass(AActor::StaticClass()).GetIcon();
+}
+
+FText UMovieGraphConditionGroupQuery_Actor::GetRowText(TSharedPtr<TSoftObjectPtr<AActor>> InActor)
+{
+	if (InActor.IsValid())
+	{
+		if (InActor.Get()->IsValid())
+		{
+			// The first Get() returns the TSoftObjectPtr, the second Get() dereferences the TSoftObjectPtr
+			return FText::FromString(InActor.Get()->Get()->GetActorLabel());
+		}
+	}
+
+	return LOCTEXT("MovieGraphActorConditionGroupQuery_InvalidActor", "(invalid)");
+}
+#endif	// WITH_EDITOR
+
+void UMovieGraphConditionGroupQuery_ActorTagName::Evaluate(const TArray<AActor*>& InActorsToQuery, const UWorld* InWorld, TSet<AActor*>& OutMatchingActors) const
+{
+	TRACE_CPUPROFILER_EVENT_SCOPE(UMovieGraphConditionGroupQuery_ActorTag::Evaluate);
+	
+	// Quick early-out if "*" is used as the wildcard. Faster than doing the wildcard matching.
+	if (TagsToMatch == TEXT("*"))
+	{
+		OutMatchingActors.Append(InActorsToQuery);
+		return;
+	}
+
+	// Actor tags can be specified on multiple lines
+	TArray<FString> AllTagNameStrings;
+	TagsToMatch.ParseIntoArrayLines(AllTagNameStrings);
+
+	for (AActor* Actor : InActorsToQuery)
+	{
+		for (const FString& TagToMatch : AllTagNameStrings)
+		{
+			bool bMatchedTag = false;
+			
+			for (const FName& ActorTag : Actor->Tags)
+			{
+				if (ActorTag.ToString().MatchesWildcard(TagToMatch))
+				{
+					OutMatchingActors.Add(Actor);
+					bMatchedTag = true;
+					break;
+				}
+			}
+
+			// Skip comparing the rest of the tags if one tag already matched 
+			if (bMatchedTag)
+			{
+				break;
+			}
+		}
+	}
+}
+
+const FSlateIcon& UMovieGraphConditionGroupQuery_ActorTagName::GetIcon() const
+{
+	// TODO: This icon is wrong
+	static const FSlateIcon ActorTagIcon = FSlateIcon(FAppStyle::GetAppStyleSetName(), "LevelEditor.Tabs.Debug");
+	return ActorTagIcon;
+}
+
+const FText& UMovieGraphConditionGroupQuery_ActorTagName::GetDisplayName() const
+{
+	static const FText DisplayName = LOCTEXT("ConditionGroupQueryDisplayName_ActorTagName", "Actor Tag Name");
+	return DisplayName;
+}
+
+#if WITH_EDITOR
+TArray<TSharedRef<SWidget>> UMovieGraphConditionGroupQuery_ActorTagName::GetWidgets()
+{
+	TArray<TSharedRef<SWidget>> Widgets;
+
+	Widgets.Add(
+		SNew(SBox)
+		.HAlign(HAlign_Fill)
+		.Padding(7.f, 2.f)
+		[
+			SNew(SMultiLineEditableTextBox)
+			.Text_Lambda([this]() { return FText::FromString(TagsToMatch); })
+			.OnTextChanged_Lambda([this](const FText& InText) { TagsToMatch = InText.ToString(); })
+			.HintText(LOCTEXT("MovieGraphActorTagNameQueryHintText", "The actor must match one or more tags. Wildcards allowed.\nEnter each tag on a separate line."))
+		]
+	);
+
+	return Widgets;
+}
+#endif
+
+void UMovieGraphConditionGroupQuery_ActorName::Evaluate(const TArray<AActor*>& InActorsToQuery, const UWorld* InWorld, TSet<AActor*>& OutMatchingActors) const
+{
+	TRACE_CPUPROFILER_EVENT_SCOPE(UMovieGraphConditionGroupQuery_ActorName::Evaluate);
+	
+	// Quick early-out if "*" is used as the wildcard. Faster than doing the wildcard matching.
+	if (WildcardSearch == TEXT("*"))
+	{
+		OutMatchingActors.Append(InActorsToQuery);
+		return;
+	}
+
+	// Actor names can be specified on multiple lines
+	TArray<FString> AllActorNames;
+	WildcardSearch.ParseIntoArrayLines(AllActorNames);
+
+	for (AActor* Actor : InActorsToQuery)
+	{
+#if WITH_EDITOR
+		for (const FString& ActorName : AllActorNames)
+		{
+			if (Actor->GetActorLabel().MatchesWildcard(ActorName))
+			{
+				OutMatchingActors.Add(Actor);
+			}
+		}
+#endif
+	}
+}
+
+const FSlateIcon& UMovieGraphConditionGroupQuery_ActorName::GetIcon() const
+{
+	static const FSlateIcon ActorTagIcon = FSlateIcon(FAppStyle::GetAppStyleSetName(), "ClassIcon.TextRenderActor");
+	return ActorTagIcon;
+}
+
+const FText& UMovieGraphConditionGroupQuery_ActorName::GetDisplayName() const
+{
+	static const FText DisplayName = LOCTEXT("ConditionGroupQueryDisplayName_ActorName", "Actor Name");
+	return DisplayName;
+}
+
+#if WITH_EDITOR
+TArray<TSharedRef<SWidget>> UMovieGraphConditionGroupQuery_ActorName::GetWidgets()
+{
+	TArray<TSharedRef<SWidget>> Widgets;
+
+	Widgets.Add(
+		SNew(SBox)
+		.HAlign(HAlign_Fill)
+		.Padding(7.f, 2.f)
+		[
+			SNew(SMultiLineEditableTextBox)
+			.Text_Lambda([this]() { return FText::FromString(WildcardSearch); })
+			.OnTextChanged_Lambda([this](const FText& InText) { WildcardSearch = InText.ToString(); })
+			.HintText(LOCTEXT("MovieGraphActorNameQueryHintText", "Actor names to query. Wildcards allowed.\nEnter each actor name on a separate line."))
+		]
+	);
+
+	return Widgets;
+}
+#endif
+
+bool UMovieGraphConditionGroupQuery_ActorName::IsEditorOnly() const
+{
+	// GetActorLabel() is editor-only
+	return true;
+}
+
+void UMovieGraphConditionGroupQuery_ActorType::Evaluate(const TArray<AActor*>& InActorsToQuery, const UWorld* InWorld, TSet<AActor*>& OutMatchingActors) const
+{
+	TRACE_CPUPROFILER_EVENT_SCOPE(UMovieGraphConditionGroupQuery_ActorType::Evaluate);
+	
+	for (AActor* Actor : InActorsToQuery)
+	{
+		if (ActorTypes.Contains(Actor->GetClass()))
+		{
+			OutMatchingActors.Add(Actor);
+		}
+	}
+}
+
+const FSlateIcon& UMovieGraphConditionGroupQuery_ActorType::GetIcon() const
+{
+	static const FSlateIcon ActorTagIcon = FSlateIcon(FAppStyle::GetAppStyleSetName(), "ClassIcon.ActorComponent");
+	return ActorTagIcon;
+}
+
+const FText& UMovieGraphConditionGroupQuery_ActorType::GetDisplayName() const
+{
+	static const FText DisplayName = LOCTEXT("ConditionGroupQueryDisplayName_ActorType", "Actor Type");
+	return DisplayName;
+}
+
+#if WITH_EDITOR
+TArray<TSharedRef<SWidget>> UMovieGraphConditionGroupQuery_ActorType::GetWidgets()
+{
+	TArray<TSharedRef<SWidget>> Widgets;
+
+	Widgets.Add(
+		SAssignNew(ActorTypesList, SMovieGraphSimpleList<UClass*>)
+			.DataSource(&ActorTypes)
+			.DataType(FText::FromString("Actor Type"))
+			.DataTypePlural(FText::FromString("Actor Types"))
+			.OnGetRowText_Static(&GetRowText)
+			.OnGetRowIcon_Static(&GetRowIcon)
+			.OnDelete_Lambda([this](UClass* InActorClass)
+			{
+				ActorTypes.Remove(InActorClass);
+				ActorTypesList->Refresh();
+			})
+	);
+
+	return Widgets;
+}
+
+bool UMovieGraphConditionGroupQuery_ActorType::HasAddMenu() const
+{
+	return true;
+}
+
+TSharedRef<SWidget> UMovieGraphConditionGroupQuery_ActorType::GetAddMenuContents(const FMovieGraphConditionGroupQueryContentsChanged& OnAddFinished)
+{
+	FClassViewerModule& ClassViewerModule = FModuleManager::LoadModuleChecked<FClassViewerModule>("ClassViewer");
+
+	FClassViewerInitializationOptions Options;
+	Options.Mode = EClassViewerMode::ClassPicker;
+	Options.NameTypeToDisplay = EClassViewerNameTypeToDisplay::DisplayName;
+	Options.bShowNoneOption = false;
+	Options.bIsActorsOnly = true;
+	Options.bShowUnloadedBlueprints = false;
+
+	// Add a class filter to disallow adding duplicates of actor types that were already picked
+	Options.ClassFilters.Add(MakeShared<UE::MovieGraph::Private::FClassViewerTypeFilter>(&ActorTypes));
+
+	const TSharedRef<SWidget> ClassViewer = ClassViewerModule.CreateClassViewer(
+		Options,
+		FOnClassPicked::CreateLambda([this, OnAddFinished](UClass* InNewClass)
+		{
+			FSlateApplication::Get().DismissAllMenus();
+			
+			ActorTypes.Add(InNewClass);
+			OnAddFinished.ExecuteIfBound();
+
+			// Ensure that the class filters run again so duplicate actor types cannot be selected
+			if (ClassViewerWidget.IsValid())
+			{
+				ClassViewerWidget->Refresh();
+				ActorTypesList->Refresh();
+			}
+		}));
+
+	ClassViewerWidget = StaticCastSharedPtr<SClassViewer>(ClassViewer.ToSharedPtr()); 
+	
+	return SNew(SBox)
+		.WidthOverride(300.f)
+		.HeightOverride(300.f)
+		[
+			ClassViewerWidget.ToSharedRef()
+		];
+}
+
+const FSlateBrush* UMovieGraphConditionGroupQuery_ActorType::GetRowIcon(UClass* InActorType)
+{
+	return FSlateIconFinder::FindIconForClass(InActorType).GetIcon();
+}
+
+FText UMovieGraphConditionGroupQuery_ActorType::GetRowText(UClass* InActorType)
+{
+	if (InActorType)
+	{
+		return InActorType->GetDisplayNameText();
+	}
+
+	return LOCTEXT("MovieGraphActorTypeConditionGroupQuery_Invalid", "(invalid)");
+}
+#endif
+
+void UMovieGraphConditionGroupQuery_ComponentTagName::Evaluate(const TArray<AActor*>& InActorsToQuery, const UWorld* InWorld, TSet<AActor*>& OutMatchingActors) const
+{
+	TRACE_CPUPROFILER_EVENT_SCOPE(UMovieGraphConditionGroupQuery_ComponentTag::Evaluate);
+	
+	// Quick early-out if "*" is used as the wildcard. Faster than doing the wildcard matching.
+	if (TagsToMatch == TEXT("*"))
+	{
+		OutMatchingActors.Append(InActorsToQuery);
+		return;
+	}
+
+	// Component tags can be specified on multiple lines
+	TArray<FString> AllTagNameStrings;
+	TagsToMatch.ParseIntoArrayLines(AllTagNameStrings);
+	
+	TInlineComponentArray<UActorComponent*> ActorComponents;
+
+	for (AActor* Actor : InActorsToQuery)
+	{
+		// Include child components so components inside of Blueprints can be found
+		constexpr bool bIncludeFromChildActors = false;
+		Actor->GetComponents<UActorComponent*>(ActorComponents, bIncludeFromChildActors);
+		
+		for (const UActorComponent* Component : ActorComponents)
+		{
+			bool bMatchedTag = false;
+			
+			for (const FString& TagToMatch : AllTagNameStrings)
+			{			
+				for (const FName& ComponentTag : Component->ComponentTags)
+				{
+					if (ComponentTag.ToString().MatchesWildcard(TagToMatch))
+					{
+						OutMatchingActors.Add(Actor);
+						bMatchedTag = true;
+						break;
+					}
+				}
+
+				// Skip comparing the rest of the tags if one tag already matched 
+				if (bMatchedTag)
+				{
+					break;
+				}
+			}
+
+			// Skip comparing the rest of the components if one component already matched
+			if (bMatchedTag)
+			{
+				break;
+			}
+		}
+
+		ActorComponents.Reset();
+	}
+}
+
+const FSlateIcon& UMovieGraphConditionGroupQuery_ComponentTagName::GetIcon() const
+{
+	// TODO: This icon is wrong
+	static const FSlateIcon ActorTagIcon = FSlateIcon(FAppStyle::GetAppStyleSetName(), "ClassIcon.DataAsset");
+	return ActorTagIcon;
+}
+
+const FText& UMovieGraphConditionGroupQuery_ComponentTagName::GetDisplayName() const
+{
+	static const FText DisplayName = LOCTEXT("ConditionGroupQueryDisplayName_ComponentTagName", "Component Tag Name");
+	return DisplayName;
+}
+
+#if WITH_EDITOR
+TArray<TSharedRef<SWidget>> UMovieGraphConditionGroupQuery_ComponentTagName::GetWidgets()
+{
+	TArray<TSharedRef<SWidget>> Widgets;
+
+	Widgets.Add(
+		SNew(SBox)
+		.HAlign(HAlign_Fill)
+		.Padding(7.f, 2.f)
+		[
+			SNew(SMultiLineEditableTextBox)
+			.Text_Lambda([this]() { return FText::FromString(TagsToMatch); })
+			.OnTextChanged_Lambda([this](const FText& InText) { TagsToMatch = InText.ToString(); })
+			.HintText(LOCTEXT("MovieGraphComponentTagNameQueryHintText", "A component on the actor must match one or more component tags.\nWildcards allowed. Enter each tag on a separate line."))
+		]
+	);
+
+	return Widgets;
+}
+#endif
+
+void UMovieGraphConditionGroupQuery_ComponentType::Evaluate(const TArray<AActor*>& InActorsToQuery, const UWorld* InWorld, TSet<AActor*>& OutMatchingActors) const
+{
+	TRACE_CPUPROFILER_EVENT_SCOPE(UMovieGraphConditionGroupQuery_ComponentType::Evaluate);
+	
+	TInlineComponentArray<UActorComponent*> ActorComponents;
+
+	for (AActor* Actor : InActorsToQuery)
+	{
+		// Include child components so components inside of Blueprints can be found
+		constexpr bool bIncludeFromChildActors = false;
+		Actor->GetComponents<UActorComponent*>(ActorComponents, bIncludeFromChildActors);
+		
+		for (const UActorComponent* Component : ActorComponents)
+		{
+			if (ComponentTypes.Contains(Component->GetClass()))
+			{
+				OutMatchingActors.Add(Actor);
+			}
+		}
+
+		ActorComponents.Reset();
+	}
+}
+
+const FSlateIcon& UMovieGraphConditionGroupQuery_ComponentType::GetIcon() const
+{
+	static const FSlateIcon ActorTagIcon = FSlateIcon(FAppStyle::GetAppStyleSetName(), "ClassIcon.ActorComponent");
+	return ActorTagIcon;
+}
+
+const FText& UMovieGraphConditionGroupQuery_ComponentType::GetDisplayName() const
+{
+	static const FText DisplayName = LOCTEXT("ConditionGroupQueryDisplayName_ComponentType", "Component Type");
+	return DisplayName;
+}
+
+#if WITH_EDITOR
+TArray<TSharedRef<SWidget>> UMovieGraphConditionGroupQuery_ComponentType::GetWidgets()
+{
+	TArray<TSharedRef<SWidget>> Widgets;
+
+	Widgets.Add(
+		SAssignNew(ComponentTypesList, SMovieGraphSimpleList<UClass*>)
+			.DataSource(&ComponentTypes)
+			.DataType(FText::FromString("Component Type"))
+			.DataTypePlural(FText::FromString("Component Types"))
+			.OnGetRowText_Static(&GetRowText)
+			.OnGetRowIcon_Static(&GetRowIcon)
+			.OnDelete_Lambda([this](UClass* InComponentType)
+			{
+				ComponentTypes.Remove(InComponentType);
+				ComponentTypesList->Refresh();
+			})			
+	);
+
+	return Widgets;
+}
+
+bool UMovieGraphConditionGroupQuery_ComponentType::HasAddMenu() const
+{
+	return true;
+}
+
+TSharedRef<SWidget> UMovieGraphConditionGroupQuery_ComponentType::GetAddMenuContents(const FMovieGraphConditionGroupQueryContentsChanged& OnAddFinished)
+{
+	FClassViewerModule& ClassViewerModule = FModuleManager::LoadModuleChecked<FClassViewerModule>("ClassViewer");
+
+	FClassViewerInitializationOptions Options;
+	Options.Mode = EClassViewerMode::ClassPicker;
+	Options.NameTypeToDisplay = EClassViewerNameTypeToDisplay::DisplayName;
+	Options.bShowNoneOption = false;
+	Options.bIsActorsOnly = false;
+	Options.bShowUnloadedBlueprints = false;
+
+	// Add a class filter to disallow adding duplicates of component types that were already picked, as well as restrict the types of classes displayed
+	// to only show component classes
+	Options.ClassFilters.Add(MakeShared<UE::MovieGraph::Private::FClassViewerTypeFilter>(&ComponentTypes, UActorComponent::StaticClass()));
+
+	const TSharedRef<SWidget> ClassViewer = ClassViewerModule.CreateClassViewer(
+		Options,
+		FOnClassPicked::CreateLambda([this, OnAddFinished](UClass* InNewClass)
+		{
+			FSlateApplication::Get().DismissAllMenus();
+			
+			ComponentTypes.Add(InNewClass);
+			OnAddFinished.ExecuteIfBound();
+
+			// Ensure that the class filters run again so duplicate actor types cannot be selected
+			if (ClassViewerWidget.IsValid())
+			{
+				ClassViewerWidget->Refresh();
+				ComponentTypesList->Refresh();
+			}
+		}));
+
+	ClassViewerWidget = StaticCastSharedPtr<SClassViewer>(ClassViewer.ToSharedPtr()); 
+	
+	return SNew(SBox)
+		.WidthOverride(300.f)
+		.HeightOverride(300.f)
+		[
+			ClassViewerWidget.ToSharedRef()
+		];
+}
+
+const FSlateBrush* UMovieGraphConditionGroupQuery_ComponentType::GetRowIcon(UClass* InComponentType)
+{
+	return FSlateIconFinder::FindIconForClass(InComponentType).GetIcon();
+}
+
+FText UMovieGraphConditionGroupQuery_ComponentType::GetRowText(UClass* InComponentType)
+{
+	if (InComponentType)
+	{
+		return InComponentType->GetDisplayNameText();
+	}
+	
+	return LOCTEXT("MovieGraphComponentTypeConditionGroupQuery_Invalid", "(invalid)");
+}
+#endif	// WITH_EDITOR
+
+void UMovieGraphConditionGroupQuery_EditorFolder::Evaluate(const TArray<AActor*>& InActorsToQuery, const UWorld* InWorld, TSet<AActor*>& OutMatchingActors) const
+{
+#if WITH_EDITOR
+	// This const cast is unfortunate, but should be harmless
+	TArray<AActor*> ActorsInFolders;
+	FActorFolders::GetActorsFromFolders(*const_cast<UWorld*>(InWorld), FolderPaths, ActorsInFolders);
+	
+	OutMatchingActors.Append(ActorsInFolders);
+#endif	// WITH_EDITOR
+}
+
+const FSlateIcon& UMovieGraphConditionGroupQuery_EditorFolder::GetIcon() const
+{
+	static const FSlateIcon EditorFolderIcon = FSlateIcon(FAppStyle::GetAppStyleSetName(), "Icons.FolderOpen");
+	return EditorFolderIcon;
+}
+
+const FText& UMovieGraphConditionGroupQuery_EditorFolder::GetDisplayName() const
+{
+	static const FText DisplayName = LOCTEXT("ConditionGroupQueryDisplayName_EditorFolder", "Editor Folder");
+	return DisplayName;
+}
+
+bool UMovieGraphConditionGroupQuery_EditorFolder::IsEditorOnlyQuery() const
+{
+	// This query is editor-only because folders do not exist outside of the editor
+	return true;
+}
+
+#if WITH_EDITOR
+TArray<TSharedRef<SWidget>> UMovieGraphConditionGroupQuery_EditorFolder::GetWidgets()
+{
+	TArray<TSharedRef<SWidget>> Widgets;
+    
+    Widgets.Add(
+    	SAssignNew(FolderPathsList, SMovieGraphSimpleList<FName>)
+    		.DataSource(&FolderPaths)
+    		.DataType(FText::FromString("Folder"))
+    		.DataTypePlural(FText::FromString("Folders"))
+    		.OnGetRowText_Static(&GetRowText)
+    		.OnGetRowIcon_Static(&GetRowIcon)
+    		.OnDelete_Lambda([this](FName InFolderPath)
+    		{
+    			FolderPaths.Remove(InFolderPath);
+    			FolderPathsList->Refresh();
+    			FolderPickerWidget->FullRefresh();
+    		})			
+    );
+
+    return Widgets;
+}
+
+bool UMovieGraphConditionGroupQuery_EditorFolder::HasAddMenu() const
+{
+	return true;
+}
+
+TSharedRef<SWidget> UMovieGraphConditionGroupQuery_EditorFolder::GetAddMenuContents(const FMovieGraphConditionGroupQueryContentsChanged& OnAddFinished)
+{
+	auto OnItemPicked = FOnSceneOutlinerItemPicked::CreateLambda([this, OnAddFinished](TSharedRef<ISceneOutlinerTreeItem> Item)
+	{
+		if (const FActorFolderTreeItem* FolderItem = Item->CastTo<FActorFolderTreeItem>())
+		{
+			if (FolderItem->IsValid())
+			{
+				const FName& FolderPath = FolderItem->GetPath();
+					
+				// Don't allow duplicate folder paths
+				if (FolderPaths.Contains(FolderPath))
+				{
+					return;
+				}
+				
+				FolderPaths.AddUnique(FolderPath);
+				OnAddFinished.ExecuteIfBound();
+				
+				if (FolderPickerWidget.IsValid())
+				{
+					FolderPickerWidget->FullRefresh();
+					FolderPathsList->Refresh();
+				}
+			}
+		}
+	});
+
+	const FCreateSceneOutlinerMode ModeFactory = FCreateSceneOutlinerMode::CreateLambda([&OnItemPicked](SSceneOutliner* Outliner)
+	{
+		return new FActorFolderPickingMode(Outliner, OnItemPicked);
+	});
+	
+	FSceneOutlinerInitializationOptions SceneOutlinerInitOptions;
+	SceneOutlinerInitOptions.bShowCreateNewFolder = false;
+	SceneOutlinerInitOptions.bFocusSearchBoxWhenOpened = true;
+	SceneOutlinerInitOptions.ModeFactory = ModeFactory;
+
+	// Don't show folders which have already been picked
+	SceneOutlinerInitOptions.Filters->AddFilterPredicate<FActorFolderTreeItem>(
+		FActorFolderTreeItem::FFilterPredicate::CreateLambda([this](const FFolder& InFolder)
+		{
+			return !FolderPaths.Contains(InFolder.GetPath());
+		}));
+
+	// Only show the name/label column, that's the only column relevant to folders
+	SceneOutlinerInitOptions.ColumnMap.Add(
+		FSceneOutlinerBuiltInColumnTypes::Label(),
+		FSceneOutlinerColumnInfo(ESceneOutlinerColumnVisibility::Visible, 0, FCreateSceneOutlinerColumn(), false, TOptional<float>(), FSceneOutlinerBuiltInColumnTypes::Label_Localized()));
+
+	FolderPickerWidget = SNew(SSceneOutliner, SceneOutlinerInitOptions)
+		.IsEnabled(FSlateApplication::Get().GetNormalExecutionAttribute());
+	
+	return
+		SNew(SBox)
+		.WidthOverride(400.f)
+		.HeightOverride(300.f)
+		[
+			FolderPickerWidget.ToSharedRef()
+		];
+}
+
+const FSlateBrush* UMovieGraphConditionGroupQuery_EditorFolder::GetRowIcon(FName InFolderPath)
+{
+	return FAppStyle::Get().GetBrush("Icons.FolderOpen");
+}
+
+FText UMovieGraphConditionGroupQuery_EditorFolder::GetRowText(FName InFolderPath)
+{
+	return FText::FromString(InFolderPath.ToString());
+}
+#endif	// WITH_EDITOR
+
+void UMovieGraphConditionGroupQuery_Sublevel::Evaluate(const TArray<AActor*>& InActorsToQuery, const UWorld* InWorld, TSet<AActor*>& OutMatchingActors) const
+{
+	for (const TSoftObjectPtr<UWorld>& World : Sublevels)
+	{
+		// Don't load the level, only use levels which are already loaded
+		const UWorld* LoadedWorld = World.Get();
+		if (!LoadedWorld)
+		{
+			const UMovieGraphCollection* ParentCollection = GetTypedOuter<UMovieGraphCollection>();
+			const FString CollectionName = ParentCollection ? ParentCollection->GetCollectionName() : TEXT("<unknown>");
+			
+			UE_LOG(LogMovieRenderPipeline, Warning, TEXT("Sublevel query in collection '%s' is excluding level (%s) because it is not loaded."), *CollectionName, *World.ToString())
+			continue;
+		}
+
+		ULevel* CurrentLevel = LoadedWorld->GetCurrentLevel();
+		if (!CurrentLevel)
+		{
+			continue;
+		}
+
+		for (TObjectPtr<AActor>& LevelActor : CurrentLevel->Actors)
+		{
+			OutMatchingActors.Add(LevelActor.Get());
+		}
+	}
+}
+
+const FSlateIcon& UMovieGraphConditionGroupQuery_Sublevel::GetIcon() const
+{
+	static const FSlateIcon SublevelIcon = FSlateIcon(FAppStyle::GetAppStyleSetName(), "Icons.Level");
+	return SublevelIcon;
+}
+
+const FText& UMovieGraphConditionGroupQuery_Sublevel::GetDisplayName() const
+{
+	static const FText DisplayName = LOCTEXT("ConditionGroupQueryDisplayName_Sublevel", "Sublevel");
+	return DisplayName;
+}
+
+#if WITH_EDITOR
+TArray<TSharedRef<SWidget>> UMovieGraphConditionGroupQuery_Sublevel::GetWidgets()
+{
+	TArray<TSharedRef<SWidget>> Widgets;
+
+	// Create the data source for the list view
+	ListDataSource.Empty();
+	for (TSoftObjectPtr<UWorld>& Sublevel : Sublevels)
+	{
+		ListDataSource.Add(MakeShared<TSoftObjectPtr<UWorld>>(Sublevel));
+	}
+
+	Widgets.Add(
+		SAssignNew(SublevelsList, SMovieGraphSimpleList<TSharedPtr<TSoftObjectPtr<UWorld>>>)
+			.DataSource(&ListDataSource)
+			.DataType(FText::FromString("Sublevel"))
+			.DataTypePlural(FText::FromString("Sublevels"))
+			.OnGetRowText_Static(&GetRowText)
+			.OnGetRowIcon_Static(&GetRowIcon)
+			.OnDelete_Lambda([this](const TSharedPtr<TSoftObjectPtr<UWorld>> InSublevel)
+			{
+				ListDataSource.Remove(InSublevel);
+				Sublevels.Remove(*InSublevel.Get());
+				
+				SublevelsList->Refresh();
+				
+				constexpr bool bUpdateSources = true;
+				RefreshLevelPicker.ExecuteIfBound(bUpdateSources);
+			})
+	);
+
+	return Widgets;
+}
+
+bool UMovieGraphConditionGroupQuery_Sublevel::HasAddMenu() const
+{
+	return true;
+}
+
+TSharedRef<SWidget> UMovieGraphConditionGroupQuery_Sublevel::GetAddMenuContents(const FMovieGraphConditionGroupQueryContentsChanged& OnAddFinished)
+{
+	FAssetPickerConfig SublevelPickerConfig;
+	{
+		SublevelPickerConfig.SelectionMode = ESelectionMode::Single;
+		SublevelPickerConfig.SaveSettingsName = TEXT("MovieRenderGraphSublevelPicker");
+		SublevelPickerConfig.RefreshAssetViewDelegates.Add(&RefreshLevelPicker);
+		SublevelPickerConfig.InitialAssetViewType = EAssetViewType::Column;
+		SublevelPickerConfig.bFocusSearchBoxWhenOpened = true;
+		SublevelPickerConfig.bAllowNullSelection = false;
+		SublevelPickerConfig.bShowBottomToolbar = true;
+		SublevelPickerConfig.bAutohideSearchBar = false;
+		SublevelPickerConfig.bAllowDragging = false;
+		SublevelPickerConfig.bCanShowClasses = false;
+		SublevelPickerConfig.bShowPathInColumnView = true;
+		SublevelPickerConfig.bShowTypeInColumnView = false;
+		SublevelPickerConfig.bSortByPathInColumnView = false;
+		SublevelPickerConfig.HiddenColumnNames = {
+			ContentBrowserItemAttributes::ItemDiskSize.ToString(),
+			ContentBrowserItemAttributes::VirtualizedData.ToString(),
+			TEXT("PrimaryAssetType"),
+			TEXT("PrimaryAssetName")
+		};
+		SublevelPickerConfig.AssetShowWarningText = LOCTEXT("ConditionGroupQuery_NoSublevelsFound", "No Sublevels Found");
+		SublevelPickerConfig.Filter.ClassPaths.Add(UWorld::StaticClass()->GetClassPathName());
+		SublevelPickerConfig.OnAssetSelected = FOnAssetSelected::CreateLambda([this, OnAddFinished](const FAssetData& InLevelAsset)
+		{
+			FSlateApplication::Get().DismissAllMenus();
+			
+			Sublevels.AddUnique(InLevelAsset.GetAsset());
+			ListDataSource.AddUnique(MakeShared<TSoftObjectPtr<UWorld>>(InLevelAsset.GetAsset()));
+			OnAddFinished.ExecuteIfBound();
+
+			if (SublevelsList.IsValid())
+			{
+				SublevelsList->Refresh();
+				
+				constexpr bool bUpdateSources = false;
+				RefreshLevelPicker.ExecuteIfBound(bUpdateSources);
+			}
+		});
+		SublevelPickerConfig.OnShouldFilterAsset = FOnShouldFilterAsset::CreateLambda([this](const FAssetData& InLevelAsset)
+		{
+			// Don't show sublevels which have already been picked
+			UWorld* Sublevel = Cast<UWorld>(InLevelAsset.GetAsset());
+			return !Sublevel || Sublevels.Contains(Sublevel);
+		});
+	}
+
+	IContentBrowserSingleton& ContentBrowser = FModuleManager::LoadModuleChecked<FContentBrowserModule>("ContentBrowser").Get();
+
+	return
+		SNew(SBox)
+		.Padding(0, 10.f, 0, 0)
+		.WidthOverride(400.f)
+		.HeightOverride(300.f)
+		[
+			ContentBrowser.CreateAssetPicker(SublevelPickerConfig)
+		];
+}
+
+const FSlateBrush* UMovieGraphConditionGroupQuery_Sublevel::GetRowIcon(TSharedPtr<TSoftObjectPtr<UWorld>> InSublevel)
+{
+	return FAppStyle::Get().GetBrush("Icons.Level");
+}
+
+FText UMovieGraphConditionGroupQuery_Sublevel::GetRowText(TSharedPtr<TSoftObjectPtr<UWorld>> InSublevel)
+{
+	if (InSublevel.IsValid())
+	{
+		if (InSublevel.Get()->IsValid())
+		{
+			// The first Get() returns the TSoftObjectPtr, the second Get() dereferences the TSoftObjectPtr
+			return FText::FromString(InSublevel.Get()->Get()->GetName());
+		}
+	}
+
+	return LOCTEXT("MovieGraphSublevelConditionGroupQuery_InvalidLevel", "(invalid)");
+}
+#endif	// WITH_EDITOR
+
+UMovieGraphConditionGroup::UMovieGraphConditionGroup()
+	: OpType(EMovieGraphConditionGroupOpType::Add)
+{
+	// The CDO will always have the default GUID
+	if (!HasAllFlags(RF_ClassDefaultObject))
+	{
+		Id = FGuid::NewGuid();
+	}
+	else
+	{
+		Id = FGuid();
+	}
+}
+
+void UMovieGraphConditionGroup::SetOperationType(const EMovieGraphConditionGroupOpType OperationType)
+{
+	// Always allow setting the operation type to Union. If not setting to Union, only allow setting the operation type if this is not the first
+	// condition group in the collection. The first condition group is always a Union.
+	if (OperationType == EMovieGraphConditionGroupOpType::Add)
+	{
+		OpType = EMovieGraphConditionGroupOpType::Add;
+		return;
+	}
+
+	const UMovieGraphCollection* ParentCollection = GetTypedOuter<UMovieGraphCollection>();
+	if (ensureMsgf(ParentCollection, TEXT("Cannot set the operation type on a condition group that doesn't have a collection outer")))
+	{
+		if (ParentCollection->GetConditionGroups().Find(this) != 0)
+		{
+			OpType = OperationType;
+		}
+	}
+}
+
+EMovieGraphConditionGroupOpType UMovieGraphConditionGroup::GetOperationType() const
+{
+	return OpType;
+}
+
+TSet<AActor*> UMovieGraphConditionGroup::Evaluate(const UWorld* InWorld) const
+{
+	TRACE_CPUPROFILER_EVENT_SCOPE(UMovieGraphConditionGroup::Evaluate);
+
+	// Reset the TSet for evaluation results; it is persisted across frames to prevent constantly re-allocating it
+	EvaluationResult.Reset();
+
+	// Generate a list of actors that can be fed to the queries once, rather than having all queries perform this
+	TArray<AActor*> AllActors;
+	for (TActorIterator<AActor> ActorItr(InWorld); ActorItr; ++ActorItr)
+	{
+		if (AActor* Actor = *ActorItr)
+		{
+			AllActors.Add(Actor);
+		}
+	}
+
+	for (int32 QueryIndex = 0; QueryIndex < Queries.Num(); ++QueryIndex)
+	{
+		const TObjectPtr<UMovieGraphConditionGroupQueryBase>& Query = Queries[QueryIndex];
+		if (!Query || !Query->IsEnabled())
+		{
+			continue;
+		}
+		
+		if (QueryIndex == 0)
+		{
+			// The first query should always be a Union
+			ensure(Query->GetOperationType() == EMovieGraphConditionGroupQueryOpType::Add);
+		}
+
+		// Similar to EvaluationResult, QueryResult is persisted+reset to prevent constantly re-allocating it
+		QueryResult.Reset();
+
+		Query->Evaluate(AllActors, InWorld, QueryResult);
+		
+		switch (Query->GetOperationType())
+		{
+		case EMovieGraphConditionGroupQueryOpType::Add:
+			// Append() is faster than Union() because we don't need to allocate a new set
+			EvaluationResult.Append(QueryResult);
+			break;
+
+		case EMovieGraphConditionGroupQueryOpType::And:
+			EvaluationResult = EvaluationResult.Intersect(QueryResult);
+			break;
+
+		case EMovieGraphConditionGroupQueryOpType::Subtract:
+			EvaluationResult = EvaluationResult.Difference(QueryResult);
+			break;
+		}
+	}
+	
+	return EvaluationResult;
+}
+
+UMovieGraphConditionGroupQueryBase* UMovieGraphConditionGroup::AddQuery(const TSubclassOf<UMovieGraphConditionGroupQueryBase>& InQueryType, const int32 InsertIndex)
+{
+	UMovieGraphConditionGroupQueryBase* NewQueryObj = NewObject<UMovieGraphConditionGroupQueryBase>(this, InQueryType.Get());
+
+	if (InsertIndex < 0)
+	{
+		Queries.Add(NewQueryObj);
+	}
+	else
+	{
+		// Clamp the insert index to a valid range in case an invalid one is provided
+		Queries.Insert(NewQueryObj, FMath::Clamp(InsertIndex, 0, Queries.Num()));
+	}
+	
+	return NewQueryObj;
+}
+
+const TArray<UMovieGraphConditionGroupQueryBase*>& UMovieGraphConditionGroup::GetQueries() const
+{
+	return Queries;
+}
+
+bool UMovieGraphConditionGroup::RemoveQuery(UMovieGraphConditionGroupQueryBase* InQuery)
+{
+	return Queries.RemoveSingle(InQuery) == 1;
+}
+
+bool UMovieGraphConditionGroup::IsFirstConditionGroup() const
+{
+	const UMovieGraphCollection* ParentCollection = GetTypedOuter<UMovieGraphCollection>();
+	if (ensureMsgf(ParentCollection, TEXT("Cannot determine if this is the first condition group when no parent collection is present")))
+	{
+		// GetConditionGroups() returns an array of non-const pointers, so Find() doesn't like having a const pointer passed to it.
+		// Find() won't mutate the condition group though, so the const_cast here is OK.
+		return ParentCollection->GetConditionGroups().Find(const_cast<UMovieGraphConditionGroup*>(this)) == 0;
+	}
+	
+	return false;
+}
+
+bool UMovieGraphConditionGroup::MoveQueryToIndex(UMovieGraphConditionGroupQueryBase* InQuery, const int32 NewIndex)
+{
+#if WITH_EDITOR
+	Modify();
+#endif
+
+	if (!InQuery)
+	{
+		return false;
+	}
+
+	const int32 ExistingIndex = Queries.Find(InQuery);
+	if (ExistingIndex == INDEX_NONE)
+	{
+		return false;
+	}
+
+	// If the new index is greater than the current index, then decrement the destination index so it remains valid after the removal below
+	int32 DestinationIndex = NewIndex;
+	if (DestinationIndex > ExistingIndex)
+	{
+		--DestinationIndex;
+	}
+
+	Queries.Remove(InQuery);
+	Queries.Insert(InQuery, DestinationIndex);
+
+	// Enforce that the first query is set to Union
+	InQuery->SetOperationType(EMovieGraphConditionGroupQueryOpType::Add);
+
+	return true;
+}
+
+const FGuid& UMovieGraphConditionGroup::GetId() const
+{
+	return Id;
+}
+
+TSet<AActor*> UMovieGraphCollection::Evaluate(const UWorld* InWorld) const
+{
+	TRACE_CPUPROFILER_EVENT_SCOPE(UMovieGraphCollection::Evaluate);
+	
+	TSet<AActor*> ResultSet;
+
+	for (int32 ConditionGroupIndex = 0; ConditionGroupIndex < ConditionGroups.Num(); ++ConditionGroupIndex)
+	{
+		const TObjectPtr<UMovieGraphConditionGroup>& ConditionGroup = ConditionGroups[ConditionGroupIndex];
+		if (!ConditionGroup)
+		{
+			continue;
+		}
+		
+		if (ConditionGroupIndex == 0)
+		{
+			// The first condition group should always be a Union
+			ensure(ConditionGroup->GetOperationType() == EMovieGraphConditionGroupOpType::Add);
+		}
+
+		const TSet<AActor*> QueryResult = ConditionGroup->Evaluate(InWorld);
+		
+		switch (ConditionGroup->GetOperationType())
+		{
+		case EMovieGraphConditionGroupOpType::Add:
+			ResultSet = ResultSet.Union(QueryResult);
+			break;
+
+		case EMovieGraphConditionGroupOpType::And:
+			ResultSet = ResultSet.Intersect(QueryResult);
+			break;
+
+		case EMovieGraphConditionGroupOpType::Subtract:
+			ResultSet = ResultSet.Difference(QueryResult);
+			break;
+		}
+	}
+	
+	return ResultSet;
+}
+
+UMovieGraphConditionGroup* UMovieGraphCollection::AddConditionGroup()
+{
+	UMovieGraphConditionGroup* NewConditionGroup = NewObject<UMovieGraphConditionGroup>(this);
+	ConditionGroups.Add(NewConditionGroup);
+	return NewConditionGroup;
+}
+
+const TArray<UMovieGraphConditionGroup*>& UMovieGraphCollection::GetConditionGroups() const
+{
+	return ConditionGroups;
+}
+
+bool UMovieGraphCollection::RemoveConditionGroup(UMovieGraphConditionGroup* InConditionGroup)
+{
+	return ConditionGroups.RemoveSingle(InConditionGroup) == 1;
+}
+
+bool UMovieGraphCollection::MoveConditionGroupToIndex(UMovieGraphConditionGroup* InConditionGroup, const int32 NewIndex)
+{
+#if WITH_EDITOR
+	Modify();
+#endif
+
+	if (!InConditionGroup)
+	{
+		return false;
+	}
+
+	const int32 ExistingIndex = ConditionGroups.Find(InConditionGroup);
+	if (ExistingIndex == INDEX_NONE)
+	{
+		return false;
+	}
+
+	// If the new index is greater than the current index, then decrement the destination index so it remains valid after the removal below
+	int32 DestinationIndex = NewIndex;
+	if (DestinationIndex > ExistingIndex)
+	{
+		--DestinationIndex;
+	}
+
+	ConditionGroups.Remove(InConditionGroup);
+	ConditionGroups.Insert(InConditionGroup, DestinationIndex);
+
+	// Enforce that the first condition group is set to Union
+	InConditionGroup->SetOperationType(EMovieGraphConditionGroupOpType::Add);
+
+	return true;
+}
+
+#if WITH_EDITOR
+void UMovieGraphCollection::PostEditChangeProperty(FPropertyChangedEvent& PropertyChangedEvent)
+{
+	Super::PostEditChangeProperty(PropertyChangedEvent);
+
+	// Name change delegate is broadcast here so it catches both SetCollectionName() and a direct change of the property via the details panel
+	if (PropertyChangedEvent.GetPropertyName() == GET_MEMBER_NAME_CHECKED(UMovieGraphCollection, CollectionName))
+	{
+		OnCollectionNameChangedDelegate.Broadcast(this);
+	}
+}
+#endif	// WITH_EDITOR
+
+void UMovieGraphCollection::SetCollectionName(const FString& InName)
+{
+	CollectionName = InName;
+}
+
+const FString& UMovieGraphCollection::GetCollectionName() const
+{
+	return CollectionName;
+}
+
+UMovieGraphCollection* UMovieGraphRenderLayer::GetCollectionByName(const FString& Name) const
+{
+	for (const UMovieGraphCollectionModifier* Modifier : Modifiers)
+	{
+		if (!Modifier)
+		{
+			continue;
+		}
+		
+		for (UMovieGraphCollection* Collection : Modifier->GetCollections())
+		{
+			if (Collection && Collection->GetCollectionName().Equals(Name))
+			{
+				return Collection;
+			}
+		}
+	}
+
+	return nullptr;
+}
+
+void UMovieGraphRenderLayer::AddModifier(UMovieGraphCollectionModifier* Modifier)
+{
+	if (!Modifiers.Contains(Modifier))
+	{
+		Modifiers.Add(Modifier);
+	}
+}
+
+void UMovieGraphRenderLayer::RemoveModifier(UMovieGraphCollectionModifier* Modifier)
+{
+	Modifiers.Remove(Modifier);
+}
+
+void UMovieGraphRenderLayer::Preview(const UWorld* World)
+{
+	if (!World)
+	{
+		return;
+	}
+	
+	// Apply all modifiers
+	for (UMovieGraphCollectionModifier* Modifier : Modifiers)
+	{
+		Modifier->ApplyModifier(World);
+	}
+}
+
+void UMovieGraphRenderLayer::UndoPreview(const UWorld* World)
+{
+	if (!World)
+	{
+		return;
+	}
+
+	// Undo actions performed by all modifiers. Do this in the reverse order that they were applied, since the undo
+	// state of one modifier may depend on modifiers that were previously applied.
+	for (int32 Index = Modifiers.Num() - 1; Index >= 0; Index--)
+	{
+		if (UMovieGraphCollectionModifier* Modifier = Modifiers[Index])
+		{
+			Modifier->UndoModifier();
+		}
+	}
+}
+
+
+UMovieGraphRenderLayerSubsystem* UMovieGraphRenderLayerSubsystem::GetFromWorld(const UWorld* World)
+{
+	if (World)
+	{
+		return UWorld::GetSubsystem<UMovieGraphRenderLayerSubsystem>(World);
+	}
+
+	return nullptr;
+}
+
+void UMovieGraphRenderLayerSubsystem::Initialize(FSubsystemCollectionBase& Collection)
+{
+	VisualizationEmptyCollection = NewObject<UMovieGraphCollection>(GetTransientPackage(), NAME_None, RF_Transient);
+
+	// By default, the visualizer should hide everything in the world
+	VisualizationModifier_HideWorld = NewObject<UMovieGraphRenderPropertyModifier>(GetTransientPackage(), NAME_None, RF_Transient);
+	VisualizationModifier_HideWorld->AddCollection(VisualizationEmptyCollection);
+	VisualizationModifier_HideWorld->SetHidden(true);
+
+	// Selectively show collections in the visualization
+	VisualizationModifier_VisibleCollections = NewObject<UMovieGraphRenderPropertyModifier>(GetTransientPackage(), NAME_None, RF_Transient);
+
+	// The visualizer render layer will hide the world, selectively show specified collections, and then apply any other provided modifiers
+	VisualizationRenderLayer = NewObject<UMovieGraphRenderLayer>(GetTransientPackage(), NAME_None, RF_Transient);
+	VisualizationRenderLayer->AddModifier(VisualizationModifier_HideWorld);
+	VisualizationRenderLayer->AddModifier(VisualizationModifier_VisibleCollections);
+}
+
+void UMovieGraphRenderLayerSubsystem::Deinitialize()
+{
+}
+
+void UMovieGraphRenderLayerSubsystem::Reset()
+{
+	ClearAllPreviews();
+	RenderLayers.Empty();
+}
+
+bool UMovieGraphRenderLayerSubsystem::AddRenderLayer(UMovieGraphRenderLayer* RenderLayer)
+{
+	if (!RenderLayer)
+	{
+		UE_LOG(LogMovieRenderPipeline, Warning, TEXT("Invalid render layer provided to AddRenderLayer()."));
+		return false;
+	}
+	
+	const bool bRenderLayerExists = RenderLayers.ContainsByPredicate([RenderLayer](const UMovieGraphRenderLayer* RL)
+	{
+		return RL && (RenderLayer->GetRenderLayerName() == RL->GetName());
+	});
+
+	if (bRenderLayerExists)
+	{
+		UE_LOG(LogMovieRenderPipeline, Warning, TEXT("Render layer '%s' already exists in the render layer subsystem; it will not be added again."), *RenderLayer->GetRenderLayerName().ToString());
+		return false;
+	}
+
+	RenderLayers.Add(RenderLayer);
+	return true;
+}
+
+void UMovieGraphRenderLayerSubsystem::RemoveRenderLayer(const FString& RenderLayerName)
+{
+	if (ActiveRenderLayer && (ActiveRenderLayer->GetName() == RenderLayerName))
+	{
+		ClearAllPreviews();
+	}
+	
+	const uint32 Index = RenderLayers.IndexOfByPredicate([&RenderLayerName](const UMovieGraphRenderLayer* RenderLayer)
+	{
+		return RenderLayer->GetRenderLayerName() == RenderLayerName;
+	});
+
+	if (Index != INDEX_NONE)
+	{
+		RenderLayers.RemoveAt(Index);
+	}
+}
+
+void UMovieGraphRenderLayerSubsystem::SetActiveRenderLayerByObj(UMovieGraphRenderLayer* RenderLayer)
+{
+	if (!RenderLayer)
+	{
+		return;
+	}
+	
+	ClearAllPreviews();
+	SetAndPreviewRenderLayer(RenderLayer);
+}
+
+void UMovieGraphRenderLayerSubsystem::SetActiveRenderLayerByName(const FName& RenderLayerName)
+{
+	const uint32 Index = RenderLayers.IndexOfByPredicate([&RenderLayerName](const UMovieGraphRenderLayer* RenderLayer)
+	{
+		return RenderLayer->GetRenderLayerName() == RenderLayerName;
+	});
+
+	if (Index != INDEX_NONE)
+	{
+		SetActiveRenderLayerByObj(RenderLayers[Index]);
+	}
+}
+
+void UMovieGraphRenderLayerSubsystem::ClearActiveRenderLayer()
+{
+	ClearAllPreviews();
+}
+
+void UMovieGraphRenderLayerSubsystem::PreviewCollection(UMovieGraphCollection* Collection)
+{
+	if (!Collection)
+	{
+		return;
+	}
+
+	ClearAllPreviews();
+	
+	ActiveCollection = Collection;
+	VisualizationModifier_VisibleCollections->AddCollection(ActiveCollection);
+
+	SetAndPreviewRenderLayer(VisualizationRenderLayer);
+}
+
+void UMovieGraphRenderLayerSubsystem::ClearCollectionPreview()
+{
+	ClearAllPreviews();
+}
+
+void UMovieGraphRenderLayerSubsystem::ClearAllPreviews()
+{
+	// Render layer previews and collection previews both use the active render layer, so undoing the preview this
+	// way will clear previews for both
+	if (ActiveRenderLayer)
+	{
+		ActiveRenderLayer->UndoPreview(GetWorld());
+
+		// Remove the modifier preview if present (requires an active render layer)
+		if (ActiveModifier)
+		{
+			ActiveRenderLayer->RemoveModifier(ActiveModifier);
+		}
+	}
+
+	// Reset the viz modifier for the collection preview
+	VisualizationModifier_VisibleCollections->SetCollections({});
+
+	ActiveRenderLayer = nullptr;
+	ActiveCollection = nullptr;
+	ActiveModifier = nullptr;
+}
+
+void UMovieGraphRenderLayerSubsystem::SetAndPreviewRenderLayer(UMovieGraphRenderLayer* RenderLayer)
+{
+	ActiveRenderLayer = RenderLayer;
+	ActiveRenderLayer->Preview(GetWorld());
+}
+
+void UMovieGraphRenderLayerSubsystem::PreviewModifier(UMovieGraphCollectionModifier* Modifier)
+{
+	if (!Modifier)
+	{
+		return;
+	}
+
+	ClearAllPreviews();
+
+	ActiveModifier = Modifier;
+	VisualizationRenderLayer->AddModifier(ActiveModifier);
+
+	// Add the modifier's collections to the viz as well, so the actors that the modifier affects are actually visible
+	// TODO: This may need special handing for visibility modifiers
+	for (UMovieGraphCollection* Collection : ActiveModifier->GetCollections())
+	{
+		VisualizationModifier_VisibleCollections->AddCollection(Collection);
+	}
+	
+	SetAndPreviewRenderLayer(VisualizationRenderLayer);
+}
+
+void UMovieGraphRenderLayerSubsystem::ClearModifierPreview()
+{
+	ClearAllPreviews();
+}
+
+#undef LOCTEXT_NAMESPACE
\ No newline at end of file
diff --git a/Engine/Plugins/MovieScene/MovieRenderPipeline/Source/MovieRenderPipelineCore/Private/Graph/MoviePipelineRenderLayerSubsystem.cpp b/Engine/Plugins/MovieScene/MovieRenderPipeline/Source/MovieRenderPipelineCore/Private/Graph/MoviePipelineRenderLayerSubsystem.cpp
deleted file mode 100644
diff --git a/Engine/Plugins/MovieScene/MovieRenderPipeline/Source/MovieRenderPipelineCore/Private/Graph/Nodes/MovieGraphCollectionNode.cpp b/Engine/Plugins/MovieScene/MovieRenderPipeline/Source/MovieRenderPipelineCore/Private/Graph/Nodes/MovieGraphCollectionNode.cpp
--- a/Engine/Plugins/MovieScene/MovieRenderPipeline/Source/MovieRenderPipelineCore/Private/Graph/Nodes/MovieGraphCollectionNode.cpp
+++ b/Engine/Plugins/MovieScene/MovieRenderPipeline/Source/MovieRenderPipelineCore/Private/Graph/Nodes/MovieGraphCollectionNode.cpp
@@ -3,7 +3,7 @@
 #include "Graph/Nodes/MovieGraphCollectionNode.h"
 
 #include "Graph/MovieGraphConfig.h"
-#include "Graph/MoviePipelineRenderLayerSubsystem.h"
+#include "Graph/MovieGraphRenderLayerSubsystem.h"
 #include "Styling/AppStyle.h"
 
 #define LOCTEXT_NAMESPACE "MovieGraph"
diff --git a/Engine/Plugins/MovieScene/MovieRenderPipeline/Source/MovieRenderPipelineCore/Private/Graph/Nodes/MovieGraphModifierNode.cpp b/Engine/Plugins/MovieScene/MovieRenderPipeline/Source/MovieRenderPipelineCore/Private/Graph/Nodes/MovieGraphModifierNode.cpp
--- a/Engine/Plugins/MovieScene/MovieRenderPipeline/Source/MovieRenderPipelineCore/Private/Graph/Nodes/MovieGraphModifierNode.cpp
+++ b/Engine/Plugins/MovieScene/MovieRenderPipeline/Source/MovieRenderPipelineCore/Private/Graph/Nodes/MovieGraphModifierNode.cpp
@@ -12,10 +12,10 @@ void UMovieGraphMergeableModifierContainer::Merge(const IMovieGraphTraversableOb
 	const UMovieGraphMergeableModifierContainer* InSourceModifiers = Cast<UMovieGraphMergeableModifierContainer>(InSourceObject);
 	checkf(InSourceModifiers, TEXT("UMovieGraphMergeableModifierContainer cannot merge with null or an object of another type."));
 
-	for (TObjectPtr<UMoviePipelineCollectionModifier>& ModifierInstance : Modifiers)
+	for (TObjectPtr<UMovieGraphCollectionModifier>& ModifierInstance : Modifiers)
 	{
 		// Find the equivalent modifier in the source modifier container
-		for (const TObjectPtr<UMoviePipelineCollectionModifier>& SourceModifier : InSourceModifiers->Modifiers)
+		for (const TObjectPtr<UMovieGraphCollectionModifier>& SourceModifier : InSourceModifiers->Modifiers)
 		{
 			if (ModifierInstance && SourceModifier && (ModifierInstance->GetClass() == SourceModifier->GetClass()))
 			{
@@ -31,7 +31,7 @@ TArray<TPair<FString, FString>> UMovieGraphMergeableModifierContainer::GetMerged
 	TArray<TPair<FString, FString>> MergedProperties;
 
 	// Include properties from modifiers if they have a bOverride_* counterpart
-	for (const TObjectPtr<UMoviePipelineCollectionModifier>& ModifierInstance : Modifiers)
+	for (const TObjectPtr<UMovieGraphCollectionModifier>& ModifierInstance : Modifiers)
 	{
 		for (TFieldIterator<FProperty> PropertyIterator(ModifierInstance->GetClass()); PropertyIterator; ++PropertyIterator)
 		{
@@ -61,7 +61,7 @@ TArray<TPair<FString, FString>> UMovieGraphMergeableModifierContainer::GetMerged
 	return MergedProperties;
 }
 
-void UMovieGraphMergeableModifierContainer::MergeProperties(const TObjectPtr<UMoviePipelineCollectionModifier>& InDestModifier, const TObjectPtr<UMoviePipelineCollectionModifier>& InSourceModifier)
+void UMovieGraphMergeableModifierContainer::MergeProperties(const TObjectPtr<UMovieGraphCollectionModifier>& InDestModifier, const TObjectPtr<UMovieGraphCollectionModifier>& InSourceModifier)
 {
 	for (TFieldIterator<FProperty> PropertyIterator(InDestModifier->GetClass()); PropertyIterator; ++PropertyIterator)
 	{
@@ -98,8 +98,8 @@ UMovieGraphModifierNode::UMovieGraphModifierNode()
 	// Add some default modifiers
 	if (!HasAnyFlags(RF_ClassDefaultObject))
 	{
-		AddModifier(UMoviePipelineVisibilityModifier::StaticClass());
-		AddModifier(UMoviePipelineMaterialModifier::StaticClass());
+		AddModifier(UMovieGraphRenderPropertyModifier::StaticClass());
+		AddModifier(UMovieGraphMaterialModifier::StaticClass());
 	}
 }
 
@@ -151,10 +151,10 @@ void UMovieGraphModifierNode::PostEditChangeProperty(FPropertyChangedEvent& Prop
 }
 #endif // WITH_EDITOR
 
-UMoviePipelineCollectionModifier* UMovieGraphModifierNode::GetModifier(TSubclassOf<UMoviePipelineCollectionModifier> ModifierType) const
+UMovieGraphCollectionModifier* UMovieGraphModifierNode::GetModifier(TSubclassOf<UMovieGraphCollectionModifier> ModifierType) const
 {
-	const TObjectPtr<UMoviePipelineCollectionModifier>* FoundModifier =
-		ModifiersContainer->Modifiers.FindByPredicate([&ModifierType](const TObjectPtr<UMoviePipelineCollectionModifier>& Modifier)
+	const TObjectPtr<UMovieGraphCollectionModifier>* FoundModifier =
+		ModifiersContainer->Modifiers.FindByPredicate([&ModifierType](const TObjectPtr<UMovieGraphCollectionModifier>& Modifier)
         {
            	return Modifier && (Modifier->GetClass() == ModifierType);
         });
@@ -162,15 +162,15 @@ UMoviePipelineCollectionModifier* UMovieGraphModifierNode::GetModifier(TSubclass
 	return (FoundModifier == nullptr) ? nullptr : *FoundModifier;
 }
 
-const TArray<UMoviePipelineCollectionModifier*>& UMovieGraphModifierNode::GetModifiers() const
+const TArray<UMovieGraphCollectionModifier*>& UMovieGraphModifierNode::GetModifiers() const
 {
 	return ModifiersContainer->Modifiers;
 }
 
-UMoviePipelineCollectionModifier* UMovieGraphModifierNode::AddModifier(TSubclassOf<UMoviePipelineCollectionModifier> ModifierType)
+UMovieGraphCollectionModifier* UMovieGraphModifierNode::AddModifier(TSubclassOf<UMovieGraphCollectionModifier> ModifierType)
 {
-	const TObjectPtr<UMoviePipelineCollectionModifier>* FoundModifier =
-		ModifiersContainer->Modifiers.FindByPredicate([&ModifierType](const TObjectPtr<UMoviePipelineCollectionModifier>& Modifier)
+	const TObjectPtr<UMovieGraphCollectionModifier>* FoundModifier =
+		ModifiersContainer->Modifiers.FindByPredicate([&ModifierType](const TObjectPtr<UMovieGraphCollectionModifier>& Modifier)
 		{
 			return Modifier && (Modifier->GetClass() == ModifierType);
 		});
@@ -183,7 +183,7 @@ UMoviePipelineCollectionModifier* UMovieGraphModifierNode::AddModifier(TSubclass
 
 	if (const UClass* ModifierClass = ModifierType.Get())
 	{
-		UMoviePipelineCollectionModifier* NewModifier = NewObject<UMoviePipelineCollectionModifier>(this, ModifierType, ModifierClass->GetFName());
+		UMovieGraphCollectionModifier* NewModifier = NewObject<UMovieGraphCollectionModifier>(this, ModifierType, ModifierClass->GetFName());
 		ModifiersContainer->Modifiers.Add(NewModifier);
 
 		return NewModifier;
@@ -192,9 +192,9 @@ UMoviePipelineCollectionModifier* UMovieGraphModifierNode::AddModifier(TSubclass
 	return nullptr;
 }
 
-bool UMovieGraphModifierNode::RemoveModifier(TSubclassOf<UMoviePipelineCollectionModifier> ModifierType)
+bool UMovieGraphModifierNode::RemoveModifier(TSubclassOf<UMovieGraphCollectionModifier> ModifierType)
 {
-	const int32 NumRemoved = ModifiersContainer->Modifiers.RemoveAll([&ModifierType](const TObjectPtr<UMoviePipelineCollectionModifier>& Modifier)
+	const int32 NumRemoved = ModifiersContainer->Modifiers.RemoveAll([&ModifierType](const TObjectPtr<UMovieGraphCollectionModifier>& Modifier)
 	{
 		return Modifier && (Modifier->GetClass() == ModifierType);
 	});
diff --git a/Engine/Plugins/MovieScene/MovieRenderPipeline/Source/MovieRenderPipelineCore/Public/Graph/MovieGraphRenderLayerSubsystem.h b/Engine/Plugins/MovieScene/MovieRenderPipeline/Source/MovieRenderPipelineCore/Public/Graph/MovieGraphRenderLayerSubsystem.h
new file mode 100644
--- /dev/null
+++ b/Engine/Plugins/MovieScene/MovieRenderPipeline/Source/MovieRenderPipelineCore/Public/Graph/MovieGraphRenderLayerSubsystem.h
@@ -0,0 +1,952 @@
+// Copyright Epic Games, Inc. All Rights Reserved.
+
+#pragma once
+
+#include "CoreMinimal.h"
+#include "Graph/MovieGraphSharedWidgets.h"
+#include "Subsystems/WorldSubsystem.h"
+#include "Textures/SlateIcon.h"
+
+#if WITH_EDITOR
+#include "ContentBrowserDelegates.h"
+#endif	// WITH_EDITOR
+
+#include "MovieGraphRenderLayerSubsystem.generated.h"
+
+class SWidget;
+
+/** Operation types available on condition groups. */
+UENUM(BlueprintType)
+enum class EMovieGraphConditionGroupOpType : uint8
+{
+	/** Adds the contents of the condition group to the results from the previous condition group (if any). */
+	Add,
+
+	/** Removes the contents of the condition group from the result of the previous condition group (if any). Any items in this condition group that aren't also found in the previous condition group will be ignored. */
+	Subtract,
+
+	/** Replaces the results of the previous condition group(s) with only the elements that exist in both that group, and this group. Intersecting with an empty condition group will result in an empty condition group. */
+	And
+};
+
+/** Operation types available on condition group queries. */
+UENUM(BlueprintType)
+enum class EMovieGraphConditionGroupQueryOpType : uint8
+{
+	/** Adds the results of the query to the results from the previous query (if any). */
+	Add,
+
+	/** Removes the results of the query from the results of the previous query (if any). Any items in this query result that aren't also found in the previous query result will be ignored. */
+	Subtract,
+	
+	/** Replaces the results of the previous queries with only the items that exist in both those queries, and this query result. Intersecting with a query which returns nothing will create an empty query result. */
+	And
+};
+
+/** Base class that all condition group queries must inherit from. */
+UCLASS(Abstract)
+class MOVIERENDERPIPELINECORE_API UMovieGraphConditionGroupQueryBase : public UObject
+{
+	GENERATED_BODY()
+
+public:
+	UMovieGraphConditionGroupQueryBase();
+
+	/** Delegate which is called when the contents of a query has changed. */
+	DECLARE_DELEGATE(FMovieGraphConditionGroupQueryContentsChanged)
+
+	/**
+	 * Sets how the condition group query interacts with the condition group. This call is ignored for the first query
+	 * in the condition group (the first is always Union).
+	 */
+	UFUNCTION(BlueprintCallable, Category = "Settings")
+	void SetOperationType(const EMovieGraphConditionGroupQueryOpType OperationType);
+
+	/** Gets the condition group query operation type. */
+	UFUNCTION(BlueprintCallable, Category = "Settings")
+	EMovieGraphConditionGroupQueryOpType GetOperationType() const;
+
+	/** Determines which of the provided actors (in the given world) match the query. Matches are added to OutMatchingActors.  */
+	UFUNCTION(BlueprintCallable, Category = "Settings")
+	virtual void Evaluate(const TArray<AActor*>& InActorsToQuery, const UWorld* InWorld, TSet<AActor*>& OutMatchingActors) const;
+
+	/**
+	 * Determines if the public properties on the query class will have their names hidden in the details panel. Returns
+	 * false by default. Most query subclasses will only have one property and do not need to clutter the UI with the
+	 * property name (eg, the "Actor Name" query only shows one text box with entries for the actor names, no need to
+	 * show the property name).
+	 */
+	UFUNCTION(BlueprintCallable, Category = "Settings")
+	virtual bool ShouldHidePropertyNames() const;
+
+	/** Gets the icon that represents this query class in the UI. */
+	virtual const FSlateIcon& GetIcon() const;
+
+	/** Gets the display name for this query class as shown in the UI. */
+	virtual const FText& GetDisplayName() const;
+
+#if WITH_EDITOR
+	/**
+	 * Gets the widgets that should be displayed for this query. If no custom widgets are specified (returning an empty array), the default
+	 * name/value widgets will be shown for all query properties tagged with EditAnywhere.
+	 */
+	virtual TArray<TSharedRef<SWidget>> GetWidgets();
+
+	/**
+	 * Returns true if this query should expose an Add menu, or false if no Add menu is visible.
+	 *
+	 * @see GetAddMenuContents()
+	 */
+	virtual bool HasAddMenu() const;
+
+	/**
+	 * Gets the contents of the "Add" menu in the UI, if any. When the Add menu updates properties within the query, OnAddFinished should be called
+	 * in order to give the UI a chance to update itself. Note that HasAddMenu() must return true in order for the contents returned from this method
+	 * to be displayed in the UI.
+	 *
+	 * @see HasAddMenu()
+	 */
+	virtual TSharedRef<SWidget> GetAddMenuContents(const FMovieGraphConditionGroupQueryContentsChanged& OnAddFinished);
+#endif
+
+	/** Determines if this query is only respected when run within the editor. Used for providing a UI hint. */
+	UFUNCTION(BlueprintCallable, Category = "Settings")
+	virtual bool IsEditorOnlyQuery() const;
+
+	/** Sets whether this query is enabled. */
+	UFUNCTION(BlueprintCallable, Category = "Settings")
+	void SetEnabled(const bool bEnabled);
+
+	/** Determines if this query is enabled. */
+	UFUNCTION(BlueprintCallable, Category = "Settings")
+	bool IsEnabled() const;
+
+	/** Determines if this is the first condition group query under the parent condition group. */
+	UFUNCTION(BlueprintCallable, Category = "Settings")
+	bool IsFirstConditionGroupQuery() const;
+
+private:
+	/** The operation type that the query is using. */
+	UPROPERTY()
+	EMovieGraphConditionGroupQueryOpType OpType;
+
+	/** Whether this query is currently enabled within the condition group. */
+	UPROPERTY()
+	bool bIsEnabled;
+};
+
+/** Query type which filters actors via an explicit actor list. */
+UCLASS(BlueprintType)
+class MOVIERENDERPIPELINECORE_API UMovieGraphConditionGroupQuery_Actor final : public UMovieGraphConditionGroupQueryBase
+{
+	GENERATED_BODY()
+
+public:
+	virtual void Evaluate(const TArray<AActor*>& InActorsToQuery, const UWorld* InWorld, TSet<AActor*>& OutMatchingActors) const override;
+	virtual const FSlateIcon& GetIcon() const override;
+	virtual const FText& GetDisplayName() const override;
+
+#if WITH_EDITOR
+	virtual TArray<TSharedRef<SWidget>> GetWidgets() override;
+	virtual bool HasAddMenu() const override;
+	virtual TSharedRef<SWidget> GetAddMenuContents(const FMovieGraphConditionGroupQueryContentsChanged& OnAddFinished) override;
+#endif
+
+private:
+#if WITH_EDITOR
+	static const FSlateBrush* GetRowIcon(TSharedPtr<TSoftObjectPtr<AActor>> InActor);
+	static FText GetRowText(TSharedPtr<TSoftObjectPtr<AActor>> InActor);
+#endif
+
+public:
+	/** The query must match one of the actors in order to be a match. If these are editor actors, they will be converted to PIE actors automatically. */
+	UPROPERTY(EditAnywhere, Category="General")
+	TArray<TSoftObjectPtr<AActor>> ActorsToMatch;
+
+private:
+#if WITH_EDITOR
+	TSharedPtr<class ISceneOutliner> ActorPickerWidget;
+
+	/** Displays the actors which have been chosen. */
+	TSharedPtr<SMovieGraphSimpleList<TSharedPtr<TSoftObjectPtr<AActor>>>> ActorsList;
+
+	// Not ideal to store a duplicate of ActorsToMatch, but SListView requires TSharedPtr<...> as the data source, and UPROPERTY does not
+	// support TSharedPtr<...>
+	TArray<TSharedPtr<TSoftObjectPtr<AActor>>> ListDataSource;
+#endif
+};
+
+/** Query type which filters actors via tags on actors. */
+UCLASS(BlueprintType)
+class MOVIERENDERPIPELINECORE_API UMovieGraphConditionGroupQuery_ActorTagName final : public UMovieGraphConditionGroupQueryBase
+{
+	GENERATED_BODY()
+
+public:
+	virtual void Evaluate(const TArray<AActor*>& InActorsToQuery, const UWorld* InWorld, TSet<AActor*>& OutMatchingActors) const override;
+	virtual const FSlateIcon& GetIcon() const override;
+	virtual const FText& GetDisplayName() const override;
+
+#if WITH_EDITOR
+	virtual TArray<TSharedRef<SWidget>> GetWidgets() override;
+#endif
+
+public:
+	/**
+	 * Tags on the actor must match one or more of the specified tags to be a match. Not case sensitive. One tag per line. Wildcards ("?" and "*") are supported but not required.
+	 * The "*" wildcard matches zero or more characters, and "?" matches exactly one character (and that character must be present).
+	 * 
+	 * Wildcard examples:
+	 * Foo* would match Foo, FooBar, and FooBaz, but not BarFoo.
+	 * *Foo* would match the above in addition to BarFoo.
+	 * Foo?Bar would match Foo.Bar and Foo_Bar, but not FooBar.
+	 * Foo? would match Food, but not FooBar or BarFoo.
+	 * Foo??? would match FooBar and FooBaz, but not Foo or Food.
+	 * ?oo? would match Food, but not Foo.
+	 * ?Foo* would match AFooBar, but not FooBar 
+	 */
+	UPROPERTY(EditAnywhere, Category="General")
+	FString TagsToMatch;
+};
+
+/** Query type which filters actors via their name (label). */
+UCLASS(BlueprintType)
+class MOVIERENDERPIPELINECORE_API UMovieGraphConditionGroupQuery_ActorName final : public UMovieGraphConditionGroupQueryBase
+{
+	GENERATED_BODY()
+
+public:
+	virtual void Evaluate(const TArray<AActor*>& InActorsToQuery, const UWorld* InWorld, TSet<AActor*>& OutMatchingActors) const override;
+	virtual const FSlateIcon& GetIcon() const override;
+	virtual const FText& GetDisplayName() const override;
+
+#if WITH_EDITOR
+	virtual TArray<TSharedRef<SWidget>> GetWidgets() override;
+#endif
+
+	virtual bool IsEditorOnly() const override;
+
+public:
+	/**
+	 * The name that the actor needs to have in order to be a match. Not case sensitive. One name per line. Wildcards ("?" and "*") are supported but not required.
+	 * The "*" wildcard matches zero or more characters, and "?" matches exactly one character (and that character must be present).
+	 * 
+	 * Wildcard examples:
+	 * Foo* would match Foo, FooBar, and FooBaz, but not BarFoo.
+	 * *Foo* would match the above in addition to BarFoo.
+	 * Foo?Bar would match Foo.Bar and Foo_Bar, but not FooBar.
+	 * Foo? would match Food, but not FooBar or BarFoo.
+	 * Foo??? would match FooBar and FooBaz, but not Foo or Food.
+	 * ?oo? would match Food, but not Foo.
+	 * ?Foo* would match AFooBar, but not FooBar 
+	 */
+	UPROPERTY(EditAnywhere, Category="General")
+	FString WildcardSearch;
+};
+
+/** Query type which filters actors by type. */
+UCLASS(BlueprintType)
+class MOVIERENDERPIPELINECORE_API UMovieGraphConditionGroupQuery_ActorType final : public UMovieGraphConditionGroupQueryBase
+{
+	GENERATED_BODY()
+
+public:
+	virtual void Evaluate(const TArray<AActor*>& InActorsToQuery, const UWorld* InWorld, TSet<AActor*>& OutMatchingActors) const override;
+	virtual const FSlateIcon& GetIcon() const override;
+	virtual const FText& GetDisplayName() const override;
+
+#if WITH_EDITOR
+	virtual TArray<TSharedRef<SWidget>> GetWidgets() override;
+	virtual bool HasAddMenu() const override;
+	virtual TSharedRef<SWidget> GetAddMenuContents(const FMovieGraphConditionGroupQueryContentsChanged& OnAddFinished) override;
+#endif
+
+public:
+	/** The type (class) that the actor needs to have in order to be a match. */
+	UPROPERTY(EditAnywhere, Category="General")
+	TArray<UClass*> ActorTypes;
+
+private:
+#if WITH_EDITOR
+	static const FSlateBrush* GetRowIcon(UClass* InActorType);
+	static FText GetRowText(UClass* InActorType);
+
+	/** Displays the actor types which have been chosen. */
+	TSharedPtr<SMovieGraphSimpleList<UClass*>> ActorTypesList;
+#endif
+	
+	/** The class viewer widget to show in the Add menu. */
+	TSharedPtr<class SClassViewer> ClassViewerWidget;
+};
+
+/** Query type which filters actors by tags on their components. */
+UCLASS(BlueprintType)
+class MOVIERENDERPIPELINECORE_API UMovieGraphConditionGroupQuery_ComponentTagName final : public UMovieGraphConditionGroupQueryBase
+{
+	GENERATED_BODY()
+
+public:
+	virtual void Evaluate(const TArray<AActor*>& InActorsToQuery, const UWorld* InWorld, TSet<AActor*>& OutMatchingActors) const override;
+	virtual const FSlateIcon& GetIcon() const override;
+	virtual const FText& GetDisplayName() const override;
+
+#if WITH_EDITOR
+	virtual TArray<TSharedRef<SWidget>> GetWidgets() override;
+#endif
+
+public:
+	/**
+	 * A component on the actor must have one or more of the specified tags to be a match. Not case sensitive. One tag per line. Wildcards ("?" and "*") are supported but not required.
+	 * The "*" wildcard matches zero or more characters, and "?" matches exactly one character (and that character must be present).
+	 * 
+	 * Wildcard examples:
+	 * Foo* would match Foo, FooBar, and FooBaz, but not BarFoo.
+	 * *Foo* would match the above in addition to BarFoo.
+	 * Foo?Bar would match Foo.Bar and Foo_Bar, but not FooBar.
+	 * Foo? would match Food, but not FooBar or BarFoo.
+	 * Foo??? would match FooBar and FooBaz, but not Foo or Food.
+	 * ?oo? would match Food, but not Foo.
+	 * ?Foo* would match AFooBar, but not FooBar 
+	 */
+	UPROPERTY(EditAnywhere, Category="General")
+	FString TagsToMatch;
+};
+
+/** Query type which filters actors via the components contained in them. */
+UCLASS(BlueprintType)
+class MOVIERENDERPIPELINECORE_API UMovieGraphConditionGroupQuery_ComponentType final : public UMovieGraphConditionGroupQueryBase
+{
+	GENERATED_BODY()
+
+public:
+	virtual void Evaluate(const TArray<AActor*>& InActorsToQuery, const UWorld* InWorld, TSet<AActor*>& OutMatchingActors) const override;
+	virtual const FSlateIcon& GetIcon() const override;
+	virtual const FText& GetDisplayName() const override;
+
+#if WITH_EDITOR
+	virtual TArray<TSharedRef<SWidget>> GetWidgets() override;
+	virtual bool HasAddMenu() const override;
+	virtual TSharedRef<SWidget> GetAddMenuContents(const FMovieGraphConditionGroupQueryContentsChanged& OnAddFinished) override;
+#endif
+
+private:
+#if WITH_EDITOR
+	static const FSlateBrush* GetRowIcon(UClass* InComponentType);
+	static FText GetRowText(UClass* InComponentType);
+
+	/** Displays the component types which have been chosen. */
+	TSharedPtr<SMovieGraphSimpleList<UClass*>> ComponentTypesList;
+#endif
+	
+	/** The class viewer widget to show in the Add menu. */
+	TSharedPtr<class SClassViewer> ClassViewerWidget;
+
+public:
+	/** The actor must have one or more of the component type(s) in order to be a match. */
+	UPROPERTY(EditAnywhere, Category="General")
+	TArray<UClass*> ComponentTypes;
+};
+
+/** Query type which filters actors via the editor folder that they're contained in. */
+UCLASS(BlueprintType)
+class MOVIERENDERPIPELINECORE_API UMovieGraphConditionGroupQuery_EditorFolder final : public UMovieGraphConditionGroupQueryBase
+{
+	GENERATED_BODY()
+
+public:
+	virtual void Evaluate(const TArray<AActor*>& InActorsToQuery, const UWorld* InWorld, TSet<AActor*>& OutMatchingActors) const override;
+	virtual const FSlateIcon& GetIcon() const override;
+	virtual const FText& GetDisplayName() const override;
+	virtual bool IsEditorOnlyQuery() const override;
+
+#if WITH_EDITOR
+	virtual TArray<TSharedRef<SWidget>> GetWidgets() override;
+	virtual bool HasAddMenu() const override;
+	virtual TSharedRef<SWidget> GetAddMenuContents(const FMovieGraphConditionGroupQueryContentsChanged& OnAddFinished) override;
+#endif
+
+private:
+#if WITH_EDITOR
+	static const FSlateBrush* GetRowIcon(FName InFolderPath);
+	static FText GetRowText(FName InFolderPath);
+
+	/** Displays the paths of folders which have been chosen. */
+	TSharedPtr<SMovieGraphSimpleList<FName>> FolderPathsList;
+
+	/** The folder browser widget to show in the Add menu. */
+	TSharedPtr<class ISceneOutliner> FolderPickerWidget;
+#endif
+
+public:
+	/** The actor must be in one of the chosen folders in order to be a match. */
+	UPROPERTY(EditAnywhere, Category="General")
+	TArray<FName> FolderPaths;
+};
+
+/** Query type which filters actors via the sublevel that they're contained in. */
+UCLASS(BlueprintType)
+class MOVIERENDERPIPELINECORE_API UMovieGraphConditionGroupQuery_Sublevel final : public UMovieGraphConditionGroupQueryBase
+{
+	GENERATED_BODY()
+
+public:
+	virtual void Evaluate(const TArray<AActor*>& InActorsToQuery, const UWorld* InWorld, TSet<AActor*>& OutMatchingActors) const override;
+	virtual const FSlateIcon& GetIcon() const override;
+	virtual const FText& GetDisplayName() const override;
+
+#if WITH_EDITOR
+	virtual TArray<TSharedRef<SWidget>> GetWidgets() override;
+	virtual bool HasAddMenu() const override;
+	virtual TSharedRef<SWidget> GetAddMenuContents(const FMovieGraphConditionGroupQueryContentsChanged& OnAddFinished) override;
+#endif
+
+private:
+#if WITH_EDITOR
+	static const FSlateBrush* GetRowIcon(TSharedPtr<TSoftObjectPtr<UWorld>> InSublevel);
+	static FText GetRowText(TSharedPtr<TSoftObjectPtr<UWorld>> InSublevel);
+
+	/** Displays the names of sublevels which have been chosen. */
+	TSharedPtr<SMovieGraphSimpleList<TSharedPtr<TSoftObjectPtr<UWorld>>>> SublevelsList;
+
+	/** Refreshes the contents of the level picker widget when called. */
+	FRefreshAssetViewDelegate RefreshLevelPicker;
+
+	// Not ideal to store a duplicate of Sublevels, but SListView requires TSharedPtr<...> as the data source, and UPROPERTY does not
+	// support TSharedPtr<...>
+	TArray<TSharedPtr<TSoftObjectPtr<UWorld>>> ListDataSource;
+#endif
+
+public:
+	/** The actor must be in one of the chosen sublevels in order to be a match. */
+	UPROPERTY(EditAnywhere, Category="General")
+	TArray<TSoftObjectPtr<UWorld>> Sublevels;
+};
+
+/** A group of queries which can be added to a collection. */
+UCLASS(BlueprintType)
+class MOVIERENDERPIPELINECORE_API UMovieGraphConditionGroup : public UObject
+{
+	GENERATED_BODY()
+
+public:
+	UMovieGraphConditionGroup();
+
+	/**
+	 * Sets how the condition group interacts with the collection. This call is ignored for the first condition group
+	 * in the collection (the first is always Union).
+	 */
+	UFUNCTION(BlueprintCallable, Category = "Settings")
+	void SetOperationType(const EMovieGraphConditionGroupOpType OperationType);
+
+	/** Gets the condition group operation type. */
+	UFUNCTION(BlueprintCallable, Category = "Settings")
+	EMovieGraphConditionGroupOpType GetOperationType() const;
+
+	/** Determines the actors that match the condition group by running the queries contained in it. */
+	UFUNCTION(BlueprintCallable, Category = "Settings")
+	TSet<AActor*> Evaluate(const UWorld* InWorld) const;
+
+	/**
+	 * Adds a new condition group query to the condition group and returns a ptr to it. The condition group owns the
+	 * created query. By default the query is added to the end, but an optional index can be provided if the query
+	 * should be placed in a specific location among the existing queries.
+	 */
+	UFUNCTION(BlueprintCallable, Category = "Settings")
+	UMovieGraphConditionGroupQueryBase* AddQuery(const TSubclassOf<UMovieGraphConditionGroupQueryBase>& InQueryType, const int32 InsertIndex = -1);
+
+	/** Gets all queries currently contained in the condition group. */
+	UFUNCTION(BlueprintCallable, Category = "Settings")
+	const TArray<UMovieGraphConditionGroupQueryBase*>& GetQueries() const;
+
+	/** Removes the specified query from the condition group if it exists. Returns true on success, else false. */
+	UFUNCTION(BlueprintCallable, Category = "Settings")
+	bool RemoveQuery(UMovieGraphConditionGroupQueryBase* InQuery);
+
+	/** Determines if this is the first condition group under the parent collection. */
+	UFUNCTION(BlueprintCallable, Category = "Settings")
+	bool IsFirstConditionGroup() const;
+
+	/**
+	 * Move the specified query to a new index within the condition group. Returns false if the query was not found or the index
+	 * specified is invalid, else true.
+	 */
+	UFUNCTION(BlueprintCallable, Category = "Settings")
+	bool MoveQueryToIndex(UMovieGraphConditionGroupQueryBase* InQuery, const int32 NewIndex);
+
+	/** Gets a persistent unique identifier for this condition group. */
+	const FGuid& GetId() const;
+
+private:
+	/** A unique identifier for this condition group. Needed in some cases because condition groups do not have names. */
+	UPROPERTY()
+	FGuid Id;
+	
+	/** The operation type that the condition group is using. */
+	UPROPERTY(EditAnywhere, Category="General")
+	EMovieGraphConditionGroupOpType OpType;
+
+	/** The queries that are contained within the condition group. */
+	UPROPERTY(EditAnywhere, Category="General")
+	TArray<TObjectPtr<UMovieGraphConditionGroupQueryBase>> Queries;
+
+	/** Persisted actor set which can be re-used for query evaluations across frames to prevent constantly re-allocating it. */
+	UPROPERTY(Transient)
+	mutable TSet<AActor*> QueryResult;
+
+	/** Persisted actor set which can be re-used for condition group evaluations across frames to prevent constantly re-allocating it. */
+	UPROPERTY(Transient)
+	mutable TSet<AActor*> EvaluationResult;
+};
+
+/** A group of actors generated by actor queries. */
+UCLASS(BlueprintType)
+class MOVIERENDERPIPELINECORE_API UMovieGraphCollection : public UObject
+{
+	GENERATED_BODY()
+
+public:
+	/** Delegate which is called when the collection name changes. */
+	DECLARE_MULTICAST_DELEGATE_OneParam(FMovieGraphCollectionNameChanged, UMovieGraphCollection*)
+	
+	UMovieGraphCollection() = default;
+	
+#if WITH_EDITOR
+	//~ Begin UObject interface
+	virtual void PostEditChangeProperty(FPropertyChangedEvent& PropertyChangedEvent) override;
+	//~ End UObject interface
+#endif
+
+	/** Sets the name of the collection as seen in the UI. */
+	UFUNCTION(BlueprintCallable, Category = "Settings")
+	void SetCollectionName(const FString& InName);
+
+	/** Gets the name of the collection as seen in the UI. */
+	UFUNCTION(BlueprintCallable, Category = "Settings")
+	const FString& GetCollectionName() const;
+
+	/**
+	 * Gets matching actors by having condition groups evaluate themselves, and performing set operations on the
+	 * condition group results (eg, union'ing condition group A and B).
+	 */
+	UFUNCTION(BlueprintCallable, Category = "Settings")
+	TSet<AActor*> Evaluate(const UWorld* InWorld) const;
+
+	/**
+	 * Adds a new condition group to the collection and returns a ptr to it. The collection owns the created
+	 * condition group.
+	 */
+	UFUNCTION(BlueprintCallable, Category = "Settings")
+	UMovieGraphConditionGroup* AddConditionGroup();
+
+	/** Gets all condition groups currently contained in the collection. */
+	UFUNCTION(BlueprintCallable, Category = "Settings")
+	const TArray<UMovieGraphConditionGroup*>& GetConditionGroups() const;
+
+	/**
+	 * Removes the specified condition group from the collection if it exists. Returns true on success, else false.
+	 * Removes all child queries that belong to this group at the same time.
+	 */
+	UFUNCTION(BlueprintCallable, Category = "Settings")
+	bool RemoveConditionGroup(UMovieGraphConditionGroup* InConditionGroup);
+
+	/**
+	 * Move the specified condition group to a new index within the collection. Returns false if the condition group was not found or the index
+	 * specified is invalid, else true.
+	 */
+	UFUNCTION(BlueprintCallable, Category = "Settings")
+	bool MoveConditionGroupToIndex(UMovieGraphConditionGroup* InConditionGroup, const int32 NewIndex);
+
+public:
+#if WITH_EDITOR
+	/** Called when the collection name changes. */
+	FMovieGraphCollectionNameChanged OnCollectionNameChangedDelegate;
+#endif
+
+private:
+	/** The display name of the collection, shown in the UI. Does not need to be unique across collections. */
+	UPROPERTY(EditAnywhere, Category="Collection")
+	FString CollectionName;
+	
+	/** The condition groups that are contained within the collection. */
+	UPROPERTY(EditAnywhere, Category="Collection")
+	TArray<TObjectPtr<UMovieGraphConditionGroup>> ConditionGroups;
+};
+
+/**
+ * Base class for providing actor modification functionality via collections.
+ */
+UCLASS(Abstract)
+class MOVIERENDERPIPELINECORE_API UMovieGraphCollectionModifier : public UObject
+{
+	GENERATED_BODY()
+
+public:
+	/** Adds a collection to the existing set of collections in this modifier. */
+	UFUNCTION(BlueprintCallable, Category = "Settings")
+	void AddCollection(UMovieGraphCollection* Collection);
+
+	/** Overwrites the existing collections with the provided array of collections. */
+	UFUNCTION(BlueprintCallable, Category = "Settings")
+	void SetCollections(const TArray<UMovieGraphCollection*> InCollections) { Collections = InCollections; }
+	
+	UFUNCTION(BlueprintCallable, Category = "Settings")
+	TArray<UMovieGraphCollection*> GetCollections() const { return Collections; }
+
+	virtual void ApplyModifier(const UWorld* World) PURE_VIRTUAL(UMovieGraphCollectionModifier::ApplyModifier, );
+	virtual void UndoModifier() PURE_VIRTUAL(UMovieGraphCollectionModifier::UndoModifier, );
+
+protected:
+	/** The collections which this modifier will operate on. */
+	UPROPERTY()
+	TArray<TObjectPtr<UMovieGraphCollection>> Collections;
+};
+
+/**
+ * Modifies actor materials.
+ */
+UCLASS(BlueprintType, DisplayName="Material")
+class MOVIERENDERPIPELINECORE_API UMovieGraphMaterialModifier : public UMovieGraphCollectionModifier
+{
+	GENERATED_BODY()
+
+public:
+	UFUNCTION(BlueprintCallable, Category = "Settings")
+	void SetMaterial(TSoftObjectPtr<UMaterialInterface> InMaterial) { Material = InMaterial; }
+
+	UFUNCTION(BlueprintCallable, Category = "Settings")
+	virtual void ApplyModifier(const UWorld* World) override;
+
+	UFUNCTION(BlueprintCallable, Category = "Settings")
+	virtual void UndoModifier() override;
+
+private:
+	typedef TTuple<int32, TSoftObjectPtr<UMaterialInterface>> FMaterialSlotAssignment;
+	typedef TMap<TSoftObjectPtr<UPrimitiveComponent>, TArray<FMaterialSlotAssignment>> FComponentToMaterialMap;
+	
+	/** Maps a component to its original material assignments (per index). */
+	FComponentToMaterialMap ModifiedComponents;
+
+public:
+	UPROPERTY(EditAnywhere, BlueprintReadWrite, Category = Overrides, meta = (InlineEditConditionToggle))
+	uint8 bOverride_Material : 1;
+	
+	UPROPERTY(BlueprintReadWrite, EditAnywhere, Category = "Settings", meta = (EditCondition = "bOverride_Material"))
+	TSoftObjectPtr<UMaterialInterface> Material;
+};
+
+/**
+ * Modifies actor visibility.
+ */
+UCLASS(BlueprintType, DisplayName="Visibility")
+class MOVIERENDERPIPELINECORE_API UMovieGraphRenderPropertyModifier : public UMovieGraphCollectionModifier
+{
+	GENERATED_BODY()
+
+public:
+	UMovieGraphRenderPropertyModifier();
+	
+	UFUNCTION(BlueprintCallable, Category = "Settings")
+	void SetHidden(const bool bInIsHidden) { bIsHidden = bInIsHidden; }
+
+	UFUNCTION(BlueprintCallable, Category = "Settings")
+	bool IsHidden() const { return bIsHidden; }
+
+	UFUNCTION(BlueprintCallable, Category = "Settings")
+	virtual void ApplyModifier(const UWorld* World) override;
+
+	UFUNCTION(BlueprintCallable, Category = "Settings")
+	virtual void UndoModifier() override;
+
+private:
+	/** Various visibility properties for an actor. */
+	struct FActorVisibilityState
+	{
+		struct FComponentState
+		{
+			TSoftObjectPtr<USceneComponent> Component = nullptr;
+			
+			// Note: The default values specified here reflect the defaults on the scene component. If a modifier property is marked as overridden, the
+			// override will initially be a no-op due to the defaults being the same.
+			uint8 bCastsShadows : 1 = true;
+			uint8 bCastShadowWhileHidden : 1 = false;
+			uint8 bAffectIndirectLightingWhileHidden : 1 = false;
+			uint8 bHoldout : 1 = false;
+		};
+
+		TSoftObjectPtr<AActor> Actor = nullptr;
+		TArray<FComponentState> Components;
+
+		uint8 bIsHidden : 1 = false;
+	};
+
+	/** Updates an actor's visibility state to the state contained in NewVisibilityState. */
+	void SetActorVisibilityState(const FActorVisibilityState& NewVisibilityState);
+
+private:
+	/** Tracks actor visibility state prior to having the modifier applied. */
+	TArray<FActorVisibilityState> ModifiedActors;
+
+public:
+	UPROPERTY(EditAnywhere, BlueprintReadWrite, Category = Overrides, meta = (InlineEditConditionToggle))
+	uint8 bOverride_bIsHidden : 1;
+
+	UPROPERTY(EditAnywhere, BlueprintReadWrite, Category = Overrides, meta = (InlineEditConditionToggle))
+	uint8 bOverride_bCastsShadows : 1;
+
+	UPROPERTY(EditAnywhere, BlueprintReadWrite, Category = Overrides, meta = (InlineEditConditionToggle))
+	uint8 bOverride_bCastShadowWhileHidden : 1;
+
+	UPROPERTY(EditAnywhere, BlueprintReadWrite, Category = Overrides, meta = (InlineEditConditionToggle))
+	uint8 bOverride_bAffectIndirectLightingWhileHidden : 1;
+
+	UPROPERTY(EditAnywhere, BlueprintReadWrite, Category = Overrides, meta = (InlineEditConditionToggle))
+	uint8 bOverride_bHoldout : 1;
+	
+	/**
+	 * If true, the actor will not be visible and will not contribute to any secondary effects (shadows, indirect
+	 * lighting) unless their respective flags are set below.
+	 */
+	UPROPERTY(BlueprintReadWrite, EditAnywhere, Category = "Settings", meta = (EditCondition = "bOverride_bIsHidden"))
+	uint8 bIsHidden : 1;
+	
+	/** If true, the primitive will cast shadows. */
+	UPROPERTY(BlueprintReadWrite, EditAnywhere, Category = "Settings", meta = (EditCondition = "bOverride_bCastsShadows"))
+	uint8 bCastsShadows : 1;
+
+	/** If true, the primitive will cast shadows even if it is hidden. */
+	UPROPERTY(BlueprintReadWrite, EditAnywhere, Category = "Settings", meta = (EditCondition = "bOverride_bCastShadowWhileHidden"))
+	uint8 bCastShadowWhileHidden : 1;
+
+	/** Controls whether the primitive should affect indirect lighting when hidden. */
+	UPROPERTY(BlueprintReadWrite, EditAnywhere, Category = "Settings", meta = (EditCondition = "bOverride_bAffectIndirectLightingWhileHidden"))
+	uint8 bAffectIndirectLightingWhileHidden : 1;
+
+	/**
+	 * If true, the primitive will render black with an alpha of 0, but all secondary effects (shadows, reflections,
+	 * indirect lighting) remain. This feature is currently only implemented in the Path Tracer.
+	 */
+	UPROPERTY(BlueprintReadWrite, EditAnywhere, Category = "Settings", meta = (EditCondition = "bOverride_bHoldout"))
+	uint8 bHoldout : 1;
+};
+
+/**
+ * Base class for providing the ability to determine if an actor matches a query.
+ */
+UCLASS(Abstract)
+class UMovieGraphCollectionQueryBase : public UObject
+{
+	GENERATED_BODY()
+
+public:
+	virtual bool DoesActorMatchQuery(const AActor* Actor) const PURE_VIRTUAL(UMovieGraphCollectionQueryBase::DoesActorMatchQuery, return false; );  
+};
+
+UENUM(BlueprintType)
+enum class EMovieGraphCollectionCommonQueryMode : uint8
+{
+	And UMETA(ToolTip = "All specifiers in the query must be true"),
+	Or UMETA(ToolTip = "At least one specifier in the query must be true")
+};
+
+/**
+ * Provides common actor querying functionality (names, tags, components, etc). These individual sub-queries can be
+ * AND'd or OR'd together (eg, matches provided names OR provided tags, vs matches provided names AND provided tags).
+ */
+UCLASS(BlueprintType)
+class UMovieGraphCollectionCommonQuery : public UMovieGraphCollectionQueryBase
+{
+	GENERATED_BODY()
+
+public:
+	// TODO: Add other common query operations (level, etc)
+	UFUNCTION(BlueprintCallable, Category = "Settings")
+	void SetActorNames(const TArray<FString>& InActorNames) { ActorNames = InActorNames; }
+
+	UFUNCTION(BlueprintCallable, Category = "Settings")
+	void SetTags(const TArray<FName>& InTags) { Tags = InTags; }
+
+	UFUNCTION(BlueprintCallable, Category = "Settings")
+	void SetComponentTypes(TArray<UClass*> InComponentTypes) { ComponentTypes = InComponentTypes; }
+
+	UFUNCTION(BlueprintCallable, Category = "Settings")
+	void SetQueryMode(const EMovieGraphCollectionCommonQueryMode InQueryMode) { QueryMode = InQueryMode; }
+
+	UFUNCTION(BlueprintCallable, Category = "Settings")
+	virtual bool DoesActorMatchQuery(const AActor* Actor) const override;
+
+public:
+	UPROPERTY(BlueprintReadWrite, EditAnywhere, Category = "General")
+	TArray<TObjectPtr<UClass>> ComponentTypes;
+
+	UPROPERTY(BlueprintReadWrite, EditAnywhere, Category = "General")
+	EMovieGraphCollectionCommonQueryMode QueryMode = EMovieGraphCollectionCommonQueryMode::And;
+
+	UPROPERTY(BlueprintReadWrite, EditAnywhere, Category = "General")
+	TArray<FString> ActorNames;
+
+	UPROPERTY(BlueprintReadWrite, EditAnywhere, Category = "General")
+	TArray<FName> Tags;
+};
+
+/** Generates a query that matches all lighting components. */
+UCLASS(BlueprintType)
+class UMovieGraphCollectionLightingQuery : public UMovieGraphCollectionQueryBase
+{
+	GENERATED_BODY()
+
+public:
+	UFUNCTION(BlueprintCallable, Category = "Settings")
+	virtual bool DoesActorMatchQuery(const AActor* Actor) const override;
+};
+
+/**
+ * Provides a means of assembling modifiers together to generate a desired view of a scene. 
+ */
+UCLASS(BlueprintType)
+class UMovieGraphRenderLayer : public UObject
+{
+	GENERATED_BODY()
+
+public:
+	UMovieGraphRenderLayer() = default;
+
+	UFUNCTION(BlueprintCallable, Category = "Settings")
+	FName GetRenderLayerName() const { return RenderLayerName; };
+
+	UFUNCTION(BlueprintCallable, Category = "Settings")
+	void SetRenderLayerName(const FName& NewName) { RenderLayerName = NewName; }
+
+	UFUNCTION(BlueprintCallable, Category = "Settings")
+	UMovieGraphCollection* GetCollectionByName(const FString& Name) const;
+
+	UFUNCTION(BlueprintCallable, Category = "Settings")
+	void AddModifier(UMovieGraphCollectionModifier* Modifier);
+
+	UFUNCTION(BlueprintCallable, Category = "Settings")
+	TArray<UMovieGraphCollectionModifier*> GetModifiers() const { return Modifiers; };
+
+	UFUNCTION(BlueprintCallable, Category = "Settings")
+	void RemoveModifier(UMovieGraphCollectionModifier* Modifier);
+
+	UFUNCTION(BlueprintCallable, Category = "Settings")
+	void Preview(const UWorld* World);
+
+	UFUNCTION(BlueprintCallable, Category = "Settings")
+	void UndoPreview(const UWorld* World);
+
+private:
+	/** The name of this render layer. */
+	UPROPERTY()
+	FName RenderLayerName;
+
+	/** The modifiers that are active when this render layer is active. */
+	UPROPERTY()
+	TArray<TObjectPtr<UMovieGraphCollectionModifier>> Modifiers;
+};
+
+/**
+ * The primary means of controlling render layers in MRQ. Render layers can be added/registered with the subsystem, then
+ * made active in order to view them. Collections and modifiers can also be viewed, but they do not need to be added to
+ * the subsystem ahead of time.
+ */
+UCLASS()
+class MOVIERENDERPIPELINECORE_API UMovieGraphRenderLayerSubsystem : public UWorldSubsystem
+{
+	GENERATED_BODY()
+
+public:
+	UMovieGraphRenderLayerSubsystem() = default;
+
+	/* Get this subsystem for a specific world. Handy for use from Python. */
+	UFUNCTION(BlueprintCallable, Category = "Settings")
+	static UMovieGraphRenderLayerSubsystem* GetFromWorld(const UWorld* World);
+
+	//~ Begin USubsystem interface
+	virtual bool ShouldCreateSubsystem(UObject* Outer) const override { return true; }
+	virtual void Initialize(FSubsystemCollectionBase& Collection) override;
+	virtual void Deinitialize() override;
+	//~ End USubsystem interface
+
+	/** Clear out all tracked render layers and collections. */
+	UFUNCTION(BlueprintCallable, Category = "Settings")
+	void Reset();
+
+	/**
+	 * Adds a render layer to the system, which can later be made active by SetActiveRenderLayer*(). Returns true
+	 * if the layer was added successfully, else false.
+	 */
+	UFUNCTION(BlueprintCallable, Category = "Settings")
+	bool AddRenderLayer(UMovieGraphRenderLayer* RenderLayer);
+
+	UFUNCTION(BlueprintCallable, Category = "Settings")
+	const TArray<UMovieGraphRenderLayer*>& GetRenderLayers() { return RenderLayers; }
+
+	UFUNCTION(BlueprintCallable, Category = "Settings")
+	void RemoveRenderLayer(const FString& RenderLayerName);
+
+	UFUNCTION(BlueprintCallable, Category = "Settings")
+	UMovieGraphRenderLayer* GetActiveRenderLayer() const { return ActiveRenderLayer; }
+
+	/** Previews the layer with the given name. The layer needs to have been registered with AddRenderLayer(). */
+	UFUNCTION(BlueprintCallable, Category = "Settings")
+	void SetActiveRenderLayerByName(const FName& RenderLayerName);
+
+	/** Previews the given layer. The layer does not need to have been registered with AddRenderLayer(). */
+	UFUNCTION(BlueprintCallable, Category = "Settings")
+	void SetActiveRenderLayerByObj(UMovieGraphRenderLayer* RenderLayer);
+
+	UFUNCTION(BlueprintCallable, Category = "Settings")
+	void ClearActiveRenderLayer();
+
+	UFUNCTION(BlueprintCallable, Category = "Settings")
+	void PreviewCollection(UMovieGraphCollection* Collection);
+
+	UFUNCTION(BlueprintCallable, Category = "Settings")
+	void ClearCollectionPreview();
+
+	UFUNCTION(BlueprintCallable, Category = "Settings")
+	void PreviewModifier(UMovieGraphCollectionModifier* Modifier);
+
+	UFUNCTION(BlueprintCallable, Category = "Settings")
+	void ClearModifierPreview();
+
+private:
+	/** Clears the render layer, collection, and modifier previews if any of them are active. */
+	void ClearAllPreviews();
+
+	/** Sets the active render layer and previews it. */
+	void SetAndPreviewRenderLayer(UMovieGraphRenderLayer* RenderLayer);
+
+private:
+	/** Render layers which have been added/registered with the subsystem. These can be found by name. */
+	UPROPERTY(Transient)
+	TArray<TObjectPtr<UMovieGraphRenderLayer>> RenderLayers;
+
+	/** The render layer that is currently being viewed/previewed. */
+	UPROPERTY(Transient)
+	TObjectPtr<UMovieGraphRenderLayer> ActiveRenderLayer = nullptr;
+
+	/** The collection that is currently being viewed/previewed. */
+	UPROPERTY(Transient)
+	TObjectPtr<UMovieGraphCollection> ActiveCollection = nullptr;
+
+	/** The modifier that is currently being viewed/previewed. */
+	UPROPERTY(Transient)
+	TObjectPtr<UMovieGraphCollectionModifier> ActiveModifier = nullptr;
+
+	/** A render layer dedicated to visualizing collections and modifiers. */
+	UPROPERTY(Transient)
+	TObjectPtr<UMovieGraphRenderLayer> VisualizationRenderLayer = nullptr;
+
+	/** Empty collection used for visualization purposes (in conjunction w/ the viz render layer). */
+	UPROPERTY(Transient)
+	TObjectPtr<UMovieGraphCollection> VisualizationEmptyCollection = nullptr;
+
+	/** A modifier used for visualization purposes (to hide the entire world). */
+	UPROPERTY(Transient)
+	TObjectPtr<UMovieGraphRenderPropertyModifier> VisualizationModifier_HideWorld = nullptr;
+
+	/** A modifier used for visualization purposes (to show collections used with the modifier). */
+	UPROPERTY(Transient)
+	TObjectPtr<UMovieGraphRenderPropertyModifier> VisualizationModifier_VisibleCollections = nullptr;
+};
diff --git a/Engine/Plugins/MovieScene/MovieRenderPipeline/Source/MovieRenderPipelineCore/Public/Graph/MoviePipelineRenderLayerSubsystem.h b/Engine/Plugins/MovieScene/MovieRenderPipeline/Source/MovieRenderPipelineCore/Public/Graph/MoviePipelineRenderLayerSubsystem.h
deleted file mode 100644
diff --git a/Engine/Plugins/MovieScene/MovieRenderPipeline/Source/MovieRenderPipelineCore/Public/Graph/Nodes/MovieGraphModifierNode.h b/Engine/Plugins/MovieScene/MovieRenderPipeline/Source/MovieRenderPipelineCore/Public/Graph/Nodes/MovieGraphModifierNode.h
--- a/Engine/Plugins/MovieScene/MovieRenderPipeline/Source/MovieRenderPipelineCore/Public/Graph/Nodes/MovieGraphModifierNode.h
+++ b/Engine/Plugins/MovieScene/MovieRenderPipeline/Source/MovieRenderPipelineCore/Public/Graph/Nodes/MovieGraphModifierNode.h
@@ -4,7 +4,7 @@
 
 #include "Graph/MovieGraphConfig.h"
 #include "Graph/MovieGraphNode.h"
-#include "Graph/MoviePipelineRenderLayerSubsystem.h"
+#include "Graph/MovieGraphRenderLayerSubsystem.h"
 
 #include "MovieGraphModifierNode.generated.h"
 
@@ -21,11 +21,11 @@ public:
 	// ~IMovieGraphTraversableObject interface
 
 private:
-	void MergeProperties(const TObjectPtr<UMoviePipelineCollectionModifier>& InDestModifier, const TObjectPtr<UMoviePipelineCollectionModifier>& InSourceModifier);
+	void MergeProperties(const TObjectPtr<UMovieGraphCollectionModifier>& InDestModifier, const TObjectPtr<UMovieGraphCollectionModifier>& InSourceModifier);
 
 public:
 	UPROPERTY()
-	TArray<TObjectPtr<UMoviePipelineCollectionModifier>> Modifiers;
+	TArray<TObjectPtr<UMovieGraphCollectionModifier>> Modifiers;
 };
 
 /** 
@@ -55,22 +55,22 @@ public:
 
 	/** Gets the modifier of the specified type, or nullptr if one does not exist on this node. */
 	UFUNCTION(BlueprintCallable, Category = "Modifiers")
-	UMoviePipelineCollectionModifier* GetModifier(TSubclassOf<UMoviePipelineCollectionModifier> ModifierType) const;
+	UMovieGraphCollectionModifier* GetModifier(TSubclassOf<UMovieGraphCollectionModifier> ModifierType) const;
 
 	/** Gets all modifiers currently added to the node. */
 	UFUNCTION(BlueprintCallable, Category = "Modifiers")
-	const TArray<UMoviePipelineCollectionModifier*>& GetModifiers() const;
+	const TArray<UMovieGraphCollectionModifier*>& GetModifiers() const;
 
 	/**
 	 * Adds a new modifier of the specified type. Returns a pointer to the new modifier, or nullptr if a modifier of the specified type already
 	 * exists on this node (only one modifier of each type can be added to the node).
 	 */
 	UFUNCTION(BlueprintCallable, Category = "Modifiers")
-	UMoviePipelineCollectionModifier* AddModifier(TSubclassOf<UMoviePipelineCollectionModifier> ModifierType);
+	UMovieGraphCollectionModifier* AddModifier(TSubclassOf<UMovieGraphCollectionModifier> ModifierType);
 
 	/** Removes the modifier of the specified type. Returns true on success, or false if a modifier of the specified type does not exist on the node. */
 	UFUNCTION(BlueprintCallable, Category = "Modifiers")
-	bool RemoveModifier(TSubclassOf<UMoviePipelineCollectionModifier> ModifierType);
+	bool RemoveModifier(TSubclassOf<UMovieGraphCollectionModifier> ModifierType);
 
 	/** Add a collection identified by the given name which will be affected by the modifiers on this node. */
 	UFUNCTION(BlueprintCallable, Category = "Modifiers")
diff --git a/Engine/Plugins/MovieScene/MovieRenderPipeline/Source/MovieRenderPipelineEditor/Private/Customizations/Graph/MovieGraphCollectionsCustomization.cpp b/Engine/Plugins/MovieScene/MovieRenderPipeline/Source/MovieRenderPipelineEditor/Private/Customizations/Graph/MovieGraphCollectionsCustomization.cpp
--- a/Engine/Plugins/MovieScene/MovieRenderPipeline/Source/MovieRenderPipelineEditor/Private/Customizations/Graph/MovieGraphCollectionsCustomization.cpp
+++ b/Engine/Plugins/MovieScene/MovieRenderPipeline/Source/MovieRenderPipelineEditor/Private/Customizations/Graph/MovieGraphCollectionsCustomization.cpp
@@ -7,7 +7,7 @@
 #include "DetailWidgetRow.h"
 #include "Framework/Application/SlateApplication.h"
 #include "Framework/MultiBox/MultiBoxBuilder.h"
-#include "Graph/MoviePipelineRenderLayerSubsystem.h"
+#include "Graph/MovieGraphRenderLayerSubsystem.h"
 #include "Graph/Nodes/MovieGraphCollectionNode.h"
 #include "IDetailChildrenBuilder.h"
 #include "MovieRenderPipelineStyle.h"
diff --git a/Engine/Plugins/MovieScene/MovieRenderPipeline/Source/MovieRenderPipelineEditor/Private/Customizations/Graph/MovieGraphModifiersCustomization.cpp b/Engine/Plugins/MovieScene/MovieRenderPipeline/Source/MovieRenderPipelineEditor/Private/Customizations/Graph/MovieGraphModifiersCustomization.cpp
--- a/Engine/Plugins/MovieScene/MovieRenderPipeline/Source/MovieRenderPipelineEditor/Private/Customizations/Graph/MovieGraphModifiersCustomization.cpp
+++ b/Engine/Plugins/MovieScene/MovieRenderPipeline/Source/MovieRenderPipelineEditor/Private/Customizations/Graph/MovieGraphModifiersCustomization.cpp
@@ -7,7 +7,7 @@
 #include "DetailWidgetRow.h"
 #include "Graph/MovieGraphConfig.h"
 #include "Graph/MovieGraphSharedWidgets.h"
-#include "Graph/MoviePipelineRenderLayerSubsystem.h"
+#include "Graph/MovieGraphRenderLayerSubsystem.h"
 #include "Graph/Nodes/MovieGraphCollectionNode.h"
 #include "Graph/Nodes/MovieGraphModifierNode.h"
 #include "Widgets/Images/SLayeredImage.h"
@@ -297,7 +297,7 @@ void FMovieGraphModifiersCustomization::CustomizeDetails(IDetailLayoutBuilder& I
 	];
 
 	// For all modifiers added to the node, add a category for each, and add each modifier's EditAnywhere properties to the category
-	for (UMoviePipelineCollectionModifier* Modifier : ModifierNode->GetModifiers())
+	for (UMovieGraphCollectionModifier* Modifier : ModifierNode->GetModifiers())
 	{
 		if (!Modifier)
 		{
diff --git a/Engine/Plugins/MovieScene/MovieRenderPipeline/Source/MovieRenderPipelineRenderPasses/Private/Graph/Nodes/MovieGraphImagePassBaseNode.cpp b/Engine/Plugins/MovieScene/MovieRenderPipeline/Source/MovieRenderPipelineRenderPasses/Private/Graph/Nodes/MovieGraphImagePassBaseNode.cpp
--- a/Engine/Plugins/MovieScene/MovieRenderPipeline/Source/MovieRenderPipelineRenderPasses/Private/Graph/Nodes/MovieGraphImagePassBaseNode.cpp
+++ b/Engine/Plugins/MovieScene/MovieRenderPipeline/Source/MovieRenderPipelineRenderPasses/Private/Graph/Nodes/MovieGraphImagePassBaseNode.cpp
@@ -1,7 +1,7 @@
 // Copyright Epic Games, Inc. All Rights Reserved.
 
 #include "Graph/Nodes/MovieGraphImagePassBaseNode.h"
-#include "Graph/MoviePipelineRenderLayerSubsystem.h" 
+#include "Graph/MovieGraphRenderLayerSubsystem.h" 
 
 
 static TAutoConsoleVariable<bool> CVarMoviePipelineFlushRenderingPerLayerHack(
@@ -52,8 +52,8 @@ void UMovieGraphImagePassBaseNode::RenderImpl(const FMovieGraphTraversalContext&
 {
 	for (const TUniquePtr<UE::MovieGraph::Rendering::FMovieGraphImagePassBase>& Instance : CurrentInstances)
 	{
-		UMoviePipelineRenderLayerSubsystem* LayerSubsystem =
-			Instance->GetRenderer()->GetWorld()->GetSubsystem<UMoviePipelineRenderLayerSubsystem>();
+		UMovieGraphRenderLayerSubsystem* LayerSubsystem =
+			Instance->GetRenderer()->GetWorld()->GetSubsystem<UMovieGraphRenderLayerSubsystem>();
 		
 		// Apply all modifiers in the evaluated graph
 		if (LayerSubsystem)
diff --git a/Engine/Plugins/MovieScene/MovieRenderPipeline/Source/MovieRenderPipelineRenderPasses/Private/Graph/Renderers/MovieGraphDeferredPass.cpp b/Engine/Plugins/MovieScene/MovieRenderPipeline/Source/MovieRenderPipelineRenderPasses/Private/Graph/Renderers/MovieGraphDeferredPass.cpp
--- a/Engine/Plugins/MovieScene/MovieRenderPipeline/Source/MovieRenderPipelineRenderPasses/Private/Graph/Renderers/MovieGraphDeferredPass.cpp
+++ b/Engine/Plugins/MovieScene/MovieRenderPipeline/Source/MovieRenderPipelineRenderPasses/Private/Graph/Renderers/MovieGraphDeferredPass.cpp
@@ -6,7 +6,7 @@
 #include "Graph/MovieGraphDataTypes.h"
 #include "Graph/MovieGraphDefaultRenderer.h"
 #include "Graph/MovieGraphPipeline.h"
-#include "Graph/MoviePipelineRenderLayerSubsystem.h"
+#include "Graph/MovieGraphRenderLayerSubsystem.h"
 #include "Graph/MovieGraphBlueprintLibrary.h"
 #include "MovieRenderOverlappedImage.h"
 #include "MoviePipelineSurfaceReader.h"
