MovieGraph: Fixed converting a job to a graph job not dirtying the queue asset it belonged to (thus not notifying you that you need to save it)

diff --git a/Engine/Plugins/MovieScene/MovieRenderPipeline/Source/MovieRenderPipelineEditor/Private/Widgets/SMoviePipelineQueueEditor.cpp b/Engine/Plugins/MovieScene/MovieRenderPipeline/Source/MovieRenderPipelineEditor/Private/Widgets/SMoviePipelineQueueEditor.cpp
--- a/Engine/Plugins/MovieScene/MovieRenderPipeline/Source/MovieRenderPipelineEditor/Private/Widgets/SMoviePipelineQueueEditor.cpp
+++ b/Engine/Plugins/MovieScene/MovieRenderPipeline/Source/MovieRenderPipelineEditor/Private/Widgets/SMoviePipelineQueueEditor.cpp
@@ -1016,6 +1016,8 @@ struct FMoviePipelineShotItem : IMoviePipelineQueueTreeItem
 		UMoviePipelineExecutorShot* Shot = WeakShot.Get();
 		if (Shot)
 		{
+			FScopedTransaction Transaction(LOCTEXT("PickShotClearConfig_Transaction", "Clear Configuration Asset"));
+			Shot->Modify();
 			Shot->SetShotOverrideConfiguration(nullptr);
 		}
 
@@ -1036,6 +1038,8 @@ struct FMoviePipelineShotItem : IMoviePipelineQueueTreeItem
 		const TSoftObjectPtr<UMovieGraphConfig> ProjectDefaultGraph = ProjectSettings->DefaultGraph;
 		if (const UMovieGraphConfig* DefaultGraph = ProjectDefaultGraph.LoadSynchronous())
 		{
+			FScopedTransaction Transaction(LOCTEXT("PickReplaceWithGraphAsset_Transaction", "Convert Job to Graph"));
+			Shot->Modify();
 			Shot->SetGraphPreset(DefaultGraph);
 		}
 	}
@@ -1047,6 +1051,8 @@ struct FMoviePipelineShotItem : IMoviePipelineQueueTreeItem
 			UMoviePipelineExecutorShot* Shot = WeakShot.Get();
 			if (ensureMsgf(Shot, TEXT("Could not assign new graph to shot: Shot is invalid.")))
 			{
+				FScopedTransaction Transaction(LOCTEXT("PickCreateNewGraphAsset_Transaction", "Assign Graph Configuration Asset"));
+				Shot->Modify();
 				Shot->SetGraphPreset(NewGraph);
 			}
 		}
@@ -1387,6 +1393,14 @@ TSharedRef<SWidget> SMoviePipelineQueueEditor::OnGenerateNewJobFromAssetMenu()
 
 void SMoviePipelineQueueEditor::AssignDefaultGraphPresetToJob(UMoviePipelineExecutorJob* InJob)
 {
+	if(!InJob)
+	{
+		return;
+	}
+	
+	FScopedTransaction Transaction(LOCTEXT("ConvertJobToGraphConfig_Transaction", "Convert Job to Graph Config"));
+	InJob->Modify();
+
 	const UMovieRenderPipelineProjectSettings* ProjectSettings = GetDefault<UMovieRenderPipelineProjectSettings>();
 	const TSoftObjectPtr<UMovieGraphConfig> ProjectDefaultGraph = ProjectSettings->DefaultGraph;
 	if (const UMovieGraphConfig* DefaultGraph = ProjectDefaultGraph.LoadSynchronous())
