Movie Graph:

-Fix: Cursor jumps to the end of the line in the File Name Format field

diff --git a/Engine/Plugins/MovieScene/MovieRenderPipeline/Source/MovieRenderPipelineEditor/Private/Widgets/SMoviePipelineFormatTokenAutoCompleteBox.cpp b/Engine/Plugins/MovieScene/MovieRenderPipeline/Source/MovieRenderPipelineEditor/Private/Widgets/SMoviePipelineFormatTokenAutoCompleteBox.cpp
--- a/Engine/Plugins/MovieScene/MovieRenderPipeline/Source/MovieRenderPipelineEditor/Private/Widgets/SMoviePipelineFormatTokenAutoCompleteBox.cpp
+++ b/Engine/Plugins/MovieScene/MovieRenderPipeline/Source/MovieRenderPipelineEditor/Private/Widgets/SMoviePipelineFormatTokenAutoCompleteBox.cpp
@@ -105,7 +105,10 @@ FReply SMoviePipelineFormatTokenAutoCompleteBox::OnKeyDown(const FGeometry& MyGe
 
 void SMoviePipelineFormatTokenAutoCompleteBox::SetText(const FText& InText)
 {
+	// SetText always sets the cursor location to the end of the line, so we need to cache and restore the location
+	const FTextLocation OriginalCursorLocation = TextBox->GetCursorLocation();
 	TextBox->SetText(InText);
+	TextBox->GoTo(OriginalCursorLocation);
 }
 
 void SMoviePipelineFormatTokenAutoCompleteBox::OnItemClicked(TSharedPtr<FString> Item) const
