MoviePipeline: Fixed static analysis warning.

#rnx

diff --git a/Engine/Plugins/MovieScene/MovieRenderPipeline/Source/MovieRenderPipelineCore/Private/Graph/MovieGraphConfig.cpp b/Engine/Plugins/MovieScene/MovieRenderPipeline/Source/MovieRenderPipelineCore/Private/Graph/MovieGraphConfig.cpp
--- a/Engine/Plugins/MovieScene/MovieRenderPipeline/Source/MovieRenderPipelineCore/Private/Graph/MovieGraphConfig.cpp
+++ b/Engine/Plugins/MovieScene/MovieRenderPipeline/Source/MovieRenderPipelineCore/Private/Graph/MovieGraphConfig.cpp
@@ -1240,6 +1240,12 @@ void UMovieGraphConfig::CopyOverriddenProperties(UMovieGraphNode* FromNode, UMov
 				while (!ConnectedValuePins.IsEmpty())
 				{
 					UMovieGraphPin* ConnectedValuePin = ConnectedValuePins[0];
+					if (!ensureMsgf(ConnectedValuePin, TEXT("Found an invalid pin on node '%s'."), *InputPin->Node->GetName()))
+					{
+						// Can't continue following the connection chain if an invalid pin was found
+						break;
+					}
+					
 					if (ConnectionPath.Contains(ConnectedValuePin))
 					{
 						// Recursive connection found
@@ -1249,7 +1255,7 @@ void UMovieGraphConfig::CopyOverriddenProperties(UMovieGraphNode* FromNode, UMov
 					}
 
 					// For the connected value to be used, the type must match and the node the value is originating from must be enabled
-					if (ConnectedValuePin && (ConnectedValuePin->Properties.Type == InputPin->Properties.Type) && ConnectedValuePin->Node && !ConnectedValuePin->Node->IsDisabled())
+					if ((ConnectedValuePin->Properties.Type == InputPin->Properties.Type) && ConnectedValuePin->Node && !ConnectedValuePin->Node->IsDisabled())
 					{
 						ConnectionPath.Add(ConnectedValuePin);
 					}
