MovieGraph: Added a LayerName to the FMovieGraphRenderDataIdentifier (separate from BranchName) to easily see the "display" name for an identifier without having to have a configuration available to look up the Render Layer node just to get the name.

[chad.williams checking in change for matt.hoffman]

#rnx

#virtualized

diff --git a/Engine/Plugins/MovieScene/MovieRenderPipeline/Source/MovieRenderPipelineCore/Private/Graph/MovieGraphBlueprintLibrary.cpp b/Engine/Plugins/MovieScene/MovieRenderPipeline/Source/MovieRenderPipelineCore/Private/Graph/MovieGraphBlueprintLibrary.cpp
--- a/Engine/Plugins/MovieScene/MovieRenderPipeline/Source/MovieRenderPipelineCore/Private/Graph/MovieGraphBlueprintLibrary.cpp
+++ b/Engine/Plugins/MovieScene/MovieRenderPipeline/Source/MovieRenderPipelineCore/Private/Graph/MovieGraphBlueprintLibrary.cpp
@@ -60,35 +60,19 @@ FString UMovieGraphBlueprintLibrary::ResolveFilenameFormatArguments(const FStrin
 		// Use the legacy system to add filename and metadata kvp
 		MoviePipeline::GetOutputStateFormatArgs(OutMergedFormatArgs.FilenameArguments, OutMergedFormatArgs.FileMetadata, FrameNumber, FrameNumberShot, FrameNumberRel, FrameNumberShotRel, TEXT("DummyCameraToken"), ShotName);
 
-		// Look up the Render Layer display name from the evaluated config if possible.
-		FString RenderLayerName = InParams.RenderDataIdentifier.RootBranchName.ToString();
-		if (InParams.EvaluatedConfig)
-		{
-			const bool bIncludeCDOs = false;
-			TObjectPtr<UMovieGraphRenderLayerNode> RenderLayerNode = InParams.EvaluatedConfig->GetSettingForBranch<UMovieGraphRenderLayerNode>(InParams.RenderDataIdentifier.RootBranchName, bIncludeCDOs);
-			if (RenderLayerNode)
-			{
-				RenderLayerName = RenderLayerNode->GetRenderLayerName();
-			}
-
-			TObjectPtr<UMovieGraphGlobalOutputSettingNode> OutputSettingNode = InParams.EvaluatedConfig->GetSettingForBranch<UMovieGraphGlobalOutputSettingNode>(UMovieGraphNode::GlobalsPinName, bIncludeCDOs);
-			if(OutputSettingNode)
-			{
-				bOverwriteExisting = OutputSettingNode->bOverwriteExistingOutput;
-			}
-		}
-
 		// Add on Render Data Identifier, overwriting the dummy camera name above.
 		OutMergedFormatArgs.FilenameArguments.Add(TEXT("camera_name"), InParams.RenderDataIdentifier.CameraName);
 		OutMergedFormatArgs.FilenameArguments.Add(TEXT("renderer_name"), InParams.RenderDataIdentifier.RendererName);
 		OutMergedFormatArgs.FilenameArguments.Add(TEXT("renderer_sub_name"), InParams.RenderDataIdentifier.SubResourceName);
-		OutMergedFormatArgs.FilenameArguments.Add(TEXT("layer_name"), RenderLayerName);
+		OutMergedFormatArgs.FilenameArguments.Add(TEXT("layer_name"), InParams.RenderDataIdentifier.LayerName);
+		OutMergedFormatArgs.FilenameArguments.Add(TEXT("branch_name"), InParams.RenderDataIdentifier.RootBranchName.ToString());
 
 		// TODO: Some of these are per render layer and need to be stored that way. EXRs will have the metadata for all the layers/cameras/etc. in one file.
 		OutMergedFormatArgs.FileMetadata.Add(TEXT("unreal/cameraName"), InParams.RenderDataIdentifier.CameraName);
 		OutMergedFormatArgs.FileMetadata.Add(TEXT("unreal/rendererName"), InParams.RenderDataIdentifier.RendererName);
 		OutMergedFormatArgs.FileMetadata.Add(TEXT("unreal/subResourceName"), InParams.RenderDataIdentifier.SubResourceName);
-		OutMergedFormatArgs.FileMetadata.Add(TEXT("unreal/layerName"), RenderLayerName);
+		OutMergedFormatArgs.FileMetadata.Add(TEXT("unreal/layerName"), InParams.RenderDataIdentifier.LayerName);
+		OutMergedFormatArgs.FileMetadata.Add(TEXT("unreal/branchName"), InParams.RenderDataIdentifier.RootBranchName.ToString());
 	}
 
 
diff --git a/Engine/Plugins/MovieScene/MovieRenderPipeline/Source/MovieRenderPipelineCore/Private/Graph/MovieGraphDefaultRenderer.cpp b/Engine/Plugins/MovieScene/MovieRenderPipeline/Source/MovieRenderPipelineCore/Private/Graph/MovieGraphDefaultRenderer.cpp
--- a/Engine/Plugins/MovieScene/MovieRenderPipeline/Source/MovieRenderPipelineCore/Private/Graph/MovieGraphDefaultRenderer.cpp
+++ b/Engine/Plugins/MovieScene/MovieRenderPipeline/Source/MovieRenderPipelineCore/Private/Graph/MovieGraphDefaultRenderer.cpp
@@ -37,7 +37,7 @@ void UMovieGraphDefaultRenderer::SetupRenderingPipelineForShot(UMoviePipelineExe
 		TSubclassOf<UMovieGraphRenderPassNode> ClassType;
 
 		/** Maps a named branch to the specific render pass node that is assigned to render it. */
-		TMap<FName, TWeakObjectPtr<UMovieGraphRenderPassNode>> BranchRenderers;
+		TMap<FMovieGraphRenderDataIdentifier, TWeakObjectPtr<UMovieGraphRenderPassNode>> BranchRenderers;
 	};
 
 	TArray<FMovieGraphPass> OutputPasses;
@@ -51,7 +51,6 @@ void UMovieGraphDefaultRenderer::SetupRenderingPipelineForShot(UMoviePipelineExe
 	{
 		// We follow each branch looking for Render Layer nodes to figure out what render layer this should be. We assume a render layer is named
 		// after the branch it is on, unless they specifically add a UMovieGraphRenderLayerNode to rename it.
-		FString RenderLayerName = Branch.ToString();
 		const bool bIncludeCDOs = false;
 		UMovieGraphRenderLayerNode* RenderLayerNode = EvaluatedConfig->GetSettingForBranch<UMovieGraphRenderLayerNode>(Branch, bIncludeCDOs);
 			
@@ -66,7 +65,7 @@ void UMovieGraphDefaultRenderer::SetupRenderingPipelineForShot(UMoviePipelineExe
 		TArray<UMovieGraphRenderPassNode*> Renderers = EvaluatedConfig->GetSettingsForBranch<UMovieGraphRenderPassNode>(Branch, bIncludeCDOs, bExactMatch);
 		if (RenderLayerNode && Renderers.Num() == 0)
 		{
-			UE_LOG(LogMovieRenderPipeline, Warning, TEXT("Found RenderLayer: \"%s\" but no Renderers defined."), *RenderLayerName);
+			UE_LOG(LogMovieRenderPipeline, Warning, TEXT("Found RenderLayer: \"%s\" but no Renderers defined."), *Branch.ToString());
 		}
 
 		for (UMovieGraphRenderPassNode* RenderPassNode : Renderers)
@@ -80,8 +79,12 @@ void UMovieGraphDefaultRenderer::SetupRenderingPipelineForShot(UMoviePipelineExe
 				ExistingPass = &OutputPasses.AddDefaulted_GetRef();
 				ExistingPass->ClassType = RenderPassNode->GetClass();
 			}
+
+			FMovieGraphRenderDataIdentifier Identifier;
+			Identifier.RootBranchName = Branch;
+			Identifier.LayerName = RenderLayerNode->LayerName;
 			
-			ExistingPass->BranchRenderers.Add(Branch, RenderPassNode);
+			ExistingPass->BranchRenderers.Add(Identifier, RenderPassNode);
 		}
 	}
 
@@ -96,12 +99,12 @@ void UMovieGraphDefaultRenderer::SetupRenderingPipelineForShot(UMoviePipelineExe
 		FMovieGraphRenderPassSetupData SetupData;
 		SetupData.Renderer = this;
 		//UE_LOG(LogMovieRenderPipeline, Warning, TEXT("\tRenderer Class: %s"), *Pass.ClassType->GetName());
-		for (const TTuple<FName, TWeakObjectPtr<UMovieGraphRenderPassNode>>& BranchRenderer : Pass.BranchRenderers)
+		for (const TTuple<FMovieGraphRenderDataIdentifier, TWeakObjectPtr<UMovieGraphRenderPassNode>>& BranchRenderer : Pass.BranchRenderers)
 		{
 			FMovieGraphRenderPassLayerData& LayerData = SetupData.Layers.AddDefaulted_GetRef();
-			LayerData.BranchName = BranchRenderer.Key;
+			LayerData.BranchName = BranchRenderer.Key.RootBranchName;
+			LayerData.LayerName = BranchRenderer.Key.LayerName;
 			LayerData.RenderPassNode = BranchRenderer.Value;
-
 			// UE_LOG(LogMovieRenderPipeline, Warning, TEXT("\t\tBranch Name: %s"), *LayerBranchName.ToString());
 		}
 
@@ -460,25 +463,11 @@ TArray<FMovieGraphImagePreviewData> UMovieGraphDefaultRenderer::GetPreviewData()
 { 
 	TArray<FMovieGraphImagePreviewData> Results;
 
-	// The evaluated config may not be available when the preview data is first requested
-	FString LayerName;
-	const UMovieGraphTimeStepBase* TimeStepInstance = GetOwningGraph()->GetTimeStepInstance();
-	const UMovieGraphEvaluatedConfig* EvaluatedConfig = TimeStepInstance ? TimeStepInstance->GetCalculatedTimeData().EvaluatedConfig : nullptr;
-
 	for (const TPair<UE::MovieGraph::DefaultRenderer::FMovieGraphImagePreviewDataPoolParams, TObjectPtr<UTextureRenderTarget2D>>& RenderTarget : PooledViewRenderTargets)
 	{
-		// Try to determine the layer name
-		if (EvaluatedConfig)
-		{
-			constexpr bool bIncludeCDOs = false;
-			const UMovieGraphRenderLayerNode* RenderLayerNode = EvaluatedConfig->GetSettingForBranch<UMovieGraphRenderLayerNode>(RenderTarget.Key.Identifier.RootBranchName, bIncludeCDOs);
-			LayerName = RenderLayerNode ? RenderLayerNode->GetRenderLayerName() : TEXT("");
-		}
-		
 		FMovieGraphImagePreviewData& Data = Results.AddDefaulted_GetRef();
 		Data.Identifier = RenderTarget.Key.Identifier;
 		Data.Texture = RenderTarget.Value.Get();
-		Data.LayerName = LayerName;
 	}
 
 	return Results;
diff --git a/Engine/Plugins/MovieScene/MovieRenderPipeline/Source/MovieRenderPipelineCore/Private/Graph/Nodes/MovieGraphAudioOutputNode.cpp b/Engine/Plugins/MovieScene/MovieRenderPipeline/Source/MovieRenderPipelineCore/Private/Graph/Nodes/MovieGraphAudioOutputNode.cpp
--- a/Engine/Plugins/MovieScene/MovieRenderPipeline/Source/MovieRenderPipelineCore/Private/Graph/Nodes/MovieGraphAudioOutputNode.cpp
+++ b/Engine/Plugins/MovieScene/MovieRenderPipeline/Source/MovieRenderPipelineCore/Private/Graph/Nodes/MovieGraphAudioOutputNode.cpp
@@ -151,8 +151,9 @@ void UMovieGraphAudioOutputNode::GenerateFinalAudioData(TArray<FFinalAudioData>&
 	const int32 ShotIndex = CachedPipeline->GetCurrentShotIndex();
 	const TObjectPtr<UMoviePipelineExecutorShot>& Shot = CachedPipeline->GetActiveShotList()[ShotIndex];
 	const FString SubResourceName = FString();
+	const FString LayerName = GlobalsPinNameString;
 	const FString CameraName = Shot->InnerName;
-	const FMovieGraphRenderDataIdentifier RenderDataIdentifier(GlobalsPinName, RendererName, SubResourceName, CameraName);
+	const FMovieGraphRenderDataIdentifier RenderDataIdentifier(GlobalsPinName, GlobalsPinNameString, RendererName, SubResourceName, CameraName);
 	
 	const FString FinalFilePath = GenerateOutputPath(RenderDataIdentifier, Shot);
 	
diff --git a/Engine/Plugins/MovieScene/MovieRenderPipeline/Source/MovieRenderPipelineCore/Private/Graph/Nodes/MovieGraphCommandLineEncoderNode.cpp b/Engine/Plugins/MovieScene/MovieRenderPipeline/Source/MovieRenderPipelineCore/Private/Graph/Nodes/MovieGraphCommandLineEncoderNode.cpp
--- a/Engine/Plugins/MovieScene/MovieRenderPipeline/Source/MovieRenderPipelineCore/Private/Graph/Nodes/MovieGraphCommandLineEncoderNode.cpp
+++ b/Engine/Plugins/MovieScene/MovieRenderPipeline/Source/MovieRenderPipelineCore/Private/Graph/Nodes/MovieGraphCommandLineEncoderNode.cpp
@@ -115,6 +115,7 @@ void UMovieGraphCommandLineEncoderNode::StartEncodingProcess(TArray<FMovieGraphR
 		// and not the copy that we're currently iterating through.
 		FMovieGraphRenderDataIdentifier CommandLineEncoderIdentifier;
 		CommandLineEncoderIdentifier.RootBranchName = GlobalsPinName;
+		CommandLineEncoderIdentifier.LayerName = GlobalsPinName.ToString();
 		CommandLineEncoderIdentifier.RendererName = FString(TEXT("CommandLineEncoder"));
 		for (FMovieGraphRenderOutputData& OutputData : InGeneratedData)
 		{
diff --git a/Engine/Plugins/MovieScene/MovieRenderPipeline/Source/MovieRenderPipelineCore/Private/Graph/Nodes/MovieGraphWidgetRendererBaseNode.cpp b/Engine/Plugins/MovieScene/MovieRenderPipeline/Source/MovieRenderPipelineCore/Private/Graph/Nodes/MovieGraphWidgetRendererBaseNode.cpp
--- a/Engine/Plugins/MovieScene/MovieRenderPipeline/Source/MovieRenderPipelineCore/Private/Graph/Nodes/MovieGraphWidgetRendererBaseNode.cpp
+++ b/Engine/Plugins/MovieScene/MovieRenderPipeline/Source/MovieRenderPipelineCore/Private/Graph/Nodes/MovieGraphWidgetRendererBaseNode.cpp
@@ -76,6 +76,7 @@ void UMovieGraphWidgetRendererBaseNode::FMovieGraphWidgetPass::Setup(TWeakObject
 	RenderPassNode = CastChecked<UMovieGraphWidgetRendererBaseNode>(InLayer.RenderPassNode);
 
 	RenderDataIdentifier.RootBranchName = LayerData.BranchName;
+	RenderDataIdentifier.LayerName = LayerData.LayerName;
 	RenderDataIdentifier.RendererName = RenderPassNode->GetRendererName();
 
 	// NOTE: Subclasses should specify the SubResourceName
diff --git a/Engine/Plugins/MovieScene/MovieRenderPipeline/Source/MovieRenderPipelineCore/Public/Graph/MovieGraphDataTypes.h b/Engine/Plugins/MovieScene/MovieRenderPipeline/Source/MovieRenderPipelineCore/Public/Graph/MovieGraphDataTypes.h
--- a/Engine/Plugins/MovieScene/MovieRenderPipeline/Source/MovieRenderPipelineCore/Public/Graph/MovieGraphDataTypes.h
+++ b/Engine/Plugins/MovieScene/MovieRenderPipeline/Source/MovieRenderPipelineCore/Public/Graph/MovieGraphDataTypes.h
@@ -47,10 +47,6 @@ struct FMovieGraphImagePreviewData
 	/** The identifier for the image, containing the branch name, renderer, etc. */
 	UPROPERTY(BlueprintReadOnly, Category = "Movie Graph")
 	FMovieGraphRenderDataIdentifier Identifier;
-	
-	/** The layer name associated with the render. */
-	UPROPERTY(BlueprintReadOnly, Category = "Movie Graph")
-	FString LayerName;
 };
 
 USTRUCT(BlueprintType)
@@ -282,6 +278,7 @@ protected:
 struct FMovieGraphRenderPassLayerData
 {
 	FName BranchName;
+	FString LayerName;
 	FGuid CameraIdentifier;
 	TWeakObjectPtr<class UMovieGraphRenderPassNode> RenderPassNode;
 };
@@ -439,6 +436,26 @@ namespace UE::MovieGraph
 			return BranchUseCounts.Num() > 1;
 		}
 
+		bool HasDataFromMultipleLayersWithName(const FString& InLayerName) const
+		{
+			TMap<FString, int32> Layers;
+			Layers.Reserve(ExpectedRenderPasses.Num());
+
+			for (const FMovieGraphRenderDataIdentifier& PassIdentifier : ExpectedRenderPasses)
+			{
+				int32& Count = Layers.FindOrAdd(PassIdentifier.LayerName);
+				Count++;
+			}
+
+			// For the given branch, check to see if there's more than one branch with that name.
+			if (int32* CountPtr = Layers.Find(InLayerName))
+			{
+				return *CountPtr > 1;
+			}
+
+			return false;
+		}
+
 		/** Check if the (expected) render passes come from different renderers, excluding composited ones. */
 		bool HasMultipleRendersPerBranch(const FName& InBranchName)
 		{
diff --git a/Engine/Plugins/MovieScene/MovieRenderPipeline/Source/MovieRenderPipelineCore/Public/Graph/MovieGraphRenderDataIdentifier.h b/Engine/Plugins/MovieScene/MovieRenderPipeline/Source/MovieRenderPipelineCore/Public/Graph/MovieGraphRenderDataIdentifier.h
--- a/Engine/Plugins/MovieScene/MovieRenderPipeline/Source/MovieRenderPipelineCore/Public/Graph/MovieGraphRenderDataIdentifier.h
+++ b/Engine/Plugins/MovieScene/MovieRenderPipeline/Source/MovieRenderPipelineCore/Public/Graph/MovieGraphRenderDataIdentifier.h
@@ -17,9 +17,10 @@ struct FMovieGraphRenderDataIdentifier
 	FMovieGraphRenderDataIdentifier()
 	{}
 
-	FMovieGraphRenderDataIdentifier(const FName& InRootBranchName, const FString& InRendererName,
+	FMovieGraphRenderDataIdentifier(const FName& InRootBranchName, const FString& InLayerName, const FString& InRendererName,
 		const FString& InSubRenderResourceName, const FString& InCameraName)
 		: RootBranchName(InRootBranchName)
+		, LayerName(InLayerName)
 		, RendererName(InRendererName)
 		, SubResourceName(InSubRenderResourceName)
 		, CameraName(InCameraName)
@@ -29,6 +30,7 @@ struct FMovieGraphRenderDataIdentifier
 	bool operator == (const FMovieGraphRenderDataIdentifier& InRHS) const
 	{
 		return RootBranchName == InRHS.RootBranchName && 
+			LayerName == InRHS.LayerName &&
 			RendererName == InRHS.RendererName &&
 			SubResourceName == InRHS.SubResourceName &&
 			CameraName == InRHS.CameraName;
@@ -42,14 +44,15 @@ struct FMovieGraphRenderDataIdentifier
 	friend uint32 GetTypeHash(FMovieGraphRenderDataIdentifier InIdentifier)
 	{
 		return HashCombineFast(GetTypeHash(InIdentifier.RootBranchName),
-			HashCombineFast(GetTypeHash(InIdentifier.RendererName),
-				HashCombineFast(GetTypeHash(InIdentifier.SubResourceName),
-					GetTypeHash(InIdentifier.CameraName))));
+			HashCombineFast(GetTypeHash(InIdentifier.LayerName),
+				HashCombineFast(GetTypeHash(InIdentifier.RendererName),
+					HashCombineFast(GetTypeHash(InIdentifier.SubResourceName),
+						GetTypeHash(InIdentifier.CameraName)))));
 	}
 
 	friend FString LexToString(const FMovieGraphRenderDataIdentifier InIdentifier)
 	{
-		return FString::Printf(TEXT("RootBranch: %s Renderer:%s SubResource: %s Camera: %s"), *InIdentifier.RootBranchName.ToString(), *InIdentifier.RendererName, *InIdentifier.SubResourceName, *InIdentifier.CameraName);
+		return FString::Printf(TEXT("LayerName: %s RootBranch: %s Renderer:%s SubResource: %s Camera: %s"), *InIdentifier.LayerName, *InIdentifier.RootBranchName.ToString(), *InIdentifier.RendererName, *InIdentifier.SubResourceName, *InIdentifier.CameraName);
 	}
 
 	bool IsBranchAndCameraEqual(const FMovieGraphRenderDataIdentifier& InIdentifier) const
@@ -60,25 +63,32 @@ struct FMovieGraphRenderDataIdentifier
 
 public:
 	/** 
-	* The root branch name that this render layer exists on. Actual display name comes from a UMovieGraphRenderLayerNode (if found in the branch)
-	* otherwise it falls back to just displaying RootBranchName. All of our internal lookups for branches are done based on the path we followed 
-	* out from the root so that we can handle overriding a render layer display name via regular node overriding.
-	* 
-	* Could have a value like "Globals" or a user-provided one "character", "background", etc.
+	* The root branch name from the Outputs node that this identifier is for. This is useful to 
+	* know which branch it came from, as RenderLayer is user-defined and can be redefined multiple
+	* times within one graph. 
 	*/
-	UPROPERTY(EditAnywhere, BlueprintReadWrite, Category = "Movie Pipeline")
+	UPROPERTY(EditAnywhere, BlueprintReadWrite, Category = "Movie Graph")
 	FName RootBranchName;
 
+	/**
+	* The Render Layer name (as defined by the "Render Layer" node.). This is effectively a 
+	* "display" name for identifiers. If there is no Render Layer node then this will be the 
+	* RootBranchName (so that the {render_layer} token works in the case of data from things 
+	* on the Globals branch).
+	*/ 
+	UPROPERTY(EditAnywhere, BlueprintReadWrite, Category = "Movie Graph")
+	FString LayerName;
+
 	/** Which renderer was used to produce this image ("panoramic" "deferred" "path tracer", etc.) */
-	UPROPERTY(EditAnywhere, BlueprintReadWrite, Category = "Movie Pipeline")
+	UPROPERTY(EditAnywhere, BlueprintReadWrite, Category = "Movie Graph")
 	FString RendererName;
 
 	/** A sub-resource name for the renderer (ie: "beauty", "object id", "depth", etc.) */
-	UPROPERTY(EditAnywhere, BlueprintReadWrite, Category = "Movie Pipeline")
+	UPROPERTY(EditAnywhere, BlueprintReadWrite, Category = "Movie Graph")
 	FString SubResourceName;
 
 	/** The name of the camera being used for this render. */
-	UPROPERTY(EditAnywhere, BlueprintReadWrite, Category = "Movie Pipeline")
+	UPROPERTY(EditAnywhere, BlueprintReadWrite, Category = "Movie Graph")
 	FString CameraName;
 };
 
diff --git a/Engine/Plugins/MovieScene/MovieRenderPipeline/Source/MovieRenderPipelineCore/Public/Graph/Nodes/MovieGraphExecuteScriptNode.h b/Engine/Plugins/MovieScene/MovieRenderPipeline/Source/MovieRenderPipelineCore/Public/Graph/Nodes/MovieGraphExecuteScriptNode.h
--- a/Engine/Plugins/MovieScene/MovieRenderPipeline/Source/MovieRenderPipelineCore/Public/Graph/Nodes/MovieGraphExecuteScriptNode.h
+++ b/Engine/Plugins/MovieScene/MovieRenderPipeline/Source/MovieRenderPipelineCore/Public/Graph/Nodes/MovieGraphExecuteScriptNode.h
@@ -36,7 +36,7 @@
 *        unreal.log_warning("OnJobStart")
 *        for shot in inOutputData.graph_data:
 *            for layerIdentifier in shot.render_layer_data:
-*                unreal.log("render layer: " + unreal.StringLibrary.conv_name_to_string(layerIdentifier.root_branch_name))
+*                unreal.log("render layer: " + layerIdentifier.layer_name)
 *                for file in shot.render_layer_data[layerIdentifier].file_paths:
 *                    unreal.log("file: " + file)
 * 
diff --git a/Engine/Plugins/MovieScene/MovieRenderPipeline/Source/MovieRenderPipelineRenderPasses/Private/Graph/Renderers/MovieGraphDeferredPass.cpp b/Engine/Plugins/MovieScene/MovieRenderPipeline/Source/MovieRenderPipelineRenderPasses/Private/Graph/Renderers/MovieGraphDeferredPass.cpp
--- a/Engine/Plugins/MovieScene/MovieRenderPipeline/Source/MovieRenderPipelineRenderPasses/Private/Graph/Renderers/MovieGraphDeferredPass.cpp
+++ b/Engine/Plugins/MovieScene/MovieRenderPipeline/Source/MovieRenderPipelineRenderPasses/Private/Graph/Renderers/MovieGraphDeferredPass.cpp
@@ -39,6 +39,7 @@ void FMovieGraphDeferredPass::Setup(TWeakObjectPtr<UMovieGraphDefaultRenderer> I
 
 
 	RenderDataIdentifier.RootBranchName = LayerData.BranchName;
+	RenderDataIdentifier.LayerName = LayerData.LayerName;
 	RenderDataIdentifier.RendererName = InRenderPassNode->GetRendererName();
 	RenderDataIdentifier.SubResourceName = TEXT("beauty");
 	
diff --git a/Engine/Plugins/MovieScene/MovieRenderPipeline/Source/MovieRenderPipelineRenderPasses/Private/MovieGraphImageSequenceOutputNode.cpp b/Engine/Plugins/MovieScene/MovieRenderPipeline/Source/MovieRenderPipelineRenderPasses/Private/MovieGraphImageSequenceOutputNode.cpp
--- a/Engine/Plugins/MovieScene/MovieRenderPipeline/Source/MovieRenderPipelineRenderPasses/Private/MovieGraphImageSequenceOutputNode.cpp
+++ b/Engine/Plugins/MovieScene/MovieRenderPipeline/Source/MovieRenderPipelineRenderPasses/Private/MovieGraphImageSequenceOutputNode.cpp
@@ -287,15 +287,23 @@ FString UMovieGraphImageSequenceOutputNode::CreateFileName(
 	// Generate one string that puts the directory combined with the filename format.
 	FString FileNameFormatString = OutputSettingNode->OutputDirectory.Path / InParentNode->FileNameFormat;
 
+	// These two checks seem slightly backwards, but are the correct way to handle it. If there is
+	// multiple branches with data, we try to separate them by layer name (which is likely to be unique).
+	// Below we do the same check again in reverse - if there are multiple layers that share this name,
+	// then we need to insert the branch_name to differentiate them.
+	// ToDo: This is overly protective and could be relaxed later, for instance
+	// if different file write nodes have chosen a separate filepath entirely.
 	if (InRawFrameData->HasDataFromMultipleBranches())
 	{
-		// There should be one render layer per branch, so with multiple branches
-		// we test for {layer_name} presence.
-		// 
-		// ToDo: This is overly protective and could be relaxed later, for instance
-		// if different file write nodes have chosen a separate filepath entirely.
 		UE::MoviePipeline::ConformOutputFormatStringToken(FileNameFormatString, TEXT("{layer_name}"), InParentNode->GetFName(), InRenderData.Key.RootBranchName);
 	}
+	// We can run into the scenario where the users have given layers the same name, so layer_name token won't help differentiate.
+	// To resolve this, we do the reverse search now - look to see if there's multiple branches with the same layer name, and if so
+	// we force the branch name into the token too.
+	if (InRawFrameData->HasDataFromMultipleLayersWithName(InRenderData.Key.LayerName))
+	{
+		UE::MoviePipeline::ConformOutputFormatStringToken(FileNameFormatString, TEXT("{branch_name}"), InParentNode->GetFName(), InRenderData.Key.RootBranchName);
+	}
 
 	if (InRawFrameData->HasMultipleRendersPerBranch(InRenderData.Key.RootBranchName))
 	{
