MoviePipeline: Fixed issue that caused graphs to have their dirty flag set as soon as they are opened.

This was fixed by making the editor representation of the graph non-transient and saved as part of the runtime graph. Before, the editor representation was re-created when the runtime graph was loaded, thus dirtying the graph as soon as it was shown in the editor.

Most of the changes here, aside from making the editor representation non-transient, all revolve around making sure delegates get registered upon *initial* object creation (ie, in a new, non-loaded graph) AND when the object is *loaded* (ie, from an existing graph). This was not a concern before because the graph was always re-created upon load.

#rnx

diff --git a/Engine/Plugins/MovieScene/MovieRenderPipeline/Source/MovieRenderPipelineCore/Private/Graph/Nodes/MovieGraphCollectionNode.cpp b/Engine/Plugins/MovieScene/MovieRenderPipeline/Source/MovieRenderPipelineCore/Private/Graph/Nodes/MovieGraphCollectionNode.cpp
--- a/Engine/Plugins/MovieScene/MovieRenderPipeline/Source/MovieRenderPipelineCore/Private/Graph/Nodes/MovieGraphCollectionNode.cpp
+++ b/Engine/Plugins/MovieScene/MovieRenderPipeline/Source/MovieRenderPipelineCore/Private/Graph/Nodes/MovieGraphCollectionNode.cpp
@@ -12,12 +12,7 @@ UMovieGraphCollectionNode::UMovieGraphCollectionNode()
 {
 	Collection = CreateDefaultSubobject<UMovieGraphCollection>(TEXT("Collection"));
 
-#if WITH_EDITOR
-	Collection->OnCollectionNameChangedDelegate.AddWeakLambda(this, [this](UMovieGraphCollection* ChangedCollection)
-	{
-		OnNodeChangedDelegate.Broadcast(this);
-	});
-#endif
+	RegisterDelegates();
 }
 
 #if WITH_EDITOR
@@ -26,9 +21,13 @@ FText UMovieGraphCollectionNode::GetNodeTitle(const bool bGetDescriptive) const
 	static const FText CollectionNodeName = LOCTEXT("NodeName_Collection", "Collection");
 	static const FText CollectionNodeDescription = LOCTEXT("NodeDescription_Collection", "Collection\n{0}");
 
-	if (bGetDescriptive && Collection && !Collection->GetCollectionName().IsEmpty())
+	const FString CollectionDisplayName = (Collection && !Collection->GetCollectionName().IsEmpty())
+		? Collection->GetCollectionName()
+		: LOCTEXT("NodeNoNameWarning_Collection", "NO NAME").ToString();
+
+	if (bGetDescriptive)
 	{
-		return FText::Format(CollectionNodeDescription, FText::FromString(Collection->GetCollectionName()));
+		return FText::Format(CollectionNodeDescription, FText::FromString(CollectionDisplayName));
 	}
 
 	return CollectionNodeName;
@@ -59,4 +58,19 @@ FString UMovieGraphCollectionNode::GetNodeInstanceName() const
 	return Collection ? Collection->GetCollectionName() : FString();
 }
 
+void UMovieGraphCollectionNode::RegisterDelegates()
+{
+	Super::RegisterDelegates();
+
+#if WITH_EDITOR
+	if (Collection)
+	{
+		Collection->OnCollectionNameChangedDelegate.AddWeakLambda(this, [this](UMovieGraphCollection* ChangedCollection)
+		{
+			OnNodeChangedDelegate.Broadcast(this);
+		});
+	}
+#endif
+}
+
 #undef LOCTEXT_NAMESPACE // "MovieGraph"
\ No newline at end of file
diff --git a/Engine/Plugins/MovieScene/MovieRenderPipeline/Source/MovieRenderPipelineCore/Private/Graph/Nodes/MovieGraphInputNode.cpp b/Engine/Plugins/MovieScene/MovieRenderPipeline/Source/MovieRenderPipelineCore/Private/Graph/Nodes/MovieGraphInputNode.cpp
--- a/Engine/Plugins/MovieScene/MovieRenderPipeline/Source/MovieRenderPipelineCore/Private/Graph/Nodes/MovieGraphInputNode.cpp
+++ b/Engine/Plugins/MovieScene/MovieRenderPipeline/Source/MovieRenderPipelineCore/Private/Graph/Nodes/MovieGraphInputNode.cpp
@@ -8,16 +8,10 @@
 
 UMovieGraphInputNode::UMovieGraphInputNode()
 {
-#if WITH_EDITOR
 	if (!HasAnyFlags(RF_ClassDefaultObject))
 	{
-		if (UMovieGraphConfig* Graph = GetGraph())
-		{
-			// Register delegates for new inputs when they're added to the graph
-			Graph->OnGraphInputAddedDelegate.AddUObject(this, &UMovieGraphInputNode::RegisterDelegates);
-		}
+		RegisterDelegates();
 	}
-#endif
 }
 
 TArray<FMovieGraphPinProperties> UMovieGraphInputNode::GetOutputPinProperties() const
@@ -130,16 +124,23 @@ void UMovieGraphInputNode::RegisterDelegates()
 {
 	Super::RegisterDelegates();
 	
-	if (const UMovieGraphConfig* Graph = GetGraph())
+	if (UMovieGraphConfig* Graph = GetGraph())
 	{
+#if WITH_EDITOR
+		// Register delegates for new inputs when they're added to the graph
+		Graph->OnGraphInputAddedDelegate.RemoveAll(this);
+		Graph->OnGraphInputAddedDelegate.AddUObject(this, &UMovieGraphInputNode::RegisterInputDelegates);
+#endif
+
+		// Register delegates for existing inputs
 		for (UMovieGraphInput* InputMember : Graph->GetInputs())
 		{
-			RegisterDelegates(InputMember);
+			RegisterInputDelegates(InputMember);
 		}
 	}
 }
 
-void UMovieGraphInputNode::RegisterDelegates(UMovieGraphInput* Input)
+void UMovieGraphInputNode::RegisterInputDelegates(UMovieGraphInput* Input)
 {
 #if WITH_EDITOR
 	if (Input)
diff --git a/Engine/Plugins/MovieScene/MovieRenderPipeline/Source/MovieRenderPipelineCore/Private/Graph/Nodes/MovieGraphModifierNode.cpp b/Engine/Plugins/MovieScene/MovieRenderPipeline/Source/MovieRenderPipelineCore/Private/Graph/Nodes/MovieGraphModifierNode.cpp
--- a/Engine/Plugins/MovieScene/MovieRenderPipeline/Source/MovieRenderPipelineCore/Private/Graph/Nodes/MovieGraphModifierNode.cpp
+++ b/Engine/Plugins/MovieScene/MovieRenderPipeline/Source/MovieRenderPipelineCore/Private/Graph/Nodes/MovieGraphModifierNode.cpp
@@ -109,7 +109,7 @@ FText UMovieGraphModifierNode::GetNodeTitle(const bool bGetDescriptive) const
 	static const FText ModifierNodeName = LOCTEXT("NodeName_Modifier", "Modifier");
 	static const FText ModifierNodeDescription = LOCTEXT("NodeDescription_Modifier", "Modifier\n{0}");
 	
-	const FString ModifierNameDisp = ModifierName.IsEmpty() ? TEXT("MISSING") : ModifierName;
+	const FString ModifierNameDisp = ModifierName.IsEmpty() ? LOCTEXT("NodeNoNameWarning_Modifier", "NO NAME").ToString() : ModifierName;
 
 	if (bGetDescriptive)
 	{
diff --git a/Engine/Plugins/MovieScene/MovieRenderPipeline/Source/MovieRenderPipelineCore/Private/Graph/Nodes/MovieGraphOutputNode.cpp b/Engine/Plugins/MovieScene/MovieRenderPipeline/Source/MovieRenderPipelineCore/Private/Graph/Nodes/MovieGraphOutputNode.cpp
--- a/Engine/Plugins/MovieScene/MovieRenderPipeline/Source/MovieRenderPipelineCore/Private/Graph/Nodes/MovieGraphOutputNode.cpp
+++ b/Engine/Plugins/MovieScene/MovieRenderPipeline/Source/MovieRenderPipelineCore/Private/Graph/Nodes/MovieGraphOutputNode.cpp
@@ -7,16 +7,10 @@
 
 UMovieGraphOutputNode::UMovieGraphOutputNode()
 {
-#if WITH_EDITOR
 	if (!HasAnyFlags(RF_ClassDefaultObject))
 	{
-		if (UMovieGraphConfig* Graph = GetGraph())
-		{
-			// Register delegates for new outputs when they're added to the graph
-			Graph->OnGraphOutputAddedDelegate.AddUObject(this, &UMovieGraphOutputNode::RegisterDelegates);
-		}
+		RegisterDelegates();
 	}
-#endif
 }
 
 TArray<FMovieGraphPinProperties> UMovieGraphOutputNode::GetInputPinProperties() const
@@ -70,16 +64,23 @@ void UMovieGraphOutputNode::RegisterDelegates()
 {
 	Super::RegisterDelegates();
 	
-	if (const UMovieGraphConfig* Graph = GetGraph())
+	if (UMovieGraphConfig* Graph = GetGraph())
 	{
+#if WITH_EDITOR
+		// Register delegates for new outputs when they're added to the graph
+		Graph->OnGraphOutputAddedDelegate.RemoveAll(this);
+		Graph->OnGraphOutputAddedDelegate.AddUObject(this, &UMovieGraphOutputNode::RegisterOutputDelegates);
+#endif
+
+		// Register delegates for existing outputs
 		for (UMovieGraphOutput* OutputMember : Graph->GetOutputs())
 		{
-			RegisterDelegates(OutputMember);
+			RegisterOutputDelegates(OutputMember);
 		}
 	}
 }
 
-void UMovieGraphOutputNode::RegisterDelegates(UMovieGraphOutput* Output)
+void UMovieGraphOutputNode::RegisterOutputDelegates(UMovieGraphOutput* Output)
 {
 #if WITH_EDITOR
 	if (Output)
diff --git a/Engine/Plugins/MovieScene/MovieRenderPipeline/Source/MovieRenderPipelineCore/Public/Graph/MovieGraphConfig.h b/Engine/Plugins/MovieScene/MovieRenderPipeline/Source/MovieRenderPipelineCore/Public/Graph/MovieGraphConfig.h
--- a/Engine/Plugins/MovieScene/MovieRenderPipeline/Source/MovieRenderPipelineCore/Public/Graph/MovieGraphConfig.h
+++ b/Engine/Plugins/MovieScene/MovieRenderPipeline/Source/MovieRenderPipelineCore/Public/Graph/MovieGraphConfig.h
@@ -792,7 +792,10 @@ public:
 #if WITH_EDITORONLY_DATA
 	// Not strongly typed to avoid a circular dependency between the editor only module
 	// and the runtime module, but it should be a UMoviePipelineEdGraph.
-	UPROPERTY(Transient)
+	//
+	// Note that the editor graph is saved with the runtime graph. This is done to prevent the runtime graph from being dirtied immediately upon loading
+	// (because the editor graph would have to be re-created from the runtime graph, thus dirtying the package).
+	UPROPERTY()
 	TObjectPtr<UEdGraph> PipelineEdGraph;
 #endif
 
diff --git a/Engine/Plugins/MovieScene/MovieRenderPipeline/Source/MovieRenderPipelineCore/Public/Graph/Nodes/MovieGraphCollectionNode.h b/Engine/Plugins/MovieScene/MovieRenderPipeline/Source/MovieRenderPipelineCore/Public/Graph/Nodes/MovieGraphCollectionNode.h
--- a/Engine/Plugins/MovieScene/MovieRenderPipeline/Source/MovieRenderPipelineCore/Public/Graph/Nodes/MovieGraphCollectionNode.h
+++ b/Engine/Plugins/MovieScene/MovieRenderPipeline/Source/MovieRenderPipelineCore/Public/Graph/Nodes/MovieGraphCollectionNode.h
@@ -30,6 +30,9 @@ public:
 
 	virtual FString GetNodeInstanceName() const override;
 
+protected:
+	virtual void RegisterDelegates() override;
+
 public:
 	UPROPERTY(EditAnywhere, BlueprintReadWrite, Category = Overrides, meta = (InlineEditConditionToggle))
 	uint8 bOverride_Collection : 1 = 1;	// The collection is customized in the details panel, so the override should always be enabled
diff --git a/Engine/Plugins/MovieScene/MovieRenderPipeline/Source/MovieRenderPipelineCore/Public/Graph/Nodes/MovieGraphInputNode.h b/Engine/Plugins/MovieScene/MovieRenderPipeline/Source/MovieRenderPipelineCore/Public/Graph/Nodes/MovieGraphInputNode.h
--- a/Engine/Plugins/MovieScene/MovieRenderPipeline/Source/MovieRenderPipelineCore/Public/Graph/Nodes/MovieGraphInputNode.h
+++ b/Engine/Plugins/MovieScene/MovieRenderPipeline/Source/MovieRenderPipelineCore/Public/Graph/Nodes/MovieGraphInputNode.h
@@ -35,7 +35,7 @@ private:
 	virtual void RegisterDelegates() override;
 
 	/** Register delegates for the provided input member. */
-	void RegisterDelegates(UMovieGraphInput* Input);
+	void RegisterInputDelegates(UMovieGraphInput* Input);
 
 	/** Update data (name, etc) on all existing output pins on this node to reflect the input members on the graph. */
 	void UpdateExistingPins(UMovieGraphMember* ChangedVariable) const;
diff --git a/Engine/Plugins/MovieScene/MovieRenderPipeline/Source/MovieRenderPipelineCore/Public/Graph/Nodes/MovieGraphOutputNode.h b/Engine/Plugins/MovieScene/MovieRenderPipeline/Source/MovieRenderPipelineCore/Public/Graph/Nodes/MovieGraphOutputNode.h
--- a/Engine/Plugins/MovieScene/MovieRenderPipeline/Source/MovieRenderPipelineCore/Public/Graph/Nodes/MovieGraphOutputNode.h
+++ b/Engine/Plugins/MovieScene/MovieRenderPipeline/Source/MovieRenderPipelineCore/Public/Graph/Nodes/MovieGraphOutputNode.h
@@ -31,7 +31,7 @@ private:
 	virtual void RegisterDelegates() override;
 
 	/** Register delegates for the provided output member. */
-	void RegisterDelegates(UMovieGraphOutput* Output);
+	void RegisterOutputDelegates(UMovieGraphOutput* Output);
 
 	/** Update data (name, etc) on all existing input pins on this node to reflect the output members on the graph. */
 	void UpdateExistingPins(UMovieGraphMember* ChangedVariable) const;
diff --git a/Engine/Plugins/MovieScene/MovieRenderPipeline/Source/MovieRenderPipelineEditor/Private/Graph/MovieEdGraph.cpp b/Engine/Plugins/MovieScene/MovieRenderPipeline/Source/MovieRenderPipelineEditor/Private/Graph/MovieEdGraph.cpp
--- a/Engine/Plugins/MovieScene/MovieRenderPipeline/Source/MovieRenderPipelineEditor/Private/Graph/MovieEdGraph.cpp
+++ b/Engine/Plugins/MovieScene/MovieRenderPipeline/Source/MovieRenderPipelineEditor/Private/Graph/MovieEdGraph.cpp
@@ -82,12 +82,17 @@ void UMoviePipelineEdGraph::InitFromRuntimeGraph(UMovieGraphConfig* InGraph)
 		}
 	}
 
-	InGraph->OnGraphChangedDelegate.AddUObject(this, &UMoviePipelineEdGraph::OnGraphConfigChanged);
-	InGraph->OnGraphNodesDeletedDelegate.AddUObject(this, &UMoviePipelineEdGraph::OnGraphNodesDeleted);
+	RegisterDelegates(InGraph);
 
 	bInitialized = true;
 }
 
+void UMoviePipelineEdGraph::RegisterDelegates(UMovieGraphConfig* InGraph)
+{
+	InGraph->OnGraphChangedDelegate.AddUObject(this, &UMoviePipelineEdGraph::OnGraphConfigChanged);
+	InGraph->OnGraphNodesDeletedDelegate.AddUObject(this, &UMoviePipelineEdGraph::OnGraphNodesDeleted);
+}
+
 void UMoviePipelineEdGraph::CreateLinks(UMoviePipelineEdGraphNodeBase* InGraphNode, bool bCreateInboundLinks, bool bCreateOutboundLinks)
 {
 	check(InGraphNode);
diff --git a/Engine/Plugins/MovieScene/MovieRenderPipeline/Source/MovieRenderPipelineEditor/Private/Graph/MovieEdGraph.h b/Engine/Plugins/MovieScene/MovieRenderPipeline/Source/MovieRenderPipelineEditor/Private/Graph/MovieEdGraph.h
--- a/Engine/Plugins/MovieScene/MovieRenderPipeline/Source/MovieRenderPipelineEditor/Private/Graph/MovieEdGraph.h
+++ b/Engine/Plugins/MovieScene/MovieRenderPipeline/Source/MovieRenderPipelineEditor/Private/Graph/MovieEdGraph.h
@@ -20,10 +20,12 @@ class UMoviePipelineEdGraph : public UEdGraph
 	GENERATED_BODY()
 
 public:
-
 	/** Initialize this Editor Graph from a Runtime Graph */
 	void InitFromRuntimeGraph(UMovieGraphConfig* InGraph);
 
+	/** Register delegates that relate to the runtime graph. */
+	void RegisterDelegates(UMovieGraphConfig* InGraph);
+
 	/** Returns the runtime UMovieGraphConfig that contains this editor graph */
 	class UMovieGraphConfig* GetPipelineGraph() const;
 
diff --git a/Engine/Plugins/MovieScene/MovieRenderPipeline/Source/MovieRenderPipelineEditor/Private/Graph/MovieEdGraphNode.cpp b/Engine/Plugins/MovieScene/MovieRenderPipeline/Source/MovieRenderPipelineEditor/Private/Graph/MovieEdGraphNode.cpp
--- a/Engine/Plugins/MovieScene/MovieRenderPipeline/Source/MovieRenderPipelineEditor/Private/Graph/MovieEdGraphNode.cpp
+++ b/Engine/Plugins/MovieScene/MovieRenderPipeline/Source/MovieRenderPipelineEditor/Private/Graph/MovieEdGraphNode.cpp
@@ -10,7 +10,6 @@
 #include "Misc/TransactionObjectEvent.h"
 #include "MovieEdGraph.h"
 #include "MovieGraphSchema.h"
-#include "PropertyBag.h"
 #include "ToolMenu.h"
 #include "EdGraph/EdGraphSchema.h"
 #include "GraphEditorActions.h"
@@ -24,7 +23,6 @@ void UMoviePipelineEdGraphNodeBase::Construct(UMovieGraphNode* InRuntimeNode)
 	check(InRuntimeNode);
 	RuntimeNode = InRuntimeNode;
 	RuntimeNode->GraphNode = this;
-	RuntimeNode->OnNodeChangedDelegate.AddUObject(this, &UMoviePipelineEdGraphNodeBase::OnRuntimeNodeChanged);
 	
 	NodePosX = InRuntimeNode->GetNodePosX();
 	NodePosY = InRuntimeNode->GetNodePosY();
@@ -33,6 +31,7 @@ void UMoviePipelineEdGraphNodeBase::Construct(UMovieGraphNode* InRuntimeNode)
 	bCommentBubblePinned = InRuntimeNode->IsCommentBubblePinned();
 	bCommentBubbleVisible = InRuntimeNode->IsCommentBubbleVisible();
 	
+	RegisterDelegates();
 	SetEnabledState(InRuntimeNode->IsDisabled() ? ENodeEnabledState::Disabled : ENodeEnabledState::Enabled);
 }
 
@@ -218,6 +217,14 @@ void UMoviePipelineEdGraphNodeBase::UpdateEnableState() const
 	}
 }
 
+void UMoviePipelineEdGraphNodeBase::RegisterDelegates()
+{
+	if (RuntimeNode)
+	{
+		RuntimeNode->OnNodeChangedDelegate.AddUObject(this, &UMoviePipelineEdGraphNodeBase::OnRuntimeNodeChanged);
+	}
+}
+
 void UMoviePipelineEdGraphNode::AllocateDefaultPins()
 {
 	if (RuntimeNode)
@@ -520,6 +527,13 @@ void UMoviePipelineEdGraphNodeBase::OnRuntimeNodeChanged(const UMovieGraphNode*
 	}
 }
 
+void UMoviePipelineEdGraphNodeBase::PostLoad()
+{
+	Super::PostLoad();
+
+	RegisterDelegates();	
+}
+
 void UMoviePipelineEdGraphNodeBase::ReconstructNode()
 {
 	// Don't reconstruct the node during copy-paste. If we allow reconstruction,
diff --git a/Engine/Plugins/MovieScene/MovieRenderPipeline/Source/MovieRenderPipelineEditor/Private/Graph/MovieEdGraphNode.h b/Engine/Plugins/MovieScene/MovieRenderPipeline/Source/MovieRenderPipelineEditor/Private/Graph/MovieEdGraphNode.h
--- a/Engine/Plugins/MovieScene/MovieRenderPipeline/Source/MovieRenderPipelineEditor/Private/Graph/MovieEdGraphNode.h
+++ b/Engine/Plugins/MovieScene/MovieRenderPipeline/Source/MovieRenderPipelineEditor/Private/Graph/MovieEdGraphNode.h
@@ -21,6 +21,7 @@ public:
 	void OnRuntimeNodeChanged(const UMovieGraphNode* InChangedNode);
 
 	//~ Begin UEdGraphNode Interface
+	virtual void PostLoad() override;
 	virtual void ReconstructNode() override;
 	virtual void AutowireNewNode(UEdGraphPin* FromPin) override;
 	virtual FLinearColor GetNodeTitleColor() const override;
@@ -71,6 +72,9 @@ protected:
 	 */
 	void UpdateEnableState() const;
 
+	/** Registers any delegates needed by the node. */
+	void RegisterDelegates();
+
 protected:
 	/** The runtime node that this editor node represents. */
 	UPROPERTY()
diff --git a/Engine/Plugins/MovieScene/MovieRenderPipeline/Source/MovieRenderPipelineEditor/Private/Widgets/Graph/SMovieGraphConfigPanel.cpp b/Engine/Plugins/MovieScene/MovieRenderPipeline/Source/MovieRenderPipelineEditor/Private/Widgets/Graph/SMovieGraphConfigPanel.cpp
--- a/Engine/Plugins/MovieScene/MovieRenderPipeline/Source/MovieRenderPipelineEditor/Private/Widgets/Graph/SMovieGraphConfigPanel.cpp
+++ b/Engine/Plugins/MovieScene/MovieRenderPipeline/Source/MovieRenderPipelineEditor/Private/Widgets/Graph/SMovieGraphConfigPanel.cpp
@@ -103,23 +103,34 @@ void SMoviePipelineGraphPanel::Construct(const FArguments& InArgs)
 
 	CurrentGraph = InArgs._Graph;
 	
-	UMoviePipelineEdGraph* GraphToEdit = Cast<UMoviePipelineEdGraph>(FBlueprintEditorUtils::CreateNewGraph(CurrentGraph, TEXT("MoviePipelineEdGraph"), UMoviePipelineEdGraph::StaticClass(), UMovieGraphSchema::StaticClass()));
+	UMoviePipelineEdGraph* EdGraph = Cast<UMoviePipelineEdGraph>(CurrentGraph->PipelineEdGraph);
 
-	// Probably not ideal.. USoundCue has a CreateGraph() node that does this (#if WITH_EDITOR) but then requires an interface
-	// for the editor half to avoid the circular dependency.
-	CurrentGraph->PipelineEdGraph = GraphToEdit;
-	GraphToEdit->InitFromRuntimeGraph(CurrentGraph);
+	// Create the EdGraph if it has not yet been created. It is saved as part of the runtime graph to prevent it from being re-created every time the
+	// graph is opened (and therefore dirtying the package).
+	if (!EdGraph)
+	{
+		EdGraph = Cast<UMoviePipelineEdGraph>(FBlueprintEditorUtils::CreateNewGraph(CurrentGraph, TEXT("MoviePipelineEdGraph"), UMoviePipelineEdGraph::StaticClass(), UMovieGraphSchema::StaticClass()));
 
+		// Probably not ideal.. USoundCue has a CreateGraph() node that does this (#if WITH_EDITOR) but then requires an interface
+		// for the editor half to avoid the circular dependency.
+		CurrentGraph->PipelineEdGraph = EdGraph;
+		EdGraph->InitFromRuntimeGraph(CurrentGraph);
 
-	const UEdGraphSchema* Schema = GraphToEdit->GetSchema();
-	Schema->CreateDefaultNodesForGraph(*GraphToEdit);
+		const UEdGraphSchema* Schema = EdGraph->GetSchema();
+		Schema->CreateDefaultNodesForGraph(*EdGraph);
+	}
+	else
+	{
+		EdGraph->RegisterDelegates(CurrentGraph);
+	}
+	
 	MakeEditorCommands();
 
 	ChildSlot
 	[
 		SAssignNew(GraphEditorWidget, SGraphEditor)
 		.IsEditable(true)
-		.GraphToEdit(GraphToEdit)
+		.GraphToEdit(EdGraph)
 		.AdditionalCommands(GraphEditorCommands)
 		.GraphEvents(InEvents)
 		.Appearance(AppearanceInfo)
