MoviePipeline: Fixed the burn-in node from preventing all renders if no burn-in class is selected. Also defaulted the burn-in class to be the one that ships with MRG, and defaulted compositing to be turned on by default.

#rnx

diff --git a/Engine/Plugins/MovieScene/MovieRenderPipeline/Source/MovieRenderPipelineCore/Private/Graph/Nodes/MovieGraphBurnInNode.cpp b/Engine/Plugins/MovieScene/MovieRenderPipeline/Source/MovieRenderPipelineCore/Private/Graph/Nodes/MovieGraphBurnInNode.cpp
--- a/Engine/Plugins/MovieScene/MovieRenderPipeline/Source/MovieRenderPipelineCore/Private/Graph/Nodes/MovieGraphBurnInNode.cpp
+++ b/Engine/Plugins/MovieScene/MovieRenderPipeline/Source/MovieRenderPipelineCore/Private/Graph/Nodes/MovieGraphBurnInNode.cpp
@@ -12,6 +12,12 @@
 #include "TextureResource.h"
 
 const FString UMovieGraphBurnInNode::RendererName = FString("BurnIn");
+const FString UMovieGraphBurnInNode::DefaultBurnInWidgetAsset = TEXT("/MovieRenderPipeline/Blueprints/Graph/DefaultGraphBurnIn.DefaultGraphBurnIn_C");
+
+UMovieGraphBurnInNode::UMovieGraphBurnInNode()
+	: BurnInClass(DefaultBurnInWidgetAsset)
+{
+}
 
 #if WITH_EDITOR
 FText UMovieGraphBurnInNode::GetNodeTitle(const bool bGetDescriptive) const
@@ -38,6 +44,18 @@ TUniquePtr<UMovieGraphWidgetRendererBaseNode::FMovieGraphWidgetPass> UMovieGraph
 	return MakeUnique<FMovieGraphBurnInPass>();
 }
 
+void UMovieGraphBurnInNode::GatherOutputPassesImpl(UMovieGraphEvaluatedConfig* InConfig, TArray<FMovieGraphRenderDataIdentifier>& OutExpectedPasses) const
+{
+	for (const TUniquePtr<FMovieGraphWidgetPass>& Instance : CurrentInstances)
+	{
+		// Only generate passes if there's a valid burn-in class
+		if (StaticCast<FMovieGraphBurnInPass*>(Instance.Get())->GetBurnInClass())
+		{
+			Instance->GatherOutputPasses(OutExpectedPasses);
+		}
+	}
+}
+
 void UMovieGraphBurnInNode::TeardownImpl()
 {
 	Super::TeardownImpl();
@@ -82,8 +100,7 @@ TObjectPtr<UMovieGraphBurnInWidget> UMovieGraphBurnInNode::FMovieGraphBurnInPass
 {
 	const UMovieGraphPipeline* Pipeline = Renderer->GetOwningGraph();
 
-	const UMovieGraphBurnInNode* BurnInNode = CastChecked<UMovieGraphBurnInNode>(RenderPassNode);
-	UClass* LoadedBurnInClass = BurnInNode->BurnInClass.TryLoadClass<UMovieGraphBurnInWidget>();
+	UClass* LoadedBurnInClass = GetBurnInClass();
 	if (!LoadedBurnInClass)
 	{
 		UE_LOG(LogMovieRenderPipeline, Error, TEXT("The burn-in widget provided in layer '%s' for renderer '%s' is not valid."), *LayerData.BranchName.ToString(), *Renderer->GetClass()->GetName());
@@ -95,6 +112,7 @@ TObjectPtr<UMovieGraphBurnInWidget> UMovieGraphBurnInNode::FMovieGraphBurnInPass
 	const TObjectPtr<UMovieGraphBurnInWidget> BurnInWidget = BurnInCDO->GetOrCreateBurnInWidget(LoadedBurnInClass, Pipeline->GetWorld());
 	if (!BurnInWidget)
 	{
+		const UMovieGraphBurnInNode* BurnInNode = CastChecked<UMovieGraphBurnInNode>(RenderPassNode);
 		UE_LOG(LogMovieRenderPipeline, Error, TEXT("Unable to load burn-in widget at path: %s"), *BurnInNode->BurnInClass.GetAssetPath().ToString());
 		return nullptr;
 	}
@@ -118,4 +136,10 @@ int32 UMovieGraphBurnInNode::FMovieGraphBurnInPass::GetCompositingSortOrder() co
 {
 	// Burn-ins should always appear over all other passes
 	return 0;
+}
+
+UClass* UMovieGraphBurnInNode::FMovieGraphBurnInPass::GetBurnInClass() const
+{
+	const UMovieGraphBurnInNode* BurnInNode = CastChecked<UMovieGraphBurnInNode>(RenderPassNode);
+	return BurnInNode->BurnInClass.TryLoadClass<UMovieGraphBurnInWidget>();
 }
\ No newline at end of file
diff --git a/Engine/Plugins/MovieScene/MovieRenderPipeline/Source/MovieRenderPipelineCore/Private/Graph/Nodes/MovieGraphWidgetRendererBaseNode.h b/Engine/Plugins/MovieScene/MovieRenderPipeline/Source/MovieRenderPipelineCore/Private/Graph/Nodes/MovieGraphWidgetRendererBaseNode.h
--- a/Engine/Plugins/MovieScene/MovieRenderPipeline/Source/MovieRenderPipelineCore/Private/Graph/Nodes/MovieGraphWidgetRendererBaseNode.h
+++ b/Engine/Plugins/MovieScene/MovieRenderPipeline/Source/MovieRenderPipelineCore/Private/Graph/Nodes/MovieGraphWidgetRendererBaseNode.h
@@ -28,7 +28,7 @@ public:
 
 	/** If true, the pass will be composited onto each render. Does not apply to multi-layer EXR files. */
 	UPROPERTY(EditAnywhere, BlueprintReadWrite, meta=(EditCondition="bOverride_bCompositeOntoFinalImage"), Category="Widget Settings")
-	bool bCompositeOntoFinalImage;
+	bool bCompositeOntoFinalImage = true;
 
 protected:
 	/** A render pass for a specific render layer, generated by the UMovieGraphWidgetRendererBaseNode CDO. */
diff --git a/Engine/Plugins/MovieScene/MovieRenderPipeline/Source/MovieRenderPipelineCore/Private/MovieRenderPipelineCoreModule.cpp b/Engine/Plugins/MovieScene/MovieRenderPipeline/Source/MovieRenderPipelineCore/Private/MovieRenderPipelineCoreModule.cpp
--- a/Engine/Plugins/MovieScene/MovieRenderPipeline/Source/MovieRenderPipelineCore/Private/MovieRenderPipelineCoreModule.cpp
+++ b/Engine/Plugins/MovieScene/MovieRenderPipeline/Source/MovieRenderPipelineCore/Private/MovieRenderPipelineCoreModule.cpp
@@ -9,6 +9,7 @@
 #include "TimerManager.h"
 #include "MoviePipeline.h"
 #include "Graph/MovieGraphPipeline.h"
+#include "Graph/Nodes/MovieGraphBurnInNode.h"
 
 FName IMoviePipelineBurnInExtension::ModularFeatureName = "ModularFeature_MoviePipelineBurnInExt";
 
@@ -20,13 +21,13 @@ void FMovieRenderPipelineCoreModule::StartupModule()
 		TArray<FString> Assets;
 		Assets.Add(UMoviePipeline::DefaultDebugWidgetAsset);
 		Assets.Add(UMovieGraphPipeline::DefaultPreviewWidgetAsset);
+		Assets.Add(UMovieGraphBurnInNode::DefaultBurnInWidgetAsset);
 
 		for (const FString& Asset : Assets)
 		{
 			TSoftObjectPtr<UObject> AssetRef = TSoftObjectPtr<UObject>(FSoftObjectPath(Asset));
 			AssetRef.LoadSynchronous();
 		}
-
 	}
 #endif
 
diff --git a/Engine/Plugins/MovieScene/MovieRenderPipeline/Source/MovieRenderPipelineCore/Public/Graph/Nodes/MovieGraphBurnInNode.h b/Engine/Plugins/MovieScene/MovieRenderPipeline/Source/MovieRenderPipelineCore/Public/Graph/Nodes/MovieGraphBurnInNode.h
--- a/Engine/Plugins/MovieScene/MovieRenderPipeline/Source/MovieRenderPipelineCore/Public/Graph/Nodes/MovieGraphBurnInNode.h
+++ b/Engine/Plugins/MovieScene/MovieRenderPipeline/Source/MovieRenderPipelineCore/Public/Graph/Nodes/MovieGraphBurnInNode.h
@@ -19,6 +19,8 @@ class MOVIERENDERPIPELINECORE_API UMovieGraphBurnInNode : public UMovieGraphWidg
 	GENERATED_BODY()
 
 public:
+	UMovieGraphBurnInNode();
+
 #if WITH_EDITOR
 	virtual FText GetNodeTitle(const bool bGetDescriptive = false) const override;
 	virtual FLinearColor GetNodeTitleColor() const override;
@@ -33,6 +35,9 @@ public:
 
 public:
 	static const FString RendererName;
+
+	/** The path to the default widget class that will be used for the burn-in. */
+	static const FString DefaultBurnInWidgetAsset;
 	
 	UPROPERTY(EditAnywhere, BlueprintReadWrite, Category = Overrides, meta = (InlineEditConditionToggle))
 	uint8 bOverride_BurnInClass : 1;
@@ -50,6 +55,9 @@ protected:
 		virtual void Render(const FMovieGraphTraversalContext& InFrameTraversalContext, const FMovieGraphTimeStepData& InTimeData) override;
 		virtual int32 GetCompositingSortOrder() const override;
 
+		/** Gets the burn-in class that this pass will be using. */
+		UClass* GetBurnInClass() const;
+
 	private:
 		TObjectPtr<UMovieGraphBurnInWidget> GetBurnInWidget() const;
 	};
@@ -65,6 +73,7 @@ protected:
 	// ~UMovieGraphWidgetRendererBaseNode Interface
 	
 	// UMovieGraphRenderPassNode Interface
+	virtual void GatherOutputPassesImpl(UMovieGraphEvaluatedConfig* InConfig, TArray<FMovieGraphRenderDataIdentifier>& OutExpectedPasses) const override;
 	virtual FString GetRendererNameImpl() const override { return RendererName; }
 	virtual void TeardownImpl() override;
 	// ~UMovieGraphRenderPassNode Interface
