Moving CL 30677046 from 28.30 of 30677046
original CL description: hange Initialization time from FDateTime::UtcNow() to FDateTime::Now() to align with other DCC's to provide default time as user system (local) time

diff --git a/Engine/Plugins/MovieScene/MovieRenderPipeline/Source/MovieRenderPipelineCore/Private/MoviePipeline.cpp b/Engine/Plugins/MovieScene/MovieRenderPipeline/Source/MovieRenderPipelineCore/Private/MoviePipeline.cpp
--- a/Engine/Plugins/MovieScene/MovieRenderPipeline/Source/MovieRenderPipelineCore/Private/MoviePipeline.cpp
+++ b/Engine/Plugins/MovieScene/MovieRenderPipeline/Source/MovieRenderPipelineCore/Private/MoviePipeline.cpp
@@ -318,7 +318,7 @@ void UMoviePipeline::Initialize(UMoviePipelineExecutorJob* InJob)
 
 	// Initialization is complete. This engine frame is a wash (because the tick started with a 
 	// delta time not generated by us) so we'll wait until the next engine frame to start rendering.
-	InitializationTime = FDateTime::UtcNow();
+	InitializationTime = FDateTime::Now();
 
 	// We need to start our CachedOutputState at the time the Game World is currently. This way
 	// as MRQ counts up time, the render world stays in sync allowing comparisons between last rendered time
diff --git a/Engine/Plugins/MovieScene/MovieRenderPipeline/Source/MovieRenderPipelineCore/Private/MoviePipelineLinearExecutor.cpp b/Engine/Plugins/MovieScene/MovieRenderPipeline/Source/MovieRenderPipelineCore/Private/MoviePipelineLinearExecutor.cpp
--- a/Engine/Plugins/MovieScene/MovieRenderPipeline/Source/MovieRenderPipelineCore/Private/MoviePipelineLinearExecutor.cpp
+++ b/Engine/Plugins/MovieScene/MovieRenderPipeline/Source/MovieRenderPipelineCore/Private/MoviePipelineLinearExecutor.cpp
@@ -35,7 +35,7 @@ void UMoviePipelineLinearExecutorBase::Execute_Implementation(UMoviePipelineQueu
 	// before moving onto the next one. This may be parallelizable in the future (either multiple PIE
 	// sessions, or multiple external processes) but ideally one render would maximize resource usage anyways...
 	Queue = InPipelineQueue;
-	InitializationTime = FDateTime::UtcNow();
+	InitializationTime = FDateTime::Now();
 	JobsStarted = 0;
 
 	UE_LOG(LogMovieRenderPipeline, Log, TEXT("MoviePipelineLinearExecutorBase starting %d jobs."), InPipelineQueue->GetJobs().Num());
diff --git a/Engine/Plugins/MovieScene/MovieRenderPipeline/Source/MovieRenderPipelineEditor/Private/MoviePipelineEditorBlueprintLibrary.cpp b/Engine/Plugins/MovieScene/MovieRenderPipeline/Source/MovieRenderPipelineEditor/Private/MoviePipelineEditorBlueprintLibrary.cpp
--- a/Engine/Plugins/MovieScene/MovieRenderPipeline/Source/MovieRenderPipelineEditor/Private/MoviePipelineEditorBlueprintLibrary.cpp
+++ b/Engine/Plugins/MovieScene/MovieRenderPipeline/Source/MovieRenderPipelineEditor/Private/MoviePipelineEditorBlueprintLibrary.cpp
@@ -258,7 +258,7 @@ FString UMoviePipelineEditorBlueprintLibrary::ResolveOutputDirectoryFromJob(UMov
 		const bool bGetNextVersion = false;
 		FMovieGraphFilenameResolveParams Params;
 		Params.Job = InJob;
-		Params.InitializationTime = FDateTime::UtcNow();
+		Params.InitializationTime = FDateTime::Now();
 		Params.Version = UMovieGraphBlueprintLibrary::ResolveVersionNumber(Params, bGetNextVersion);
 		Params.RenderDataIdentifier.RootBranchName = UMovieGraphNode::GlobalsPinName;
 		
diff --git a/Engine/Plugins/MovieScene/MovieRenderPipeline/Source/MovieRenderPipelineEditor/Private/MoviePipelinePIEExecutor.cpp b/Engine/Plugins/MovieScene/MovieRenderPipeline/Source/MovieRenderPipelineEditor/Private/MoviePipelinePIEExecutor.cpp
--- a/Engine/Plugins/MovieScene/MovieRenderPipeline/Source/MovieRenderPipelineEditor/Private/MoviePipelinePIEExecutor.cpp
+++ b/Engine/Plugins/MovieScene/MovieRenderPipeline/Source/MovieRenderPipelineEditor/Private/MoviePipelinePIEExecutor.cpp
@@ -257,10 +257,18 @@ void UMoviePipelinePIEExecutor::OnPIEStartupFinished(bool)
 		if (PipelineAsLegacy)
 		{
 			PipelineAsLegacy->Initialize(Queue->GetJobs()[CurrentPipelineIndex]);
+			if (CustomInitializationTime.IsSet())
+			{
+				PipelineAsLegacy->SetInitializationTime(CustomInitializationTime.GetValue());
+			}
 		}
 		else if (UMovieGraphPipeline* PipelineAsGraph = Cast<UMovieGraphPipeline>(ActiveMoviePipeline))
 		{
 			PipelineAsGraph->Initialize(Queue->GetJobs()[CurrentPipelineIndex], FMovieGraphInitConfig());
+			if (CustomInitializationTime.IsSet())
+			{
+				PipelineAsGraph->SetInitializationTime(CustomInitializationTime.GetValue());
+			}
 		}
 		RemainingInitializationFrames = -1;
 	}
