MovieGraph: Fixed Custom Frame Rate setting not actually changing the frame rate.

diff --git a/Engine/Plugins/MovieScene/MovieRenderPipeline/Source/MovieRenderPipelineCore/Private/Graph/MovieGraphCoreTimeStep.cpp b/Engine/Plugins/MovieScene/MovieRenderPipeline/Source/MovieRenderPipelineCore/Private/Graph/MovieGraphCoreTimeStep.cpp
--- a/Engine/Plugins/MovieScene/MovieRenderPipeline/Source/MovieRenderPipelineCore/Private/Graph/MovieGraphCoreTimeStep.cpp
+++ b/Engine/Plugins/MovieScene/MovieRenderPipeline/Source/MovieRenderPipelineCore/Private/Graph/MovieGraphCoreTimeStep.cpp
@@ -548,7 +548,17 @@ void UMovieGraphCoreTimeStep::UpdateFrameMetrics()
 	// We inherit a tick resolution from the level sequence so that we can represent the same
 	// range of time that the level sequence does.
 	FrameData.TickResolution = GetOwningGraph()->GetDataSourceInstance()->GetTickResolution();
-	FrameData.FrameRate = GetOwningGraph()->GetDataSourceInstance()->GetDisplayRate(); // ToDo, needs to come from config (config can override)
+	FrameData.FrameRate = GetOwningGraph()->GetDataSourceInstance()->GetDisplayRate();
+
+	// We don't always have a config set up when this function is called.
+	if (CurrentFrameData.EvaluatedConfig)
+	{
+		if (UMovieGraphGlobalOutputSettingNode* OutputSetting = CurrentFrameData.EvaluatedConfig->GetSettingForBranch<UMovieGraphGlobalOutputSettingNode>(UMovieGraphNode::GlobalsPinName))
+		{
+			FrameData.FrameRate = UMovieGraphBlueprintLibrary::GetEffectiveFrameRate(OutputSetting, FrameData.FrameRate);
+		}
+	}
+
 	FrameData.FrameTimePerOutputFrame = FFrameRate::TransformTime(FFrameTime(FFrameNumber(1)), FrameData.FrameRate, FrameData.TickResolution);
 
 	// Manually perform blending of the Post Process Volumes/Camera/Camera Modifiers to match what the renderer will do.
@@ -579,9 +589,9 @@ void UMovieGraphCoreTimeStep::UpdateFrameMetrics()
 	EMoviePipelineShutterTiming ShutterTiming = EMoviePipelineShutterTiming::FrameCenter;
 	
 	// We don't always have a config set up when this function is called.
-	if (CurrentTimeStepData.EvaluatedConfig)
+	if (CurrentFrameData.EvaluatedConfig)
 	{
-		if (UMovieGraphCameraSettingNode* CameraSetting = CurrentTimeStepData.EvaluatedConfig->GetSettingForBranch< UMovieGraphCameraSettingNode>(UMovieGraphNode::GlobalsPinName))
+		if (UMovieGraphCameraSettingNode* CameraSetting = CurrentFrameData.EvaluatedConfig->GetSettingForBranch<UMovieGraphCameraSettingNode>(UMovieGraphNode::GlobalsPinName))
 		{
 			ShutterTiming = CameraSetting->ShutterTiming;
 		}
diff --git a/Engine/Plugins/MovieScene/MovieRenderPipeline/Source/MovieRenderPipelineCore/Private/Graph/MovieGraphPipeline.cpp b/Engine/Plugins/MovieScene/MovieRenderPipeline/Source/MovieRenderPipelineCore/Private/Graph/MovieGraphPipeline.cpp
--- a/Engine/Plugins/MovieScene/MovieRenderPipeline/Source/MovieRenderPipelineCore/Private/Graph/MovieGraphPipeline.cpp
+++ b/Engine/Plugins/MovieScene/MovieRenderPipeline/Source/MovieRenderPipelineCore/Private/Graph/MovieGraphPipeline.cpp
@@ -439,7 +439,7 @@ void UMovieGraphPipeline::BuildShotListFromDataSource()
 		Shot->ShotInfo.CurrentTimeInRoot = Shot->ShotInfo.TotalOutputRangeRoot.GetLowerBoundValue();
 		Shot->ShotInfo.NumEngineWarmUpFramesRemaining = WarmUpNode->NumWarmUpFrames;
 		Shot->ShotInfo.bEmulateFirstFrameMotionBlur = WarmUpNode->bEmulateMotionBlur;
-		Shot->ShotInfo.CalculateWorkMetrics(GetDataSourceInstance());
+		Shot->ShotInfo.CalculateWorkMetrics();
 		Shot->ShotInfo.VersionNumber = ResolveVersionForShot(Shot, EvaluatedConfig);
 	}
 
diff --git a/Engine/Plugins/MovieScene/MovieRenderPipeline/Source/MovieRenderPipelineCore/Private/Graph/MovieGraphSequenceDataSource.cpp b/Engine/Plugins/MovieScene/MovieRenderPipeline/Source/MovieRenderPipelineCore/Private/Graph/MovieGraphSequenceDataSource.cpp
--- a/Engine/Plugins/MovieScene/MovieRenderPipeline/Source/MovieRenderPipelineCore/Private/Graph/MovieGraphSequenceDataSource.cpp
+++ b/Engine/Plugins/MovieScene/MovieRenderPipeline/Source/MovieRenderPipelineCore/Private/Graph/MovieGraphSequenceDataSource.cpp
@@ -338,7 +338,9 @@ void UMovieGraphSequenceDataSource::JumpDataSource(const FFrameTime& InTimeToJum
 
 		// SetPlaybackPosition takes time in display rate and not tick resolution, so we convert
 		// it from tick resolution to display rate. JumpDataSource provides it in tick resolution
-		// to be consistent with other functions in this class.
+		// to be consistent with other functions in this class. This doesn't use the "Effective"
+		// framerate of the render because we never modified the Sequence, so we need to use original
+		// DisplayRate to jump to the correct frame.
 		FFrameRate TickResolution = LevelSequenceActor->GetSequence()->GetMovieScene()->GetTickResolution();
 		FFrameRate DisplayRate = LevelSequenceActor->GetSequence()->GetMovieScene()->GetDisplayRate();
 		FFrameTime RequestTime = FFrameRate::TransformTime(InTimeToJumpTo, TickResolution, DisplayRate);
diff --git a/Engine/Plugins/MovieScene/MovieRenderPipeline/Source/MovieRenderPipelineCore/Public/MovieRenderPipelineDataTypes.cpp b/Engine/Plugins/MovieScene/MovieRenderPipeline/Source/MovieRenderPipelineCore/Public/MovieRenderPipelineDataTypes.cpp
--- a/Engine/Plugins/MovieScene/MovieRenderPipeline/Source/MovieRenderPipelineCore/Public/MovieRenderPipelineDataTypes.cpp
+++ b/Engine/Plugins/MovieScene/MovieRenderPipeline/Source/MovieRenderPipelineCore/Public/MovieRenderPipelineDataTypes.cpp
@@ -9,11 +9,10 @@
 
 #include UE_INLINE_GENERATED_CPP_BY_NAME(MovieRenderPipelineDataTypes)
 
-FFrameNumber FMoviePipelineCameraCutInfo::GetOutputFrameCountEstimate(const UMovieGraphDataSourceBase* InDataSource) const
+FFrameNumber FMoviePipelineCameraCutInfo::GetOutputFrameCountEstimate() const
 {
-	// The graph uses the data source; the old version uses members in this struct
-	const FFrameRate SourceRate = InDataSource ? InDataSource->GetDisplayRate() : CachedFrameRate;
-	const FFrameRate DestinationRate = InDataSource ? InDataSource->GetTickResolution() : CachedTickResolution;
+	const FFrameRate SourceRate = CachedFrameRate;
+	const FFrameRate DestinationRate = CachedTickResolution;
 	
 	// TotalRange is stored in Tick Resolution, so we convert 1 frame of Frame Rate to the number of ticks.
 	FFrameNumber OneFrameInTicks = FFrameRate::TransformTime(FFrameTime(FFrameNumber(1)), SourceRate, DestinationRate).FloorToFrame();
@@ -25,10 +24,10 @@ FFrameNumber FMoviePipelineCameraCutInfo::GetOutputFrameCountEstimate(const UMov
 	return FFrameNumber(NumFrames);
 }
 
-void FMoviePipelineCameraCutInfo::CalculateWorkMetrics(const UMovieGraphDataSourceBase* InDataSource)
+void FMoviePipelineCameraCutInfo::CalculateWorkMetrics()
 {
 	// Initial Range + Handle Frames
-	FFrameNumber OutputFrameCount = GetOutputFrameCountEstimate(InDataSource);
+	FFrameNumber OutputFrameCount = GetOutputFrameCountEstimate();
 	WorkMetrics.TotalOutputFrameCount = OutputFrameCount.Value;
 	WorkMetrics.TotalSubSampleCount = NumSpatialSamples * NumTemporalSamples * NumTiles.X * NumTiles.Y; // Samples to generate an output frame.
 	WorkMetrics.TotalEngineWarmUpFrameCount = NumEngineWarmUpFramesRemaining;
diff --git a/Engine/Plugins/MovieScene/MovieRenderPipeline/Source/MovieRenderPipelineCore/Public/MovieRenderPipelineDataTypes.h b/Engine/Plugins/MovieScene/MovieRenderPipeline/Source/MovieRenderPipelineCore/Public/MovieRenderPipelineDataTypes.h
--- a/Engine/Plugins/MovieScene/MovieRenderPipeline/Source/MovieRenderPipelineCore/Public/MovieRenderPipelineDataTypes.h
+++ b/Engine/Plugins/MovieScene/MovieRenderPipeline/Source/MovieRenderPipelineCore/Public/MovieRenderPipelineDataTypes.h
@@ -491,12 +491,10 @@ public:
 
 	bool IsInitialized() const { return State != EMovieRenderShotState::Uninitialized; }
 	void SetNextStateAfter(const EMovieRenderShotState InCurrentState);
-
-	/** The data source is required for calculating some metrics (graph only). */
-	void CalculateWorkMetrics(const UMovieGraphDataSourceBase* InDataSource = nullptr);
+	void CalculateWorkMetrics();
 
 private:
-	FFrameNumber GetOutputFrameCountEstimate(const UMovieGraphDataSourceBase* InDataSource) const;
+	FFrameNumber GetOutputFrameCountEstimate() const;
 
 public:
 	
