MoviePipeline: Fixed an assert when rendering if the Water plugin was not enabled.

diff --git a/Engine/Plugins/MovieScene/MovieRenderPipeline/Source/MovieRenderPipelineCore/Private/Graph/Nodes/MovieGraphGlobalGameOverrides.cpp b/Engine/Plugins/MovieScene/MovieRenderPipeline/Source/MovieRenderPipelineCore/Private/Graph/Nodes/MovieGraphGlobalGameOverrides.cpp
--- a/Engine/Plugins/MovieScene/MovieRenderPipeline/Source/MovieRenderPipelineCore/Private/Graph/Nodes/MovieGraphGlobalGameOverrides.cpp
+++ b/Engine/Plugins/MovieScene/MovieRenderPipeline/Source/MovieRenderPipelineCore/Private/Graph/Nodes/MovieGraphGlobalGameOverrides.cpp
@@ -169,6 +169,10 @@ void UMovieGraphGlobalGameOverridesNode::ApplySettings(const bool bOverrideValue
 
 		// Cloth's time step smoothing messes up the change in number of simulation substeps that fixes the cloth simulation behavior when using Temporal Samples.
 		MOVIEPIPELINE_STORE_AND_OVERRIDE_CVAR_INT(PreviousChaosClothUseTimeStepSmoothing, TEXT("p.ChaosCloth.UseTimeStepSmoothing"), 0, bOverrideValues);
+
+		// Water skips water info texture when the world's game viewport rendering is disabled so we need to prevent this from happening.
+		MOVIEPIPELINE_STORE_AND_OVERRIDE_CVAR_INT_IF_EXIST(PreviousSkipWaterInfoTextureRenderWhenWorldRenderingDisabled, TEXT("r.Water.SkipWaterInfoTextureRenderWhenWorldRenderingDisabled"), 0, bOverrideValues);
+
 	}
 	
 	// Must come after the above cvars so that if one of those cvars is also specified by the scalability level, then we
diff --git a/Engine/Plugins/MovieScene/MovieRenderPipeline/Source/MovieRenderPipelineCore/Private/MoviePipelineGameOverrideSetting.cpp b/Engine/Plugins/MovieScene/MovieRenderPipeline/Source/MovieRenderPipelineCore/Private/MoviePipelineGameOverrideSetting.cpp
--- a/Engine/Plugins/MovieScene/MovieRenderPipeline/Source/MovieRenderPipelineCore/Private/MoviePipelineGameOverrideSetting.cpp
+++ b/Engine/Plugins/MovieScene/MovieRenderPipeline/Source/MovieRenderPipelineCore/Private/MoviePipelineGameOverrideSetting.cpp
@@ -138,7 +138,7 @@ void UMoviePipelineGameOverrideSetting::ApplyCVarSettings(const bool bOverrideVa
 	MOVIEPIPELINE_STORE_AND_OVERRIDE_CVAR_INT(PreviousChaosClothUseTimeStepSmoothing, TEXT("p.ChaosCloth.UseTimeStepSmoothing"), 0, bOverrideValues);
 
 	// Water skips water info texture when the world's game viewport rendering is disabled so we need to prevent this from happening.
-	MOVIEPIPELINE_STORE_AND_OVERRIDE_CVAR_INT(PreviousSkipWaterInfoTextureRenderWhenWorldRenderingDisabled, TEXT("r.Water.SkipWaterInfoTextureRenderWhenWorldRenderingDisabled"), 0, bOverrideValues);
+	MOVIEPIPELINE_STORE_AND_OVERRIDE_CVAR_INT_IF_EXIST(PreviousSkipWaterInfoTextureRenderWhenWorldRenderingDisabled, TEXT("r.Water.SkipWaterInfoTextureRenderWhenWorldRenderingDisabled"), 0, bOverrideValues);
 
 	// Must come after the above cvars so that if one of those cvars is also specified by the Scalability level, then we restore to the value in the original scalability level
 	// not the value we cached in the Cinematic level (if applied).
diff --git a/Engine/Plugins/MovieScene/MovieRenderPipeline/Source/MovieRenderPipelineCore/Public/Graph/Nodes/MovieGraphGlobalGameOverrides.h b/Engine/Plugins/MovieScene/MovieRenderPipeline/Source/MovieRenderPipelineCore/Public/Graph/Nodes/MovieGraphGlobalGameOverrides.h
--- a/Engine/Plugins/MovieScene/MovieRenderPipeline/Source/MovieRenderPipelineCore/Public/Graph/Nodes/MovieGraphGlobalGameOverrides.h
+++ b/Engine/Plugins/MovieScene/MovieRenderPipeline/Source/MovieRenderPipelineCore/Public/Graph/Nodes/MovieGraphGlobalGameOverrides.h
@@ -202,6 +202,7 @@ private:
 	float PreviousChaosImmPhysicsMinStepTime;
 	int32 PreviousSkipRedundantTransformUpdate;
 	int32 PreviousChaosClothUseTimeStepSmoothing;
+	int32 PreviousSkipWaterInfoTextureRenderWhenWorldRenderingDisabled;
 
 #if WITH_EDITOR
 	int32 PreviousGeoCacheStreamerShowNotification;
