Movie Graph:

-Fix: Unconnected user variables show up in MRQ window

diff --git a/Engine/Plugins/MovieScene/MovieRenderPipeline/Source/MovieRenderPipelineCore/Private/Graph/MovieGraphNode.cpp b/Engine/Plugins/MovieScene/MovieRenderPipeline/Source/MovieRenderPipelineCore/Private/Graph/MovieGraphNode.cpp
--- a/Engine/Plugins/MovieScene/MovieRenderPipeline/Source/MovieRenderPipelineCore/Private/Graph/MovieGraphNode.cpp
+++ b/Engine/Plugins/MovieScene/MovieRenderPipeline/Source/MovieRenderPipelineCore/Private/Graph/MovieGraphNode.cpp
@@ -365,6 +365,19 @@ UMovieGraphPin* UMovieGraphNode::GetFirstConnectedInputPin() const
 	return nullptr;
 }
 
+UMovieGraphPin* UMovieGraphNode::GetFirstConnectedOutputPin() const
+{
+	for (const TObjectPtr<UMovieGraphPin>& OutputPin : OutputPins)
+	{
+		if (OutputPin->IsConnected())
+		{
+			return OutputPin.Get();
+		}
+	}
+
+	return nullptr;
+}
+
 bool UMovieGraphNode::CanBeDisabled() const
 {
 	// By default, all nodes can be disabled
diff --git a/Engine/Plugins/MovieScene/MovieRenderPipeline/Source/MovieRenderPipelineCore/Private/MovieJobVariableAssignmentContainer.cpp b/Engine/Plugins/MovieScene/MovieRenderPipeline/Source/MovieRenderPipelineCore/Private/MovieJobVariableAssignmentContainer.cpp
--- a/Engine/Plugins/MovieScene/MovieRenderPipeline/Source/MovieRenderPipelineCore/Private/MovieJobVariableAssignmentContainer.cpp
+++ b/Engine/Plugins/MovieScene/MovieRenderPipeline/Source/MovieRenderPipelineCore/Private/MovieJobVariableAssignmentContainer.cpp
@@ -4,6 +4,7 @@
 
 #include "Graph/MovieGraphCommon.h"
 #include "Graph/MovieGraphConfig.h"
+#include "Graph/Nodes/MovieGraphVariableNode.h"
 
 // The prefix for property descriptions that are in charge of variable assignment EditCondition enable/disable state
 static const TCHAR* EditConditionPrefix = TEXT("bOverride_");
@@ -306,7 +307,7 @@ bool UMovieJobVariableAssignmentContainer::GenerateVariableOverride(const UMovie
 }
 
 #if WITH_EDITOR
-void UMovieJobVariableAssignmentContainer::UpdateGraphVariableOverrides()
+void UMovieJobVariableAssignmentContainer::UpdateGraphVariableOverrides(const bool bOnlyIncludeConnectedVariables/*= false*/)
 {
 	auto GetVariableGuidFromDesc = [](const FPropertyBagPropertyDesc& Desc) -> FGuid
 	{
@@ -445,6 +446,26 @@ void UMovieJobVariableAssignmentContainer::UpdateGraphVariableOverrides()
 		}
 	}
 
+	// Optionally, remove any descs that don't correspond to a variable node with an outgoing connection whose Guid matches the desc's Guid
+	if (bOnlyIncludeConnectedVariables)
+	{
+		const int32 NumDescsRemoved = ModifiedDescs.RemoveAll([&GetVariableGuidFromDesc, this](const FPropertyBagPropertyDesc& Desc)
+		{
+			const FGuid DescGuid = GetVariableGuidFromDesc(Desc);
+			return !GraphPreset->GetNodes().ContainsByPredicate([&DescGuid](UMovieGraphNode* Node)
+				{
+					const UMovieGraphVariableNode* VariableNode = Cast<UMovieGraphVariableNode>(Node);
+					return VariableNode && VariableNode->GetFirstConnectedOutputPin() && 
+						VariableNode->GetVariable() && (VariableNode->GetVariable()->GetGuid() == DescGuid);
+				});
+		});
+
+		if (NumDescsRemoved > 0) 
+		{
+			bNeedsToRegenerate = true;
+		}
+	}
+
 	if (bNeedsToRegenerate)
 	{
 		const UPropertyBag* NewBagStruct = UPropertyBag::GetOrCreateFromDescs(ModifiedDescs);
diff --git a/Engine/Plugins/MovieScene/MovieRenderPipeline/Source/MovieRenderPipelineCore/Private/MoviePipelineQueue.cpp b/Engine/Plugins/MovieScene/MovieRenderPipeline/Source/MovieRenderPipelineCore/Private/MoviePipelineQueue.cpp
--- a/Engine/Plugins/MovieScene/MovieRenderPipeline/Source/MovieRenderPipelineCore/Private/MoviePipelineQueue.cpp
+++ b/Engine/Plugins/MovieScene/MovieRenderPipeline/Source/MovieRenderPipelineCore/Private/MoviePipelineQueue.cpp
@@ -19,13 +19,15 @@ namespace UE::MovieGraph::Private
 	TObjectPtr<UMovieJobVariableAssignmentContainer> GetOrCreateJobVariableAssignmentsForGraph(
 		const UMovieGraphConfig* InGraph, TArray<TObjectPtr<UMovieJobVariableAssignmentContainer>>& InVariableAssignments, UObject* InAssignmentsOwner)
 	{
+		constexpr bool bOnlyIncludeConnectedVariables = true;
+		
 		for (TObjectPtr<UMovieJobVariableAssignmentContainer>& VariableAssignment : InVariableAssignments)
 		{
 			const TSoftObjectPtr<UMovieGraphConfig> SoftGraphConfig = VariableAssignment->GetGraphConfig();
 			if (SoftGraphConfig.Get() == InGraph)
 			{
 #if WITH_EDITOR
-				VariableAssignment->UpdateGraphVariableOverrides();
+				VariableAssignment->UpdateGraphVariableOverrides(bOnlyIncludeConnectedVariables);
 #endif
 				
 				return VariableAssignment;
@@ -38,7 +40,7 @@ namespace UE::MovieGraph::Private
 		NewVariableAssignments->SetGraphConfig(InGraph);
 
 #if WITH_EDITOR
-		NewVariableAssignments->UpdateGraphVariableOverrides();
+		NewVariableAssignments->UpdateGraphVariableOverrides(bOnlyIncludeConnectedVariables);
 #endif
 	
 		return NewVariableAssignments;
diff --git a/Engine/Plugins/MovieScene/MovieRenderPipeline/Source/MovieRenderPipelineCore/Public/Graph/MovieGraphNode.h b/Engine/Plugins/MovieScene/MovieRenderPipeline/Source/MovieRenderPipelineCore/Public/Graph/MovieGraphNode.h
--- a/Engine/Plugins/MovieScene/MovieRenderPipeline/Source/MovieRenderPipelineCore/Public/Graph/MovieGraphNode.h
+++ b/Engine/Plugins/MovieScene/MovieRenderPipeline/Source/MovieRenderPipelineCore/Public/Graph/MovieGraphNode.h
@@ -241,6 +241,9 @@ public:
 
 	/** Gets the first input pin on the node which has a connection, or nullptr if no pins are connected. */
 	UMovieGraphPin* GetFirstConnectedInputPin() const;
+	
+	/** Gets the first output pin on the node which has a connection, or nullptr if no pins are connected. */
+	UMovieGraphPin* GetFirstConnectedOutputPin() const;
 
 	/** Gets the GUID which uniquely identifies this node. */
 	const FGuid& GetGuid() const { return Guid; }
diff --git a/Engine/Plugins/MovieScene/MovieRenderPipeline/Source/MovieRenderPipelineCore/Public/MovieJobVariableAssignmentContainer.h b/Engine/Plugins/MovieScene/MovieRenderPipeline/Source/MovieRenderPipelineCore/Public/MovieJobVariableAssignmentContainer.h
--- a/Engine/Plugins/MovieScene/MovieRenderPipeline/Source/MovieRenderPipelineCore/Public/MovieJobVariableAssignmentContainer.h
+++ b/Engine/Plugins/MovieScene/MovieRenderPipeline/Source/MovieRenderPipelineCore/Public/MovieJobVariableAssignmentContainer.h
@@ -230,9 +230,10 @@ public:
 	 * the graph variable name, value type, object type, and container type. Additionally, stale overrides that have no
 	 * corresponding graph variable will be removed, and overrides will be created for graph variables which do not have
 	 * existing overrides.
+	 * @param bOnlyIncludeConnectedVariables Optionally, if the stored variable overrides should only include variables that 1) have nodes referencing them in the graph, and 2) at least one of those variable nodes has a connection to another node, set bOnlyIncludeConnectedVariables to true. This will ensure that variable overrides that have no chance of being utilized in the graph are not available to be changed.
 	 */
 	UFUNCTION(BlueprintCallable, Category = "Experimental")
-	void UpdateGraphVariableOverrides();
+	void UpdateGraphVariableOverrides(const bool bOnlyIncludeConnectedVariables = false);
 #endif
 
 protected:
