MRG: Fixed filename token validation and multi-layer EXR sorting.

Note that this doesn't attempt to solve the redundant validation warnings on output file nodes with different types.

#rnx

diff --git a/Engine/Plugins/MovieScene/MovieRenderPipeline/Source/MovieRenderPipelineCore/Private/Graph/MovieGraphDataTypes.cpp b/Engine/Plugins/MovieScene/MovieRenderPipeline/Source/MovieRenderPipelineCore/Private/Graph/MovieGraphDataTypes.cpp
--- a/Engine/Plugins/MovieScene/MovieRenderPipeline/Source/MovieRenderPipelineCore/Private/Graph/MovieGraphDataTypes.cpp
+++ b/Engine/Plugins/MovieScene/MovieRenderPipeline/Source/MovieRenderPipelineCore/Private/Graph/MovieGraphDataTypes.cpp
@@ -38,3 +38,66 @@ const MoviePipeline::FAudioState& UMovieGraphAudioRendererBase::GetAudioState()
 {
 	return AudioState;
 }
+
+UE::MovieGraph::FMovieGraphRenderDataValidationInfo UE::MovieGraph::FMovieGraphOutputMergerFrame::GetValidationInfo(const FMovieGraphRenderDataIdentifier& InRenderID, bool bInDiscardCompositedRenders) const
+{
+	const FName& ActiveBranchName = InRenderID.RootBranchName;
+	const FString& ActiveRendererName = InRenderID.RendererName;
+
+	FMovieGraphRenderDataValidationInfo ValidationInfo;
+
+	TSet<FString>	LayerCounts;
+	TSet<FName>		BranchCounts;
+	TSet<FString>	BranchRendererCounts;
+
+	const int32 ReserveNum = ExpectedRenderPasses.Num();
+	LayerCounts.Reserve(ReserveNum);
+	BranchCounts.Reserve(ReserveNum);
+	BranchRendererCounts.Reserve(ReserveNum);
+
+	for (const FMovieGraphRenderDataIdentifier& PassIdentifier : ExpectedRenderPasses)
+	{
+		LayerCounts.Add(PassIdentifier.LayerName);
+		BranchCounts.Add(PassIdentifier.RootBranchName);
+
+		// We only count renderers on the active branch
+		if (PassIdentifier.RootBranchName == ActiveBranchName)
+		{
+			BranchRendererCounts.Add(PassIdentifier.RendererName);
+
+			// Since renderers will be resolved separately, we only count subresources on the active renderer for the active branch
+			if (PassIdentifier.RendererName == ActiveRendererName)
+			{
+				ValidationInfo.ActiveRendererSubresourceCount++;
+			}
+		}
+	}
+
+	// Remove any renderers that will be composited on later
+	if (bInDiscardCompositedRenders)
+	{
+		for (const TPair<FMovieGraphRenderDataIdentifier, TUniquePtr<FImagePixelData>>& RenderData : ImageOutputData)
+		{
+			UE::MovieGraph::FMovieGraphSampleState* Payload = RenderData.Value->GetPayload<UE::MovieGraph::FMovieGraphSampleState>();
+			check(Payload);
+			if (Payload->bCompositeOnOtherRenders)
+			{
+				if (RenderData.Key.RootBranchName == ActiveBranchName)
+				{
+					BranchRendererCounts.Remove(RenderData.Key.RendererName);
+
+					if (RenderData.Key.RendererName == ActiveRendererName)
+					{
+						ValidationInfo.ActiveRendererSubresourceCount--;
+					}
+				}
+			}
+		}
+	}
+
+	ValidationInfo.LayerCount = LayerCounts.Num();
+	ValidationInfo.BranchCount = BranchCounts.Num();
+	ValidationInfo.ActiveBranchRendererCount = BranchRendererCounts.Num();
+
+	return ValidationInfo;
+}
diff --git a/Engine/Plugins/MovieScene/MovieRenderPipeline/Source/MovieRenderPipelineCore/Private/Graph/MovieGraphOutputMerger.cpp b/Engine/Plugins/MovieScene/MovieRenderPipeline/Source/MovieRenderPipelineCore/Private/Graph/MovieGraphOutputMerger.cpp
--- a/Engine/Plugins/MovieScene/MovieRenderPipeline/Source/MovieRenderPipelineCore/Private/Graph/MovieGraphOutputMerger.cpp
+++ b/Engine/Plugins/MovieScene/MovieRenderPipeline/Source/MovieRenderPipelineCore/Private/Graph/MovieGraphOutputMerger.cpp
@@ -94,8 +94,16 @@ namespace UE::MovieGraph
 		
 		if (FinishedPasses == TotalPasses)
 		{
-			// ToDo: We potentialy want to sort the output using a stable-sort for EXRs so they can choose which is the primary rgba layer.
-			
+			// Sort the output frames. This is only really important for multi-channel formats like EXR, but it lets passes
+			// specify which one should be the thumbnail/default rgba channels instead of a first-come-first-serve.
+			OutputFrame->ImageOutputData.ValueStableSort([](const TUniquePtr<FImagePixelData>& First, const TUniquePtr<FImagePixelData>& Second) -> bool
+				{
+					FMovieGraphSampleState* FirstPayload = First->GetPayload<FMovieGraphSampleState>();
+					FMovieGraphSampleState* SecondPayload = Second->GetPayload<FMovieGraphSampleState>();
+
+					return FirstPayload->CompositingSortOrder < SecondPayload->CompositingSortOrder;
+				}
+			);
 			// Move this frame into our FinishedFrames array so the Game Thread can read it at its leisure
 			FMovieGraphOutputMergerFrame FinalFrame;
 			ensureMsgf(PendingData.RemoveAndCopyValue(RenderedFrameNumber, FinalFrame), TEXT("Could not find frame in pending data, output will be skipped!"));
diff --git a/Engine/Plugins/MovieScene/MovieRenderPipeline/Source/MovieRenderPipelineCore/Private/Graph/Nodes/MovieGraphBurnInNode.cpp b/Engine/Plugins/MovieScene/MovieRenderPipeline/Source/MovieRenderPipelineCore/Private/Graph/Nodes/MovieGraphBurnInNode.cpp
--- a/Engine/Plugins/MovieScene/MovieRenderPipeline/Source/MovieRenderPipelineCore/Private/Graph/Nodes/MovieGraphBurnInNode.cpp
+++ b/Engine/Plugins/MovieScene/MovieRenderPipeline/Source/MovieRenderPipelineCore/Private/Graph/Nodes/MovieGraphBurnInNode.cpp
@@ -135,7 +135,7 @@ void UMovieGraphBurnInNode::FMovieGraphBurnInPass::Render(const FMovieGraphTrave
 int32 UMovieGraphBurnInNode::FMovieGraphBurnInPass::GetCompositingSortOrder() const
 {
 	// Burn-ins should always appear over all other passes
-	return 0;
+	return 100;
 }
 
 UClass* UMovieGraphBurnInNode::FMovieGraphBurnInPass::GetBurnInClass() const
diff --git a/Engine/Plugins/MovieScene/MovieRenderPipeline/Source/MovieRenderPipelineCore/Private/Graph/Nodes/MovieGraphUIRendererNode.cpp b/Engine/Plugins/MovieScene/MovieRenderPipeline/Source/MovieRenderPipelineCore/Private/Graph/Nodes/MovieGraphUIRendererNode.cpp
--- a/Engine/Plugins/MovieScene/MovieRenderPipeline/Source/MovieRenderPipelineCore/Private/Graph/Nodes/MovieGraphUIRendererNode.cpp
+++ b/Engine/Plugins/MovieScene/MovieRenderPipeline/Source/MovieRenderPipelineCore/Private/Graph/Nodes/MovieGraphUIRendererNode.cpp
@@ -51,7 +51,8 @@ TSharedPtr<SWidget> UMovieGraphUIRendererNode::FMovieGraphUIPass::GetWidget()
 
 int32 UMovieGraphUIRendererNode::FMovieGraphUIPass::GetCompositingSortOrder() const
 {
-	return 100;
+	// We want UI Widgets to be underneath the burn-in.
+	return 80;
 }
 
 TUniquePtr<UMovieGraphWidgetRendererBaseNode::FMovieGraphWidgetPass> UMovieGraphUIRendererNode::GeneratePass()
diff --git a/Engine/Plugins/MovieScene/MovieRenderPipeline/Source/MovieRenderPipelineCore/Private/MoviePipelineUtils.cpp b/Engine/Plugins/MovieScene/MovieRenderPipeline/Source/MovieRenderPipelineCore/Private/MoviePipelineUtils.cpp
--- a/Engine/Plugins/MovieScene/MovieRenderPipeline/Source/MovieRenderPipelineCore/Private/MoviePipelineUtils.cpp
+++ b/Engine/Plugins/MovieScene/MovieRenderPipeline/Source/MovieRenderPipelineCore/Private/MoviePipelineUtils.cpp
@@ -695,7 +695,7 @@ namespace UE
 			}
 		}
 
-		void ValidateOutputFormatString(FString& InOutFilenameFormatString, const bool bTestRenderPass, const bool bTestFrameNumber, const bool bTestCameraName, const bool bIncludeSubRendererName)
+		void ValidateOutputFormatString(FString& InOutFilenameFormatString, const bool bTestRenderPass, const bool bTestFrameNumber, const bool bTestCameraName)
 		{
 			const FString FrameNumberIdentifiers[] = { TEXT("{frame_number}"), TEXT("{frame_number_shot}"), TEXT("{frame_number_rel}"), TEXT("{frame_number_shot_rel}") };
 
@@ -793,14 +793,6 @@ namespace UE
 				}
 			}
 
-			if (bIncludeSubRendererName)
-			{
-				if (!InOutFilenameFormatString.Contains(TEXT("{renderer_sub_name}"), ESearchCase::IgnoreCase))
-				{
-					InOutFilenameFormatString.Append(TEXT(".{renderer_sub_name}"));
-				}
-			}
-
 			if (!InOutFilenameFormatString.Contains(TEXT("{file_dup}"), ESearchCase::IgnoreCase))
 			{
 				InOutFilenameFormatString.Append(TEXT("{file_dup}"));
diff --git a/Engine/Plugins/MovieScene/MovieRenderPipeline/Source/MovieRenderPipelineCore/Public/Graph/MovieGraphDataTypes.h b/Engine/Plugins/MovieScene/MovieRenderPipeline/Source/MovieRenderPipelineCore/Public/Graph/MovieGraphDataTypes.h
--- a/Engine/Plugins/MovieScene/MovieRenderPipeline/Source/MovieRenderPipelineCore/Public/Graph/MovieGraphDataTypes.h
+++ b/Engine/Plugins/MovieScene/MovieRenderPipeline/Source/MovieRenderPipelineCore/Public/Graph/MovieGraphDataTypes.h
@@ -389,6 +389,17 @@ namespace UE::MovieGraph
 		ESceneCaptureSource SceneCaptureSource;
 	};
 
+	/**
+	* A collection of validation information extracted from the FMovieGraphOutputMergerFrame 
+	*/
+	struct FMovieGraphRenderDataValidationInfo
+	{
+		int32 LayerCount = 0;
+		int32 BranchCount = 0;
+		int32 ActiveBranchRendererCount = 0;
+		int32 ActiveRendererSubresourceCount = 0;
+	};
+
 	/**
 	* A collection of all of the data needed to produce the one output frame on disk. This holds
 	* all of the pixel data associated with the output frame at once, so we can pass it to our 
@@ -422,72 +433,8 @@ namespace UE::MovieGraph
 		/** Additional metadata to be added to the output (if supported by the output container). */
 		TMap<FString, FString> FileMetadata;
 
-		/** Check if the (expected) render passes come from different branches. */
-		bool HasDataFromMultipleBranches() const
-		{
-			TSet<FName> BranchUseCounts;
-			BranchUseCounts.Reserve(ExpectedRenderPasses.Num());
-
-			for (const FMovieGraphRenderDataIdentifier& PassIdentifier : ExpectedRenderPasses)
-			{
-				BranchUseCounts.Add(PassIdentifier.RootBranchName);
-			}
-
-			return BranchUseCounts.Num() > 1;
-		}
-
-		bool HasDataFromMultipleLayersWithName(const FString& InLayerName) const
-		{
-			TMap<FString, int32> Layers;
-			Layers.Reserve(ExpectedRenderPasses.Num());
-
-			for (const FMovieGraphRenderDataIdentifier& PassIdentifier : ExpectedRenderPasses)
-			{
-				int32& Count = Layers.FindOrAdd(PassIdentifier.LayerName);
-				Count++;
-			}
-
-			// For the given branch, check to see if there's more than one branch with that name.
-			if (int32* CountPtr = Layers.Find(InLayerName))
-			{
-				return *CountPtr > 1;
-			}
-
-			return false;
-		}
-
-		/** Check if the (expected) render passes come from different renderers, excluding composited ones. */
-		bool HasMultipleRendersPerBranch(const FName& InBranchName)
-		{
-			TSet<uint32> BranchRendererUseCounts;
-			BranchRendererUseCounts.Reserve(ExpectedRenderPasses.Num());
-
-			for (const FMovieGraphRenderDataIdentifier& PassIdentifier : ExpectedRenderPasses)
-			{
-				if (PassIdentifier.RootBranchName == InBranchName)
-				{
-					BranchRendererUseCounts.Add(GetTypeHash(PassIdentifier.RendererName));
-				}
-			}
-
-			// Remove any renderers that will be composited on later
-			for (const TPair<FMovieGraphRenderDataIdentifier, TUniquePtr<FImagePixelData>>& RenderData : ImageOutputData)
-			{
-				if (RenderData.Key.RootBranchName != InBranchName)
-				{
-					continue;
-				}
-
-				UE::MovieGraph::FMovieGraphSampleState* Payload = RenderData.Value->GetPayload<UE::MovieGraph::FMovieGraphSampleState>();
-				check(Payload);
-				if (Payload->bCompositeOnOtherRenders)
-				{
-					BranchRendererUseCounts.Remove(GetTypeHash(RenderData.Key.RendererName));
-				}
-			}
-
-			return BranchRendererUseCounts.Num() > 1;
-		}
+		/** Get filename token validation info from the (expected) render passes. */
+		MOVIERENDERPIPELINECORE_API FMovieGraphRenderDataValidationInfo GetValidationInfo(const FMovieGraphRenderDataIdentifier& InRenderID, bool bInDiscardCompositedRenders = true) const;
 	};
 
 	/**
diff --git a/Engine/Plugins/MovieScene/MovieRenderPipeline/Source/MovieRenderPipelineCore/Public/MoviePipelineUtils.h b/Engine/Plugins/MovieScene/MovieRenderPipeline/Source/MovieRenderPipelineCore/Public/MoviePipelineUtils.h
--- a/Engine/Plugins/MovieScene/MovieRenderPipeline/Source/MovieRenderPipelineCore/Public/MoviePipelineUtils.h
+++ b/Engine/Plugins/MovieScene/MovieRenderPipeline/Source/MovieRenderPipelineCore/Public/MoviePipelineUtils.h
@@ -132,7 +132,7 @@ namespace UE
 	namespace MoviePipeline
 	{
 		MOVIERENDERPIPELINECORE_API void ConformOutputFormatStringToken(FString& InOutFilenameFormatString, const FStringView InToken, const FName& InNodeName, const FName& InBranchName);
-		MOVIERENDERPIPELINECORE_API void ValidateOutputFormatString(FString& InOutFilenameFormatString, const bool bTestRenderPass, const bool bTestFrameNumber, const bool bIncludeCameraName = false, const bool bIncludeSubRendererName = false);
+		MOVIERENDERPIPELINECORE_API void ValidateOutputFormatString(FString& InOutFilenameFormatString, const bool bTestRenderPass, const bool bTestFrameNumber, const bool bIncludeCameraName = false);
 		MOVIERENDERPIPELINECORE_API void RemoveFrameNumberFormatStrings(FString& InOutFilenameFormatString, const bool bIncludeShots);
 		/** De-duplicates the provided array of strings by appending (1), (2), etc to the end of duplicates. */
 		MOVIERENDERPIPELINECORE_API void DeduplicateNameArray(TArray<FString>& InOutNames);
diff --git a/Engine/Plugins/MovieScene/MovieRenderPipeline/Source/MovieRenderPipelineRenderPasses/Private/Graph/Renderers/MovieGraphDeferredPass.cpp b/Engine/Plugins/MovieScene/MovieRenderPipeline/Source/MovieRenderPipelineRenderPasses/Private/Graph/Renderers/MovieGraphDeferredPass.cpp
--- a/Engine/Plugins/MovieScene/MovieRenderPipeline/Source/MovieRenderPipelineRenderPasses/Private/Graph/Renderers/MovieGraphDeferredPass.cpp
+++ b/Engine/Plugins/MovieScene/MovieRenderPipeline/Source/MovieRenderPipelineRenderPasses/Private/Graph/Renderers/MovieGraphDeferredPass.cpp
@@ -271,6 +271,7 @@ void FMovieGraphDeferredPass::Render(const FMovieGraphTraversalContext& InFrameT
 			SampleState.OverscanFraction = OverscanFraction;
 			SampleState.bAllowOCIO = ParentNodeThisFrame->GetAllowOCIO();
 			SampleState.SceneCaptureSource = SceneCaptureSource;
+			SampleState.CompositingSortOrder = 10;
 		}
 
 		if (UMovieGraphImagePassBaseNode* ParentNode = GetParentNode(InFrameTraversalContext.Time.EvaluatedConfig))
@@ -297,6 +298,8 @@ void FMovieGraphDeferredPass::Render(const FMovieGraphTraversalContext& InFrameT
 				UE::MovieGraph::FMovieGraphSampleState PassSampleState = SampleState;
 				PassSampleState.TraversalContext.RenderDataIdentifier = Identifier;
 				
+				// Give a lower priority to materials so they show up after the main pass in multi-layer exrs.
+				PassSampleState.CompositingSortOrder = SampleState.CompositingSortOrder + 1;
 				BufferPipe->AddEndpoint(MakeForwardingEndpoint(PassSampleState, InTimeData));
 
 				NewView->FinalPostProcessSettings.BufferVisualizationPipes.Add(VisMaterial->GetFName(), BufferPipe);
diff --git a/Engine/Plugins/MovieScene/MovieRenderPipeline/Source/MovieRenderPipelineRenderPasses/Private/MovieGraphImageSequenceOutputNode.cpp b/Engine/Plugins/MovieScene/MovieRenderPipeline/Source/MovieRenderPipelineRenderPasses/Private/MovieGraphImageSequenceOutputNode.cpp
--- a/Engine/Plugins/MovieScene/MovieRenderPipeline/Source/MovieRenderPipelineRenderPasses/Private/MovieGraphImageSequenceOutputNode.cpp
+++ b/Engine/Plugins/MovieScene/MovieRenderPipeline/Source/MovieRenderPipelineRenderPasses/Private/MovieGraphImageSequenceOutputNode.cpp
@@ -70,23 +70,6 @@ namespace UE::MovieGraph::Private
 
 		return Params;
 	}
-
-	/** Determine whether this token should be added based on the number of unique SubResourceNames in ExpectedRenderPasses. */
-	bool ShouldIncludeSubRendererName(const TArray<FMovieGraphRenderDataIdentifier>& ExpectedRenderPasses)
-	{
-		TSet<FString> PassNames;
-		for (const FMovieGraphRenderDataIdentifier& Pass : ExpectedRenderPasses)
-		{
-			PassNames.Add(Pass.SubResourceName);
-
-			if (PassNames.Num() > 1)
-			{
-				return true;
-			}
-		}
-
-		return false;
-	}
 	
 #if WITH_OCIO
 	struct FOpenColorIOPixelPreProcessor
@@ -244,8 +227,8 @@ TArray<TPair<FMovieGraphRenderDataIdentifier, TUniquePtr<FImagePixelData>>> UMov
 		CompositedPasses.Add(MoveTemp(CompositePass));
 	}
 
-	// Sort composited passes if multiple were found. Passes with a smaller sort order go to the end of the array so they
-	// get composited on top of passes with a higher sort order.
+	// Sort composited passes if multiple were found. Passes with a higher sort order go to the end of the array so they
+	// get composited on top of passes with a lower sort order.
 	CompositedPasses.Sort([](
 		const TPair<FMovieGraphRenderDataIdentifier, TUniquePtr<FImagePixelData>>& PassA,
 		const TPair<FMovieGraphRenderDataIdentifier, TUniquePtr<FImagePixelData>>& PassB)
@@ -255,7 +238,7 @@ TArray<TPair<FMovieGraphRenderDataIdentifier, TUniquePtr<FImagePixelData>>> UMov
 		check(PayloadA);
 		check(PayloadB);
 
-		return PayloadA->CompositingSortOrder > PayloadB->CompositingSortOrder;
+		return PayloadA->CompositingSortOrder < PayloadB->CompositingSortOrder;
 	});
 
 	return CompositedPasses;
@@ -287,27 +270,36 @@ FString UMovieGraphImageSequenceOutputNode::CreateFileName(
 	// Generate one string that puts the directory combined with the filename format.
 	FString FileNameFormatString = OutputSettingNode->OutputDirectory.Path / InParentNode->FileNameFormat;
 
-	// These two checks seem slightly backwards, but are the correct way to handle it. If there is
-	// multiple branches with data, we try to separate them by layer name (which is likely to be unique).
-	// Below we do the same check again in reverse - if there are multiple layers that share this name,
-	// then we need to insert the branch_name to differentiate them.
 	// ToDo: This is overly protective and could be relaxed later, for instance
 	// if different file write nodes have chosen a separate filepath entirely.
-	if (InRawFrameData->HasDataFromMultipleBranches())
+	UE::MovieGraph::FMovieGraphRenderDataValidationInfo ValidationInfo = InRawFrameData->GetValidationInfo(InRenderData.Key);
+
+	// Since there can only be one layer per branch, we restrain layer/branch validation to multi-branch graphs.
+	if (ValidationInfo.BranchCount > 1)
 	{
-		UE::MoviePipeline::ConformOutputFormatStringToken(FileNameFormatString, TEXT("{layer_name}"), InParentNode->GetFName(), InRenderData.Key.RootBranchName);
+		// We can run into the scenario where the users have given layers the same name, so layer_name token won't help differentiate.
+		// To resolve this, we look to see if there's multiple branches with the same layer name, and if so we force the branch name into the token too.
+		if (ValidationInfo.LayerCount < ValidationInfo.BranchCount)
+		{
+			UE::MoviePipeline::ConformOutputFormatStringToken(FileNameFormatString, TEXT("{branch_name}"), InParentNode->GetFName(), InRenderData.Key.RootBranchName);
+		}
+		else
+		{
+			// Otherwise, we separate each branch by its unique layer name.
+			UE::MoviePipeline::ConformOutputFormatStringToken(FileNameFormatString, TEXT("{layer_name}"), InParentNode->GetFName(), InRenderData.Key.RootBranchName);
+		}
 	}
-	// We can run into the scenario where the users have given layers the same name, so layer_name token won't help differentiate.
-	// To resolve this, we do the reverse search now - look to see if there's multiple branches with the same layer name, and if so
-	// we force the branch name into the token too.
-	if (InRawFrameData->HasDataFromMultipleLayersWithName(InRenderData.Key.LayerName))
+
+	// We only add the renderer name token if multiple (non-composited) renderers are present on the active branch (eg, in the case of optional PPMs).
+	if (ValidationInfo.ActiveBranchRendererCount > 1)
 	{
-		UE::MoviePipeline::ConformOutputFormatStringToken(FileNameFormatString, TEXT("{branch_name}"), InParentNode->GetFName(), InRenderData.Key.RootBranchName);
+		UE::MoviePipeline::ConformOutputFormatStringToken(FileNameFormatString, TEXT("{renderer_name}"), InParentNode->GetFName(), InRenderData.Key.RootBranchName);
 	}
 
-	if (InRawFrameData->HasMultipleRendersPerBranch(InRenderData.Key.RootBranchName))
+	// We only add the subresource token if a (non-composited) renderer on the active branch is producing more than one subresource (eg, in the case of optional PPMs).
+	if (ValidationInfo.ActiveRendererSubresourceCount > 1)
 	{
-		UE::MoviePipeline::ConformOutputFormatStringToken(FileNameFormatString, TEXT("{renderer_name}"), InParentNode->GetFName(), InRenderData.Key.RootBranchName);
+		UE::MoviePipeline::ConformOutputFormatStringToken(FileNameFormatString, TEXT("{renderer_sub_name}"), InParentNode->GetFName(), InRenderData.Key.RootBranchName);
 	}
 
 	// ToDo: Add {camera_name} validation once relevant
@@ -316,8 +308,7 @@ FString UMovieGraphImageSequenceOutputNode::CreateFileName(
 	constexpr bool bIncludeRenderPass = false;
 	constexpr bool bTestFrameNumber = true;
 	constexpr bool bIncludeCameraName = false;
-	const bool bIncludeSubRendererName = UE::MovieGraph::Private::ShouldIncludeSubRendererName(InRawFrameData->ExpectedRenderPasses);
-	UE::MoviePipeline::ValidateOutputFormatString(FileNameFormatString, bIncludeRenderPass, bTestFrameNumber,bIncludeCameraName, bIncludeSubRendererName);
+	UE::MoviePipeline::ValidateOutputFormatString(FileNameFormatString, bIncludeRenderPass, bTestFrameNumber, bIncludeCameraName);
 	
 	// Map the .ext to be specific to our output data.
 	TMap<FString, FString> AdditionalFormatArgs;
@@ -725,10 +716,44 @@ void UMovieGraphImageSequenceOutputNode_MultiLayerEXR::OnReceiveImageDataImpl(UM
 				UMoviePipelineExecutorShot* CurrentShot = InPipeline->GetActiveShotList()[ShotIndex];
 				int32 NumCameras = CurrentShot->SidecarCameras.Num();
 
-				LayerName = FString::Printf(TEXT("%s_%s"), *RenderID.RootBranchName.ToString(), *RenderID.RendererName);
+				UE::MovieGraph::FMovieGraphRenderDataValidationInfo ValidationInfo = InRawFrameData->GetValidationInfo(RenderID, /*bInDiscardCompositedRenders*/ false);
+				TArray<FString> Tokens;
+
+				if (ValidationInfo.BranchCount > 1)
+				{
+					if (ValidationInfo.LayerCount < ValidationInfo.BranchCount)
+					{
+						Tokens.Add(RenderID.RootBranchName.ToString());
+					}
+					else
+					{
+						Tokens.Add(RenderID.LayerName);
+					}
+				}
+
+				if (ValidationInfo.ActiveBranchRendererCount > 1)
+				{
+					Tokens.Add(RenderID.RendererName);
+				}
+
+				if (ValidationInfo.ActiveRendererSubresourceCount > 1)
+				{
+					Tokens.Add(RenderID.SubResourceName);
+				}
+
 				if (NumCameras > 1)
 				{
-					LayerName = FString::Printf(TEXT("%s_%s"), *LayerName, *RenderID.CameraName);
+					Tokens.Add(RenderID.CameraName);
+				}
+
+				if (ensureMsgf(!Tokens.IsEmpty(), TEXT("Missing expected EXR layer token.")))
+				{
+					LayerName = Tokens[0];
+					
+					for (int32 Index = 1; Index < Tokens.Num(); ++Index)
+					{
+						LayerName = FString::Printf(TEXT("%s_%s"), *LayerName, *Tokens[Index]);
+					}
 				}
 			}
 
@@ -851,16 +876,13 @@ FString UMovieGraphImageSequenceOutputNode_MultiLayerEXR::ResolveOutputFilename(
 	
 	// If we have more than one resolution we'll store it as "_Add" / "_Add(1)" etc via {ExtraTag}.
 	FString FileNameFormatString = InParentNode->FileNameFormat + "{ExtraTag}";
-
-	const FString FilePathFormatString = OutputSettings->OutputDirectory.Path / FileNameFormatString;
 	
 	// If we're writing more than one render pass out, we need to ensure the file name has the format string in it so we don't
 	// overwrite the same file multiple times. Burn In overlays don't count because they get composited on top of an existing file.
 	constexpr bool bIncludeRenderPass = false;
 	constexpr bool bTestFrameNumber = true;
 	constexpr bool bIncludeCameraName = false;
-	const bool bIncludeSubRendererName = UE::MovieGraph::Private::ShouldIncludeSubRendererName(InRawFrameData->ExpectedRenderPasses);
-	UE::MoviePipeline::ValidateOutputFormatString(FileNameFormatString, bIncludeRenderPass, bTestFrameNumber, bIncludeCameraName, bIncludeSubRendererName);
+	UE::MoviePipeline::ValidateOutputFormatString(FileNameFormatString, bIncludeRenderPass, bTestFrameNumber, bIncludeCameraName);
 	
 	// Create specific data that needs to override 
 	TMap<FString, FString> FormatOverrides;
@@ -891,6 +913,8 @@ FString UMovieGraphImageSequenceOutputNode_MultiLayerEXR::ResolveOutputFilename(
 	FMovieGraphFilenameResolveParams Params = UE::MovieGraph::Private::MakeResolveParams(
 		TempRenderDataIdentifier, InPipeline, InRawFrameData->EvaluatedConfig.Get(), InRawFrameData->TraversalContext, FormatOverrides);
 	
+	const FString FilePathFormatString = OutputSettings->OutputDirectory.Path / FileNameFormatString;
+
 	FString FinalFilePath = UMovieGraphBlueprintLibrary::ResolveFilenameFormatArguments(FilePathFormatString, Params, OutResolveArgs);
 
 	if (FPaths::IsRelative(FinalFilePath))
