MRG: Adding default post-processing materials on deferred renderer and path tracer nodes.

#rnx

diff --git a/Engine/Plugins/MovieScene/MovieRenderPipeline/Source/MovieRenderPipelineRenderPasses/Private/Graph/Nodes/MovieGraphDeferredPassNode.cpp b/Engine/Plugins/MovieScene/MovieRenderPipeline/Source/MovieRenderPipelineRenderPasses/Private/Graph/Nodes/MovieGraphDeferredPassNode.cpp
--- a/Engine/Plugins/MovieScene/MovieRenderPipeline/Source/MovieRenderPipelineRenderPasses/Private/Graph/Nodes/MovieGraphDeferredPassNode.cpp
+++ b/Engine/Plugins/MovieScene/MovieRenderPipeline/Source/MovieRenderPipelineRenderPasses/Private/Graph/Nodes/MovieGraphDeferredPassNode.cpp
@@ -16,6 +16,18 @@ UMovieGraphDeferredRenderPassNode::UMovieGraphDeferredRenderPassNode()
 	, bAllowOCIO(true)
 	, ViewModeIndex(VMI_Lit)
 {
+
+	// To help user knowledge we pre-seed the additional post processing materials with an array of potentially common passes.
+	TArray<FString> DefaultPostProcessMaterials;
+	DefaultPostProcessMaterials.Add(DefaultDepthAsset);
+	DefaultPostProcessMaterials.Add(DefaultMotionVectorsAsset);
+
+	for (FString& MaterialPath : DefaultPostProcessMaterials)
+	{
+		FMoviePipelinePostProcessPass& NewPass = AdditionalPostProcessMaterials.AddDefaulted_GetRef();
+		NewPass.Material = TSoftObjectPtr<UMaterialInterface>(FSoftObjectPath(MaterialPath));
+		NewPass.bEnabled = false;
+	}
 }
 
 void UMovieGraphDeferredRenderPassNode::GetFormatResolveArgs(FMovieGraphResolveArgs& OutMergedFormatArgs, const FMovieGraphRenderDataIdentifier& InRenderDataIdentifier) const
diff --git a/Engine/Plugins/MovieScene/MovieRenderPipeline/Source/MovieRenderPipelineRenderPasses/Private/Graph/Nodes/MovieGraphImagePassBaseNode.cpp b/Engine/Plugins/MovieScene/MovieRenderPipeline/Source/MovieRenderPipelineRenderPasses/Private/Graph/Nodes/MovieGraphImagePassBaseNode.cpp
--- a/Engine/Plugins/MovieScene/MovieRenderPipeline/Source/MovieRenderPipelineRenderPasses/Private/Graph/Nodes/MovieGraphImagePassBaseNode.cpp
+++ b/Engine/Plugins/MovieScene/MovieRenderPipeline/Source/MovieRenderPipelineRenderPasses/Private/Graph/Nodes/MovieGraphImagePassBaseNode.cpp
@@ -14,6 +14,9 @@ static TAutoConsoleVariable<bool> CVarMoviePipelineFlushRenderingPerLayerHack(
 	TEXT(" 1 - Flush (but negatively impacts performance)"),
 	ECVF_Default);
 
+FString UMovieGraphImagePassBaseNode::DefaultDepthAsset = TEXT("/MovieRenderPipeline/Materials/MovieRenderQueue_WorldDepth.MovieRenderQueue_WorldDepth");
+FString UMovieGraphImagePassBaseNode::DefaultMotionVectorsAsset = TEXT("/MovieRenderPipeline/Materials/MovieRenderQueue_MotionVectors.MovieRenderQueue_MotionVectors");
+
 UMovieGraphImagePassBaseNode::UMovieGraphImagePassBaseNode()
 {
 	ShowFlags = CreateDefaultSubobject<UMovieGraphShowFlags>(TEXT("ShowFlags"));
diff --git a/Engine/Plugins/MovieScene/MovieRenderPipeline/Source/MovieRenderPipelineRenderPasses/Private/Graph/Nodes/MovieGraphPathTracerPassNode.cpp b/Engine/Plugins/MovieScene/MovieRenderPipeline/Source/MovieRenderPipelineRenderPasses/Private/Graph/Nodes/MovieGraphPathTracerPassNode.cpp
--- a/Engine/Plugins/MovieScene/MovieRenderPipeline/Source/MovieRenderPipelineRenderPasses/Private/Graph/Nodes/MovieGraphPathTracerPassNode.cpp
+++ b/Engine/Plugins/MovieScene/MovieRenderPipeline/Source/MovieRenderPipelineRenderPasses/Private/Graph/Nodes/MovieGraphPathTracerPassNode.cpp
@@ -19,6 +19,18 @@ UMovieGraphPathTracerRenderPassNode::UMovieGraphPathTracerRenderPassNode()
 {
 	ShowFlags->ApplyDefaultShowFlagValue(VMI_PathTracing, true);
 	// TODO: Showflag for SetMotionBlur()?
+
+	// To help user knowledge we pre-seed the additional post processing materials with an array of potentially common passes.
+	TArray<FString> DefaultPostProcessMaterials;
+	DefaultPostProcessMaterials.Add(DefaultDepthAsset);
+	DefaultPostProcessMaterials.Add(DefaultMotionVectorsAsset);
+
+	for (FString& MaterialPath : DefaultPostProcessMaterials)
+	{
+		FMoviePipelinePostProcessPass& NewPass = AdditionalPostProcessMaterials.AddDefaulted_GetRef();
+		NewPass.Material = TSoftObjectPtr<UMaterialInterface>(FSoftObjectPath(MaterialPath));
+		NewPass.bEnabled = false;
+	}
 }
 
 #if WITH_EDITOR
diff --git a/Engine/Plugins/MovieScene/MovieRenderPipeline/Source/MovieRenderPipelineRenderPasses/Public/Graph/Nodes/MovieGraphImagePassBaseNode.h b/Engine/Plugins/MovieScene/MovieRenderPipeline/Source/MovieRenderPipelineRenderPasses/Public/Graph/Nodes/MovieGraphImagePassBaseNode.h
--- a/Engine/Plugins/MovieScene/MovieRenderPipeline/Source/MovieRenderPipelineRenderPasses/Public/Graph/Nodes/MovieGraphImagePassBaseNode.h
+++ b/Engine/Plugins/MovieScene/MovieRenderPipeline/Source/MovieRenderPipelineRenderPasses/Public/Graph/Nodes/MovieGraphImagePassBaseNode.h
@@ -73,6 +73,10 @@ protected:
 	TObjectPtr<UMovieGraphShowFlags> ShowFlags;
 
 	TArray<TUniquePtr<UE::MovieGraph::Rendering::FMovieGraphImagePassBase>> CurrentInstances;
+
+	static FString DefaultDepthAsset;
+	static FString DefaultMotionVectorsAsset;
+
 protected:
 	// UMovieGraphRenderPassNode Interface
 	virtual FString GetRendererNameImpl() const { return TEXT("UnnamedImageRenderPass"); }
