MoviePipeline: Fixed incorrect redirectors from a renaming refactor in the graph.

#rnx

diff --git a/Engine/Plugins/MovieScene/MovieRenderPipeline/Config/BaseMovieRenderPipeline.ini b/Engine/Plugins/MovieScene/MovieRenderPipeline/Config/BaseMovieRenderPipeline.ini
--- a/Engine/Plugins/MovieScene/MovieRenderPipeline/Config/BaseMovieRenderPipeline.ini
+++ b/Engine/Plugins/MovieScene/MovieRenderPipeline/Config/BaseMovieRenderPipeline.ini
@@ -8,10 +8,11 @@
 +ClassRedirects=(OldName="/Script/MovieRenderPipelineCore.MovieGraphImagePassBaseNode",NewName="/Script/MovieRenderPipelineRenderPasses.MovieGraphImagePassBaseNode")
 +ClassRedirects=(OldName="/Script/MovieRenderPipelineCore.MovieGraphPathTracedRenderPassNode",NewName="/Script/MovieRenderPipelineRenderPasses.MovieGraphPathTracerRenderPassNode")
 +ClassRedirects=(OldName="/Script/MovieRenderPipelineRenderPasses.MovieGraphPathTracedRenderPassNode",NewName="/Script/MovieRenderPipelineRenderPasses.MovieGraphPathTracerRenderPassNode")
-+ClassRedirects=(OldName="/Script/MovieRenderPipelineCore.MoviePipelineMaterialModifier",NewName="/Script/MovieRenderPipelineRenderPasses.MovieGraphMaterialModifier")
-+ClassRedirects=(OldName="/Script/MovieRenderPipelineCore.MoviePipelineVisibilityModifier",NewName="/Script/MovieRenderPipelineRenderPasses.MovieGraphRenderPropertyModifier")
-+ClassRedirects=(OldName="/Script/MovieRenderPipelineCore.MoviePipelineCollectionCommonQuery",NewName="/Script/MovieRenderPipelineRenderPasses.MovieGraphCollectionCommonQuery")
-+ClassRedirects=(OldName="/Script/MovieRenderPipelineCore.MoviePipelineCollectionLightingQuery",NewName="/Script/MovieRenderPipelineRenderPasses.MovieGraphCollectionLightingQuery")
-+ClassRedirects=(OldName="/Script/MovieRenderPipelineCore.MoviePipelineCollection",NewName="/Script/MovieRenderPipelineRenderPasses.MovieGraphCollection")
-+ClassRedirects=(OldName="/Script/MovieRenderPipelineCore.MoviePipelineCollectionQuery",NewName="/Script/MovieRenderPipelineRenderPasses.MovieGraphCollectionQueryBase")
++ClassRedirects=(OldName="/Script/MovieRenderPipelineCore.MoviePipelineCollectionModifier",NewName="/Script/MovieRenderPipelineCore.MovieGraphCollectionModifier")
++ClassRedirects=(OldName="/Script/MovieRenderPipelineCore.MoviePipelineMaterialModifier",NewName="/Script/MovieRenderPipelineCore.MovieGraphMaterialModifier")
++ClassRedirects=(OldName="/Script/MovieRenderPipelineCore.MoviePipelineVisibilityModifier",NewName="/Script/MovieRenderPipelineCore.MovieGraphRenderPropertyModifier")
++ClassRedirects=(OldName="/Script/MovieRenderPipelineCore.MoviePipelineCollectionCommonQuery",NewName="/Script/MovieRenderPipelineCore.MovieGraphCollectionCommonQuery")
++ClassRedirects=(OldName="/Script/MovieRenderPipelineCore.MoviePipelineCollectionLightingQuery",NewName="/Script/MovieRenderPipelineCore.MovieGraphCollectionLightingQuery")
++ClassRedirects=(OldName="/Script/MovieRenderPipelineCore.MoviePipelineCollection",NewName="/Script/MovieRenderPipelineCore.MovieGraphCollection")
++ClassRedirects=(OldName="/Script/MovieRenderPipelineCore.MoviePipelineCollectionQuery",NewName="/Script/MovieRenderPipelineCore.MovieGraphCollectionQueryBase")
 +EnumRedirects=(OldName="EMoviePipelineCollectionCommonQueryMode",NewName="EMovieGraphCollectionCommonQueryMode")
\ No newline at end of file
