MoviePipeline: Fixed an issue where mixed local/utc time was still being used in time estimate calculations, resulting in a large estimate/elapsed time (depending on your timezone!)

diff --git a/Engine/Plugins/MovieScene/MovieRenderPipeline/Source/MovieRenderPipelineCore/Private/Graph/MovieGraphBlueprintLibrary.cpp b/Engine/Plugins/MovieScene/MovieRenderPipeline/Source/MovieRenderPipelineCore/Private/Graph/MovieGraphBlueprintLibrary.cpp
--- a/Engine/Plugins/MovieScene/MovieRenderPipeline/Source/MovieRenderPipelineCore/Private/Graph/MovieGraphBlueprintLibrary.cpp
+++ b/Engine/Plugins/MovieScene/MovieRenderPipeline/Source/MovieRenderPipelineCore/Private/Graph/MovieGraphBlueprintLibrary.cpp
@@ -103,7 +103,7 @@ FString UMovieGraphBlueprintLibrary::ResolveFilenameFormatArguments(const FStrin
 		}
 		
 		// Add KVP data from the job (date, time, job name, job author, job comment)
-		UE::MoviePipeline::GetSharedFormatArguments(OutMergedFormatArgs.FilenameArguments, OutMergedFormatArgs.FileMetadata, InParams.InitializationTime, InParams.Version, InParams.Job);
+		UE::MoviePipeline::GetSharedFormatArguments(OutMergedFormatArgs.FilenameArguments, OutMergedFormatArgs.FileMetadata, InParams.InitializationTime, InParams.Version, InParams.Job, InParams.InitializationTimeOffset);
 	}
 	
 
diff --git a/Engine/Plugins/MovieScene/MovieRenderPipeline/Source/MovieRenderPipelineCore/Private/Graph/MovieGraphPipeline.cpp b/Engine/Plugins/MovieScene/MovieRenderPipeline/Source/MovieRenderPipelineCore/Private/Graph/MovieGraphPipeline.cpp
--- a/Engine/Plugins/MovieScene/MovieRenderPipeline/Source/MovieRenderPipelineCore/Private/Graph/MovieGraphPipeline.cpp
+++ b/Engine/Plugins/MovieScene/MovieRenderPipeline/Source/MovieRenderPipelineCore/Private/Graph/MovieGraphPipeline.cpp
@@ -84,6 +84,7 @@ void UMovieGraphPipeline::Initialize(UMoviePipelineExecutorJob* InJob, const FMo
 	CurrentJob = InJob;
 	CurrentShotIndex = 0;
 	GraphInitializationTime = FDateTime::UtcNow();
+	InitializationTimeOffset = FDateTime::Now() - FDateTime::UtcNow();
 	CVarManager = MakeShared<UE::MovieGraph::Private::FMovieGraphCVarManager>();
 
 	DuplicateJobAndConfiguration();
@@ -987,6 +988,7 @@ int32 UMovieGraphPipeline::ResolveVersionForShot(const TObjectPtr<UMoviePipeline
 	// version). Note that to do this correctly, we need a way of asking these nodes for their renderer name and all
 	// possible sub-resource names.
 	ResolveParams.InitializationTime = GetInitializationTime();
+	ResolveParams.InitializationTimeOffset = GetInitializationTimeOffset();
 	ResolveParams.DefaultFrameRate = ShotTimeStepInstance.FrameRate;
 	ResolveParams.FrameNumberOffset = BranchOutputSettingNode->FrameNumberOffset;
 	ResolveParams.RenderDataIdentifier.CameraName = Shot->InnerName;
diff --git a/Engine/Plugins/MovieScene/MovieRenderPipeline/Source/MovieRenderPipelineCore/Private/Graph/Nodes/MovieGraphAudioOutputNode.cpp b/Engine/Plugins/MovieScene/MovieRenderPipeline/Source/MovieRenderPipelineCore/Private/Graph/Nodes/MovieGraphAudioOutputNode.cpp
--- a/Engine/Plugins/MovieScene/MovieRenderPipeline/Source/MovieRenderPipelineCore/Private/Graph/Nodes/MovieGraphAudioOutputNode.cpp
+++ b/Engine/Plugins/MovieScene/MovieRenderPipeline/Source/MovieRenderPipelineCore/Private/Graph/Nodes/MovieGraphAudioOutputNode.cpp
@@ -123,6 +123,7 @@ FString UMovieGraphAudioOutputNode::GenerateOutputPath(const FMovieGraphRenderDa
 
 	FMovieGraphFilenameResolveParams ResolveParams;
 	ResolveParams.InitializationTime = CachedPipeline->GetInitializationTime();
+	ResolveParams.InitializationTimeOffset = CachedPipeline->GetInitializationTimeOffset();
 	ResolveParams.Job = CachedPipeline->GetCurrentJob();
 	ResolveParams.Shot = InShot;
 	ResolveParams.FileNameFormatOverrides = {
diff --git a/Engine/Plugins/MovieScene/MovieRenderPipeline/Source/MovieRenderPipelineCore/Private/Graph/Nodes/MovieGraphCommandLineEncoderNode.cpp b/Engine/Plugins/MovieScene/MovieRenderPipeline/Source/MovieRenderPipelineCore/Private/Graph/Nodes/MovieGraphCommandLineEncoderNode.cpp
--- a/Engine/Plugins/MovieScene/MovieRenderPipeline/Source/MovieRenderPipelineCore/Private/Graph/Nodes/MovieGraphCommandLineEncoderNode.cpp
+++ b/Engine/Plugins/MovieScene/MovieRenderPipeline/Source/MovieRenderPipelineCore/Private/Graph/Nodes/MovieGraphCommandLineEncoderNode.cpp
@@ -394,6 +394,7 @@ FString UMovieGraphCommandLineEncoderNode::GetResolvedOutputFilename(const FMovi
 	// don't have versions when using image sequences).
 	FMovieGraphFilenameResolveParams ResolveParams;
 	ResolveParams.InitializationTime = CachedPipeline->GetInitializationTime();
+	ResolveParams.InitializationTimeOffset = CachedPipeline->GetInitializationTimeOffset();
 	ResolveParams.Job = CachedPipeline->GetCurrentJob();
 	ResolveParams.Shot = Shot.Get();
 	ResolveParams.FileNameFormatOverrides = FormatOverrides;
@@ -446,6 +447,7 @@ void UMovieGraphCommandLineEncoderNode::GenerateTemporaryEncoderInputFiles(const
 
 		FMovieGraphFilenameResolveParams ResolveParams;
 		ResolveParams.InitializationTime = CachedPipeline->GetInitializationTime();
+		ResolveParams.InitializationTimeOffset = CachedPipeline->GetInitializationTimeOffset();
 		ResolveParams.Job = CachedPipeline->GetCurrentJob();
 		ResolveParams.Shot = InParams.Shot.Get();
 		ResolveParams.FileNameFormatOverrides = FormatOverrides;
diff --git a/Engine/Plugins/MovieScene/MovieRenderPipeline/Source/MovieRenderPipelineCore/Private/MoviePipeline.cpp b/Engine/Plugins/MovieScene/MovieRenderPipeline/Source/MovieRenderPipelineCore/Private/MoviePipeline.cpp
--- a/Engine/Plugins/MovieScene/MovieRenderPipeline/Source/MovieRenderPipelineCore/Private/MoviePipeline.cpp
+++ b/Engine/Plugins/MovieScene/MovieRenderPipeline/Source/MovieRenderPipelineCore/Private/MoviePipeline.cpp
@@ -318,7 +318,8 @@ void UMoviePipeline::Initialize(UMoviePipelineExecutorJob* InJob)
 
 	// Initialization is complete. This engine frame is a wash (because the tick started with a 
 	// delta time not generated by us) so we'll wait until the next engine frame to start rendering.
-	InitializationTime = FDateTime::Now();
+	InitializationTime = FDateTime::UtcNow();
+	InitializationTimeOffset = FDateTime::Now() - InitializationTime;
 
 	// We need to start our CachedOutputState at the time the Game World is currently. This way
 	// as MRQ counts up time, the render world stays in sync allowing comparisons between last rendered time
@@ -331,6 +332,7 @@ void UMoviePipeline::Initialize(UMoviePipelineExecutorJob* InJob)
 	// version number comes after theo shot name, then each shot may have its own 'local' version number
 	FMoviePipelineFilenameResolveParams Params;
 	Params.InitializationTime = InitializationTime;
+	Params.InitializationTimeOffset = InitializationTimeOffset;
 	Params.Job = GetCurrentJob();
 
 	for (UMoviePipelineExecutorShot* Shot : ActiveShotList)
@@ -1582,6 +1584,7 @@ void UMoviePipeline::ResolveFilenameFormatArguments(const FString& InFormatStrin
 	Params.bForceRelativeFrameNumbers = bForceRelativeFrameNumbers;
 	Params.FileNameFormatOverrides = InFormatOverrides;
 	Params.InitializationTime = InitializationTime;
+	Params.InitializationTimeOffset = InitializationTimeOffset;
 	Params.Job = GetCurrentJob();
 	Params.AdditionalFrameNumberOffset = InFrameNumberOffset;
 
diff --git a/Engine/Plugins/MovieScene/MovieRenderPipeline/Source/MovieRenderPipelineCore/Private/MoviePipelineBlueprintLibrary.cpp b/Engine/Plugins/MovieScene/MovieRenderPipeline/Source/MovieRenderPipelineCore/Private/MoviePipelineBlueprintLibrary.cpp
--- a/Engine/Plugins/MovieScene/MovieRenderPipeline/Source/MovieRenderPipelineCore/Private/MoviePipelineBlueprintLibrary.cpp
+++ b/Engine/Plugins/MovieScene/MovieRenderPipeline/Source/MovieRenderPipelineCore/Private/MoviePipelineBlueprintLibrary.cpp
@@ -994,7 +994,7 @@ void UMoviePipelineBlueprintLibrary::ResolveFilenameFormatArguments(const FStrin
 	// And from ourself
 	{
 		// Use the shared function as UMoviePipelinePrimaryConfig::GetFormatArguments to ensure all the time variables get overwritten.
-		UE::MoviePipeline::GetSharedFormatArguments(OutMergedFormatArgs.FilenameArguments, OutMergedFormatArgs.FileMetadata, InParams.InitializationTime, InParams.InitializationVersion, InParams.Job);
+		UE::MoviePipeline::GetSharedFormatArguments(OutMergedFormatArgs.FilenameArguments, OutMergedFormatArgs.FileMetadata, InParams.InitializationTime, InParams.InitializationVersion, InParams.Job, InParams.InitializationTimeOffset);
 		
 		// By default, we don't want to show frame duplication numbers. If we need to start writing them,
 		// they need to come before the frame number (so that tools recognize them as a sequence).
diff --git a/Engine/Plugins/MovieScene/MovieRenderPipeline/Source/MovieRenderPipelineCore/Private/MoviePipelineCommandLineEncoder.cpp b/Engine/Plugins/MovieScene/MovieRenderPipeline/Source/MovieRenderPipelineCore/Private/MoviePipelineCommandLineEncoder.cpp
--- a/Engine/Plugins/MovieScene/MovieRenderPipeline/Source/MovieRenderPipelineCore/Private/MoviePipelineCommandLineEncoder.cpp
+++ b/Engine/Plugins/MovieScene/MovieRenderPipeline/Source/MovieRenderPipelineCore/Private/MoviePipelineCommandLineEncoder.cpp
@@ -184,6 +184,7 @@ void UMoviePipelineCommandLineEncoder::StartEncodingProcess(TArray<FMoviePipelin
 		{
 			FMoviePipelineFilenameResolveParams ResolveParams;
 			ResolveParams.InitializationTime = GetPipeline()->GetInitializationTime();
+			ResolveParams.InitializationTimeOffset = GetPipeline()->GetInitializationTimeOffset();
 			ResolveParams.Job = GetPipeline()->GetCurrentJob();
 			ResolveParams.ShotOverride = RenderPass.Value.Shot.Get();
 			ResolveParams.FileNameOverride = FileNameFormatString;
diff --git a/Engine/Plugins/MovieScene/MovieRenderPipeline/Source/MovieRenderPipelineCore/Private/MoviePipelineLinearExecutor.cpp b/Engine/Plugins/MovieScene/MovieRenderPipeline/Source/MovieRenderPipelineCore/Private/MoviePipelineLinearExecutor.cpp
--- a/Engine/Plugins/MovieScene/MovieRenderPipeline/Source/MovieRenderPipelineCore/Private/MoviePipelineLinearExecutor.cpp
+++ b/Engine/Plugins/MovieScene/MovieRenderPipeline/Source/MovieRenderPipelineCore/Private/MoviePipelineLinearExecutor.cpp
@@ -35,7 +35,7 @@ void UMoviePipelineLinearExecutorBase::Execute_Implementation(UMoviePipelineQueu
 	// before moving onto the next one. This may be parallelizable in the future (either multiple PIE
 	// sessions, or multiple external processes) but ideally one render would maximize resource usage anyways...
 	Queue = InPipelineQueue;
-	InitializationTime = FDateTime::Now();
+	InitializationTime = FDateTime::UtcNow();
 	JobsStarted = 0;
 
 	UE_LOG(LogMovieRenderPipeline, Log, TEXT("MoviePipelineLinearExecutorBase starting %d jobs."), InPipelineQueue->GetJobs().Num());
diff --git a/Engine/Plugins/MovieScene/MovieRenderPipeline/Source/MovieRenderPipelineCore/Private/MoviePipelinePrimaryConfig.cpp b/Engine/Plugins/MovieScene/MovieRenderPipeline/Source/MovieRenderPipelineCore/Private/MoviePipelinePrimaryConfig.cpp
--- a/Engine/Plugins/MovieScene/MovieRenderPipeline/Source/MovieRenderPipelineCore/Private/MoviePipelinePrimaryConfig.cpp
+++ b/Engine/Plugins/MovieScene/MovieRenderPipeline/Source/MovieRenderPipelineCore/Private/MoviePipelinePrimaryConfig.cpp
@@ -130,9 +130,10 @@ void UMoviePipelinePrimaryConfig::GetFormatArguments(FMoviePipelineFormatArgs& I
 
 		// Normally these are filled when resolving the file name by the job (so that the time is shared), but stub them in here so
 		// they show up in the UI with a value.
-		FDateTime CurrentTime = FDateTime::Now();
+		FDateTime CurrentTime = FDateTime::UtcNow();
 		int32 DummyVersionNumber = 1;
-		UE::MoviePipeline::GetSharedFormatArguments(InOutFormatArgs.FilenameArguments, InOutFormatArgs.FileMetadata, CurrentTime, DummyVersionNumber, InOutFormatArgs.InJob);
+		const FTimespan InitializationTimeOffset = FDateTime::Now() - FDateTime::UtcNow();
+		UE::MoviePipeline::GetSharedFormatArguments(InOutFormatArgs.FilenameArguments, InOutFormatArgs.FileMetadata, CurrentTime, DummyVersionNumber, InOutFormatArgs.InJob, InitializationTimeOffset);
 
 		// Let the output state fill out some too, since its the keeper of the information.
 		UMoviePipelineOutputSetting* OutputSettings = FindSetting<UMoviePipelineOutputSetting>();
diff --git a/Engine/Plugins/MovieScene/MovieRenderPipeline/Source/MovieRenderPipelineCore/Private/MoviePipelineUtils.cpp b/Engine/Plugins/MovieScene/MovieRenderPipeline/Source/MovieRenderPipelineCore/Private/MoviePipelineUtils.cpp
--- a/Engine/Plugins/MovieScene/MovieRenderPipeline/Source/MovieRenderPipelineCore/Private/MoviePipelineUtils.cpp
+++ b/Engine/Plugins/MovieScene/MovieRenderPipeline/Source/MovieRenderPipelineCore/Private/MoviePipelineUtils.cpp
@@ -852,24 +852,33 @@ namespace UE
 			return FPlatformProcess::UserName(false);
 		}
 
-		void GetSharedFormatArguments(TMap<FString, FString>& InFilenameArguments, TMap<FString, FString>& InFileMetadata, const FDateTime& InDateTime, const int32 InVersionNumber, const UMoviePipelineExecutorJob* InJob)
+		void GetSharedFormatArguments(TMap<FString, FString>& InFilenameArguments, TMap<FString, FString>& InFileMetadata, const FDateTime& InDateTime, const int32 InVersionNumber, const UMoviePipelineExecutorJob* InJob, const FTimespan& InInitializationTimeOffset)
 		{
-			InFilenameArguments.Add(TEXT("date"), InDateTime.ToString(TEXT("%Y.%m.%d")));
-			InFilenameArguments.Add(TEXT("time"), InDateTime.ToString(TEXT("%H.%M.%S")));
-			InFilenameArguments.Add(TEXT("year"), InDateTime.ToString(TEXT("%Y")));
-			InFilenameArguments.Add(TEXT("month"), InDateTime.ToString(TEXT("%m")));
-			InFilenameArguments.Add(TEXT("day"), InDateTime.ToString(TEXT("%d")));
+			FDateTime DateTimeLocal = InDateTime + InInitializationTimeOffset; 
+
+			FString LocalDateStr = DateTimeLocal.ToString(TEXT("%Y.%m.%d"));
+			FString LocalTimeStr = DateTimeLocal.ToString(TEXT("%H.%M.%S"));
+			FString LocalYearStr = DateTimeLocal.ToString(TEXT("%Y"));
+			FString LocalMonthStr = DateTimeLocal.ToString(TEXT("%m"));
+			FString LocalDayStr = DateTimeLocal.ToString(TEXT("%d"));
+
+
+			InFilenameArguments.Add(TEXT("date"), LocalDateStr);
+			InFilenameArguments.Add(TEXT("time"), LocalTimeStr);
+			InFilenameArguments.Add(TEXT("year"), LocalYearStr);
+			InFilenameArguments.Add(TEXT("month"), LocalMonthStr);
+			InFilenameArguments.Add(TEXT("day"), LocalDayStr);
 
 
 			FString VersionText = FString::Printf(TEXT("v%0*d"), 3, InVersionNumber);
 
 			InFilenameArguments.Add(TEXT("version"), VersionText);
 
-			InFileMetadata.Add(TEXT("unreal/jobDate"), InDateTime.ToString(TEXT("%Y.%m.%d")));
-			InFileMetadata.Add(TEXT("unreal/jobTime"), InDateTime.ToString(TEXT("%H.%M.%S")));
-			InFileMetadata.Add(TEXT("unreal/jobYear"), InDateTime.ToString(TEXT("%Y")));
-			InFileMetadata.Add(TEXT("unreal/jobMonth"), InDateTime.ToString(TEXT("%m")));
-			InFileMetadata.Add(TEXT("unreal/jobDay"), InDateTime.ToString(TEXT("%d")));
+			InFileMetadata.Add(TEXT("unreal/jobDate"), LocalDateStr);
+			InFileMetadata.Add(TEXT("unreal/jobTime"), LocalTimeStr);
+			InFileMetadata.Add(TEXT("unreal/jobYear"), LocalYearStr);
+			InFileMetadata.Add(TEXT("unreal/jobMonth"), LocalMonthStr);
+			InFileMetadata.Add(TEXT("unreal/jobDay"), LocalDayStr);
 
 			InFileMetadata.Add(TEXT("unreal/jobVersion"), FString::FromInt(InVersionNumber));
 
diff --git a/Engine/Plugins/MovieScene/MovieRenderPipeline/Source/MovieRenderPipelineCore/Public/Graph/MovieGraphFilenameResolveParams.h b/Engine/Plugins/MovieScene/MovieRenderPipeline/Source/MovieRenderPipelineCore/Public/Graph/MovieGraphFilenameResolveParams.h
--- a/Engine/Plugins/MovieScene/MovieRenderPipeline/Source/MovieRenderPipelineCore/Public/Graph/MovieGraphFilenameResolveParams.h
+++ b/Engine/Plugins/MovieScene/MovieRenderPipeline/Source/MovieRenderPipelineCore/Public/Graph/MovieGraphFilenameResolveParams.h
@@ -122,9 +122,16 @@ struct FMovieGraphFilenameResolveParams
 	UPROPERTY(BlueprintReadWrite, EditAnywhere, Category = "Movie Graph")
 	bool bEnsureAbsolutePath;
 	
-	/** The initialization time for this job. Used to resolve time-based format arguments. */
+	/**
+	* The initialization time for this job. Used to resolve time-based format arguments. This should be in UTC, if you want filenames in a local timezone
+	* then you should set InitializationTimeOffset to your offset from UTC.
+	*/
 	UPROPERTY(BlueprintReadWrite, EditAnywhere, Category = "Movie Graph")
 	FDateTime InitializationTime;
+
+	/** What offset should be applied to InitializationTime when generating {time} related filename tokens? Likely your offset from UTC if you want local time. */
+	UPROPERTY(BlueprintReadWrite, EditAnywhere, Category = "Movie Render Pipeline")
+	FTimespan InitializationTimeOffset;
 	
 	/** When converting frame numbers to strings, how many digits should we pad them up to? ie: 5 => 0005 with a count of 4. */
 	UPROPERTY(BlueprintReadWrite, EditAnywhere, Category = "Movie Graph")
diff --git a/Engine/Plugins/MovieScene/MovieRenderPipeline/Source/MovieRenderPipelineCore/Public/Graph/MovieGraphPipeline.h b/Engine/Plugins/MovieScene/MovieRenderPipeline/Source/MovieRenderPipelineCore/Public/Graph/MovieGraphPipeline.h
--- a/Engine/Plugins/MovieScene/MovieRenderPipeline/Source/MovieRenderPipelineCore/Public/Graph/MovieGraphPipeline.h
+++ b/Engine/Plugins/MovieScene/MovieRenderPipeline/Source/MovieRenderPipelineCore/Public/Graph/MovieGraphPipeline.h
@@ -52,16 +52,26 @@ public:
 	UFUNCTION(BlueprintPure, Category = "Movie Graph")
 	FDateTime GetInitializationTime() const { return GraphInitializationTime; }
 
+	/** 
+	* The offset that should be applied to the GetInitializationTime() when generating
+	* the {time} related filename tokens. GetInitializationTime() is in UTC so this is
+	* either zero (if you called SetInitializationTime) or your offset from UTC.
+	*/
+	UFUNCTION(BlueprintPure, Category = "Movie Render Pipeline")
+	FTimespan GetInitializationTimeOffset() const { return InitializationTimeOffset; }
+
+
 	/**
 	* Override the time this movie pipeline was initialized at. This can be used for render farms
 	* to ensure that jobs on all machines use the same date/time instead of each calculating it locally.
+	* Clears the auto-calculated InitializationTimeOffset, meaning time tokens will be written in UTC.
 	*
 	* Needs to be called after ::Initialize(...)
 	*
-	* @param InDateTime - The DateTime object to return for GetInitializationTime.
+	* @param InDateTime - Expected to be in UTC timezone.
 	*/
 	UFUNCTION(BlueprintCallable, Category = "Movie Graph")
-	void SetInitializationTime(const FDateTime& InDateTime) { GraphInitializationTime = InDateTime; }
+	void SetInitializationTime(const FDateTime& InDateTime) { GraphInitializationTime = InDateTime; InitializationTimeOffset = FTimespan(); }
 
 	/** Gets the graph config for the shot if one was specified for the shot. Otherwise, gets the graph config for the associated primary job. */
 	UMovieGraphConfig* GetRootGraphForShot(UMoviePipelineExecutorShot* InShot) const;
@@ -257,6 +267,9 @@ protected:
 	/** True if we're in a TransitionToState call. Used to prevent reentrancy. */
 	bool bIsTransitioningState;
 
+	/** When we originally initialize we store the offset from UTC (which is what GetInitializationTime() is in), but we clear this if you call SetInitializationTime. */
+	FTimespan InitializationTimeOffset;
+
 	/** True if RequestShutdown() was called. At the start of the next frame we will stop producing frames and start shutting down. */
 	FThreadSafeBool bShutdownRequested;
 
diff --git a/Engine/Plugins/MovieScene/MovieRenderPipeline/Source/MovieRenderPipelineCore/Public/MoviePipeline.h b/Engine/Plugins/MovieScene/MovieRenderPipeline/Source/MovieRenderPipelineCore/Public/MoviePipeline.h
--- a/Engine/Plugins/MovieScene/MovieRenderPipeline/Source/MovieRenderPipelineCore/Public/MoviePipeline.h
+++ b/Engine/Plugins/MovieScene/MovieRenderPipeline/Source/MovieRenderPipelineCore/Public/MoviePipeline.h
@@ -68,16 +68,25 @@ public:
 	UFUNCTION(BlueprintPure, Category = "Movie Render Pipeline")
 	FDateTime GetInitializationTime() const { return InitializationTime; }
 
+	/** 
+	* The offset that should be applied to the GetInitializationTime() when generating
+	* the {time} related filename tokens. GetInitializationTime() is in UTC so this is
+	* either zero (if you called SetInitializationTime) or your offset from UTC.
+	*/
+	UFUNCTION(BlueprintPure, Category = "Movie Render Pipeline")
+	FTimespan GetInitializationTimeOffset() const { return InitializationTimeOffset; }
+
 	/**
 	* Override the time this movie pipeline was initialized at. This can be used for render farms
 	* to ensure that jobs on all machines use the same date/time instead of each calculating it locally.
+	* Clears the auto-calculated InitializationTimeOffset, meaning time tokens will be written in UTC.
 	*
 	* Needs to be called after ::Initialize(...)
 	*
-	* @param InDateTime - The DateTime object to return for GetInitializationTime.
+	* @param InDateTime - Expected to be in UTC timezone.
 	*/
 	UFUNCTION(BlueprintCallable, Category = "Movie Render Pipeline")
-	void SetInitializationTime(const FDateTime& InDateTime) { InitializationTime = InDateTime; }
+	void SetInitializationTime(const FDateTime& InDateTime) { InitializationTime = InDateTime; InitializationTimeOffset = FTimespan(); }
 
 	/** Deprecated. Use OnMoviePipelineWorkFinished() instead. */
 	UE_DEPRECATED(4.27, "Use OnMoviePipelineWorkFinished() instead.")
@@ -383,6 +392,9 @@ private:
 	/** When using temporal sub-frame stepping common counts (such as 3) don't result in whole ticks. We keep track of how many ticks we lose so we can add them the next time there's a chance. */
 	float AccumulatedTickSubFrameDeltas;
 
+	/** When we originally initialize we store the offset from UTC (which is what GetInitializationTime() is in), but we clear this if you call SetInitializationTime. */
+	FTimespan InitializationTimeOffset;
+
 	/** Deprecated. */
 	FMoviePipelineFinishedNative OnMoviePipelineFinishedDelegateNative;
 
diff --git a/Engine/Plugins/MovieScene/MovieRenderPipeline/Source/MovieRenderPipelineCore/Public/MoviePipelineUtils.h b/Engine/Plugins/MovieScene/MovieRenderPipeline/Source/MovieRenderPipelineCore/Public/MoviePipelineUtils.h
--- a/Engine/Plugins/MovieScene/MovieRenderPipeline/Source/MovieRenderPipelineCore/Public/MoviePipelineUtils.h
+++ b/Engine/Plugins/MovieScene/MovieRenderPipeline/Source/MovieRenderPipelineCore/Public/MoviePipelineUtils.h
@@ -138,7 +138,7 @@ namespace UE
 		MOVIERENDERPIPELINECORE_API void DeduplicateNameArray(TArray<FString>& InOutNames);
 
 		MOVIERENDERPIPELINECORE_API FString GetJobAuthor(const UMoviePipelineExecutorJob* InJob);
-		MOVIERENDERPIPELINECORE_API void GetSharedFormatArguments(TMap<FString, FString>& InFilenameArguments, TMap<FString, FString>& InFileMetadata, const FDateTime& InDateTime, const int32 InVersionNumber, const UMoviePipelineExecutorJob* InJob);
+		MOVIERENDERPIPELINECORE_API void GetSharedFormatArguments(TMap<FString, FString>& InFilenameArguments, TMap<FString, FString>& InFileMetadata, const FDateTime& InDateTime, const int32 InVersionNumber, const UMoviePipelineExecutorJob* InJob, const FTimespan& InInitializationTimeOffset = FTimespan());
 		MOVIERENDERPIPELINECORE_API void GetHardwareUsageMetadata(TMap<FString, FString>& InFileMetadata, const FString& InOutputDir);
 		MOVIERENDERPIPELINECORE_API void GetMetadataFromCineCamera(class UCineCameraComponent* InComponent, const FString& InCameraName, const FString& InRenderPassName, TMap<FString, FString>& InOutMetadata);
 		MOVIERENDERPIPELINECORE_API void GetMetadataFromCameraLocRot(const FString& InCameraName, const FString& InRenderPassName, const FVector& InCurLoc, const FRotator& InCurRot, const FVector& InPrevLoc, const FRotator& InPrevRot, TMap<FString, FString>& InOutMetadata);
diff --git a/Engine/Plugins/MovieScene/MovieRenderPipeline/Source/MovieRenderPipelineCore/Public/MovieRenderPipelineDataTypes.h b/Engine/Plugins/MovieScene/MovieRenderPipeline/Source/MovieRenderPipelineCore/Public/MovieRenderPipelineDataTypes.h
--- a/Engine/Plugins/MovieScene/MovieRenderPipeline/Source/MovieRenderPipelineCore/Public/MovieRenderPipelineDataTypes.h
+++ b/Engine/Plugins/MovieScene/MovieRenderPipeline/Source/MovieRenderPipelineCore/Public/MovieRenderPipelineDataTypes.h
@@ -856,6 +856,10 @@ struct FMoviePipelineFilenameResolveParams
 	UPROPERTY(BlueprintReadWrite, EditAnywhere, Category = "Movie Render Pipeline")
 	FDateTime InitializationTime;
 
+	/** What offset should be applied to InitializationTime when generating {time} related filename tokens? Likely your offset from UTC if you want local time. */
+	UPROPERTY(BlueprintReadWrite, EditAnywhere, Category = "Movie Render Pipeline")
+	FTimespan InitializationTimeOffset;
+
 	/** The version for this job. Used to resolve version format arguments. */
 	UPROPERTY(BlueprintReadWrite, EditAnywhere, Category = "Movie Render Pipeline")
 	int32 InitializationVersion;
diff --git a/Engine/Plugins/MovieScene/MovieRenderPipeline/Source/MovieRenderPipelineEditor/Private/Customizations/Graph/MovieGraphFormatTokenCustomization.cpp b/Engine/Plugins/MovieScene/MovieRenderPipeline/Source/MovieRenderPipelineEditor/Private/Customizations/Graph/MovieGraphFormatTokenCustomization.cpp
--- a/Engine/Plugins/MovieScene/MovieRenderPipeline/Source/MovieRenderPipelineEditor/Private/Customizations/Graph/MovieGraphFormatTokenCustomization.cpp
+++ b/Engine/Plugins/MovieScene/MovieRenderPipeline/Source/MovieRenderPipelineEditor/Private/Customizations/Graph/MovieGraphFormatTokenCustomization.cpp
@@ -95,6 +95,7 @@ void FMovieGraphFormatTokenCustomization::GetFormatArguments(FMoviePipelineForma
 
 	static const FDateTime CurrentTime = FDateTime::Now();
 	static constexpr int32 DummyVersionNumber = 1;
+	const FTimespan InitializationTimeOffset = FDateTime::Now() - FDateTime::UtcNow();
 	UE::MoviePipeline::GetSharedFormatArguments(
-		InOutFormatArgs.FilenameArguments, InOutFormatArgs.FileMetadata, CurrentTime, DummyVersionNumber, InOutFormatArgs.InJob);
+		InOutFormatArgs.FilenameArguments, InOutFormatArgs.FileMetadata, CurrentTime, DummyVersionNumber, InOutFormatArgs.InJob, InitializationTimeOffset);
 }
diff --git a/Engine/Plugins/MovieScene/MovieRenderPipeline/Source/MovieRenderPipelineEditor/Private/MoviePipelineEditorBlueprintLibrary.cpp b/Engine/Plugins/MovieScene/MovieRenderPipeline/Source/MovieRenderPipelineEditor/Private/MoviePipelineEditorBlueprintLibrary.cpp
--- a/Engine/Plugins/MovieScene/MovieRenderPipeline/Source/MovieRenderPipelineEditor/Private/MoviePipelineEditorBlueprintLibrary.cpp
+++ b/Engine/Plugins/MovieScene/MovieRenderPipeline/Source/MovieRenderPipelineEditor/Private/MoviePipelineEditorBlueprintLibrary.cpp
@@ -258,7 +258,9 @@ FString UMoviePipelineEditorBlueprintLibrary::ResolveOutputDirectoryFromJob(UMov
 		const bool bGetNextVersion = false;
 		FMovieGraphFilenameResolveParams Params;
 		Params.Job = InJob;
-		Params.InitializationTime = FDateTime::Now();
+		Params.InitializationTime = FDateTime::UtcNow();
+		Params.InitializationTimeOffset = FDateTime::Now() - FDateTime::UtcNow();
+
 		Params.Version = UMovieGraphBlueprintLibrary::ResolveVersionNumber(Params, bGetNextVersion);
 		Params.RenderDataIdentifier.RootBranchName = UMovieGraphNode::GlobalsPinName;
 		
@@ -281,6 +283,7 @@ FString UMoviePipelineEditorBlueprintLibrary::ResolveOutputDirectoryFromJob(UMov
 		Params.ShotNameOverride = DummyTag;
 		Params.CameraNameOverride = DummyTag;
 		Params.InitializationTime = FDateTime::UtcNow();
+		Params.InitializationTimeOffset = FDateTime::Now() - FDateTime::UtcNow();
 		Params.InitializationVersion = UMoviePipelineBlueprintLibrary::ResolveVersionNumber(Params, bGetNextVersion);
 
 		FMoviePipelineFormatArgs Dummy;
diff --git a/Engine/Plugins/MovieScene/MovieRenderPipeline/Source/MovieRenderPipelineEditor/Private/MoviePipelinePIEExecutor.cpp b/Engine/Plugins/MovieScene/MovieRenderPipeline/Source/MovieRenderPipelineEditor/Private/MoviePipelinePIEExecutor.cpp
--- a/Engine/Plugins/MovieScene/MovieRenderPipeline/Source/MovieRenderPipelineEditor/Private/MoviePipelinePIEExecutor.cpp
+++ b/Engine/Plugins/MovieScene/MovieRenderPipeline/Source/MovieRenderPipelineEditor/Private/MoviePipelinePIEExecutor.cpp
@@ -300,7 +300,11 @@ void UMoviePipelinePIEExecutor::OnTick()
 			}
 			else if (UMovieGraphPipeline* PipelineAsGraph = Cast<UMovieGraphPipeline>(ActiveMoviePipeline))
 			{
-				// PipelineAsGraph->Initialize(Queue->GetJobs()[CurrentPipelineIndex], FMovieGraphInitConfig());
+				PipelineAsGraph->Initialize(Queue->GetJobs()[CurrentPipelineIndex], FMovieGraphInitConfig());
+				if (CustomInitializationTime.IsSet())
+				{
+					PipelineAsGraph->SetInitializationTime(CustomInitializationTime.GetValue());
+				}
 			}
 		}
 
diff --git a/Engine/Plugins/MovieScene/MovieRenderPipeline/Source/MovieRenderPipelineRenderPasses/Private/MovieGraphImageSequenceOutputNode.cpp b/Engine/Plugins/MovieScene/MovieRenderPipeline/Source/MovieRenderPipelineRenderPasses/Private/MovieGraphImageSequenceOutputNode.cpp
--- a/Engine/Plugins/MovieScene/MovieRenderPipeline/Source/MovieRenderPipelineRenderPasses/Private/MovieGraphImageSequenceOutputNode.cpp
+++ b/Engine/Plugins/MovieScene/MovieRenderPipeline/Source/MovieRenderPipelineRenderPasses/Private/MovieGraphImageSequenceOutputNode.cpp
@@ -65,6 +65,7 @@ namespace UE::MovieGraph::Private
 		Params.bEnsureAbsolutePath = true;
 		Params.FileNameFormatOverrides = InAdditionalFormatArgs;
 		Params.InitializationTime = InPipeline->GetInitializationTime();
+		Params.InitializationTimeOffset = InPipeline->GetInitializationTimeOffset();
 		Params.Shot = Shot;
 		Params.Job = InPipeline->GetCurrentJob();
 
