MoviePipeline: Ensure that all graph-related file format tokens are available in the autocomplete for the File Name Format property within the graph editor.

* Updated methods that return/resolve tokens to be tolerant of missing resolve context. This allows us to get unresolved tokens back when tokens do not need to be resolved.
* The details panel customization now uses ResolveFilenameFormatArguments(), so the tokens that are presented in the UI are the exact same as the tokens used in the graph's resolve process.
* Tokens in the details panel customization are now shown alphabetically.

#rnx

diff --git a/Engine/Plugins/MovieScene/MovieRenderPipeline/Source/MovieRenderPipelineCore/Private/Graph/MovieGraphBlueprintLibrary.cpp b/Engine/Plugins/MovieScene/MovieRenderPipeline/Source/MovieRenderPipelineCore/Private/Graph/MovieGraphBlueprintLibrary.cpp
--- a/Engine/Plugins/MovieScene/MovieRenderPipeline/Source/MovieRenderPipelineCore/Private/Graph/MovieGraphBlueprintLibrary.cpp
+++ b/Engine/Plugins/MovieScene/MovieRenderPipeline/Source/MovieRenderPipelineCore/Private/Graph/MovieGraphBlueprintLibrary.cpp
@@ -90,18 +90,20 @@ FString UMovieGraphBlueprintLibrary::ResolveFilenameFormatArguments(const FStrin
 		OutMergedFormatArgs.FilenameArguments.Add(TEXT("frame_rate"), FString::SanitizeFloat(FrameRate));
 		OutMergedFormatArgs.FileMetadata.Add(TEXT("unreal/frameRate"), FString::SanitizeFloat(FrameRate));
 
+		FString LevelName;
+		FString SequenceName;
 		if (InParams.Job)
 		{
-			FString LevelName = InParams.Job->Map.GetAssetName();
-			FString SequenceName = InParams.Job->Sequence.GetAssetName();
-			
-			OutMergedFormatArgs.FilenameArguments.Add(TEXT("level_name"), SequenceName);
-			OutMergedFormatArgs.FilenameArguments.Add(TEXT("sequence_name"), SequenceName);
-			
-			OutMergedFormatArgs.FileMetadata.Add(TEXT("unreal/levelName"), LevelName);
-			OutMergedFormatArgs.FileMetadata.Add(TEXT("unreal/sequenceName"), SequenceName);
+			LevelName = InParams.Job->Map.GetAssetName();
+			SequenceName = InParams.Job->Sequence.GetAssetName();
 		}
 		
+		OutMergedFormatArgs.FilenameArguments.Add(TEXT("level_name"), LevelName);
+		OutMergedFormatArgs.FilenameArguments.Add(TEXT("sequence_name"), SequenceName);
+		
+		OutMergedFormatArgs.FileMetadata.Add(TEXT("unreal/levelName"), LevelName);
+		OutMergedFormatArgs.FileMetadata.Add(TEXT("unreal/sequenceName"), SequenceName);
+		
 		// Add KVP data from the job (date, time, job name, job author, job comment)
 		UE::MoviePipeline::GetSharedFormatArguments(OutMergedFormatArgs.FilenameArguments, OutMergedFormatArgs.FileMetadata, InParams.InitializationTime, InParams.Version, InParams.Job, InParams.InitializationTimeOffset);
 	}
@@ -116,13 +118,17 @@ FString UMovieGraphBlueprintLibrary::ResolveFilenameFormatArguments(const FStrin
 	// once per run there. We don't cache the returned results here because you could potentially add/remove classes (via Blueprints) which would invalidate our cache.
 	for (UClass* InClass : AllSettingsNodeClasses)
 	{
-		UMovieGraphSettingNode* SettingInstance = nullptr;
+		const UMovieGraphSettingNode* SettingInstance = nullptr;
 		if (InParams.EvaluatedConfig)
 		{
 			const bool bIncludeCDOs = true;
 			const bool bExactMatch = true;
 			SettingInstance = InParams.EvaluatedConfig->GetSettingForBranch(InClass, InParams.RenderDataIdentifier.RootBranchName, bIncludeCDOs, bExactMatch);
 		}
+		else
+		{
+			SettingInstance = GetDefault<UMovieGraphSettingNode>(InClass);
+		}
 
 		if (SettingInstance)
 		{
@@ -156,6 +162,12 @@ FString UMovieGraphBlueprintLibrary::ResolveFilenameFormatArguments(const FStrin
 		NamedArgs.Add(Argument.Key, Argument.Value);
 	}
 
+	// If no format string is provided, there's nothing left to do.
+	if (InFormatString.IsEmpty())
+	{
+		return FString();
+	}
+
 	// Apply all of our named args to the file and generate a path.
 	FString BaseFilename = FString::Format(*InFormatString, NamedArgs);
 	if (InParams.bEnsureAbsolutePath)
diff --git a/Engine/Plugins/MovieScene/MovieRenderPipeline/Source/MovieRenderPipelineCore/Private/MoviePipelineUtils.cpp b/Engine/Plugins/MovieScene/MovieRenderPipeline/Source/MovieRenderPipelineCore/Private/MoviePipelineUtils.cpp
--- a/Engine/Plugins/MovieScene/MovieRenderPipeline/Source/MovieRenderPipelineCore/Private/MoviePipelineUtils.cpp
+++ b/Engine/Plugins/MovieScene/MovieRenderPipeline/Source/MovieRenderPipelineCore/Private/MoviePipelineUtils.cpp
@@ -882,14 +882,15 @@ namespace UE
 
 			InFileMetadata.Add(TEXT("unreal/jobVersion"), FString::FromInt(InVersionNumber));
 
-			if (InJob)
-			{
-				InFilenameArguments.Add(TEXT("job_author"), GetJobAuthor(InJob));
-				InFilenameArguments.Add(TEXT("job_name"), InJob->JobName);
-				InFileMetadata.Add(TEXT("unreal/jobName"), InJob->JobName);
-				InFileMetadata.Add(TEXT("unreal/jobAuthor"), GetJobAuthor(InJob));
-				InFileMetadata.Add(TEXT("unreal/jobComment"), InJob->Comment);
-			}
+			const FString JobAuthor = InJob ? GetJobAuthor(InJob) : FString();
+			const FString JobName = InJob ? InJob->JobName : FString();
+			const FString JobComment = InJob ? InJob->Comment : FString();
+
+			InFilenameArguments.Add(TEXT("job_author"), JobAuthor);
+			InFilenameArguments.Add(TEXT("job_name"), JobName);
+			InFileMetadata.Add(TEXT("unreal/jobName"), JobName);
+			InFileMetadata.Add(TEXT("unreal/jobAuthor"), JobAuthor);
+			InFileMetadata.Add(TEXT("unreal/jobComment"), JobComment);
 		}
 
 		void GetCachedGPUDriverInfo(TMap<FString, FString>& InFileMetadata)
diff --git a/Engine/Plugins/MovieScene/MovieRenderPipeline/Source/MovieRenderPipelineCore/Public/Graph/MovieGraphBlueprintLibrary.h b/Engine/Plugins/MovieScene/MovieRenderPipeline/Source/MovieRenderPipelineCore/Public/Graph/MovieGraphBlueprintLibrary.h
--- a/Engine/Plugins/MovieScene/MovieRenderPipeline/Source/MovieRenderPipelineCore/Public/Graph/MovieGraphBlueprintLibrary.h
+++ b/Engine/Plugins/MovieScene/MovieRenderPipeline/Source/MovieRenderPipelineCore/Public/Graph/MovieGraphBlueprintLibrary.h
@@ -34,12 +34,12 @@ public:
 	* Takes a Movie Graph format string (in the form of {token}), a list of parameters (which normally come from the running UMovieGraphPipeline) and
 	* resolves them into a string. Unknown tokens are ignored. Which tokens can be resolved depends on the contents of InParams, tokens from settings
 	* rely on a evaluated config being provided, etc.
-	* @param	InFormatString		- Format string to attempt to resolve.
+	* @param	InFormatString		- Format string to attempt to resolve. Leave blank if just the format args should be populated.
 	* @param	InParams			- A list of parameters to use as source data for the resolve step. Normally comes from the UMovieGraphPipeline instance,
 	*								- but takes (mostly) POD here to allow using this function outside of the render runtime.
 	* @param	OutMergedFormatArgs - The set of KVP for both filename formats and file metadata that is generated as a result of this. Provided in case you
 	* 								- needed to do your own string resolving with the final dataset.
-	* @return						- The resolved format string.
+	* @return						- The resolved format string. Returns an empty string if InFormatString is blank.
 	*/
 	UFUNCTION(BlueprintCallable, Category = "Movie Graph")
 	static FString ResolveFilenameFormatArguments(const FString& InFormatString, const FMovieGraphFilenameResolveParams& InParams, FMovieGraphResolveArgs& OutMergedFormatArgs);
diff --git a/Engine/Plugins/MovieScene/MovieRenderPipeline/Source/MovieRenderPipelineEditor/Private/Customizations/Graph/MovieGraphFormatTokenCustomization.cpp b/Engine/Plugins/MovieScene/MovieRenderPipeline/Source/MovieRenderPipelineEditor/Private/Customizations/Graph/MovieGraphFormatTokenCustomization.cpp
--- a/Engine/Plugins/MovieScene/MovieRenderPipeline/Source/MovieRenderPipelineEditor/Private/Customizations/Graph/MovieGraphFormatTokenCustomization.cpp
+++ b/Engine/Plugins/MovieScene/MovieRenderPipeline/Source/MovieRenderPipelineEditor/Private/Customizations/Graph/MovieGraphFormatTokenCustomization.cpp
@@ -2,8 +2,8 @@
 
 #include "MovieGraphFormatTokenCustomization.h"
 
+#include "Graph/MovieGraphBlueprintLibrary.h"
 #include "Graph/Nodes/MovieGraphFileOutputNode.h"
-#include "MoviePipelineUtils.h"
 
 #include "DetailLayoutBuilder.h"
 #include "DetailWidgetRow.h"
@@ -55,12 +55,17 @@ void FMovieGraphFormatTokenCustomization::OnPropertyChange()
 TArray<FString> FMovieGraphFormatTokenCustomization::GetSuggestions() const
 {
 	TArray<FString> Suggestions;
-	FMoviePipelineFormatArgs FormatArgs;
+	FMovieGraphResolveArgs FormatArgs;
 	GetFormatArguments(FormatArgs);
 
-	for (const TPair<FString, FString>& KVP : FormatArgs.FilenameArguments)
+	// Display the token names alphabetically
+	TArray<FString> FilenameFormatTokens;
+	FormatArgs.FilenameArguments.GetKeys(FilenameFormatTokens);
+	FilenameFormatTokens.Sort();
+
+	for (const FString& FilenameFormatToken : FilenameFormatTokens)
 	{
-		Suggestions.Add(KVP.Key);
+		Suggestions.Add(FilenameFormatToken);
 	}
 
 	return Suggestions;
@@ -71,31 +76,11 @@ void FMovieGraphFormatTokenCustomization::OnTextChanged(const FText& InValue)
 	OutputFormatPropertyHandle->SetValue(InValue.ToString());
 }
 
-void FMovieGraphFormatTokenCustomization::GetFormatArguments(FMoviePipelineFormatArgs& InOutFormatArgs)
+void FMovieGraphFormatTokenCustomization::GetFormatArguments(FMovieGraphResolveArgs& InOutFormatArgs)
 {
-	static const FString LevelName = TEXT("Level Name");
-	static const FString SequenceName = TEXT("Sequence Name");
-	static const FString JobAuthor = TEXT("Job Author");
-	static const FString JobName = TEXT("Job Name");
-	static const FString JobComment = TEXT("Job Comment");
-	static constexpr double FrameRate = 0.0;
-	static const FString FramePlaceholderNumber = TEXT("0");
-
-	MoviePipeline::GetOutputStateFormatArgs(
-		InOutFormatArgs.FilenameArguments, InOutFormatArgs.FileMetadata,
-		FramePlaceholderNumber, FramePlaceholderNumber,
-		FramePlaceholderNumber, FramePlaceholderNumber,
-		TEXT("CameraName"), TEXT("ShotName"));
-
-	InOutFormatArgs.FilenameArguments.Add(TEXT("level_name"), LevelName);
-	InOutFormatArgs.FilenameArguments.Add(TEXT("sequence_name"), SequenceName);
-	InOutFormatArgs.FilenameArguments.Add(TEXT("job_author"), JobAuthor);
-	InOutFormatArgs.FilenameArguments.Add(TEXT("job_name"), JobName);
-	InOutFormatArgs.FilenameArguments.Add(TEXT("frame_rate"), FString::SanitizeFloat(FrameRate));
-
-	static const FDateTime CurrentTime = FDateTime::Now();
-	static constexpr int32 DummyVersionNumber = 1;
-	const FTimespan InitializationTimeOffset = FDateTime::Now() - FDateTime::UtcNow();
-	UE::MoviePipeline::GetSharedFormatArguments(
-		InOutFormatArgs.FilenameArguments, InOutFormatArgs.FileMetadata, CurrentTime, DummyVersionNumber, InOutFormatArgs.InJob, InitializationTimeOffset);
+	// Just fetch the format arguments (by keeping the format string empty). The tokens themselves will not be resolved correctly here (no context is
+	// provided in the resolve params), but all we care about here is the token list, not the resolved token values.
+	const FString FormatString;
+	const FMovieGraphFilenameResolveParams ResolveParams;
+	UMovieGraphBlueprintLibrary::ResolveFilenameFormatArguments(FormatString, ResolveParams, InOutFormatArgs);
 }
diff --git a/Engine/Plugins/MovieScene/MovieRenderPipeline/Source/MovieRenderPipelineEditor/Private/Customizations/Graph/MovieGraphFormatTokenCustomization.h b/Engine/Plugins/MovieScene/MovieRenderPipeline/Source/MovieRenderPipelineEditor/Private/Customizations/Graph/MovieGraphFormatTokenCustomization.h
--- a/Engine/Plugins/MovieScene/MovieRenderPipeline/Source/MovieRenderPipelineEditor/Private/Customizations/Graph/MovieGraphFormatTokenCustomization.h
+++ b/Engine/Plugins/MovieScene/MovieRenderPipeline/Source/MovieRenderPipelineEditor/Private/Customizations/Graph/MovieGraphFormatTokenCustomization.h
@@ -27,7 +27,7 @@ protected:
 
 	void OnTextChanged(const FText& InValue);
 
-	static void GetFormatArguments(FMoviePipelineFormatArgs& InOutFormatArgs);
+	static void GetFormatArguments(FMovieGraphResolveArgs& InOutFormatArgs);
 	
 	TWeakObjectPtr<UObject> CustomizedObject;
 	TSharedPtr<IPropertyHandle> OutputFormatPropertyHandle;
