MoviePipeline: Updated graph evaluation to ensure that job-level variable overrides are respected. When scripting was introduced, these stopped working because scripting duplicates graphs, and the graph pointers stored in the variable overrides did not point to the duplicated graphs.

Additionally, subgraphs are now accounted for in the duplication logic, so any scripting that modifies subgraphs will not affect the on-disk asset.

#rnx

diff --git a/Engine/Plugins/MovieScene/MovieRenderPipeline/Source/MovieRenderPipelineCore/Private/Graph/MovieGraphPipeline.cpp b/Engine/Plugins/MovieScene/MovieRenderPipeline/Source/MovieRenderPipelineCore/Private/Graph/MovieGraphPipeline.cpp
--- a/Engine/Plugins/MovieScene/MovieRenderPipeline/Source/MovieRenderPipelineCore/Private/Graph/MovieGraphPipeline.cpp
+++ b/Engine/Plugins/MovieScene/MovieRenderPipeline/Source/MovieRenderPipelineCore/Private/Graph/MovieGraphPipeline.cpp
@@ -16,6 +16,7 @@
 #include "Graph/Nodes/MovieGraphGlobalOutputSettingNode.h"
 #include "Graph/Nodes/MovieGraphWarmUpSettingNode.h"
 #include "Graph/Nodes/MovieGraphExecuteScriptNode.h"
+#include "Graph/Nodes/MovieGraphSubgraphNode.h"
 #include "Graph/MovieGraphBlueprintLibrary.h"
 #include "MovieRenderPipelineCoreModule.h"
 #include "Misc/CoreDelegates.h"
@@ -145,33 +146,43 @@ void UMovieGraphPipeline::Initialize(UMoviePipelineExecutorJob* InJob, const FMo
 
 void UMovieGraphPipeline::DuplicateJobAndConfiguration()
 {
+	// Contains all duplicated graphs. Maps the original graph (key) to the duplicated graph (value).
+	TMap<UMovieGraphConfig*, UMovieGraphConfig*> DuplicatedGraphs;
+	
 	// Scripting is likely to want to modify both the job (to set variable assignments) and 
 	// the configuration itself (to add nodes, or override an output directory, etc. If scripts
 	// directly modified the job/configuration it would lead to a lot of unintentional mutation
 	// of assets and queues, so we instead choose to duplicate the job and configurations for
 	// the duration of a render, and all of the Graph Pipeline code should look at the duplicates.
 	FObjectDuplicationParameters JobDuplicationParms = FObjectDuplicationParameters(CurrentJob, GetTransientPackage());
+	JobDuplicationParms.DestName = FName(FString::Format(TEXT("{0}_Duplicate"), {CurrentJob->GetFName().ToString()}));
 	CurrentJobDuplicate = Cast<UMoviePipelineExecutorJob>(StaticDuplicateObjectEx(JobDuplicationParms));
 
 	// The duplicate job is a mix of duplicated objects and non-duplicated objects. Objects that 
 	// don't have the job as their outer will not have been duplicate (which is good for the World/Sequence),
 	// but this also means that the graph configurations were not duplicated (as they are assets), so we need
 	// to manually duplicate them and update the pointers in the duplicated job.
-	FObjectDuplicationParameters PrimaryConfigDuplicationParams(CurrentJob->GetGraphPreset(), GetTransientPackage());
-	UMovieGraphConfig* DuplicatePrimaryConfig = Cast<UMovieGraphConfig>(StaticDuplicateObjectEx(PrimaryConfigDuplicationParams));
-	CurrentJobDuplicate->SetGraphPreset(DuplicatePrimaryConfig);
+	UMovieGraphConfig* DuplicatePrimaryConfig = DuplicateConfigRecursive(CurrentJob->GetGraphPreset(), DuplicatedGraphs);
+	UpdateVariableAssignmentsHelper(CurrentJobDuplicate.Get(), DuplicatedGraphs);
+
+	// Set the graph preset WITHOUT updating assignments. We're doing all updates manually here. SetGraphPreset() will attempt to update the
+	// assignments in the shots as well, but those have not yet been updated to use the duplicated graphs. Therefore, skip updating assignments
+	// completely to avoid the shot assignments from being wiped out.
+	constexpr bool bUpdateVariableAssignments = false;
+	CurrentJobDuplicate->SetGraphPreset(DuplicatePrimaryConfig, bUpdateVariableAssignments);
+
 	for (int32 Index = 0; Index < CurrentJob->ShotInfo.Num(); Index++)
 	{
 		// Now for each shot we need to duplicate its config (if any)
 		if (UMovieGraphConfig* ShotGraphPreset = CurrentJob->ShotInfo[Index]->GetGraphPreset())
 		{
-			// We use the transient package here and above because the _configs_ don't belong to the executor job usually (they belong
-			// to an asset package)
-			FObjectDuplicationParameters ShotConfigDuplicationParams(ShotGraphPreset, GetTransientPackage());
-			UMovieGraphConfig* DuplicateShotConfig = Cast<UMovieGraphConfig>(StaticDuplicateObjectEx(ShotConfigDuplicationParams));
-
-			CurrentJobDuplicate->ShotInfo[Index]->SetGraphPreset(DuplicateShotConfig);
+			UMovieGraphConfig* DuplicateShotConfig = DuplicateConfigRecursive(ShotGraphPreset, DuplicatedGraphs);
+			CurrentJobDuplicate->ShotInfo[Index]->SetGraphPreset(DuplicateShotConfig, bUpdateVariableAssignments);
 		}
+
+		// Always update the variable assignments, regardless of whether the shot has a graph assigned to it (it may be overriding the primary graph's
+		// variable assignments)
+		UpdateVariableAssignmentsHelper(CurrentJobDuplicate->ShotInfo[Index].Get(), DuplicatedGraphs);
 	}
 
 	// We only look in the primary configuration for the job for script nodes (and not shot specific overrides). If we looked
@@ -207,6 +218,97 @@ void UMovieGraphPipeline::DuplicateJobAndConfiguration()
 
 }
 
+UMovieGraphConfig* UMovieGraphPipeline::DuplicateConfigRecursive(UMovieGraphConfig* InGraphToDuplicate, TMap<UMovieGraphConfig*, UMovieGraphConfig*>& OutDuplicatedGraphs)
+{
+	UMovieGraphConfig* DuplicateConfig;
+
+	// Duplicate the graph. If the graph has been duplicated already, don't re-duplicate it, but continue updating variable assignments.
+	if (OutDuplicatedGraphs.Contains(InGraphToDuplicate))
+	{
+		DuplicateConfig = OutDuplicatedGraphs[InGraphToDuplicate];
+	}
+	else
+	{
+		// The transient package is used because graphs don't belong to the executor job usually (they belong to an asset package)
+		FObjectDuplicationParameters GraphDuplicationParams(InGraphToDuplicate, GetTransientPackage());
+		GraphDuplicationParams.DestName = FName(FString::Format(TEXT("{0}_Duplicate"), {InGraphToDuplicate->GetFName().ToString()}));
+		DuplicateConfig = Cast<UMovieGraphConfig>(StaticDuplicateObjectEx(GraphDuplicationParams));
+		
+		OutDuplicatedGraphs.Add(InGraphToDuplicate, DuplicateConfig);
+	}
+
+	// Duplicate sub-graphs also.
+	for (const TObjectPtr<UMovieGraphNode>& Node : DuplicateConfig->GetNodes())
+	{
+		UMovieGraphSubgraphNode* SubgraphNode = Cast<UMovieGraphSubgraphNode>(Node);
+		if (!SubgraphNode)
+		{
+			continue;
+		}
+
+		// Only duplicate if the subgraph node has a graph asset assigned to it.
+		if (UMovieGraphConfig* SubgraphConfig = SubgraphNode->GetSubgraphAsset())
+		{
+			// Don't recurse into this graph if it was already duplicated. Check BOTH the keys (the original graph) AND value (the duplicated graph)
+			// to prevent recursion. Checking the key ensures that we only duplicate if this graph has never been encountered. Checking the value
+			// ensures that we don't re-duplicate a graph that has already been duplicated (the subgraph node was already updated).
+			bool bHasBeenDuplicated = false;
+			for (const TPair<UMovieGraphConfig*, UMovieGraphConfig*>& DuplicateMapping : OutDuplicatedGraphs)
+			{
+				if ((DuplicateMapping.Key == SubgraphConfig) || (DuplicateMapping.Value == SubgraphConfig))
+				{
+					bHasBeenDuplicated = true;
+					break;
+				}
+			}
+			
+			if (!bHasBeenDuplicated)
+			{
+				DuplicateConfigRecursive(SubgraphConfig, OutDuplicatedGraphs);
+			}
+
+			// Update the subgraph node to use the duplicated graph. This should always be done, even if the graph was already duplicated (since
+			// a graph can be included as a subgraph in multiple locations).
+			if (UMovieGraphConfig** DuplicatedGraph = OutDuplicatedGraphs.Find(SubgraphConfig))
+			{
+				SubgraphNode->SetSubGraphAsset(*DuplicatedGraph);
+			}
+		}
+	}
+	
+	return DuplicateConfig;
+}
+
+template <typename JobType>
+void UMovieGraphPipeline::UpdateVariableAssignmentsHelper(JobType* InTargetJob, TMap<UMovieGraphConfig*, UMovieGraphConfig*>& InOriginalToDuplicateGraphMap)
+{
+	// Remaps the provided variable assignments to point to the the duplicated graphs.
+	auto UpdateVariableAssignments = [&InOriginalToDuplicateGraphMap](TArray<TObjectPtr<UMovieJobVariableAssignmentContainer>>& InVariableAssignments)
+	{
+		for (const TObjectPtr<UMovieJobVariableAssignmentContainer>& VariableAssignment : InVariableAssignments)
+		{
+			for (const TPair<UMovieGraphConfig*, UMovieGraphConfig*>& GraphMapping : InOriginalToDuplicateGraphMap)
+			{
+				if (VariableAssignment->GetGraphConfig() == GraphMapping.Key)
+				{
+					VariableAssignment->SetGraphConfig(GraphMapping.Value);
+					break;
+				}
+			}
+		}
+	};
+
+	// Update the variable assignments according to the original-to-duplicate map. Applies to primary jobs and shots.
+	constexpr bool bUpdateAssignments = false;
+	UpdateVariableAssignments(InTargetJob->GetGraphVariableAssignments(bUpdateAssignments));
+
+	// Update the primary-level graph (and subgraph) assignments on shots as well.
+	if constexpr (std::is_same_v<JobType, UMoviePipelineExecutorShot>)
+	{
+		UpdateVariableAssignments(InTargetJob->GetPrimaryGraphVariableAssignments(bUpdateAssignments));
+	}
+}
+
 void UMovieGraphPipeline::LoadPreviewWidget()
 {
 	// ToDo: Allow overriding this widget so that users can style it to match their project.
diff --git a/Engine/Plugins/MovieScene/MovieRenderPipeline/Source/MovieRenderPipelineCore/Private/Graph/Nodes/MovieGraphVariableNode.cpp b/Engine/Plugins/MovieScene/MovieRenderPipeline/Source/MovieRenderPipelineCore/Private/Graph/Nodes/MovieGraphVariableNode.cpp
--- a/Engine/Plugins/MovieScene/MovieRenderPipeline/Source/MovieRenderPipelineCore/Private/Graph/Nodes/MovieGraphVariableNode.cpp
+++ b/Engine/Plugins/MovieScene/MovieRenderPipeline/Source/MovieRenderPipelineCore/Private/Graph/Nodes/MovieGraphVariableNode.cpp
@@ -147,19 +147,19 @@ bool UMovieGraphVariableNode::ContextHasEnabledAssignmentForVariable(const FMovi
 	
 	if (Shot)
 	{
-		ShotVariableAssignments = Shot->GetOrCreateJobVariableAssignmentsForGraph(Shot->GetGraphPreset());
+		ShotVariableAssignments = Shot->GetOrCreateJobVariableAssignmentsForGraph(GetGraph());
 		
 		// The shot can also override variables on the primary job's graph (in addition to the shot-level ones fetched/created above)
 		if (PrimaryJob)
 		{
 			constexpr bool bIsForPrimaryOverrides = true;
-			ShotVariableAssignments_PrimaryOverrides = Shot->GetOrCreateJobVariableAssignmentsForGraph(PrimaryJob->GetGraphPreset(), bIsForPrimaryOverrides);
+			ShotVariableAssignments_PrimaryOverrides = Shot->GetOrCreateJobVariableAssignmentsForGraph(GetGraph(), bIsForPrimaryOverrides);
 		}
 	}
 
 	if (PrimaryJob)
 	{
-		JobVariableAssignments = PrimaryJob->GetOrCreateJobVariableAssignmentsForGraph(PrimaryJob->GetGraphPreset());
+		JobVariableAssignments = PrimaryJob->GetOrCreateJobVariableAssignmentsForGraph(GetGraph());
 	}
 	
 	// Check the shot job first for an enabled job variable assignment for this variable. Shot jobs take precedence over primary jobs.
diff --git a/Engine/Plugins/MovieScene/MovieRenderPipeline/Source/MovieRenderPipelineCore/Private/MoviePipelineQueue.cpp b/Engine/Plugins/MovieScene/MovieRenderPipeline/Source/MovieRenderPipelineCore/Private/MoviePipelineQueue.cpp
--- a/Engine/Plugins/MovieScene/MovieRenderPipeline/Source/MovieRenderPipelineCore/Private/MoviePipelineQueue.cpp
+++ b/Engine/Plugins/MovieScene/MovieRenderPipeline/Source/MovieRenderPipelineCore/Private/MoviePipelineQueue.cpp
@@ -310,11 +310,14 @@ void UMoviePipelineExecutorJob::RefreshAllVariableAssignments()
 	}
 }
 
-void UMoviePipelineExecutorJob::SetGraphPreset(const UMovieGraphConfig* InGraphPreset)
+void UMoviePipelineExecutorJob::SetGraphPreset(const UMovieGraphConfig* InGraphPreset, const bool bUpdateVariableAssignments)
 {
 	GraphPreset = InGraphPreset;
 
-	RefreshAllVariableAssignments();
+	if (bUpdateVariableAssignments)
+	{
+		RefreshAllVariableAssignments();
+	}
 
 	OnJobGraphPresetChanged.Broadcast(this, GraphPreset.Get());
 }
@@ -388,11 +391,14 @@ void UMoviePipelineExecutorJob::OnGraphPreSave(UObject* InObject, FObjectPreSave
 	RefreshAllVariableAssignments();
 }
 
-void UMoviePipelineExecutorShot::SetGraphPreset(const UMovieGraphConfig* InGraphPreset)
+void UMoviePipelineExecutorShot::SetGraphPreset(const UMovieGraphConfig* InGraphPreset, const bool bUpdateVariableAssignments)
 {
 	GraphPreset = InGraphPreset;
 
-	RefreshAllVariableAssignments();
+	if (bUpdateVariableAssignments)
+	{
+		RefreshAllVariableAssignments();
+	}
 
 	OnShotGraphPresetChanged.Broadcast(this, GraphPreset.Get());
 }
diff --git a/Engine/Plugins/MovieScene/MovieRenderPipeline/Source/MovieRenderPipelineCore/Public/Graph/MovieGraphPipeline.h b/Engine/Plugins/MovieScene/MovieRenderPipeline/Source/MovieRenderPipelineCore/Public/Graph/MovieGraphPipeline.h
--- a/Engine/Plugins/MovieScene/MovieRenderPipeline/Source/MovieRenderPipelineCore/Public/Graph/MovieGraphPipeline.h
+++ b/Engine/Plugins/MovieScene/MovieRenderPipeline/Source/MovieRenderPipelineCore/Public/Graph/MovieGraphPipeline.h
@@ -165,6 +165,16 @@ protected:
 	virtual void ExecutePreShotScripts(const TObjectPtr<UMoviePipelineExecutorShot>& InShot);
 	virtual void ExecutePostShotScripts(const FMoviePipelineOutputData& InData);
 
+	/**
+	 * Helps duplicate graph configs. Prevents re-duplications, duplicates sub-graphs (updates subgraph nodes accordingly), and potentially more.
+	 * Returns the duplicated graph.
+	 */
+	UMovieGraphConfig* DuplicateConfigRecursive(UMovieGraphConfig* InGraphToDuplicate, TMap<UMovieGraphConfig*, UMovieGraphConfig*>& OutDuplicatedGraphs);
+
+	/** Helps update variable assignments on the provided job to use duplicated graphs (reflected in the original-to-duplicate graph mapping). */
+	template<typename JobType>
+	void UpdateVariableAssignmentsHelper(JobType* InTargetJob, TMap<UMovieGraphConfig*, UMovieGraphConfig*>& InOriginalToDuplicateGraphMap);
+
 	// Update our data source to isolate the shot we're currently working on, so that expanded shots don't interfere with each other.
 	virtual void SetSoloShot(const TObjectPtr<UMoviePipelineExecutorShot>& InShot);
 	
diff --git a/Engine/Plugins/MovieScene/MovieRenderPipeline/Source/MovieRenderPipelineCore/Public/MoviePipelineQueue.h b/Engine/Plugins/MovieScene/MovieRenderPipeline/Source/MovieRenderPipelineCore/Public/MoviePipelineQueue.h
--- a/Engine/Plugins/MovieScene/MovieRenderPipeline/Source/MovieRenderPipelineCore/Public/MoviePipelineQueue.h
+++ b/Engine/Plugins/MovieScene/MovieRenderPipeline/Source/MovieRenderPipelineCore/Public/MoviePipelineQueue.h
@@ -153,9 +153,12 @@ public:
 	/**
 	 * Sets the graph-style preset that this job will use. Note that this will cause the graph to switch over to using
 	 * graph-style configuration if it is not already using it.
+	 *
+	 * @param InGraphPreset - The graph preset to assign to the shot.
+	 * @param bUpdateVariableAssignments - Set to false if variable assignments should NOT be automatically updated to reflect the graph preset being used. This is normally not what you want and should be used with caution.
 	 */
 	UFUNCTION(BlueprintCallable, Category = "Movie Render Pipeline")
-	void SetGraphPreset(const UMovieGraphConfig* InGraphPreset);
+	void SetGraphPreset(const UMovieGraphConfig* InGraphPreset, const bool bUpdateVariableAssignments = true);
 
 	/** Gets the job variable assignments for a specific graph. Creates a new variable assignments container if one was not found for the given graph. */
 	TObjectPtr<UMovieJobVariableAssignmentContainer> GetOrCreateJobVariableAssignmentsForGraph(const UMovieGraphConfig* InGraph, const bool bIsForPrimaryOverrides = false);
@@ -463,9 +466,12 @@ public:
 	/**
 	 * Sets the graph-style preset that this job will use. Note that this will cause the graph to switch over to using
 	 * graph-style configuration if it is not already using it.
+	 *
+	 * @param InGraphPreset - The graph preset to assign to the job.
+	 * @param bUpdateVariableAssignments - Set to false if variable assignments should NOT be automatically updated to reflect the graph preset being used. This is normally not what you want and should be used with caution.
 	 */
 	UFUNCTION(BlueprintCallable, Category = "Movie Render Pipeline")
-	void SetGraphPreset(const UMovieGraphConfig* InGraphPreset);
+	void SetGraphPreset(const UMovieGraphConfig* InGraphPreset, const bool bUpdateVariableAssignments = true);
 
 	UFUNCTION(BlueprintSetter, Category = "Movie Render Pipeline")
 	void SetSequence(FSoftObjectPath InSequence);
