Movie graph:

-Fix: Crash when executing multiple Undo commands to undo node connections and creations in MRG

diff --git a/Engine/Plugins/MovieScene/MovieRenderPipeline/Source/MovieRenderPipelineEditor/Private/Graph/MovieGraphSchema.cpp b/Engine/Plugins/MovieScene/MovieRenderPipeline/Source/MovieRenderPipelineEditor/Private/Graph/MovieGraphSchema.cpp
--- a/Engine/Plugins/MovieScene/MovieRenderPipeline/Source/MovieRenderPipelineEditor/Private/Graph/MovieGraphSchema.cpp
+++ b/Engine/Plugins/MovieScene/MovieRenderPipeline/Source/MovieRenderPipelineEditor/Private/Graph/MovieGraphSchema.cpp
@@ -433,6 +433,7 @@ UEdGraphNode* FMovieGraphSchemaAction_NewNode::PerformAction(UEdGraph* ParentGra
 	UMovieGraphConfig* RuntimeGraph = CastChecked<UMoviePipelineEdGraph>(ParentGraph)->GetPipelineGraph();
 	const FScopedTransaction Transaction(LOCTEXT("GraphEditor_NewNode", "Create Pipeline Graph Node."));
 	RuntimeGraph->Modify();
+	ParentGraph->Modify();
 
 	UMovieGraphNode* RuntimeNode = RuntimeGraph->ConstructRuntimeNode<UMovieGraphNode>(NodeClass);
 
@@ -477,6 +478,7 @@ UEdGraphNode* FMovieGraphSchemaAction_NewVariableNode::PerformAction(UEdGraph* P
 	UMovieGraphConfig* RuntimeGraph = CastChecked<UMoviePipelineEdGraph>(ParentGraph)->GetPipelineGraph();
 	const FScopedTransaction Transaction(LOCTEXT("GraphEditor_NewVariableNode", "Add New Variable Accessor Node"));
 	RuntimeGraph->Modify();
+	ParentGraph->Modify();
 
 	UMovieGraphNode* RuntimeNode = RuntimeGraph->ConstructRuntimeNode<UMovieGraphNode>(NodeClass);
 	if (UMovieGraphVariableNode* VariableNode = Cast<UMovieGraphVariableNode>(RuntimeNode))
