MoviePipeline: Removed old, unused prototype cruft from the graph's Render Layer Subsystem. Also renamed some methods to avoid "preview" wording (which only made sense in the prototype) and updated documentation.

#rnx

diff --git a/Engine/Plugins/MovieScene/MovieRenderPipeline/Config/BaseMovieRenderPipeline.ini b/Engine/Plugins/MovieScene/MovieRenderPipeline/Config/BaseMovieRenderPipeline.ini
--- a/Engine/Plugins/MovieScene/MovieRenderPipeline/Config/BaseMovieRenderPipeline.ini
+++ b/Engine/Plugins/MovieScene/MovieRenderPipeline/Config/BaseMovieRenderPipeline.ini
@@ -11,8 +11,4 @@
 +ClassRedirects=(OldName="/Script/MovieRenderPipelineCore.MoviePipelineCollectionModifier",NewName="/Script/MovieRenderPipelineCore.MovieGraphCollectionModifier")
 +ClassRedirects=(OldName="/Script/MovieRenderPipelineCore.MoviePipelineMaterialModifier",NewName="/Script/MovieRenderPipelineCore.MovieGraphMaterialModifier")
 +ClassRedirects=(OldName="/Script/MovieRenderPipelineCore.MoviePipelineVisibilityModifier",NewName="/Script/MovieRenderPipelineCore.MovieGraphRenderPropertyModifier")
-+ClassRedirects=(OldName="/Script/MovieRenderPipelineCore.MoviePipelineCollectionCommonQuery",NewName="/Script/MovieRenderPipelineCore.MovieGraphCollectionCommonQuery")
-+ClassRedirects=(OldName="/Script/MovieRenderPipelineCore.MoviePipelineCollectionLightingQuery",NewName="/Script/MovieRenderPipelineCore.MovieGraphCollectionLightingQuery")
-+ClassRedirects=(OldName="/Script/MovieRenderPipelineCore.MoviePipelineCollection",NewName="/Script/MovieRenderPipelineCore.MovieGraphCollection")
-+ClassRedirects=(OldName="/Script/MovieRenderPipelineCore.MoviePipelineCollectionQuery",NewName="/Script/MovieRenderPipelineCore.MovieGraphCollectionQueryBase")
-+EnumRedirects=(OldName="EMoviePipelineCollectionCommonQueryMode",NewName="EMovieGraphCollectionCommonQueryMode")
\ No newline at end of file
++ClassRedirects=(OldName="/Script/MovieRenderPipelineCore.MoviePipelineCollection",NewName="/Script/MovieRenderPipelineCore.MovieGraphCollection")
\ No newline at end of file
diff --git a/Engine/Plugins/MovieScene/MovieRenderPipeline/Source/MovieRenderPipelineCore/Private/Graph/MovieGraphRenderLayerSubsystem.cpp b/Engine/Plugins/MovieScene/MovieRenderPipeline/Source/MovieRenderPipelineCore/Private/Graph/MovieGraphRenderLayerSubsystem.cpp
--- a/Engine/Plugins/MovieScene/MovieRenderPipeline/Source/MovieRenderPipelineCore/Private/Graph/MovieGraphRenderLayerSubsystem.cpp
+++ b/Engine/Plugins/MovieScene/MovieRenderPipeline/Source/MovieRenderPipelineCore/Private/Graph/MovieGraphRenderLayerSubsystem.cpp
@@ -319,95 +319,6 @@ void UMovieGraphRenderPropertyModifier::SetActorVisibilityState(const FActorVisi
 	}
 }
 
-// TODO: This really should be "DoesComponentMatchQuery()"
-bool UMovieGraphCollectionCommonQuery::DoesActorMatchQuery(const AActor* Actor) const
-{
-	if (!Actor)
-	{
-		return false;
-	}
-
-	// TODO: This method should short-circuit so all query types aren't executed for every actor when the query mode is OR
-
-	bool bMatchesActorNames = false;
-	if (!ActorNames.IsEmpty())
-	{
-		bMatchesActorNames = ActorNames.Contains(Actor->GetName());
-	}
-	
-	bool bMatchesComponentTypes = false;
-	if (!ComponentTypes.IsEmpty())
-	{
-		const bool bIncludeFromChildActors = true;
-		TArray<UActorComponent*> ActorComponents;
-		Actor->GetComponents(ActorComponents, bIncludeFromChildActors);
-		for (const UActorComponent* Component : ActorComponents)
-		{
-			for (const UClass* ComponentType : ComponentTypes)
-			{
-				if (Component->IsA(ComponentType))
-				{
-					bMatchesComponentTypes = true;
-					break;
-				}
-			}
-		}
-	}
-
-	bool bMatchesTags = false;
-	if (!Tags.IsEmpty())
-	{
-		for (const FName& Tag : Tags)
-		{
-			if (Actor->Tags.Contains(Tag))
-			{
-				bMatchesTags = true;
-				break;
-			}
-		}
-	}
-
-	const bool bUsingActorNames = !ActorNames.IsEmpty();
-	const bool bUsingComponentTypes = !ComponentTypes.IsEmpty();
-	const bool bUsingTags = !Tags.IsEmpty();
-
-	if (QueryMode == EMovieGraphCollectionCommonQueryMode::And)
-	{
-		return (!bUsingActorNames || bMatchesActorNames) &&
-			   (!bUsingTags || bMatchesTags) &&
-			   (!bUsingComponentTypes || bMatchesComponentTypes);
-	}
-	
-	return bMatchesActorNames || bMatchesTags || bMatchesComponentTypes;
-}
-
-bool UMovieGraphCollectionLightingQuery::DoesActorMatchQuery(const AActor* Actor) const
-{
-	const TArray<UClass*> LightingComponentTypes = {
-		ULightComponentBase::StaticClass(),
-		UReflectionCaptureComponent::StaticClass(),
-		USkyAtmosphereComponent::StaticClass(),
-		USkyLightComponent::StaticClass()
-	};
-	
-	TArray<UActorComponent*> ActorComponents;
-	constexpr bool bIncludeFromChildActors = true;
-	Actor->GetComponents(ActorComponents, bIncludeFromChildActors);
-	
-	for (const UActorComponent* Component : ActorComponents)
-	{
-		for (const UClass* ComponentType : LightingComponentTypes)
-		{
-			if (Component->IsA(ComponentType))
-			{
-				return true;
-			}
-		}
-	}
-
-	return false;
-}
-
 void UMovieGraphCollectionModifier::AddCollection(UMovieGraphCollection* Collection)
 {
 	// Don't allow adding a duplicate collection
@@ -1850,7 +1761,7 @@ void UMovieGraphRenderLayer::RemoveModifier(UMovieGraphCollectionModifier* Modif
 	Modifiers.Remove(Modifier);
 }
 
-void UMovieGraphRenderLayer::Preview(const UWorld* World)
+void UMovieGraphRenderLayer::Apply(const UWorld* World)
 {
 	if (!World)
 	{
@@ -1864,13 +1775,8 @@ void UMovieGraphRenderLayer::Preview(const UWorld* World)
 	}
 }
 
-void UMovieGraphRenderLayer::UndoPreview(const UWorld* World)
+void UMovieGraphRenderLayer::Revert()
 {
-	if (!World)
-	{
-		return;
-	}
-
 	// Undo actions performed by all modifiers. Do this in the reverse order that they were applied, since the undo
 	// state of one modifier may depend on modifiers that were previously applied.
 	for (int32 Index = Modifiers.Num() - 1; Index >= 0; Index--)
@@ -1895,20 +1801,6 @@ UMovieGraphRenderLayerSubsystem* UMovieGraphRenderLayerSubsystem::GetFromWorld(c
 
 void UMovieGraphRenderLayerSubsystem::Initialize(FSubsystemCollectionBase& Collection)
 {
-	VisualizationEmptyCollection = NewObject<UMovieGraphCollection>(GetTransientPackage(), NAME_None, RF_Transient);
-
-	// By default, the visualizer should hide everything in the world
-	VisualizationModifier_HideWorld = NewObject<UMovieGraphRenderPropertyModifier>(GetTransientPackage(), NAME_None, RF_Transient);
-	VisualizationModifier_HideWorld->AddCollection(VisualizationEmptyCollection);
-	VisualizationModifier_HideWorld->SetHidden(true);
-
-	// Selectively show collections in the visualization
-	VisualizationModifier_VisibleCollections = NewObject<UMovieGraphRenderPropertyModifier>(GetTransientPackage(), NAME_None, RF_Transient);
-
-	// The visualizer render layer will hide the world, selectively show specified collections, and then apply any other provided modifiers
-	VisualizationRenderLayer = NewObject<UMovieGraphRenderLayer>(GetTransientPackage(), NAME_None, RF_Transient);
-	VisualizationRenderLayer->AddModifier(VisualizationModifier_HideWorld);
-	VisualizationRenderLayer->AddModifier(VisualizationModifier_VisibleCollections);
 }
 
 void UMovieGraphRenderLayerSubsystem::Deinitialize()
@@ -1917,7 +1809,7 @@ void UMovieGraphRenderLayerSubsystem::Deinitialize()
 
 void UMovieGraphRenderLayerSubsystem::Reset()
 {
-	ClearAllPreviews();
+	RevertAndClearActiveRenderLayer();
 	RenderLayers.Empty();
 }
 
@@ -1948,7 +1840,7 @@ void UMovieGraphRenderLayerSubsystem::RemoveRenderLayer(const FString& RenderLay
 {
 	if (ActiveRenderLayer && (ActiveRenderLayer->GetName() == RenderLayerName))
 	{
-		ClearAllPreviews();
+		RevertAndClearActiveRenderLayer();
 	}
 	
 	const uint32 Index = RenderLayers.IndexOfByPredicate([&RenderLayerName](const UMovieGraphRenderLayer* RenderLayer)
@@ -1969,8 +1861,8 @@ void UMovieGraphRenderLayerSubsystem::SetActiveRenderLayerByObj(UMovieGraphRende
 		return;
 	}
 	
-	ClearAllPreviews();
-	SetAndPreviewRenderLayer(RenderLayer);
+	RevertAndClearActiveRenderLayer();
+	SetAndApplyRenderLayer(RenderLayer);
 }
 
 void UMovieGraphRenderLayerSubsystem::SetActiveRenderLayerByName(const FName& RenderLayerName)
@@ -1988,83 +1880,23 @@ void UMovieGraphRenderLayerSubsystem::SetActiveRenderLayerByName(const FName& Re
 
 void UMovieGraphRenderLayerSubsystem::ClearActiveRenderLayer()
 {
-	ClearAllPreviews();
-}
-
-void UMovieGraphRenderLayerSubsystem::PreviewCollection(UMovieGraphCollection* Collection)
-{
-	if (!Collection)
-	{
-		return;
-	}
-
-	ClearAllPreviews();
-	
-	ActiveCollection = Collection;
-	VisualizationModifier_VisibleCollections->AddCollection(ActiveCollection);
-
-	SetAndPreviewRenderLayer(VisualizationRenderLayer);
-}
-
-void UMovieGraphRenderLayerSubsystem::ClearCollectionPreview()
-{
-	ClearAllPreviews();
+	RevertAndClearActiveRenderLayer();
 }
 
-void UMovieGraphRenderLayerSubsystem::ClearAllPreviews()
+void UMovieGraphRenderLayerSubsystem::RevertAndClearActiveRenderLayer()
 {
-	// Render layer previews and collection previews both use the active render layer, so undoing the preview this
-	// way will clear previews for both
 	if (ActiveRenderLayer)
 	{
-		ActiveRenderLayer->UndoPreview(GetWorld());
-
-		// Remove the modifier preview if present (requires an active render layer)
-		if (ActiveModifier)
-		{
-			ActiveRenderLayer->RemoveModifier(ActiveModifier);
-		}
+		ActiveRenderLayer->Revert();
 	}
 
-	// Reset the viz modifier for the collection preview
-	VisualizationModifier_VisibleCollections->SetCollections({});
-
 	ActiveRenderLayer = nullptr;
-	ActiveCollection = nullptr;
-	ActiveModifier = nullptr;
 }
 
-void UMovieGraphRenderLayerSubsystem::SetAndPreviewRenderLayer(UMovieGraphRenderLayer* RenderLayer)
+void UMovieGraphRenderLayerSubsystem::SetAndApplyRenderLayer(UMovieGraphRenderLayer* RenderLayer)
 {
 	ActiveRenderLayer = RenderLayer;
-	ActiveRenderLayer->Preview(GetWorld());
-}
-
-void UMovieGraphRenderLayerSubsystem::PreviewModifier(UMovieGraphCollectionModifier* Modifier)
-{
-	if (!Modifier)
-	{
-		return;
-	}
-
-	ClearAllPreviews();
-
-	ActiveModifier = Modifier;
-	VisualizationRenderLayer->AddModifier(ActiveModifier);
-
-	// Add the modifier's collections to the viz as well, so the actors that the modifier affects are actually visible
-	// TODO: This may need special handing for visibility modifiers
-	for (UMovieGraphCollection* Collection : ActiveModifier->GetCollections())
-	{
-		VisualizationModifier_VisibleCollections->AddCollection(Collection);
-	}
-	
-	SetAndPreviewRenderLayer(VisualizationRenderLayer);
-}
-
-void UMovieGraphRenderLayerSubsystem::ClearModifierPreview()
-{
-	ClearAllPreviews();
+	ActiveRenderLayer->Apply(GetWorld());
 }
 
 #undef LOCTEXT_NAMESPACE
\ No newline at end of file
diff --git a/Engine/Plugins/MovieScene/MovieRenderPipeline/Source/MovieRenderPipelineCore/Public/Graph/MovieGraphRenderLayerSubsystem.h b/Engine/Plugins/MovieScene/MovieRenderPipeline/Source/MovieRenderPipelineCore/Public/Graph/MovieGraphRenderLayerSubsystem.h
--- a/Engine/Plugins/MovieScene/MovieRenderPipeline/Source/MovieRenderPipelineCore/Public/Graph/MovieGraphRenderLayerSubsystem.h
+++ b/Engine/Plugins/MovieScene/MovieRenderPipeline/Source/MovieRenderPipelineCore/Public/Graph/MovieGraphRenderLayerSubsystem.h
@@ -729,76 +729,6 @@ public:
 	uint8 bHoldout : 1;
 };
 
-/**
- * Base class for providing the ability to determine if an actor matches a query.
- */
-UCLASS(Abstract)
-class UMovieGraphCollectionQueryBase : public UObject
-{
-	GENERATED_BODY()
-
-public:
-	virtual bool DoesActorMatchQuery(const AActor* Actor) const PURE_VIRTUAL(UMovieGraphCollectionQueryBase::DoesActorMatchQuery, return false; );  
-};
-
-UENUM(BlueprintType)
-enum class EMovieGraphCollectionCommonQueryMode : uint8
-{
-	And UMETA(ToolTip = "All specifiers in the query must be true"),
-	Or UMETA(ToolTip = "At least one specifier in the query must be true")
-};
-
-/**
- * Provides common actor querying functionality (names, tags, components, etc). These individual sub-queries can be
- * AND'd or OR'd together (eg, matches provided names OR provided tags, vs matches provided names AND provided tags).
- */
-UCLASS(BlueprintType)
-class UMovieGraphCollectionCommonQuery : public UMovieGraphCollectionQueryBase
-{
-	GENERATED_BODY()
-
-public:
-	// TODO: Add other common query operations (level, etc)
-	UFUNCTION(BlueprintCallable, Category = "Settings")
-	void SetActorNames(const TArray<FString>& InActorNames) { ActorNames = InActorNames; }
-
-	UFUNCTION(BlueprintCallable, Category = "Settings")
-	void SetTags(const TArray<FName>& InTags) { Tags = InTags; }
-
-	UFUNCTION(BlueprintCallable, Category = "Settings")
-	void SetComponentTypes(TArray<UClass*> InComponentTypes) { ComponentTypes = InComponentTypes; }
-
-	UFUNCTION(BlueprintCallable, Category = "Settings")
-	void SetQueryMode(const EMovieGraphCollectionCommonQueryMode InQueryMode) { QueryMode = InQueryMode; }
-
-	UFUNCTION(BlueprintCallable, Category = "Settings")
-	virtual bool DoesActorMatchQuery(const AActor* Actor) const override;
-
-public:
-	UPROPERTY(BlueprintReadWrite, EditAnywhere, Category = "General")
-	TArray<TObjectPtr<UClass>> ComponentTypes;
-
-	UPROPERTY(BlueprintReadWrite, EditAnywhere, Category = "General")
-	EMovieGraphCollectionCommonQueryMode QueryMode = EMovieGraphCollectionCommonQueryMode::And;
-
-	UPROPERTY(BlueprintReadWrite, EditAnywhere, Category = "General")
-	TArray<FString> ActorNames;
-
-	UPROPERTY(BlueprintReadWrite, EditAnywhere, Category = "General")
-	TArray<FName> Tags;
-};
-
-/** Generates a query that matches all lighting components. */
-UCLASS(BlueprintType)
-class UMovieGraphCollectionLightingQuery : public UMovieGraphCollectionQueryBase
-{
-	GENERATED_BODY()
-
-public:
-	UFUNCTION(BlueprintCallable, Category = "Settings")
-	virtual bool DoesActorMatchQuery(const AActor* Actor) const override;
-};
-
 /**
  * Provides a means of assembling modifiers together to generate a desired view of a scene. 
  */
@@ -829,10 +759,10 @@ public:
 	void RemoveModifier(UMovieGraphCollectionModifier* Modifier);
 
 	UFUNCTION(BlueprintCallable, Category = "Settings")
-	void Preview(const UWorld* World);
+	void Apply(const UWorld* World);
 
 	UFUNCTION(BlueprintCallable, Category = "Settings")
-	void UndoPreview(const UWorld* World);
+	void Revert();
 
 private:
 	/** The name of this render layer. */
@@ -878,75 +808,43 @@ public:
 	UFUNCTION(BlueprintCallable, Category = "Settings")
 	bool AddRenderLayer(UMovieGraphRenderLayer* RenderLayer);
 
+	/** Gets all render layers which are currently tracked by the system. */
 	UFUNCTION(BlueprintCallable, Category = "Settings")
 	const TArray<UMovieGraphRenderLayer*>& GetRenderLayers() { return RenderLayers; }
 
+	/** Removes the render layer with the given name. After removal it can no longer be made active with SetActiveRenderLayerBy*(). */
 	UFUNCTION(BlueprintCallable, Category = "Settings")
 	void RemoveRenderLayer(const FString& RenderLayerName);
 
+	/** Gets the currently active render layer (the layer with its modifiers applied). */
 	UFUNCTION(BlueprintCallable, Category = "Settings")
 	UMovieGraphRenderLayer* GetActiveRenderLayer() const { return ActiveRenderLayer; }
 
-	/** Previews the layer with the given name. The layer needs to have been registered with AddRenderLayer(). */
+	/** Applies the layer with the given name. The layer needs to have been registered with AddRenderLayer(). */
 	UFUNCTION(BlueprintCallable, Category = "Settings")
 	void SetActiveRenderLayerByName(const FName& RenderLayerName);
 
-	/** Previews the given layer. The layer does not need to have been registered with AddRenderLayer(). */
+	/** Applies the given layer. The layer does not need to have been registered with AddRenderLayer(). */
 	UFUNCTION(BlueprintCallable, Category = "Settings")
 	void SetActiveRenderLayerByObj(UMovieGraphRenderLayer* RenderLayer);
 
+	/** Clears the currently active render layer and reverts its modifiers. */
 	UFUNCTION(BlueprintCallable, Category = "Settings")
 	void ClearActiveRenderLayer();
 
-	UFUNCTION(BlueprintCallable, Category = "Settings")
-	void PreviewCollection(UMovieGraphCollection* Collection);
-
-	UFUNCTION(BlueprintCallable, Category = "Settings")
-	void ClearCollectionPreview();
-
-	UFUNCTION(BlueprintCallable, Category = "Settings")
-	void PreviewModifier(UMovieGraphCollectionModifier* Modifier);
-
-	UFUNCTION(BlueprintCallable, Category = "Settings")
-	void ClearModifierPreview();
-
 private:
-	/** Clears the render layer, collection, and modifier previews if any of them are active. */
-	void ClearAllPreviews();
+	/** Clears the currently active render layer and reverts its modifiers. */
+	void RevertAndClearActiveRenderLayer();
 
-	/** Sets the active render layer and previews it. */
-	void SetAndPreviewRenderLayer(UMovieGraphRenderLayer* RenderLayer);
+	/** Sets the active render layer and applies its modifiers. */
+	void SetAndApplyRenderLayer(UMovieGraphRenderLayer* RenderLayer);
 
 private:
 	/** Render layers which have been added/registered with the subsystem. These can be found by name. */
 	UPROPERTY(Transient)
 	TArray<TObjectPtr<UMovieGraphRenderLayer>> RenderLayers;
 
-	/** The render layer that is currently being viewed/previewed. */
+	/** The render layer that currently has its modifiers applied. */
 	UPROPERTY(Transient)
 	TObjectPtr<UMovieGraphRenderLayer> ActiveRenderLayer = nullptr;
-
-	/** The collection that is currently being viewed/previewed. */
-	UPROPERTY(Transient)
-	TObjectPtr<UMovieGraphCollection> ActiveCollection = nullptr;
-
-	/** The modifier that is currently being viewed/previewed. */
-	UPROPERTY(Transient)
-	TObjectPtr<UMovieGraphCollectionModifier> ActiveModifier = nullptr;
-
-	/** A render layer dedicated to visualizing collections and modifiers. */
-	UPROPERTY(Transient)
-	TObjectPtr<UMovieGraphRenderLayer> VisualizationRenderLayer = nullptr;
-
-	/** Empty collection used for visualization purposes (in conjunction w/ the viz render layer). */
-	UPROPERTY(Transient)
-	TObjectPtr<UMovieGraphCollection> VisualizationEmptyCollection = nullptr;
-
-	/** A modifier used for visualization purposes (to hide the entire world). */
-	UPROPERTY(Transient)
-	TObjectPtr<UMovieGraphRenderPropertyModifier> VisualizationModifier_HideWorld = nullptr;
-
-	/** A modifier used for visualization purposes (to show collections used with the modifier). */
-	UPROPERTY(Transient)
-	TObjectPtr<UMovieGraphRenderPropertyModifier> VisualizationModifier_VisibleCollections = nullptr;
 };
