Movie Graph:

-Fix: MovieGraph: DefaultPipeline name conflicts with MoviePipeline/MovieGraphPipeline

diff --git a/Engine/Plugins/MovieScene/MovieRenderPipeline/Source/MovieRenderPipelineEditor/Public/MovieRenderPipelineSettings.h b/Engine/Plugins/MovieScene/MovieRenderPipeline/Source/MovieRenderPipelineEditor/Public/MovieRenderPipelineSettings.h
--- a/Engine/Plugins/MovieScene/MovieRenderPipeline/Source/MovieRenderPipelineEditor/Public/MovieRenderPipelineSettings.h
+++ b/Engine/Plugins/MovieScene/MovieRenderPipeline/Source/MovieRenderPipelineEditor/Public/MovieRenderPipelineSettings.h
@@ -23,6 +23,14 @@ public:
 	GENERATED_BODY()
 	
 	MOVIERENDERPIPELINEEDITOR_API UMovieRenderPipelineProjectSettings();
+	
+	/**
+	* This allows you to implement your own Pipeline to handle timing and rendering of a movie. Changing
+	* this will allow you to re-use the existing UI/Executors while providing your own logic for producing
+	* a single render.
+	*/
+	UPROPERTY(config, EditAnywhere, BlueprintReadWrite, NoClear, meta = (MetaClass = "/Script/MovieRenderPipelineCore.MoviePipelineBase"), Category="Movie Render Pipeline", DisplayName = "Default Configuration Pipeline")
+	FSoftClassPath DefaultPipeline;
 
 	/**
 	* Which directory should we try to save presets in by default?
@@ -62,14 +70,6 @@ public:
 	*/
 	UPROPERTY(config, EditAnywhere, BlueprintReadWrite, meta = (MetaClass = "/Script/MovieRenderPipelineCore.MoviePipelineExecutorJob"), Category="Movie Render Pipeline")
 	FSoftClassPath DefaultExecutorJob;
-	
-	/**
-	* This allows you to implement your own Pipeline to handle timing and rendering of a movie. Changing
-	* this will allow you to re-use the existing UI/Executors while providing your own logic for producing
-	* a single render.
-	*/
-	UPROPERTY(config, EditAnywhere, BlueprintReadWrite, meta = (MetaClass = "/Script/MovieRenderPipelineCore.MoviePipelineBase"), Category="Movie Render Pipeline")
-	FSoftClassPath DefaultPipeline;
 
 	/** The graph that newly-created graph assets will be based off of. */
 	UPROPERTY(config, EditAnywhere, BlueprintReadWrite, Category="Movie Render Pipeline")
