MoviePipeline: Replaced the 'Replace with Preset' action with a new 'Clear Graph' action for shots using graphs in the queue UI. 'Replace with Preset' doesn't work for shots, and doesn't make sense either (since a shot cannot have a legacy config when the parent is a graph).

Also removed the 'Replace with Graph' action for shots for a similar reason. A parent must first be converted to a graph, which will subsequently convert all of the shots to graphs. Shots cannot have a graph config when the parent has a legacy config.

#rnx

diff --git a/Engine/Plugins/MovieScene/MovieRenderPipeline/Source/MovieRenderPipelineEditor/Private/Widgets/SMoviePipelineQueueEditor.cpp b/Engine/Plugins/MovieScene/MovieRenderPipeline/Source/MovieRenderPipelineEditor/Private/Widgets/SMoviePipelineQueueEditor.cpp
--- a/Engine/Plugins/MovieScene/MovieRenderPipeline/Source/MovieRenderPipelineEditor/Private/Widgets/SMoviePipelineQueueEditor.cpp
+++ b/Engine/Plugins/MovieScene/MovieRenderPipeline/Source/MovieRenderPipelineEditor/Private/Widgets/SMoviePipelineQueueEditor.cpp
@@ -434,19 +434,24 @@ public:
 		UMoviePipelineExecutorJob* Job = WeakJob.Get();
 
 		UClass* ConfigType = Job->IsUsingGraphConfiguration() ? UMovieGraphConfig::StaticClass() : UMoviePipelinePrimaryConfig::StaticClass();
+		constexpr bool bIsShot = false;
+		const FExecuteAction ClearGraphAction = nullptr;
 		
 		return OnGenerateConfigPresetPickerMenuFromClass(
 			ConfigType,
 			WeakJob,
+			bIsShot,
 			FOnAssetSelected::CreateRaw(this, &FMoviePipelineQueueJobTreeItem::OnPickPresetFromAsset),
 			FExecuteAction::CreateRaw(this, &FMoviePipelineQueueJobTreeItem::OnPickNewPreset),
 			FExecuteAction::CreateRaw(this, &FMoviePipelineQueueJobTreeItem::OnReplaceWithRenderGraph),
-			FExecuteAction::CreateRaw(this, &FMoviePipelineQueueJobTreeItem::OnCreateNewGraphAndAssign)
+			FExecuteAction::CreateRaw(this, &FMoviePipelineQueueJobTreeItem::OnCreateNewGraphAndAssign),
+			ClearGraphAction
 			);
 	}
 
 	static TSharedRef<SWidget> OnGenerateConfigPresetPickerMenuFromClass(UClass* InClass,
-		TWeakObjectPtr<UMoviePipelineExecutorJob> TargetJob, FOnAssetSelected InOnAssetSelected, FExecuteAction InNewConfig, FExecuteAction InNewRenderGraph, FExecuteAction InCreateNewGraphAndAssign)
+		TWeakObjectPtr<UMoviePipelineExecutorJob> TargetJob, const bool bIsShot, FOnAssetSelected InOnAssetSelected, FExecuteAction InNewConfig,
+		FExecuteAction InNewRenderGraph, FExecuteAction InCreateNewGraphAndAssign, FExecuteAction InClearGraph)
 	{
 		FMenuBuilder MenuBuilder(true, nullptr);
 		IContentBrowserSingleton& ContentBrowser = FModuleManager::LoadModuleChecked<FContentBrowserModule>("ContentBrowser").Get();
@@ -513,25 +518,45 @@ public:
 					NAME_None,
 					EUserInterfaceActionType::Button
 				);
-				MenuBuilder.AddMenuEntry(
-					LOCTEXT("ReplaceWithGraph_Label", "Replace with Graph (Experimental)"),
-					LOCTEXT("ReplaceWithGraph_Tooltip", "Replaces the current configuration with a new graph representation."),
-					FSlateIcon(),
-					FUIAction(InNewRenderGraph),
-					NAME_None,
-					EUserInterfaceActionType::Button
-				);
+
+				// Only primary jobs can do a Replace with Graph. Shots cannot have a graph config with a parent legacy config.
+				if (!bIsShot)
+				{
+					MenuBuilder.AddMenuEntry(
+						LOCTEXT("ReplaceWithGraph_Label", "Replace with Graph (Experimental)"),
+						LOCTEXT("ReplaceWithGraph_Tooltip", "Replaces the current configuration with a new graph representation."),
+						FSlateIcon(),
+						FUIAction(InNewRenderGraph),
+						NAME_None,
+						EUserInterfaceActionType::Button
+					);
+				}
 			}
 			else
 			{
-				MenuBuilder.AddMenuEntry(
-					LOCTEXT("ReplaceWithPreset_Label", "Replace with Preset"),
-					LOCTEXT("ReplaceWithPreset_Tooltip", "Replaces the current configuration with a new default (non-graph) config."),
-					FSlateIcon(),
-					FUIAction(InNewConfig),
-					NAME_None,
-					EUserInterfaceActionType::Button
-				);
+				// Shots can clear their graph, meaning they will inherit from the parent graph. Primary jobs can revert to using the legacy system.
+				if (bIsShot)
+				{
+					MenuBuilder.AddMenuEntry(
+						LOCTEXT("ClearGraph_Label", "Clear Graph"),
+						LOCTEXT("ClearGraph_Tooltip", "Remove the graph assigned to this shot. After removal, the shot will inherit the graph from the primary job."),
+						FSlateIcon(),
+						FUIAction(InClearGraph),
+						NAME_None,
+						EUserInterfaceActionType::Button
+					);
+				}
+				else
+				{
+					MenuBuilder.AddMenuEntry(
+						LOCTEXT("ReplaceWithPreset_Label", "Replace with Preset"),
+						LOCTEXT("ReplaceWithPreset_Tooltip", "Replaces the current configuration with a new default (non-graph) config."),
+						FSlateIcon(),
+						FUIAction(InNewConfig),
+						NAME_None,
+						EUserInterfaceActionType::Button
+					);
+				}
 			}
 		}
 
@@ -1058,6 +1083,14 @@ struct FMoviePipelineShotItem : IMoviePipelineQueueTreeItem
 		}
 	}
 
+	void OnClearGraph() const
+	{
+		if (UMoviePipelineExecutorShot* Shot = WeakShot.Get())
+		{
+			Shot->SetGraphPreset(nullptr);
+		}
+	}
+
 	EVisibility GetShotConfigModifiedVisibility() const
 	{
 		if (const UMoviePipelineExecutorShot* Shot = WeakShot.Get())
@@ -1085,14 +1118,18 @@ struct FMoviePipelineShotItem : IMoviePipelineQueueTreeItem
 	{
 		const UMoviePipelineExecutorShot* Shot = WeakShot.Get();
 		UClass* ConfigType = Shot && Shot->IsUsingGraphConfiguration() ? UMovieGraphConfig::StaticClass() : UMoviePipelinePrimaryConfig::StaticClass();
+		constexpr bool bIsShot = true;
+		const FExecuteAction ReplaceWithGraph = nullptr;
 		
 		return FMoviePipelineQueueJobTreeItem::OnGenerateConfigPresetPickerMenuFromClass(
 			ConfigType,
 			WeakJob,
+			bIsShot,
 			FOnAssetSelected::CreateRaw(this, &FMoviePipelineShotItem::OnPickShotPresetFromAsset),
 			FExecuteAction::CreateRaw(this, &FMoviePipelineShotItem::OnPickNewShotPreset),
-			FExecuteAction::CreateRaw(this, &FMoviePipelineShotItem::OnReplaceWithRenderGraph),
-			FExecuteAction::CreateRaw(this, &FMoviePipelineShotItem::OnCreateNewGraphAndAssign)
+			ReplaceWithGraph,
+			FExecuteAction::CreateRaw(this, &FMoviePipelineShotItem::OnCreateNewGraphAndAssign),
+			FExecuteAction::CreateRaw(this, &FMoviePipelineShotItem::OnClearGraph)
 			);
 	}
 };
