MoviePipeline: Added a warning about Holdout being active in a graph's modifier, but alpha-in-postprocessing is not enabled in the project's Rendering settings. Also added a warning about using a Path Traced Renderer node, but path tracing isn't enabled.

#rnx

diff --git a/Engine/Plugins/MovieScene/MovieRenderPipeline/Source/MovieRenderPipelineCore/Private/Graph/MovieGraphRenderLayerSubsystem.cpp b/Engine/Plugins/MovieScene/MovieRenderPipeline/Source/MovieRenderPipelineCore/Private/Graph/MovieGraphRenderLayerSubsystem.cpp
--- a/Engine/Plugins/MovieScene/MovieRenderPipeline/Source/MovieRenderPipelineCore/Private/Graph/MovieGraphRenderLayerSubsystem.cpp
+++ b/Engine/Plugins/MovieScene/MovieRenderPipeline/Source/MovieRenderPipelineCore/Private/Graph/MovieGraphRenderLayerSubsystem.cpp
@@ -6,9 +6,8 @@
 #include "Components/SkyAtmosphereComponent.h"
 #include "Components/VolumetricCloudComponent.h"
 #include "Components/ExponentialHeightFogComponent.h"
-#include "Components/SkyLightComponent.h"
-#include "Components/SphereReflectionCaptureComponent.h"
 #include "EngineUtils.h"
+#include "Engine/RendererSettings.h"
 #include "Framework/Application/SlateApplication.h"
 #include "Materials/MaterialInterface.h"
 #include "Modules/ModuleManager.h"
@@ -173,6 +172,16 @@ void UMovieGraphRenderPropertyModifier::ApplyModifier(const UWorld* World)
 	NewComponentState.bCastShadowWhileHidden = bCastShadowWhileHidden;
 	NewComponentState.bAffectIndirectLightingWhileHidden = bAffectIndirectLightingWhileHidden;
 	NewComponentState.bHoldout = bHoldout;
+
+	// Generate a warning if holdout is being used, but alpha is not enabled in post processing. Without that setting enabled, holdout will not work.
+	const URendererSettings* RendererSettings = GetDefault<URendererSettings>();
+	if (bHoldout && (RendererSettings->bEnableAlphaChannelInPostProcessing != EAlphaChannelMode::AllowThroughTonemapper))
+	{
+		// TODO: Ideally this is called in a general-purpose validation step instead, but that framework does not exist yet.
+		UE_LOG(LogMovieRenderPipeline, Warning, TEXT("A modifier with 'Holdout' is active, but 'Enable alpha channel support in post processing' in "
+													 "the project's Rendering settings is not set to 'Allow through tonemapper'. Holdout requires 'Allow "
+													 "through tonemapper' to work properly."));
+	}
 	
 	for (const UMovieGraphCollection* Collection : Collections)
 	{
diff --git a/Engine/Plugins/MovieScene/MovieRenderPipeline/Source/MovieRenderPipelineRenderPasses/Private/Graph/Nodes/MovieGraphPathTracerPassNode.cpp b/Engine/Plugins/MovieScene/MovieRenderPipeline/Source/MovieRenderPipelineRenderPasses/Private/Graph/Nodes/MovieGraphPathTracerPassNode.cpp
--- a/Engine/Plugins/MovieScene/MovieRenderPipeline/Source/MovieRenderPipelineRenderPasses/Private/Graph/Nodes/MovieGraphPathTracerPassNode.cpp
+++ b/Engine/Plugins/MovieScene/MovieRenderPipeline/Source/MovieRenderPipelineRenderPasses/Private/Graph/Nodes/MovieGraphPathTracerPassNode.cpp
@@ -3,6 +3,7 @@
 #include "Graph/Nodes/MovieGraphPathTracerPassNode.h"
 #include "Graph/Renderers/MovieGraphPathTracerPass.h"
 #include "Engine/EngineBaseTypes.h"
+#include "RenderUtils.h"
 #include "ShowFlags.h"
 
 TUniquePtr<UE::MovieGraph::Rendering::FMovieGraphImagePassBase> UMovieGraphPathTracerRenderPassNode::CreateInstance() const
@@ -50,6 +51,24 @@ void UMovieGraphPathTracerRenderPassNode::SetupImpl(const FMovieGraphRenderPassS
 		bOriginalProgressDisplayCvarValue = ProgressDisplayCvar->GetBool();
 		ProgressDisplayCvar->Set(false);
 	}
+
+	bool bSupportsPathTracing = false;
+	if (IsRayTracingEnabled())
+	{
+		if (const IConsoleVariable* PathTracingCVar = IConsoleManager::Get().FindConsoleVariable(TEXT("r.PathTracing")))
+		{
+			bSupportsPathTracing = PathTracingCVar->GetInt() != 0;
+		}
+	}
+
+	// Warn if the path tracer is being used, but it's not enabled in the Rendering settings. The path tracer won't work otherwise.
+	if (!bSupportsPathTracing)
+	{
+		// TODO: Ideally this is called in a general-purpose validation step instead, but that framework does not exist yet.
+		UE_LOG(LogMovieRenderPipeline, Warning, TEXT("An active Path Traced Renderer node was found, but path tracing support is not enabled. To get "
+													 "renders with path tracing, enable 'Support Hardware Ray Tracing' and 'Path Tracing' in the "
+													 "project's Rendering settings."));
+	}
 }
 
 void UMovieGraphPathTracerRenderPassNode::TeardownImpl()
