MoviePipeline: Updated the CLE node to check that the encoder executable in Project Settings is valid (otherwise the problem gets lost in a somewhat cryptic warning in the output log). Also updated the pipeline to prevent re-entrancy in TeardownShot(). If the pipeline shuts down in very specific ways, like the CLE node shutting down due to a bad encoder executable, there are multiple codepaths that want to call TeardownShot() which can cause some subtle issues.

#rnx

diff --git a/Engine/Plugins/MovieScene/MovieRenderPipeline/Source/MovieRenderPipelineCore/Private/Graph/MovieGraphPipeline.cpp b/Engine/Plugins/MovieScene/MovieRenderPipeline/Source/MovieRenderPipelineCore/Private/Graph/MovieGraphPipeline.cpp
--- a/Engine/Plugins/MovieScene/MovieRenderPipeline/Source/MovieRenderPipelineCore/Private/Graph/MovieGraphPipeline.cpp
+++ b/Engine/Plugins/MovieScene/MovieRenderPipeline/Source/MovieRenderPipelineCore/Private/Graph/MovieGraphPipeline.cpp
@@ -32,6 +32,7 @@ FString UMovieGraphPipeline::DefaultPreviewWidgetAsset = TEXT("/MovieRenderPipel
 UMovieGraphPipeline::UMovieGraphPipeline()
 	: CurrentShotIndex(-1)
 	, bIsTransitioningState(false)
+	, bIsTearingDownShot(false)
 	, PipelineState(EMovieRenderPipelineState::Uninitialized)
 {
 	OutputMerger = MakeShared<UE::MovieGraph::FMovieGraphOutputMerger>(this);
@@ -821,6 +822,14 @@ void UMovieGraphPipeline::SetupShot(const TObjectPtr<UMoviePipelineExecutorShot>
 
 void UMovieGraphPipeline::TeardownShot(const TObjectPtr<UMoviePipelineExecutorShot>& InShot)
 {
+	// No re-entrancy. Multiple TeardownShot() calls can occur in some situations when the pipeline is shutting down early.
+	if (bIsTearingDownShot)
+	{
+		return;
+	}
+
+	TGuardValue<bool> TeardownShotGuard(bIsTearingDownShot, true);
+	
 	// Teardown happens at the start of the first frame the shot is finished so we'll stop recording
 	// audio, which will prevent it from capturing any samples for this frame. We don't do a similar
 	// start in InitializeShot() because we don't want to record samples during warm up/motion blur.
diff --git a/Engine/Plugins/MovieScene/MovieRenderPipeline/Source/MovieRenderPipelineCore/Private/Graph/Nodes/MovieGraphCommandLineEncoderNode.cpp b/Engine/Plugins/MovieScene/MovieRenderPipeline/Source/MovieRenderPipelineCore/Private/Graph/Nodes/MovieGraphCommandLineEncoderNode.cpp
--- a/Engine/Plugins/MovieScene/MovieRenderPipeline/Source/MovieRenderPipelineCore/Private/Graph/Nodes/MovieGraphCommandLineEncoderNode.cpp
+++ b/Engine/Plugins/MovieScene/MovieRenderPipeline/Source/MovieRenderPipelineCore/Private/Graph/Nodes/MovieGraphCommandLineEncoderNode.cpp
@@ -205,10 +205,25 @@ bool UMovieGraphCommandLineEncoderNode::AreSettingsValid(TArray<FText>& OutError
 
 	// Validate project settings
 	{
-		if (EncoderSettings->ExecutablePath.IsEmpty())
+		FString EncoderPath = EncoderSettings->ExecutablePath;
+		if (EncoderPath.IsEmpty())
 		{
 			OutErrors.Add(LOCTEXT("CommandLineEncode_MissingExecutable", "No encoder executable has been specified in the Project Settings. Please set an encoder executable in Project Settings > Movie Pipeline CLI Encoder."));
 		}
+		else
+		{
+			if (FPaths::IsRelative(EncoderPath))
+			{
+				EncoderPath = FPaths::ConvertRelativePathToFull(EncoderPath);
+			}
+		
+			if (!FPaths::FileExists(EncoderPath))
+			{
+				OutErrors.Add(FText::Format(
+					LOCTEXT("CommandLineEncode_InvalidExecutable", "Invalid encoder executable path [{0}] was specified in the Project Settings. Please set a valid encoder executable in Project Settings > Movie Pipeline CLI Encoder."),
+					FText::FromString(EncoderPath)));
+			}
+		}
 	}
 
 	// Validate node settings
@@ -569,7 +584,8 @@ void UMovieGraphCommandLineEncoderNode::LaunchEncoder(const FEncoderParams& InPa
 	else
 	{
 		UE_LOG(LogMovieRenderPipeline, Error, TEXT("Failed to launch encoder process, see output log for more details."));
-		CachedPipeline->Shutdown(true);
+		constexpr bool bIsError = true;
+		CachedPipeline->RequestShutdown(bIsError);
 	}
 }
 
diff --git a/Engine/Plugins/MovieScene/MovieRenderPipeline/Source/MovieRenderPipelineCore/Public/Graph/MovieGraphPipeline.h b/Engine/Plugins/MovieScene/MovieRenderPipeline/Source/MovieRenderPipelineCore/Public/Graph/MovieGraphPipeline.h
--- a/Engine/Plugins/MovieScene/MovieRenderPipeline/Source/MovieRenderPipelineCore/Public/Graph/MovieGraphPipeline.h
+++ b/Engine/Plugins/MovieScene/MovieRenderPipeline/Source/MovieRenderPipelineCore/Public/Graph/MovieGraphPipeline.h
@@ -277,6 +277,9 @@ protected:
 	/** True if we're in a TransitionToState call. Used to prevent reentrancy. */
 	bool bIsTransitioningState;
 
+	/** True if we're in a TeardownShot call. Used to prevent reentrancy. */
+	bool bIsTearingDownShot;
+
 	/** When we originally initialize we store the offset from UTC (which is what GetInitializationTime() is in), but we clear this if you call SetInitializationTime. */
 	FTimespan InitializationTimeOffset;
 
