MoviePipeline: The .wav Output node's File Name Format should default to '{sequence_name}', not '{sequence_name}.{layer_name}.{frame_number}' which makes no sense for .wav files.

Additionally, changes to the File Name Format were not respected before (the value was coming directly from the CDO), so that is fixed as well.

#rnx

diff --git a/Engine/Plugins/MovieScene/MovieRenderPipeline/Source/MovieRenderPipelineCore/Private/Graph/Nodes/MovieGraphAudioOutputNode.cpp b/Engine/Plugins/MovieScene/MovieRenderPipeline/Source/MovieRenderPipelineCore/Private/Graph/Nodes/MovieGraphAudioOutputNode.cpp
--- a/Engine/Plugins/MovieScene/MovieRenderPipeline/Source/MovieRenderPipelineCore/Private/Graph/Nodes/MovieGraphAudioOutputNode.cpp
+++ b/Engine/Plugins/MovieScene/MovieRenderPipeline/Source/MovieRenderPipelineCore/Private/Graph/Nodes/MovieGraphAudioOutputNode.cpp
@@ -19,6 +19,11 @@ static TAutoConsoleVariable<float> CVarWaveOutputDelay(
 	TEXT("encoder runs, the encoder may fail.\n"),
 	ECVF_Default);
 
+UMovieGraphAudioOutputNode::UMovieGraphAudioOutputNode()
+{
+	FileNameFormat = TEXT("{sequence_name}");
+}
+
 void UMovieGraphAudioOutputNode::BuildNewProcessCommandLineArgsImpl(TArray<FString>& InOutUnrealURLParams, TArray<FString>& InOutCommandLineArgs, TArray<FString>& InOutDeviceProfileCvars, TArray<FString>& InOutExecCmds) const
 {
 	// Always add this so that audio is muted, it'll never line up during preview anyways.
@@ -109,10 +114,12 @@ bool UMovieGraphAudioOutputNode::IsFinishedWritingToDiskImpl() const
 
 FString UMovieGraphAudioOutputNode::GenerateOutputPath(const FMovieGraphRenderDataIdentifier& InRenderIdentifier, const TObjectPtr<UMoviePipelineExecutorShot>& InShot) const
 {
+	// Note that we fetch the audio node from the evaluated graph because this method is executing on the CDO (and we want the evaluated FileNameFormat)
 	constexpr bool bIncludeCDOs = true;
 	constexpr bool bExactMatch = true;
 	const UMovieGraphGlobalOutputSettingNode* OutputNode = EvaluatedGraph->GetSettingForBranch<UMovieGraphGlobalOutputSettingNode>(GlobalsPinName, bIncludeCDOs, bExactMatch);
-	FString FileNameFormatString = OutputNode->OutputDirectory.Path / FileNameFormat;
+	const UMovieGraphAudioOutputNode* AudioNode = EvaluatedGraph->GetSettingForBranch<UMovieGraphAudioOutputNode>(GlobalsPinName, bIncludeCDOs, bExactMatch);
+	FString FileNameFormatString = OutputNode->OutputDirectory.Path / AudioNode->FileNameFormat;
 
 	constexpr bool bIncludeRenderPass = false;
 	constexpr bool bTestFrameNumber = false;
diff --git a/Engine/Plugins/MovieScene/MovieRenderPipeline/Source/MovieRenderPipelineCore/Public/Graph/Nodes/MovieGraphAudioOutputNode.h b/Engine/Plugins/MovieScene/MovieRenderPipeline/Source/MovieRenderPipelineCore/Public/Graph/Nodes/MovieGraphAudioOutputNode.h
--- a/Engine/Plugins/MovieScene/MovieRenderPipeline/Source/MovieRenderPipelineCore/Public/Graph/Nodes/MovieGraphAudioOutputNode.h
+++ b/Engine/Plugins/MovieScene/MovieRenderPipeline/Source/MovieRenderPipelineCore/Public/Graph/Nodes/MovieGraphAudioOutputNode.h
@@ -16,7 +16,7 @@ class UMovieGraphAudioOutputNode : public UMovieGraphFileOutputNode
 	GENERATED_BODY()
 	
 public:
-	UMovieGraphAudioOutputNode() = default;
+	UMovieGraphAudioOutputNode();
 
 	// UMovieGraphSettingNode Interface
 	virtual void BuildNewProcessCommandLineArgsImpl(TArray<FString>& InOutUnrealURLParams, TArray<FString>& InOutCommandLineArgs, TArray<FString>& InOutDeviceProfileCvars, TArray<FString>& InOutExecCmds) const override;
