MoviePipeline: Fixed a number of undo-related bugs in the Members panel.

* Undoing the creation of a variable/input/output did not work correctly and could cause a crash.
* Undoing the deletion of a variable/input/output did not work correctly and could cause a crash.
* Undoing the creation of a variable/input/output and then subsequently creating a new one would cause a crash.
* After an undo of any of the above, the Members panel would not refresh properly.

#rnx

diff --git a/Engine/Plugins/MovieScene/MovieRenderPipeline/Source/MovieRenderPipelineCore/Private/Graph/MovieGraphConfig.cpp b/Engine/Plugins/MovieScene/MovieRenderPipeline/Source/MovieRenderPipelineCore/Private/Graph/MovieGraphConfig.cpp
--- a/Engine/Plugins/MovieScene/MovieRenderPipeline/Source/MovieRenderPipelineCore/Private/Graph/MovieGraphConfig.cpp
+++ b/Engine/Plugins/MovieScene/MovieRenderPipeline/Source/MovieRenderPipelineCore/Private/Graph/MovieGraphConfig.cpp
@@ -668,6 +668,8 @@ RetType* UMovieGraphConfig::AddMember(TArray<TObjectPtr<ArrType>>& InMemberArray
 {
 	static_assert(std::is_base_of_v<UMovieGraphMember, RetType>, "RetType is not derived from UMovieGraphMember");
 	
+	Modify();
+	
 	// TODO: This can be replaced with just CreateDefaultSubobject() when AddDefaultMembers() isn't called from PostLoad()
 	//
 	// This method will be called in two cases: 1) when default members are being added to a new graph when it is being
@@ -845,6 +847,8 @@ bool UMovieGraphConfig::DeleteVariableMember(UMovieGraphVariable* VariableMember
 	{
 		return false;
 	}
+
+	Modify();
 	
 	// Find all accessor nodes using this graph variable
 	TArray<TObjectPtr<UMovieGraphNode>> NodesToRemove =
@@ -898,6 +902,8 @@ bool UMovieGraphConfig::DeleteInputMember(UMovieGraphInput* InputMemberToDelete)
 {
 	if (InputMemberToDelete)
 	{
+		Modify();
+		
 		Inputs.RemoveSingle(InputMemberToDelete);
 		RemoveOutboundEdges(InputNode, FName(InputMemberToDelete->GetMemberName()));
 
@@ -914,6 +920,8 @@ bool UMovieGraphConfig::DeleteOutputMember(UMovieGraphOutput* OutputMemberToDele
 {
 	if (OutputMemberToDelete)
 	{
+		Modify();
+		
 		Outputs.RemoveSingle(OutputMemberToDelete);
 		RemoveInboundEdges(OutputNode, FName(OutputMemberToDelete->GetMemberName()));
 
diff --git a/Engine/Plugins/MovieScene/MovieRenderPipeline/Source/MovieRenderPipelineEditor/Private/Graph/MovieEdGraphNode.cpp b/Engine/Plugins/MovieScene/MovieRenderPipeline/Source/MovieRenderPipelineEditor/Private/Graph/MovieEdGraphNode.cpp
--- a/Engine/Plugins/MovieScene/MovieRenderPipeline/Source/MovieRenderPipelineEditor/Private/Graph/MovieEdGraphNode.cpp
+++ b/Engine/Plugins/MovieScene/MovieRenderPipeline/Source/MovieRenderPipelineEditor/Private/Graph/MovieEdGraphNode.cpp
@@ -544,10 +544,11 @@ void UMoviePipelineEdGraphNodeBase::ReconstructNode()
 		return;
 	}
 
+	UMoviePipelineEdGraph* Graph = CastChecked<UMoviePipelineEdGraph>(GetGraph());
+	Graph->Modify();
+
 	ReconstructPins();
 
-	UMoviePipelineEdGraph* Graph = CastChecked<UMoviePipelineEdGraph>(GetGraph());
-	
 	// Reconstruct connections
 	const bool bCreateInbound = true;
 	const bool bCreateOutbound = true;
@@ -558,6 +559,8 @@ void UMoviePipelineEdGraphNodeBase::ReconstructNode()
 
 void UMoviePipelineEdGraphNodeBase::ReconstructPins()
 {
+	Modify();
+	
 	// Store copy of old pins
 	TArray<UEdGraphPin*> OldPins = MoveTemp(Pins);
 	Pins.Reset();
diff --git a/Engine/Plugins/MovieScene/MovieRenderPipeline/Source/MovieRenderPipelineEditor/Private/Graph/SMovieGraphMembersTabContent.cpp b/Engine/Plugins/MovieScene/MovieRenderPipeline/Source/MovieRenderPipelineEditor/Private/Graph/SMovieGraphMembersTabContent.cpp
--- a/Engine/Plugins/MovieScene/MovieRenderPipeline/Source/MovieRenderPipelineEditor/Private/Graph/SMovieGraphMembersTabContent.cpp
+++ b/Engine/Plugins/MovieScene/MovieRenderPipeline/Source/MovieRenderPipelineEditor/Private/Graph/SMovieGraphMembersTabContent.cpp
@@ -8,6 +8,7 @@
 #include "Graph/MovieGraphSchema.h"
 #include "Graph/Nodes/MovieGraphVariableNode.h"
 #include "MovieEdGraphNode.h"
+#include "ScopedTransaction.h"
 #include "SGraphActionMenu.h"
 #include "Toolkits/AssetEditorToolkit.h"
 
@@ -170,8 +171,9 @@ void SMovieGraphMembersTabContent::DeleteSelectedMembers()
 	{
 		if (UMovieGraphMember* GraphMember = UE::MovieGraph::Private::GetMemberFromAction(SelectedAction.Get()))
 		{
-			MemberChangedHandles.Remove(GraphMember);
+			FScopedTransaction Transaction(LOCTEXT("DeleteGraphMember", "Delete Graph Member"));
 			
+			MemberChangedHandles.Remove(GraphMember);
 			CurrentGraph->DeleteMember(GraphMember);
 		}
 	}
@@ -204,6 +206,18 @@ bool SMovieGraphMembersTabContent::CanDeleteSelectedMembers() const
 	return true;
 }
 
+void SMovieGraphMembersTabContent::PostUndo(bool bSuccess)
+{
+	// Normally the UI relies on delegates to determine when to refresh. However, undo/redo do not fire those delegates, so refresh whenever there
+	// is an undo/redo.
+	RefreshMemberActions();
+}
+
+void SMovieGraphMembersTabContent::PostRedo(bool bSuccess)
+{
+	RefreshMemberActions();
+}
+
 void SMovieGraphMembersTabContent::CollectAllActions(FGraphActionListBuilderBase& OutAllActions)
 {
 	static const FText UserVariablesCategory = LOCTEXT("UserVariablesCategory", "User Variables");
@@ -390,14 +404,17 @@ FReply SMovieGraphMembersTabContent::OnAddButtonClickedOnSection(const int32 InS
 
 	if (Section == EActionSection::Inputs)
 	{
+		FScopedTransaction Transaction(LOCTEXT("AddNewInput", "Add New Input"));
 		CurrentGraph->AddInput();
 	}
 	else if (Section == EActionSection::Outputs)
 	{
+		FScopedTransaction Transaction(LOCTEXT("AddNewOutput", "Add New Output"));
 		CurrentGraph->AddOutput();
 	}
 	else if (Section == EActionSection::Variables)
 	{
+		FScopedTransaction Transaction(LOCTEXT("AddNewVariable", "Add New Variable"));
 		CurrentGraph->AddVariable();
 	}
 
diff --git a/Engine/Plugins/MovieScene/MovieRenderPipeline/Source/MovieRenderPipelineEditor/Private/Graph/SMovieGraphMembersTabContent.h b/Engine/Plugins/MovieScene/MovieRenderPipeline/Source/MovieRenderPipelineEditor/Private/Graph/SMovieGraphMembersTabContent.h
--- a/Engine/Plugins/MovieScene/MovieRenderPipeline/Source/MovieRenderPipelineEditor/Private/Graph/SMovieGraphMembersTabContent.h
+++ b/Engine/Plugins/MovieScene/MovieRenderPipeline/Source/MovieRenderPipelineEditor/Private/Graph/SMovieGraphMembersTabContent.h
@@ -2,6 +2,7 @@
 
 #pragma once
 
+#include "EditorUndoClient.h"
 #include "GraphEditorDragDropAction.h"
 #include "SGraphActionMenu.h"
 #include "Widgets/Input/SButton.h"
@@ -16,7 +17,7 @@ struct FGraphActionListBuilderBase;
 /**
  * Contents of the "Members" tab in the graph asset editor.
  */
-class SMovieGraphMembersTabContent : public SCompoundWidget
+class SMovieGraphMembersTabContent : public SCompoundWidget, public FSelfRegisteringEditorUndoClient
 {
 public:
 	DECLARE_DELEGATE_TwoParams(FOnActionSelected, const TArray<TSharedPtr<FEdGraphSchemaAction>>&, ESelectInfo::Type);
@@ -55,6 +56,11 @@ public:
 	/** Determines if all selected member(s) can be deleted. */
 	bool CanDeleteSelectedMembers() const;
 
+	//~ Begin FSelfRegisteringEditorUndoClient Interface
+	virtual void PostUndo(bool bSuccess) override;
+	virtual void PostRedo(bool bSuccess) override;
+	//~ End FSelfRegisteringEditorUndoClient Interface
+
 private:
 	/** The section identifier in the action widget. */
 	enum class EActionSection : uint8
