MoviePipeline: Prevented the Global Game Overrides node from taking effect if it is not contained within a graph. This differs from how the legacy system works.

#rnx

diff --git a/Engine/Plugins/MovieScene/MovieRenderPipeline/Source/MovieRenderPipelineCore/Private/Graph/MovieGraphPipeline.cpp b/Engine/Plugins/MovieScene/MovieRenderPipeline/Source/MovieRenderPipelineCore/Private/Graph/MovieGraphPipeline.cpp
--- a/Engine/Plugins/MovieScene/MovieRenderPipeline/Source/MovieRenderPipelineCore/Private/Graph/MovieGraphPipeline.cpp
+++ b/Engine/Plugins/MovieScene/MovieRenderPipeline/Source/MovieRenderPipelineCore/Private/Graph/MovieGraphPipeline.cpp
@@ -803,8 +803,9 @@ void UMovieGraphPipeline::SetupShot(const TObjectPtr<UMoviePipelineExecutorShot>
 	const UMovieGraphEvaluatedConfig* EvaluatedConfig = TimeStepData.EvaluatedConfig;
 
 	// Apply any global game overrides, which includes cvars. This needs to be done before the CVarManager sets cvars
-	// so any user-specified cvars can override cvars set via the global game overrides.
-	constexpr bool bIncludeCDOs = true;
+	// so any user-specified cvars can override cvars set via the global game overrides. Note that the CDO is intentionally
+	// not fetched here so users have a way of opting out of this node if needed.
+	constexpr bool bIncludeCDOs = false;
 	constexpr bool bExactMatch = true;
 	if (UMovieGraphGlobalGameOverridesNode* GlobalGameOverridesNode = EvaluatedConfig->GetSettingForBranch<UMovieGraphGlobalGameOverridesNode>(UMovieGraphNode::GlobalsPinName, bIncludeCDOs, bExactMatch))
 	{
@@ -896,7 +897,7 @@ void UMovieGraphPipeline::TeardownShot(const TObjectPtr<UMoviePipelineExecutorSh
 
 	// Revert cvars set by the global game overrides. Needs to be done after the CVarManager reverts (since the global
 	// game overrides are applied first in SetupShot).
-	bIncludeCDOs = true;
+	bIncludeCDOs = false;
 	bExactMatch = true;
 	if (UMovieGraphGlobalGameOverridesNode* GlobalGameOverridesNode = EvaluatedConfig->GetSettingForBranch<UMovieGraphGlobalGameOverridesNode>(UMovieGraphNode::GlobalsPinName, bIncludeCDOs, bExactMatch))
 	{
diff --git a/Engine/Plugins/MovieScene/MovieRenderPipeline/Source/MovieRenderPipelineCore/Private/Graph/Nodes/MovieGraphGlobalGameOverrides.cpp b/Engine/Plugins/MovieScene/MovieRenderPipeline/Source/MovieRenderPipelineCore/Private/Graph/Nodes/MovieGraphGlobalGameOverrides.cpp
--- a/Engine/Plugins/MovieScene/MovieRenderPipeline/Source/MovieRenderPipelineCore/Private/Graph/Nodes/MovieGraphGlobalGameOverrides.cpp
+++ b/Engine/Plugins/MovieScene/MovieRenderPipeline/Source/MovieRenderPipelineCore/Private/Graph/Nodes/MovieGraphGlobalGameOverrides.cpp
@@ -197,12 +197,13 @@ TSubclassOf<AGameModeBase> UMovieGraphGlobalGameOverridesNode::GetGameModeOverri
 		FString OutTraversalError;
 		if (const UMovieGraphEvaluatedConfig* EvaluatedGraph = GraphConfig->CreateFlattenedGraph(TraversalContext, OutTraversalError))
 		{
-			constexpr bool bIncludeCDOs = true;
+			// Note that the CDO is not fetched here. Users need to explicitly include the Global Game Overrides node and specify the game mode override.
+			constexpr bool bIncludeCDOs = false;
 			constexpr bool bExactMatch = true;
 			UMovieGraphGlobalGameOverridesNode* GameOverridesNode =
 				EvaluatedGraph->GetSettingForBranch<UMovieGraphGlobalGameOverridesNode>(GlobalsPinName, bIncludeCDOs, bExactMatch);
 
-			return GameOverridesNode->GameModeOverride;
+			return GameOverridesNode ? GameOverridesNode->GameModeOverride : nullptr;
 		}
 	}
 	else
