MoviePipeline: Slightly tweaked the logic in the graph that generates the warning about alpha-in-postprocessing not being set properly. Both alpha-in-postprocessing settings are actually OK for holdout compatibility.

#rnx

diff --git a/Engine/Plugins/MovieScene/MovieRenderPipeline/Source/MovieRenderPipelineCore/Private/Graph/MovieGraphRenderLayerSubsystem.cpp b/Engine/Plugins/MovieScene/MovieRenderPipeline/Source/MovieRenderPipelineCore/Private/Graph/MovieGraphRenderLayerSubsystem.cpp
--- a/Engine/Plugins/MovieScene/MovieRenderPipeline/Source/MovieRenderPipelineCore/Private/Graph/MovieGraphRenderLayerSubsystem.cpp
+++ b/Engine/Plugins/MovieScene/MovieRenderPipeline/Source/MovieRenderPipelineCore/Private/Graph/MovieGraphRenderLayerSubsystem.cpp
@@ -175,12 +175,12 @@ void UMovieGraphRenderPropertyModifier::ApplyModifier(const UWorld* World)
 
 	// Generate a warning if holdout is being used, but alpha is not enabled in post processing. Without that setting enabled, holdout will not work.
 	const URendererSettings* RendererSettings = GetDefault<URendererSettings>();
-	if (bHoldout && (RendererSettings->bEnableAlphaChannelInPostProcessing != EAlphaChannelMode::AllowThroughTonemapper))
+	if (bHoldout && (RendererSettings->bEnableAlphaChannelInPostProcessing == EAlphaChannelMode::Disabled))
 	{
 		// TODO: Ideally this is called in a general-purpose validation step instead, but that framework does not exist yet.
 		UE_LOG(LogMovieRenderPipeline, Warning, TEXT("A modifier with 'Holdout' is active, but 'Enable alpha channel support in post processing' in "
-													 "the project's Rendering settings is not set to 'Allow through tonemapper'. Holdout requires 'Allow "
-													 "through tonemapper' to work properly."));
+													 "the project's Rendering settings is set to Disabled. Holdout will not work properly if this "
+													 "setting is Disabled."));
 	}
 	
 	for (const UMovieGraphCollection* Collection : Collections)
