MoviePipeline: Updated the sublevel query in the graph to ensure that it always returns actors from the *current world* (in most cases, this means that actors in the sublevel need to be converted to PIE actors). Without this fix, there could be instances where the sublevel query would not return the correct actors, causing modifiers to modify the wrong actors.

#rnx

diff --git a/Engine/Plugins/MovieScene/MovieRenderPipeline/Source/MovieRenderPipelineCore/Private/Graph/MovieGraphRenderLayerSubsystem.cpp b/Engine/Plugins/MovieScene/MovieRenderPipeline/Source/MovieRenderPipelineCore/Private/Graph/MovieGraphRenderLayerSubsystem.cpp
--- a/Engine/Plugins/MovieScene/MovieRenderPipeline/Source/MovieRenderPipelineCore/Private/Graph/MovieGraphRenderLayerSubsystem.cpp
+++ b/Engine/Plugins/MovieScene/MovieRenderPipeline/Source/MovieRenderPipelineCore/Private/Graph/MovieGraphRenderLayerSubsystem.cpp
@@ -441,68 +441,77 @@ bool UMovieGraphConditionGroupQueryBase::IsFirstConditionGroupQuery() const
 	return false;
 }
 
-void UMovieGraphConditionGroupQuery_Actor::Evaluate(const TArray<AActor*>& InActorsToQuery, const UWorld* InWorld, TSet<AActor*>& OutMatchingActors) const
+AActor* UMovieGraphConditionGroupQueryBase::GetActorForCurrentWorld(AActor* InActorToConvert)
 {
-	TRACE_CPUPROFILER_EVENT_SCOPE(UMovieGraphConditionGroupQuery_Actor::Evaluate);
-
 #if WITH_EDITOR
 	const bool bIsPIE = GEditor->IsPlaySessionInProgress();
 #else
 	const bool bIsPIE = false;
 #endif
-
-	// Convert the actors in the query to PIE (or editor) equivalents once, rather than constantly in the loop.
-	TArray<AActor*> ActorsToMatch_Converted;
-	ActorsToMatch_Converted.Reserve(ActorsToMatch.Num());
-	for (const TSoftObjectPtr<AActor>& SoftActorToMatch : ActorsToMatch)
+	
+	if (!InActorToConvert)
 	{
-		if (!SoftActorToMatch.IsValid())
-		{
-			continue;
-		}
+		return nullptr;
+	}
 		
-		AActor* ActorToMatch = SoftActorToMatch.Get();
-		const UWorld* ActorWorld = ActorToMatch->GetWorld();
-		const bool bIsEditorActor = ActorWorld && ActorWorld->IsEditorWorld();
+	const UWorld* ActorWorld = InActorToConvert->GetWorld();
+	const bool bIsEditorActor = ActorWorld && ActorWorld->IsEditorWorld();
 
-		// If a PIE session is NOT in progress, make sure that the actor is the editor equivalent
-		if (!bIsPIE)
+	// If a PIE session is NOT in progress, make sure that the actor is the editor equivalent
+	if (!bIsPIE)
+	{
+		// Only do PIE -> editor actor conversion when the actor is NOT from the editor
+		if (!bIsEditorActor)
 		{
-			// Only do PIE -> editor actor conversion when the actor is NOT from the editor
-			if (!bIsEditorActor)
-			{
 #if WITH_EDITOR
-				if (AActor* EditorActor = EditorUtilities::GetEditorWorldCounterpartActor(ActorToMatch))
-				{
-					ActorsToMatch_Converted.Add(EditorActor);
-				}
-#endif
-			}
-			else
+			if (AActor* EditorActor = EditorUtilities::GetEditorWorldCounterpartActor(InActorToConvert))
 			{
-				// Just use ActorToMatch as-is if it's not from PIE
-				ActorsToMatch_Converted.Add(ActorToMatch);
+				return EditorActor;
 			}
+#endif
 		}
-
-		// If a PIE session IS active, try to get the PIE equivalent of the editor actor
 		else
 		{
-			// Only do editor -> PIE actor conversion when the actor is from an editor world
-			if (bIsEditorActor)
-			{
+			// Just use InActorToConvert as-is if it's not from PIE
+			return InActorToConvert;
+		}
+	}
+
+	// If a PIE session IS active, try to get the PIE equivalent of the editor actor
+	else
+	{
+		// Only do editor -> PIE actor conversion when the actor is from an editor world
+		if (bIsEditorActor)
+		{
 #if WITH_EDITOR
-				if (AActor* PieActor = EditorUtilities::GetSimWorldCounterpartActor(ActorToMatch))
-				{
-					ActorsToMatch_Converted.Add(PieActor);
-				}
-#endif
-			}
-			else
+			if (AActor* PieActor = EditorUtilities::GetSimWorldCounterpartActor(InActorToConvert))
 			{
-				// Just use ActorToMatch as-is if it's not from an editor actor
-				ActorsToMatch_Converted.Add(ActorToMatch);
-			}	
+				return PieActor;
+			}
+#endif
+		}
+		else
+		{
+			// Just use InActorToConvert as-is if it's not from an editor actor
+			return InActorToConvert;
+		}	
+	}
+
+	return nullptr;
+}
+
+void UMovieGraphConditionGroupQuery_Actor::Evaluate(const TArray<AActor*>& InActorsToQuery, const UWorld* InWorld, TSet<AActor*>& OutMatchingActors) const
+{
+	TRACE_CPUPROFILER_EVENT_SCOPE(UMovieGraphConditionGroupQuery_Actor::Evaluate);
+
+	// Convert the actors in the query to PIE (or editor) equivalents once, rather than constantly in the loop.
+	TArray<AActor*> ActorsToMatch_Converted;
+	ActorsToMatch_Converted.Reserve(ActorsToMatch.Num());
+	for (const TSoftObjectPtr<AActor>& SoftActorToMatch : ActorsToMatch)
+	{
+		if (AActor* ConvertedActor = GetActorForCurrentWorld(SoftActorToMatch.Get()))
+		{
+			OutMatchingActors.Add(ConvertedActor);
 		}
 	}
 	
@@ -1310,9 +1319,13 @@ void UMovieGraphConditionGroupQuery_Sublevel::Evaluate(const TArray<AActor*>& In
 			continue;
 		}
 
-		for (TObjectPtr<AActor>& LevelActor : CurrentLevel->Actors)
+		for (const TObjectPtr<AActor>& LevelActor : CurrentLevel->Actors)
 		{
-			OutMatchingActors.Add(LevelActor.Get());
+			// The actors accessed directly from the level may need to be converted into the current world (most likely editor -> PIE)
+			if (AActor* ConvertedActor = GetActorForCurrentWorld(LevelActor.Get()))
+			{
+				OutMatchingActors.Add(ConvertedActor);
+			}
 		}
 	}
 }
diff --git a/Engine/Plugins/MovieScene/MovieRenderPipeline/Source/MovieRenderPipelineCore/Public/Graph/MovieGraphRenderLayerSubsystem.h b/Engine/Plugins/MovieScene/MovieRenderPipeline/Source/MovieRenderPipelineCore/Public/Graph/MovieGraphRenderLayerSubsystem.h
--- a/Engine/Plugins/MovieScene/MovieRenderPipeline/Source/MovieRenderPipelineCore/Public/Graph/MovieGraphRenderLayerSubsystem.h
+++ b/Engine/Plugins/MovieScene/MovieRenderPipeline/Source/MovieRenderPipelineCore/Public/Graph/MovieGraphRenderLayerSubsystem.h
@@ -125,6 +125,13 @@ public:
 	UFUNCTION(BlueprintCallable, Category = "Settings")
 	bool IsFirstConditionGroupQuery() const;
 
+protected:
+	/**
+	 * Utility that returns the given actor in the current world. If currently in PIE, converts editor actors to PIE actors, and vice-versa. If no
+	 * conversion is needed, returns the provided actor as-is.
+	 */
+	static AActor* GetActorForCurrentWorld(AActor* InActorToConvert);
+
 private:
 	/** The operation type that the query is using. */
 	UPROPERTY()
