MoviePipeline: Added a debug feature to disable shader flushing during renders. It is not reccomended that you enable this cvar, as your final rendered frames may contain missing objects or default materials.

diff --git a/Engine/Plugins/MovieScene/MovieRenderPipeline/Source/MovieRenderPipelineCore/Private/MoviePipelineRendering.cpp b/Engine/Plugins/MovieScene/MovieRenderPipeline/Source/MovieRenderPipelineCore/Private/MoviePipelineRendering.cpp
--- a/Engine/Plugins/MovieScene/MovieRenderPipeline/Source/MovieRenderPipelineCore/Private/MoviePipelineRendering.cpp
+++ b/Engine/Plugins/MovieScene/MovieRenderPipeline/Source/MovieRenderPipelineCore/Private/MoviePipelineRendering.cpp
@@ -39,6 +39,12 @@
 #include "Materials/MaterialInterface.h"
 #include "ContentStreaming.h"
 
+static TAutoConsoleVariable<bool> CVarMoviePipelineDisableShaderFlushing(
+	TEXT("MoviePipeline.DisableShaderFlushingDebug"), false,
+	TEXT("If true, the Movie Pipeline won't wait for any outstanding shader or asset compilation.")
+	TEXT("If false (default), any outstanding shaders and assets will be flushed each frame before rendering.")
+	TEXT("If true, rendered frames may be missing objects (meshes, particles, etc.) or objects may show the default checkerboard material."),
+	ECVF_Default);
 
 #define LOCTEXT_NAMESPACE "MoviePipeline"
 
@@ -509,13 +515,17 @@ void UMoviePipeline::FlushAsyncEngineSystems()
 		GetWorld()->BlockTillLevelStreamingCompleted();
 	}
 
-	// Ensure we have complete shader maps for all materials used by primitives in the world.
-	// This way we will never render with the default material.
-	UMaterialInterface::SubmitRemainingJobsForWorld(GetWorld());
+	const bool bDisableShaderFlushing = CVarMoviePipelineDisableShaderFlushing.GetValueOnGameThread();
+	if (!bDisableShaderFlushing)
+	{
+		// Ensure we have complete shader maps for all materials used by primitives in the world.
+		// This way we will never render with the default material.
+		UMaterialInterface::SubmitRemainingJobsForWorld(GetWorld());
 
-	// Flush all assets still being compiled asynchronously.
-	// A progressbar is already in place so the user can get feedback while waiting for everything to settle.
-	FAssetCompilingManager::Get().FinishAllCompilation();
+		// Flush all assets still being compiled asynchronously.
+		// A progressbar is already in place so the user can get feedback while waiting for everything to settle.
+		FAssetCompilingManager::Get().FinishAllCompilation();
+	}
 
 	// Flush streaming managers
 	{
