Movie Graph:

-Fix: MRG doesn't generate Unreal Insights Trace file

diff --git a/Engine/Plugins/MovieScene/MovieRenderPipeline/Source/MovieRenderPipelineCore/Private/Graph/MovieGraphBlueprintLibrary.cpp b/Engine/Plugins/MovieScene/MovieRenderPipeline/Source/MovieRenderPipelineCore/Private/Graph/MovieGraphBlueprintLibrary.cpp
--- a/Engine/Plugins/MovieScene/MovieRenderPipeline/Source/MovieRenderPipelineCore/Private/Graph/MovieGraphBlueprintLibrary.cpp
+++ b/Engine/Plugins/MovieScene/MovieRenderPipeline/Source/MovieRenderPipelineCore/Private/Graph/MovieGraphBlueprintLibrary.cpp
@@ -83,8 +83,13 @@ FString UMovieGraphBlueprintLibrary::ResolveFilenameFormatArguments(const FStrin
 		if (InParams.EvaluatedConfig)
 		{
 			const bool bIncludeCDOs = false;
-			UMovieGraphGlobalOutputSettingNode* OutputSettingNode = InParams.EvaluatedConfig->GetSettingForBranch<UMovieGraphGlobalOutputSettingNode>(InParams.RenderDataIdentifier.RootBranchName, bIncludeCDOs);
-			FrameRate = GetEffectiveFrameRate(OutputSettingNode, InParams.DefaultFrameRate).AsDecimal();
+			UMovieGraphGlobalOutputSettingNode* OutputSettingNode = InParams.EvaluatedConfig->GetSettingForBranch<UMovieGraphGlobalOutputSettingNode>(UMovieGraphSettingNode::GlobalsPinName, bIncludeCDOs);
+
+			if (OutputSettingNode)
+			{
+				FrameRate = GetEffectiveFrameRate(OutputSettingNode, InParams.DefaultFrameRate).AsDecimal();
+				bOverwriteExisting = OutputSettingNode->bOverwriteExistingOutput;
+			}
 		}
 
 		OutMergedFormatArgs.FilenameArguments.Add(TEXT("frame_rate"), FString::SanitizeFloat(FrameRate));
@@ -179,6 +184,12 @@ FString UMovieGraphBlueprintLibrary::ResolveFilenameFormatArguments(const FStrin
 	FPaths::NormalizeFilename(BaseFilename);
 	FPaths::RemoveDuplicateSlashes(BaseFilename);
 
+	// In the event of multiple dots in the filename, we need to replace them with a single dot to ensure a valid file name
+	while (BaseFilename.Contains(".."))
+	{
+		BaseFilename.ReplaceInline(TEXT(".."), TEXT("."));
+	}
+
 	// If we end with a "." character, remove it. The extension will put it back on. We can end up with this sometimes
 	// after resolving file format strings, ie: {sequence_name}.{frame_number} becomes {sequence_name}. for
 	// videos (which can't use frame_numbers).
diff --git a/Engine/Plugins/MovieScene/MovieRenderPipeline/Source/MovieRenderPipelineCore/Private/Graph/MovieGraphFilenameResolveParams.cpp b/Engine/Plugins/MovieScene/MovieRenderPipeline/Source/MovieRenderPipelineCore/Private/Graph/MovieGraphFilenameResolveParams.cpp
new file mode 100644
--- /dev/null
+++ b/Engine/Plugins/MovieScene/MovieRenderPipeline/Source/MovieRenderPipelineCore/Private/Graph/MovieGraphFilenameResolveParams.cpp
@@ -0,0 +1,64 @@
+// Copyright Epic Games, Inc. All Rights Reserved.
+
+#include "Graph/MovieGraphFilenameResolveParams.h"
+
+#include "Graph/MovieGraphPipeline.h"
+#include "Graph/Nodes/MovieGraphGlobalOutputSettingNode.h"
+#include "MoviePipelineQueue.h"
+
+FMovieGraphFilenameResolveParams FMovieGraphFilenameResolveParams::MakeResolveParams(
+	const FMovieGraphRenderDataIdentifier& InRenderId,
+	const UMovieGraphPipeline* InPipeline,
+	const TObjectPtr<UMovieGraphEvaluatedConfig>& InEvaluatedConfig,
+	const FMovieGraphTraversalContext& InTraversalContext,
+	const TMap<FString, FString>& InAdditionalFormatArgs)
+{
+    FMovieGraphFilenameResolveParams Params = FMovieGraphFilenameResolveParams();
+	
+	if (ensureAlwaysMsgf(InPipeline, TEXT("InPipeline is not valid - ResolveParams will be created, but will be missing critical information.")))
+	{
+		Params.InitializationTime = InPipeline->GetInitializationTime();
+		Params.InitializationTimeOffset = InPipeline->GetInitializationTimeOffset();
+		Params.Job = InPipeline->GetCurrentJob();
+		
+		if (InPipeline->GetActiveShotList().IsValidIndex(InTraversalContext.ShotIndex))
+		{
+			const TObjectPtr<UMoviePipelineExecutorShot>& Shot = InPipeline->GetActiveShotList()[InTraversalContext.ShotIndex];
+			Params.Version = Shot->ShotInfo.VersionNumber;
+			Params.Shot = Shot;
+		}
+	}
+        
+	Params.RenderDataIdentifier = InRenderId;
+	
+	Params.RootFrameNumber = InTraversalContext.Time.RootFrameNumber.Value;
+	Params.ShotFrameNumber = InTraversalContext.Time.ShotFrameNumber.Value;
+	Params.RootFrameNumberRel = InTraversalContext.Time.OutputFrameNumber;
+	Params.ShotFrameNumberRel = InTraversalContext.Time.ShotOutputFrameNumber;
+	//Params.FileMetadata = ToDo: Track File Metadata
+
+	if (InEvaluatedConfig)
+	{
+		const UMovieGraphGlobalOutputSettingNode* OutputSettingNode = InEvaluatedConfig->GetSettingForBranch<UMovieGraphGlobalOutputSettingNode>(UMovieGraphNode::GlobalsPinName);
+		if (IsValid(OutputSettingNode))
+		{
+			Params.ZeroPadFrameNumberCount = OutputSettingNode->ZeroPadFrameNumbers;
+			Params.FrameNumberOffset = OutputSettingNode->FrameNumberOffset;
+		}
+		Params.EvaluatedConfig = InEvaluatedConfig;
+	}
+
+    // If time dilation is in effect, RootFrameNumber and ShotFrameNumber will contain duplicates and the files will overwrite each other, 
+    // so we force them into relative mode and then warn users we did that (as their numbers will jump from say 1001 -> 0000).
+    bool bForceRelativeFrameNumbers = false;
+    // if (FileNameFormatString.Contains(TEXT("{frame")) && InTraversalContext.Time.IsTimeDilated() && !FileNameFormatString.Contains(TEXT("_rel}")))
+    // {
+    // 	UE_LOG(LogMovieRenderPipeline, Warning, TEXT("Time Dilation was used but output format does not use relative time, forcing relative numbers. Change {frame_number} to {frame_number_rel} (or shot version) to remove this message."));
+    // 	bForceRelativeFrameNumbers = true;
+    // }
+    Params.bForceRelativeFrameNumbers = bForceRelativeFrameNumbers;
+	Params.bEnsureAbsolutePath = true;
+	Params.FileNameFormatOverrides = InAdditionalFormatArgs;
+
+	return Params;
+}
diff --git a/Engine/Plugins/MovieScene/MovieRenderPipeline/Source/MovieRenderPipelineCore/Private/Graph/MovieGraphPipeline.cpp b/Engine/Plugins/MovieScene/MovieRenderPipeline/Source/MovieRenderPipelineCore/Private/Graph/MovieGraphPipeline.cpp
--- a/Engine/Plugins/MovieScene/MovieRenderPipeline/Source/MovieRenderPipelineCore/Private/Graph/MovieGraphPipeline.cpp
+++ b/Engine/Plugins/MovieScene/MovieRenderPipeline/Source/MovieRenderPipelineCore/Private/Graph/MovieGraphPipeline.cpp
@@ -2,28 +2,32 @@
 
 #include "Graph/MovieGraphPipeline.h"
 
+#include "MoviePipelineQueue.h"
+#include "MoviePipelineUtils.h"
+#include "MovieRenderPipelineCoreModule.h"
+#include "MovieScene.h"
+#include "Graph/MovieGraphBlueprintLibrary.h"
 #include "Graph/MovieGraphCVarManager.h"
 #include "Graph/MovieGraphDataTypes.h"
 #include "Graph/MovieGraphLinearTimeStep.h"
 #include "Graph/MovieGraphOutputMerger.h"
 #include "Graph/MovieGraphRenderLayerSubsystem.h"
 #include "Graph/Nodes/MovieGraphCollectionNode.h"
+#include "Graph/Nodes/MovieGraphDebugNode.h"
+#include "Graph/Nodes/MovieGraphExecuteScriptNode.h"
 #include "Graph/Nodes/MovieGraphFileOutputNode.h"
 #include "Graph/Nodes/MovieGraphGlobalGameOverrides.h"
+#include "Graph/Nodes/MovieGraphGlobalOutputSettingNode.h"
 #include "Graph/Nodes/MovieGraphModifierNode.h"
 #include "Graph/Nodes/MovieGraphRenderLayerNode.h"
 #include "Graph/Nodes/MovieGraphSamplingMethodNode.h"
-#include "Graph/Nodes/MovieGraphGlobalOutputSettingNode.h"
-#include "Graph/Nodes/MovieGraphWarmUpSettingNode.h"
-#include "Graph/Nodes/MovieGraphExecuteScriptNode.h"
 #include "Graph/Nodes/MovieGraphSubgraphNode.h"
-#include "Graph/MovieGraphBlueprintLibrary.h"
-#include "MovieRenderPipelineCoreModule.h"
-#include "Misc/CoreDelegates.h"
-#include "MoviePipelineQueue.h"
-#include "MovieScene.h"
-#include "RenderingThread.h"
+#include "Graph/Nodes/MovieGraphWarmUpSettingNode.h"
+
+#include "HAL/PlatformFileManager.h"
 #include "ImageWriteQueue.h"
+#include "RenderingThread.h"
+#include "Misc/CoreDelegates.h"
 #include "Modules/ModuleManager.h"
 #include "UObject/Package.h"
 
@@ -92,6 +96,22 @@ void UMovieGraphPipeline::Initialize(UMoviePipelineExecutorJob* InJob, const FMo
 	DuplicateJobAndConfiguration();
 	ExecutePreJobScripts();
 
+	// Lay groundwork for Unreal Insights
+	const FMovieGraphTraversalContext Context = GetCurrentTraversalContext(false);
+	FString OutError;
+
+	if (UMovieGraphEvaluatedConfig* FlattenedConfig = GetCurrentJob()->GetGraphPreset()->CreateFlattenedGraph(Context, OutError))
+	{
+		constexpr bool bIncludeCDOs = false;
+		constexpr bool bExactMatch = true;
+		const UMovieGraphDebugSettingNode* DebugSetting =
+			FlattenedConfig->GetSettingForBranch<UMovieGraphDebugSettingNode>(UMovieGraphNode::GlobalsPinName, bIncludeCDOs, bExactMatch);
+		if (DebugSetting && DebugSetting->bCaptureUnrealInsightsTrace)
+		{
+			StartUnrealInsightsCapture(FlattenedConfig);
+		}
+	}
+
 	// Create instances of our different classes from the InitConfig
 	GraphRendererInstance = NewObject<UMovieGraphRendererBase>(this, InitConfig.RendererClass);
 	GraphDataSourceInstance = NewObject<UMovieGraphDataSourceBase>(this, InitConfig.DataSourceClass);
@@ -773,6 +793,71 @@ void UMovieGraphPipeline::BeginExport()
 	}
 }
 
+void UMovieGraphPipeline::StartUnrealInsightsCapture(UMovieGraphEvaluatedConfig* EvaluatedConfig)
+{
+	check(EvaluatedConfig);
+
+	bool bIncludeCDOs = false;
+	bool bExactMatch = true;
+	UMovieGraphDebugSettingNode* DebugSetting =
+		EvaluatedConfig->GetSettingForBranch<UMovieGraphDebugSettingNode>(UMovieGraphSettingNode::GlobalsPinName, bIncludeCDOs, bExactMatch);
+	if (!ensureAlwaysMsgf(DebugSetting, TEXT("Failed to find UMovieGraphFileOutputNode. Aborting.")))
+	{
+		return;
+	}
+
+	bIncludeCDOs = true;
+	bExactMatch = true;
+	UMovieGraphGlobalOutputSettingNode* OutputSetting =
+		EvaluatedConfig->GetSettingForBranch<UMovieGraphGlobalOutputSettingNode>(UMovieGraphSettingNode::GlobalsPinName, bIncludeCDOs, bExactMatch);
+	FString FileNameFormatString = OutputSetting->OutputDirectory.Path / DebugSetting->UnrealInsightsTraceFileNameFormat;
+
+	const bool bOverwriteExistingOutput = OutputSetting->bOverwriteExistingOutput;
+
+	// Generate a filename for this encoded file
+	TMap<FString, FString> FormatOverrides;
+	FormatOverrides.Add(TEXT("ext"), TEXT("utrace"));
+
+	FMovieGraphRenderDataIdentifier TempRenderDataIdentifier;
+	TempRenderDataIdentifier.RootBranchName = UMovieGraphSettingNode::GlobalsPinName;
+
+	FMovieGraphFilenameResolveParams Params = FMovieGraphFilenameResolveParams::MakeResolveParams(
+		TempRenderDataIdentifier, this, EvaluatedConfig, GetCurrentTraversalContext(false), FormatOverrides);
+
+	FMovieGraphResolveArgs FinalFormatArgs;
+	constexpr bool bIncludeRenderPass = false;
+	constexpr bool bTestFrameNumber = false;
+	constexpr bool bIncludeCameraName = false;
+	UE::MoviePipeline::ValidateOutputFormatString(FileNameFormatString, bIncludeRenderPass, bTestFrameNumber, bIncludeCameraName);
+	FString FinalFilePath = UMovieGraphBlueprintLibrary::ResolveFilenameFormatArguments(FileNameFormatString, Params, FinalFormatArgs);
+
+	if (FPaths::IsRelative(FinalFilePath))
+	{
+		FinalFilePath = FPaths::ConvertRelativePathToFull(FinalFilePath);
+	}
+
+	// If the end user opts to delete existing files and one with the same name exists, delete it
+	if (bOverwriteExistingOutput && !UE::MoviePipeline::CanWriteToFile(*FinalFilePath, false)) // false otherwise it will always return true
+	{
+		FPlatformFileManager::Get().GetPlatformFile().DeleteFile(*FinalFilePath);
+	}
+
+	const bool bTraceStarted = FTraceAuxiliary::Start(FTraceAuxiliary::EConnectionType::File, *FinalFilePath);
+	if (bTraceStarted)
+	{
+		UE_LOG(LogMovieRenderPipeline, Log, TEXT("Started capturing UnrealInsights trace file to %s"), *FinalFilePath);
+	}
+	else
+	{
+		UE_LOG(LogMovieRenderPipeline, Warning, TEXT("Failed to start capturing UnrealInsights trace. Is there already a trace session in progress?"));
+	}
+}
+
+void UMovieGraphPipeline::StopUnrealInsightsCapture()
+{
+	FTraceAuxiliary::Stop();
+}
+
 void UMovieGraphPipeline::SetupShot(const TObjectPtr<UMoviePipelineExecutorShot>& InShot)
 {
 	ExecutePreShotScripts(InShot);
@@ -1334,6 +1419,16 @@ void UMovieGraphPipeline::TransitionToState(const EMovieRenderPipelineState InNe
 				PreviewWidget = nullptr;
 			}
 
+			// Stop Insights trace
+			constexpr bool bIncludeCDOs = false;
+			constexpr bool bExactMatch = true;
+			const UMovieGraphDebugSettingNode* DebugSetting =
+				PostRenderEvaluatedGraph->GetSettingForBranch<UMovieGraphDebugSettingNode>(UMovieGraphNode::GlobalsPinName, bIncludeCDOs, bExactMatch);
+			if (DebugSetting && DebugSetting->bCaptureUnrealInsightsTrace)
+			{
+				StopUnrealInsightsCapture();
+			}
+
 			// Job-level evaluated graph should not be referenced after export has finished
 			PostRenderEvaluatedGraph = nullptr;
 
@@ -1364,15 +1459,6 @@ void UMovieGraphPipeline::TransitionToState(const EMovieRenderPipelineState InNe
 
 			//UE_LOG(LogMovieRenderPipeline, Log, TEXT("Movie Pipeline completed. Duration: %s"), *(FDateTime::UtcNow() - InitializationTime).ToString());
 
-			//UMovieGraphDebugSettings* DebugSetting = FindOrAddSettingForShot<UMovieGraphDebugSettings>(nullptr);
-			//if (DebugSetting)
-			//{
-			//	if (DebugSetting->bCaptureUnrealInsightsTrace)
-			//	{
-			//		StopUnrealInsightsCapture();
-			//	}
-			//}
-
 			OutputNodesDataSentTo.Reset();
 
 			GraphDataSourceInstance->RestoreCachedDataPostJob();
@@ -1478,7 +1564,7 @@ void UMovieGraphPipeline::ProcessOutstandingFinishedFrames()
 			// Get a list of all output nodes for this particular render layer. We specifically skip CDOs here because we're just trying
 			// to find out which output nodes the user _wanted_ to place items onto, we later collect only the CDOs so that we have
 			// a central point for actually handling the file writing.
-			const bool bIncludeCDOs = false;
+			constexpr bool bIncludeCDOs = false;
 			const FName BranchName = RenderData.Key.RootBranchName;
 			TArray<UMovieGraphFileOutputNode*> OutputNodeInstances = OutputFrame.EvaluatedConfig->GetSettingsForBranch<UMovieGraphFileOutputNode>(BranchName, bIncludeCDOs);
 			for (const UMovieGraphFileOutputNode* Instance : OutputNodeInstances)
diff --git a/Engine/Plugins/MovieScene/MovieRenderPipeline/Source/MovieRenderPipelineCore/Public/Graph/MovieGraphFilenameResolveParams.h b/Engine/Plugins/MovieScene/MovieRenderPipeline/Source/MovieRenderPipelineCore/Public/Graph/MovieGraphFilenameResolveParams.h
--- a/Engine/Plugins/MovieScene/MovieRenderPipeline/Source/MovieRenderPipelineCore/Public/Graph/MovieGraphFilenameResolveParams.h
+++ b/Engine/Plugins/MovieScene/MovieRenderPipeline/Source/MovieRenderPipelineCore/Public/Graph/MovieGraphFilenameResolveParams.h
@@ -9,8 +9,11 @@
 // Forward Declare
 class UMoviePipelineExecutorJob;
 class UMovieGraphEvaluatedConfig;
+class UMovieGraphPipeline;
 class UMoviePipelineExecutorShot;
 
+struct FMovieGraphTraversalContext;
+
 /**
 * This data structure contains a list of key-value pairs (as strings) for both filename resolving, and file metadata.
 * They are stored as two separate arrays (as metadata is often a bit more verbose/human friendly) and it is not required
@@ -42,7 +45,7 @@ struct FMovieGraphResolveArgs
 * be preserved so the resulting string back would say "my_seq_name/{frame_number}".
 */
 USTRUCT(BlueprintType)
-struct FMovieGraphFilenameResolveParams
+struct MOVIERENDERPIPELINECORE_API FMovieGraphFilenameResolveParams
 {
 	GENERATED_BODY()
 
@@ -63,6 +66,14 @@ struct FMovieGraphFilenameResolveParams
 		, DefaultFrameRate(FFrameRate(24, 1))
 	{
 	}
+
+	/** Convenience function to make filename resolve parameters for output files, or OCIO contexts. */
+	static FMovieGraphFilenameResolveParams MakeResolveParams(
+		const FMovieGraphRenderDataIdentifier& InRenderId,
+		const UMovieGraphPipeline* InPipeline,
+		const TObjectPtr<UMovieGraphEvaluatedConfig>& InEvaluatedConfig,
+		const FMovieGraphTraversalContext& InTraversalContext,
+		const TMap<FString, FString>& InAdditionalFormatArgs = {});
 	
 	/**
 	* This is used to fill out tokens related to the render data that the file represents.
diff --git a/Engine/Plugins/MovieScene/MovieRenderPipeline/Source/MovieRenderPipelineCore/Public/Graph/MovieGraphPipeline.h b/Engine/Plugins/MovieScene/MovieRenderPipeline/Source/MovieRenderPipelineCore/Public/Graph/MovieGraphPipeline.h
--- a/Engine/Plugins/MovieScene/MovieRenderPipeline/Source/MovieRenderPipelineCore/Public/Graph/MovieGraphPipeline.h
+++ b/Engine/Plugins/MovieScene/MovieRenderPipeline/Source/MovieRenderPipelineCore/Public/Graph/MovieGraphPipeline.h
@@ -156,6 +156,11 @@ protected:
 	/** Begins the export process for a primary job (not called for shot jobs). */
 	virtual void BeginExport();
 	
+	/** Attempts to start an Unreal Insights capture to a file on disk adjacent to the movie output. */
+	void StartUnrealInsightsCapture(UMovieGraphEvaluatedConfig* EvaluatedConfig);
+	/** Attempts to stop an already started Unreal Insights capture. */
+	void StopUnrealInsightsCapture();
+	
 	virtual void LoadPreviewWidget();
 	virtual void SetPreviewWidgetVisibleImpl(bool bInIsVisible);
 
diff --git a/Engine/Plugins/MovieScene/MovieRenderPipeline/Source/MovieRenderPipelineCore/Public/Graph/Nodes/MovieGraphDebugNode.h b/Engine/Plugins/MovieScene/MovieRenderPipeline/Source/MovieRenderPipelineCore/Public/Graph/Nodes/MovieGraphDebugNode.h
--- a/Engine/Plugins/MovieScene/MovieRenderPipeline/Source/MovieRenderPipelineCore/Public/Graph/Nodes/MovieGraphDebugNode.h
+++ b/Engine/Plugins/MovieScene/MovieRenderPipeline/Source/MovieRenderPipelineCore/Public/Graph/Nodes/MovieGraphDebugNode.h
@@ -29,6 +29,9 @@ public:
 	
 	UPROPERTY(EditAnywhere, BlueprintReadWrite, Category = Overrides, meta = (InlineEditConditionToggle))
 	uint8 bOverride_bCaptureUnrealInsightsTrace : 1;
+	
+	UPROPERTY(EditAnywhere, BlueprintReadWrite, Category = Overrides, meta = (InlineEditConditionToggle))
+	uint8 bOverride_UnrealInsightsTraceFileNameFormat : 1;
 
 	/** 
 	* If true, automatically trigger RenderDoc to capture rendering information. RenderDoc plugin must be enabled, 
@@ -43,4 +46,11 @@ public:
 	*/
 	UPROPERTY(EditAnywhere, BlueprintReadWrite, Category = "Settings", meta = (EditCondition = "bOverride_bCaptureUnrealInsightsTrace"))
 	bool bCaptureUnrealInsightsTrace;
+
+	/** 
+	* If true, automatically capture an Unreal Insights trace file for the duration of the render.
+	* Resulting capture will be in the global Output Directory for the job.
+	*/
+	UPROPERTY(EditAnywhere, BlueprintReadWrite, Category = "Settings", meta = (EditCondition = "bOverride_UnrealInsightsTraceFileNameFormat"))
+	FString UnrealInsightsTraceFileNameFormat = TEXT("{sequence_name}_UnrealInsights");
 };
\ No newline at end of file
diff --git a/Engine/Plugins/MovieScene/MovieRenderPipeline/Source/MovieRenderPipelineEditor/Private/Customizations/Graph/MovieGraphFormatTokenCustomization.cpp b/Engine/Plugins/MovieScene/MovieRenderPipeline/Source/MovieRenderPipelineEditor/Private/Customizations/Graph/MovieGraphFormatTokenCustomization.cpp
--- a/Engine/Plugins/MovieScene/MovieRenderPipeline/Source/MovieRenderPipelineEditor/Private/Customizations/Graph/MovieGraphFormatTokenCustomization.cpp
+++ b/Engine/Plugins/MovieScene/MovieRenderPipeline/Source/MovieRenderPipelineEditor/Private/Customizations/Graph/MovieGraphFormatTokenCustomization.cpp
@@ -3,6 +3,7 @@
 #include "MovieGraphFormatTokenCustomization.h"
 
 #include "Graph/MovieGraphBlueprintLibrary.h"
+#include "Graph/Nodes/MovieGraphDebugNode.h"
 #include "Graph/Nodes/MovieGraphFileOutputNode.h"
 
 #include "DetailLayoutBuilder.h"
@@ -18,6 +19,18 @@ void FMovieGraphFormatTokenCustomization::CustomizeDetails(IDetailLayoutBuilder&
 	// This should work just fine for UMovieGraphFileOutputNode and UMovieGraphCommandLineEncoderNode as long as the property names stay in sync
 	OutputFormatPropertyHandle = DetailBuilder.GetProperty(GET_MEMBER_NAME_CHECKED(UMovieGraphFileOutputNode, FileNameFormat));
 
+	// If we can't find it, try the debug node
+	if (!OutputFormatPropertyHandle->IsValidHandle())
+	{
+		OutputFormatPropertyHandle = DetailBuilder.GetProperty(GET_MEMBER_NAME_CHECKED(UMovieGraphDebugSettingNode, UnrealInsightsTraceFileNameFormat));
+	}
+
+	// If we still can't find it, early out
+	if (!OutputFormatPropertyHandle->IsValidHandle())
+	{
+		return;
+	}
+
 	FText StartingDisplayText;
 	OutputFormatPropertyHandle->GetValueAsDisplayText(StartingDisplayText);
 
diff --git a/Engine/Plugins/MovieScene/MovieRenderPipeline/Source/MovieRenderPipelineEditor/Private/Graph/MovieGraphAssetToolkit.cpp b/Engine/Plugins/MovieScene/MovieRenderPipeline/Source/MovieRenderPipelineEditor/Private/Graph/MovieGraphAssetToolkit.cpp
--- a/Engine/Plugins/MovieScene/MovieRenderPipeline/Source/MovieRenderPipelineEditor/Private/Graph/MovieGraphAssetToolkit.cpp
+++ b/Engine/Plugins/MovieScene/MovieRenderPipeline/Source/MovieRenderPipelineEditor/Private/Graph/MovieGraphAssetToolkit.cpp
@@ -30,6 +30,7 @@
 #include "PropertyEditorModule.h"
 #include "Selection.h"
 #include "ToolMenus.h"
+#include "Graph/Nodes/MovieGraphDebugNode.h"
 #include "Widgets/Docking/SDockTab.h"
 #include "Widgets/Graph/SMovieGraphConfigPanel.h"
 
@@ -403,6 +404,10 @@ TSharedRef<SDockTab> FMovieGraphAssetToolkit::SpawnTab_RenderGraphDetails(const
 	SelectedGraphObjectsDetailsWidget->RegisterInstancedCustomPropertyLayout(
 		UMovieGraphCommandLineEncoderNode::StaticClass(),
 		FOnGetDetailCustomizationInstance::CreateStatic(&FMovieGraphFormatTokenCustomization::MakeInstance));
+
+	SelectedGraphObjectsDetailsWidget->RegisterInstancedCustomPropertyLayout(
+		UMovieGraphDebugSettingNode::StaticClass(),
+		FOnGetDetailCustomizationInstance::CreateStatic(&FMovieGraphFormatTokenCustomization::MakeInstance));
 	
 	TSharedRef<SWidget> CustomContent = SAssignNew(NameAreaCustomContent, SHorizontalBox)
 	+ SHorizontalBox::Slot()
diff --git a/Engine/Plugins/MovieScene/MovieRenderPipeline/Source/MovieRenderPipelineRenderPasses/Private/MovieGraphImageSequenceOutputNode.cpp b/Engine/Plugins/MovieScene/MovieRenderPipeline/Source/MovieRenderPipelineRenderPasses/Private/MovieGraphImageSequenceOutputNode.cpp
--- a/Engine/Plugins/MovieScene/MovieRenderPipeline/Source/MovieRenderPipelineRenderPasses/Private/MovieGraphImageSequenceOutputNode.cpp
+++ b/Engine/Plugins/MovieScene/MovieRenderPipeline/Source/MovieRenderPipelineRenderPasses/Private/MovieGraphImageSequenceOutputNode.cpp
@@ -26,52 +26,7 @@
 #endif // WITH_OCIO
 
 namespace UE::MovieGraph::Private
-{
-	/** Convenience function to make filename resolve parameters for output files, or OCIO contexts. */
-	FMovieGraphFilenameResolveParams MakeResolveParams(
-		const FMovieGraphRenderDataIdentifier& InRenderId,
-		const UMovieGraphPipeline* InPipeline,
-		const TObjectPtr<UMovieGraphEvaluatedConfig>& InEvaluatedConfig,
-		const FMovieGraphTraversalContext& InTraversalContext,
-		const TMap<FString, FString>& InAdditionalFormatArgs = {})
-	{
-		const TObjectPtr<UMoviePipelineExecutorShot>& Shot = InPipeline->GetActiveShotList()[InTraversalContext.ShotIndex];
-
-		FMovieGraphFilenameResolveParams Params = FMovieGraphFilenameResolveParams();
-		Params.RenderDataIdentifier = InRenderId;
-		Params.RootFrameNumber = InTraversalContext.Time.RootFrameNumber.Value;
-		Params.ShotFrameNumber = InTraversalContext.Time.ShotFrameNumber.Value;
-		Params.RootFrameNumberRel = InTraversalContext.Time.OutputFrameNumber;
-		Params.ShotFrameNumberRel = InTraversalContext.Time.ShotOutputFrameNumber;
-		//Params.FileMetadata = ToDo: Track File Metadata
-		const UMovieGraphGlobalOutputSettingNode* OutputSettingNode = InEvaluatedConfig->GetSettingForBranch<UMovieGraphGlobalOutputSettingNode>(UMovieGraphNode::GlobalsPinName);
-		if (IsValid(OutputSettingNode))
-		{
-			Params.ZeroPadFrameNumberCount = OutputSettingNode->ZeroPadFrameNumbers;
-			Params.FrameNumberOffset = OutputSettingNode->FrameNumberOffset;
-		}
-		Params.EvaluatedConfig = InEvaluatedConfig;
-		Params.Version = Shot->ShotInfo.VersionNumber;
-
-		// If time dilation is in effect, RootFrameNumber and ShotFrameNumber will contain duplicates and the files will overwrite each other, 
-		// so we force them into relative mode and then warn users we did that (as their numbers will jump from say 1001 -> 0000).
-		bool bForceRelativeFrameNumbers = false;
-		// if (FileNameFormatString.Contains(TEXT("{frame")) && InTraversalContext.Time.IsTimeDilated() && !FileNameFormatString.Contains(TEXT("_rel}")))
-		// {
-		// 	UE_LOG(LogMovieRenderPipeline, Warning, TEXT("Time Dilation was used but output format does not use relative time, forcing relative numbers. Change {frame_number} to {frame_number_rel} (or shot version) to remove this message."));
-		// 	bForceRelativeFrameNumbers = true;
-		// }
-		Params.bForceRelativeFrameNumbers = bForceRelativeFrameNumbers;
-		Params.bEnsureAbsolutePath = true;
-		Params.FileNameFormatOverrides = InAdditionalFormatArgs;
-		Params.InitializationTime = InPipeline->GetInitializationTime();
-		Params.InitializationTimeOffset = InPipeline->GetInitializationTimeOffset();
-		Params.Shot = Shot;
-		Params.Job = InPipeline->GetCurrentJob();
-
-		return Params;
-	}
-	
+{	
 #if WITH_OCIO
 	struct FOpenColorIOPixelPreProcessor
 	{
@@ -104,7 +59,7 @@ namespace UE::MovieGraph::Private
 		TMap<FString, FString> OutContext;
 		OutContext.Reserve(InContext.Num());
 
-		FMovieGraphFilenameResolveParams Params = MakeResolveParams(InRenderId, InPipeline, InEvaluatedConfig, InTraversalContext);
+		FMovieGraphFilenameResolveParams Params = FMovieGraphFilenameResolveParams::MakeResolveParams(InRenderId, InPipeline, InEvaluatedConfig, InTraversalContext);
 
 		for (const TPair<FString, FString>& Pair : InContext)
 		{
@@ -317,7 +272,7 @@ FString UMovieGraphImageSequenceOutputNode::CreateFileName(
 
 	UE::MovieGraph::FMovieGraphSampleState* Payload = InRenderData.Value->GetPayload<UE::MovieGraph::FMovieGraphSampleState>();
 
-	FMovieGraphFilenameResolveParams Params = UE::MovieGraph::Private::MakeResolveParams(
+	FMovieGraphFilenameResolveParams Params = FMovieGraphFilenameResolveParams::MakeResolveParams(
 		InRenderData.Key,
 		InPipeline,
 		InRawFrameData->EvaluatedConfig.Get(),
@@ -911,7 +866,7 @@ FString UMovieGraphImageSequenceOutputNode_MultiLayerEXR::ResolveOutputFilename(
 	FMovieGraphRenderDataIdentifier TempRenderDataIdentifier;
 	TempRenderDataIdentifier.RootBranchName = InBranchName;
 
-	FMovieGraphFilenameResolveParams Params = UE::MovieGraph::Private::MakeResolveParams(
+	FMovieGraphFilenameResolveParams Params = FMovieGraphFilenameResolveParams::MakeResolveParams(
 		TempRenderDataIdentifier, InPipeline, InRawFrameData->EvaluatedConfig.Get(), InRawFrameData->TraversalContext, FormatOverrides);
 	
 	const FString FilePathFormatString = OutputSettings->OutputDirectory.Path / FileNameFormatString;
