MoviePipeline: Updated the graph modifier logic to treat volumetrics specially when they set 'Affect Indirect Lighting While Hidden'. Volumetrics do not have this property (unlike UPrimitiveComponent), so instead volumetrics set the 'Hidden in Main' property.

#rnx

diff --git a/Engine/Plugins/MovieScene/MovieRenderPipeline/Source/MovieRenderPipelineCore/Private/Graph/MovieGraphRenderLayerSubsystem.cpp b/Engine/Plugins/MovieScene/MovieRenderPipeline/Source/MovieRenderPipelineCore/Private/Graph/MovieGraphRenderLayerSubsystem.cpp
--- a/Engine/Plugins/MovieScene/MovieRenderPipeline/Source/MovieRenderPipelineCore/Private/Graph/MovieGraphRenderLayerSubsystem.cpp
+++ b/Engine/Plugins/MovieScene/MovieRenderPipeline/Source/MovieRenderPipelineCore/Private/Graph/MovieGraphRenderLayerSubsystem.cpp
@@ -166,13 +166,6 @@ void UMovieGraphRenderPropertyModifier::ApplyModifier(const UWorld* World)
 	FActorVisibilityState NewVisibilityState;
 	NewVisibilityState.bIsHidden = bIsHidden;
 
-	// Create a stub component representation since we're just using this to propagate to all children right now.
-	FActorVisibilityState::FComponentState& NewComponentState = NewVisibilityState.Components.AddDefaulted_GetRef();
-	NewComponentState.bCastsShadows = bCastsShadows;
-	NewComponentState.bCastShadowWhileHidden = bCastShadowWhileHidden;
-	NewComponentState.bAffectIndirectLightingWhileHidden = bAffectIndirectLightingWhileHidden;
-	NewComponentState.bHoldout = bHoldout;
-
 	// Generate a warning if holdout is being used, but alpha is not enabled in post processing. Without that setting enabled, holdout will not work.
 	const URendererSettings* RendererSettings = GetDefault<URendererSettings>();
 	if (bHoldout && (RendererSettings->bEnableAlphaChannelInPostProcessing == EAlphaChannelMode::Disabled))
@@ -195,52 +188,71 @@ void UMovieGraphRenderPropertyModifier::ApplyModifier(const UWorld* World)
 
 		for (const AActor* Actor : MatchingActors)
 		{
+			NewVisibilityState.Actor = Actor;
+			
 			// Save out visibility state before the modifier is applied
 			FActorVisibilityState OriginalVisibilityState;
 			OriginalVisibilityState.Actor = Actor;
 			OriginalVisibilityState.bIsHidden = Actor->IsHidden();
 
-			const bool bIncludeFromChildActors = true;
+			constexpr bool bIncludeFromChildActors = true;
 			TInlineComponentArray<USceneComponent*> Components;
 			Actor->GetComponents<USceneComponent>(Components, bIncludeFromChildActors);
 
 			OriginalVisibilityState.Components.Reserve(Components.Num());
-			for(const USceneComponent* SceneComponent : Components)
+			NewVisibilityState.Components.Empty(Components.Num());
+			
+			for (const USceneComponent* SceneComponent : Components)
 			{
-				FActorVisibilityState::FComponentState& ComponentState = OriginalVisibilityState.Components.AddDefaulted_GetRef();
-				ComponentState.Component = SceneComponent;
+#if WITH_EDITORONLY_DATA
+				// Don't bother processing editor-only components (editor billboard icons, text, etc)
+				if (SceneComponent->IsEditorOnly())
+				{
+					continue;
+				}
+#endif // WITH_EDITORONLY_DATA
+				
+				FActorVisibilityState::FComponentState& CachedComponentState = OriginalVisibilityState.Components.AddDefaulted_GetRef();
+				CachedComponentState.Component = SceneComponent;
 
 				// Cache the state
 				if (const UPrimitiveComponent* AsPrimitiveComponent = Cast<UPrimitiveComponent>(SceneComponent))
 				{
-					ComponentState.bCastsShadows = AsPrimitiveComponent->CastShadow;
-					ComponentState.bCastShadowWhileHidden = AsPrimitiveComponent->bCastHiddenShadow;
-					ComponentState.bAffectIndirectLightingWhileHidden = AsPrimitiveComponent->bAffectIndirectLightingWhileHidden;
-					ComponentState.bHoldout = AsPrimitiveComponent->bHoldout;
+					CachedComponentState.bCastsShadows = AsPrimitiveComponent->CastShadow;
+					CachedComponentState.bCastShadowWhileHidden = AsPrimitiveComponent->bCastHiddenShadow;
+					CachedComponentState.bAffectIndirectLightingWhileHidden = AsPrimitiveComponent->bAffectIndirectLightingWhileHidden;
+					CachedComponentState.bHoldout = AsPrimitiveComponent->bHoldout;
 				}
 				// Volumetrics are special cases as they don't inherit from UPrimitiveComponent, and don't support all of the flags.
 				else if (const UVolumetricCloudComponent* AsVolumetricCloudComponent = Cast<UVolumetricCloudComponent>(SceneComponent))
 				{
-					ComponentState.bHoldout = AsVolumetricCloudComponent->bHoldout;
+					CachedComponentState.bHoldout = AsVolumetricCloudComponent->bHoldout;
+					CachedComponentState.bAffectIndirectLightingWhileHidden = !AsVolumetricCloudComponent->bRenderInMainPass;
 				}
 				else if (const USkyAtmosphereComponent* AsSkyAtmosphereComponent = Cast<USkyAtmosphereComponent>(SceneComponent))
 				{
-					ComponentState.bHoldout = AsSkyAtmosphereComponent->bHoldout;
+					CachedComponentState.bHoldout = AsSkyAtmosphereComponent->bHoldout;
+					CachedComponentState.bAffectIndirectLightingWhileHidden = !AsSkyAtmosphereComponent->bRenderInMainPass;
 				}
 				else if (const UExponentialHeightFogComponent* AsExponentialHeightFogComponent = Cast<UExponentialHeightFogComponent>(SceneComponent))
 				{
-					ComponentState.bHoldout = AsExponentialHeightFogComponent->bHoldout;
+					CachedComponentState.bHoldout = AsExponentialHeightFogComponent->bHoldout;
+					CachedComponentState.bAffectIndirectLightingWhileHidden = !AsExponentialHeightFogComponent->bRenderInMainPass;
 				}
 
-				// Then override it. We override it one component at a time so that we can avoid making a bunch of
-				// copies of NewVisibilityState for the varying number of components you might have.
-				NewVisibilityState.Actor = Actor;
-				NewVisibilityState.Components[0].Component = SceneComponent;
-				SetActorVisibilityState(NewVisibilityState);
+				// SetActorVisibilityState() relies on inspecting individual components to determine what to do, hence why we need to generate a
+				// separate state for each component. Ideally this would not be needed in order to prevent constantly regenerating these structs.
+				FActorVisibilityState::FComponentState& NewComponentState = NewVisibilityState.Components.AddDefaulted_GetRef();
+				NewComponentState.Component = SceneComponent;
+				NewComponentState.bCastsShadows = bCastsShadows;
+				NewComponentState.bCastShadowWhileHidden = bCastShadowWhileHidden;
+				NewComponentState.bAffectIndirectLightingWhileHidden = bAffectIndirectLightingWhileHidden;
+				NewComponentState.bHoldout = bHoldout;
 			}
 			
+			SetActorVisibilityState(NewVisibilityState);
+			
 			ModifiedActors.Add(OriginalVisibilityState);
-
 		}
 	}
 }
@@ -263,14 +275,33 @@ void UMovieGraphRenderPropertyModifier::SetActorVisibilityState(const FActorVisi
 		return;
 	}
 
-	if (bOverride_bIsHidden)
+	// In most cases, if the hidden state is being modified, the hidden state should be set. However, there is an exception for volumetrics.
+	// If volumetrics set the 'Affect Indirect Lighting While Hidden' flag to true, the volumetric component needs to set the 'Render in Main' flag
+	// instead, and the 'Hidden' flag should NOT be set on the *actor*. Setting the Hidden flag on the actor in this case will override the behavior
+	// of 'Render in Main' and volumetrics will not affect indirect lighting.
+	bool bShouldSetActorHiddenState = bOverride_bIsHidden;
+
+	// Volumetrics are a special case and their visibility properties need to be handled separately
+	auto SetStateForVolumetrics = [this, &bShouldSetActorHiddenState]<typename VolumetricsType>(VolumetricsType& VolumetricsComponent, const FActorVisibilityState::FComponentState& NewComponentState)
 	{
-		Actor->SetActorHiddenInGame(NewVisibilityState.bIsHidden);
+		if (bOverride_bHoldout)
+		{
+			VolumetricsComponent->SetHoldout(NewComponentState.bHoldout);
+		}
 
-#if WITH_EDITOR
-		Actor->SetIsTemporarilyHiddenInEditor(NewVisibilityState.bIsHidden);
-#endif
-	}
+		if (bOverride_bAffectIndirectLightingWhileHidden)
+		{
+			// If the component should affect indirect while hidden, then we need to use 'Render in Main' instead.
+			VolumetricsComponent->SetRenderInMainPass(!NewComponentState.bAffectIndirectLightingWhileHidden);
+
+			// Don't allow the actor to hide itself if this component is not going to be rendered in the main pass. Hiding the actor will
+			// negate the effects of setting Render In Main Pass.
+			if (NewComponentState.bAffectIndirectLightingWhileHidden)
+			{
+				bShouldSetActorHiddenState = false;
+			}
+		}
+	};
 
 	for (const FActorVisibilityState::FComponentState& ComponentState : NewVisibilityState.Components)
 	{
@@ -306,26 +337,26 @@ void UMovieGraphRenderPropertyModifier::SetActorVisibilityState(const FActorVisi
 		// Volumetrics are special cases as they don't inherit from UPrimitiveComponent, and don't support all of the flags.
 		else if (UVolumetricCloudComponent* AsVolumetricCloudComponent = Cast<UVolumetricCloudComponent>(ComponentState.Component.Get()))
 		{
-			if (bOverride_bHoldout)
-			{
-				AsVolumetricCloudComponent->SetHoldout(ComponentState.bHoldout);
-			}
+			SetStateForVolumetrics(AsVolumetricCloudComponent, ComponentState);
 		}
 		else if (USkyAtmosphereComponent* AsSkyAtmosphereComponent = Cast<USkyAtmosphereComponent>(ComponentState.Component.Get()))
 		{
-			if (bOverride_bHoldout)
-			{
-				AsSkyAtmosphereComponent->SetHoldout(ComponentState.bHoldout);
-			}
+			SetStateForVolumetrics(AsSkyAtmosphereComponent, ComponentState);
 		}
 		else if (UExponentialHeightFogComponent* AsExponentialHeightFogComponent = Cast<UExponentialHeightFogComponent>(ComponentState.Component.Get()))
 		{
-			if (bOverride_bHoldout)
-			{
-				AsExponentialHeightFogComponent->SetHoldout(ComponentState.bHoldout);
-			}
+			SetStateForVolumetrics(AsExponentialHeightFogComponent, ComponentState);
 		}
 	}
+	
+	if (bShouldSetActorHiddenState)
+	{
+		Actor->SetActorHiddenInGame(NewVisibilityState.bIsHidden);
+
+#if WITH_EDITOR
+		Actor->SetIsTemporarilyHiddenInEditor(NewVisibilityState.bIsHidden);
+#endif
+	}
 }
 
 void UMovieGraphCollectionModifier::AddCollection(UMovieGraphCollection* Collection)
