MovieGraph: Fix crash when undoing/redoing after changing an Input variable's underlying type.

diff --git a/Engine/Plugins/MovieScene/MovieRenderPipeline/Source/MovieRenderPipelineCore/Public/Graph/MovieGraphValueContainer.cpp b/Engine/Plugins/MovieScene/MovieRenderPipeline/Source/MovieRenderPipelineCore/Public/Graph/MovieGraphValueContainer.cpp
--- a/Engine/Plugins/MovieScene/MovieRenderPipeline/Source/MovieRenderPipelineCore/Public/Graph/MovieGraphValueContainer.cpp
+++ b/Engine/Plugins/MovieScene/MovieRenderPipeline/Source/MovieRenderPipelineCore/Public/Graph/MovieGraphValueContainer.cpp
@@ -30,6 +30,7 @@ void UMovieGraphValueContainer::SetPropertyName(const FName& InName)
 			return;
 		}
 
+		Modify();
 		PropertyName = InName;
 
 		// Changing the property name requires a new desc and a migration
@@ -134,71 +135,85 @@ FString UMovieGraphValueContainer::GetValueSerializedString()
 
 bool UMovieGraphValueContainer::SetValueBool(const bool bInValue)
 {
+	Modify();
 	return Value.SetValueBool(PropertyName, bInValue) == EPropertyBagResult::Success;
 }
 
 bool UMovieGraphValueContainer::SetValueByte(const uint8 InValue)
 {
+	Modify();
 	return Value.SetValueByte(PropertyName, InValue) == EPropertyBagResult::Success;
 }
 
 bool UMovieGraphValueContainer::SetValueInt32(const int32 InValue)
 {
+	Modify();
 	return Value.SetValueInt32(PropertyName, InValue) == EPropertyBagResult::Success;
 }
 
 bool UMovieGraphValueContainer::SetValueInt64(const int64 InValue)
 {
+	Modify();
 	return Value.SetValueInt64(PropertyName, InValue) == EPropertyBagResult::Success;
 }
 
 bool UMovieGraphValueContainer::SetValueFloat(const float InValue)
 {
+	Modify();
 	return Value.SetValueFloat(PropertyName, InValue) == EPropertyBagResult::Success;
 }
 
 bool UMovieGraphValueContainer::SetValueDouble(const double InValue)
 {
+	Modify();
 	return Value.SetValueDouble(PropertyName, InValue) == EPropertyBagResult::Success;
 }
 
 bool UMovieGraphValueContainer::SetValueName(const FName InValue)
 {
+	Modify();
 	return Value.SetValueName(PropertyName, InValue) == EPropertyBagResult::Success;
 }
 
 bool UMovieGraphValueContainer::SetValueString(const FString& InValue)
 {
+	Modify();
 	return Value.SetValueString(PropertyName, InValue) == EPropertyBagResult::Success;
 }
 
 bool UMovieGraphValueContainer::SetValueText(const FText& InValue)
 {
+	Modify();
 	return Value.SetValueText(PropertyName, InValue) == EPropertyBagResult::Success;
 }
 
 bool UMovieGraphValueContainer::SetValueEnum(const uint8 InValue, const UEnum* Enum)
 {
+	Modify();
 	return Value.SetValueEnum(PropertyName, InValue, Enum) == EPropertyBagResult::Success;
 }
 
 bool UMovieGraphValueContainer::SetValueStruct(FConstStructView InValue)
 {
+	Modify();
 	return Value.SetValueStruct(PropertyName, InValue) == EPropertyBagResult::Success;
 }
 
 bool UMovieGraphValueContainer::SetValueObject(UObject* InValue)
 {
+	Modify();
 	return Value.SetValueObject(PropertyName, InValue) == EPropertyBagResult::Success;
 }
 
 bool UMovieGraphValueContainer::SetValueClass(UClass* InValue)
 {
+	Modify();
 	return Value.SetValueClass(PropertyName, InValue) == EPropertyBagResult::Success;
 }
 
 bool UMovieGraphValueContainer::SetValueSerializedString(const FString& NewValue)
 {
+	Modify();
 	return Value.SetValueSerializedString(PropertyName, NewValue) == EPropertyBagResult::Success;
 }
 
@@ -216,6 +231,8 @@ void UMovieGraphValueContainer::SetValueType(EMovieGraphValueType ValueType, UOb
 {
 	if (const FPropertyBagPropertyDesc* Desc = Value.FindPropertyDescByName(PropertyName))
 	{
+		Modify();
+
 		FPropertyBagPropertyDesc NewDesc;
 		NewDesc.Name = Desc->Name;
 		NewDesc.ValueType = static_cast<EPropertyBagPropertyType>(ValueType);
@@ -248,6 +265,8 @@ void UMovieGraphValueContainer::SetValueTypeObject(const UObject* ValueTypeObjec
 {
 	if (const FPropertyBagPropertyDesc* Desc = Value.FindPropertyDescByName(PropertyName))
 	{
+		Modify();
+
 		FPropertyBagPropertyDesc NewDesc(*Desc);
 		NewDesc.ValueTypeObject = ValueTypeObject;
 
@@ -270,6 +289,8 @@ void UMovieGraphValueContainer::SetValueContainerType(EMovieGraphContainerType C
 {
 	if (const FPropertyBagPropertyDesc* Desc = Value.FindPropertyDescByName(PropertyName))
 	{
+		Modify();
+
 		FPropertyBagPropertyDesc NewDesc(*Desc);
 		NewDesc.ContainerTypes = { static_cast<EPropertyBagContainerType>(ContainerType) };
 
@@ -287,6 +308,8 @@ void UMovieGraphValueContainer::SetFromDesc(const FPropertyBagPropertyDesc* InDe
 {
 	if (InDesc)
 	{
+		Modify();
+
 		FPropertyBagPropertyDesc NewDesc;
 		NewDesc.Name = InDesc->Name;
 		NewDesc.ValueType = InDesc->ValueType;
diff --git a/Engine/Plugins/MovieScene/MovieRenderPipeline/Source/MovieRenderPipelineEditor/Private/Graph/MovieEdGraphNode.cpp b/Engine/Plugins/MovieScene/MovieRenderPipeline/Source/MovieRenderPipelineEditor/Private/Graph/MovieEdGraphNode.cpp
--- a/Engine/Plugins/MovieScene/MovieRenderPipeline/Source/MovieRenderPipelineEditor/Private/Graph/MovieEdGraphNode.cpp
+++ b/Engine/Plugins/MovieScene/MovieRenderPipeline/Source/MovieRenderPipelineEditor/Private/Graph/MovieEdGraphNode.cpp
@@ -523,7 +523,17 @@ void UMoviePipelineEdGraphNodeBase::OnRuntimeNodeChanged(const UMovieGraphNode*
 {
 	if (InChangedNode == GetRuntimeNode())
 	{
-		ReconstructNode();
+		// During Undo/Redo ReconstructNode gets called twice. When the runtime UObject gets
+		// its properties restored, we hear the delegate broadcast and this function gets run,
+		// and then the editor objects are restored (and are rebuilt after undo/redo). This
+		// creates a problem where when redoing, it restores the editor nodes to a temporary
+		// mid-transaction state, which then causes a crash.
+		// To avoid this we skip calling ReconstructNode during undo/redo, knowing that it will be
+		// reconstructed later alongside the whole graph.
+		if (!GIsTransacting)
+		{
+			ReconstructNode();
+		}
 	}
 }
 
@@ -591,6 +601,8 @@ void UMoviePipelineEdGraphNodeBase::ReconstructPins()
 		OldPin->SubPins.Remove(nullptr);
 		DestroyPin(OldPin);
 	}
+
+	GetGraph()->NotifyGraphChanged();
 }
 
 void UMoviePipelineEdGraphNodeBase::PrepareForCopying()
