MoviePipeline: Polished up the UI of the two cvar-related nodes in the graph.

* The nodes were renamed to Apply Console Variable Preset and Set Console Variable.
* Both of these nodes now have helpful description text under the node title (preset asset name + num cvars in the preset, or the cvar being set and the value of that cvar).
* A details customization was added for the preset node to show all of the cvars contained within the preset asset (and provide a way to create a node overriding one of the preset's cvars).
* Another customization was added to provide cvar autocomplete for the Set Console Variable node.
* Updated SConsoleVariablesEditorCustomConsoleInputBox to provide support for committing text when focus is lost. This allows the details customization to only cue the graph to update after the user is done typing rather than after every character is entered.

diff --git a/Engine/Plugins/Editor/ConsoleVariablesEditor/Source/ConsoleVariablesEditor/Public/Views/Widgets/SConsoleVariablesEditorCustomConsoleInputBox.cpp b/Engine/Plugins/Editor/ConsoleVariablesEditor/Source/ConsoleVariablesEditor/Public/Views/Widgets/SConsoleVariablesEditorCustomConsoleInputBox.cpp
--- a/Engine/Plugins/Editor/ConsoleVariablesEditor/Source/ConsoleVariablesEditor/Public/Views/Widgets/SConsoleVariablesEditorCustomConsoleInputBox.cpp
+++ b/Engine/Plugins/Editor/ConsoleVariablesEditor/Source/ConsoleVariablesEditor/Public/Views/Widgets/SConsoleVariablesEditorCustomConsoleInputBox.cpp
@@ -24,6 +24,7 @@ void SConsoleVariablesEditorCustomConsoleInputBox::Construct(
 	const FArguments& InArgs)
 {
 	bHideOnFocusLost = InArgs._HideOnFocusLost;
+	bCommitOnFocusLost = InArgs._CommitOnFocusLost;
 	bClearOnCommit = InArgs._ClearOnCommit;
 	IsEnabledAttribute = InArgs._IsEnabled;
 	Font = InArgs._Font;
@@ -70,6 +71,11 @@ void SConsoleVariablesEditorCustomConsoleInputBox::Construct(
 					{
 						SuggestionBox->SetIsOpen(false);
 
+						if (bCommitOnFocusLost)
+						{
+							CommitInput();
+						}
+
 						if (bHideOnFocusLost)
 						{
 							SetVisibility(EVisibility::Collapsed);
diff --git a/Engine/Plugins/Editor/ConsoleVariablesEditor/Source/ConsoleVariablesEditor/Public/Views/Widgets/SConsoleVariablesEditorCustomConsoleInputBox.h b/Engine/Plugins/Editor/ConsoleVariablesEditor/Source/ConsoleVariablesEditor/Public/Views/Widgets/SConsoleVariablesEditorCustomConsoleInputBox.h
--- a/Engine/Plugins/Editor/ConsoleVariablesEditor/Source/ConsoleVariablesEditor/Public/Views/Widgets/SConsoleVariablesEditorCustomConsoleInputBox.h
+++ b/Engine/Plugins/Editor/ConsoleVariablesEditor/Source/ConsoleVariablesEditor/Public/Views/Widgets/SConsoleVariablesEditorCustomConsoleInputBox.h
@@ -15,10 +15,15 @@ public:
 	DECLARE_DELEGATE_OneParam(FOnTextChanged, const FText&);
 
 	SLATE_BEGIN_ARGS(SConsoleVariablesEditorCustomConsoleInputBox)
+		: _CommitOnFocusLost(false)
 	{}
+	
 		/** If true, hide the input box when focus is lost. */
 		SLATE_ARGUMENT(bool, HideOnFocusLost)
 
+		/** If true, commit the input box's text when the input box loses focus (calls the OnTextCommitted event). */
+		SLATE_ARGUMENT(bool, CommitOnFocusLost)
+
 		/** If true, clear the input box when text is committed. */
 		SLATE_ARGUMENT(bool, ClearOnCommit)
 
@@ -138,6 +143,7 @@ private:
 	bool bIgnoreUIUpdate = false;
 
 	bool bHideOnFocusLost = true;
+	bool bCommitOnFocusLost = false;
 	bool bClearOnCommit = true;
 	TAttribute<bool> IsEnabledAttribute;
 	FSlateFontInfo Font;
diff --git a/Engine/Plugins/MovieScene/MovieRenderPipeline/Source/MovieRenderPipelineCore/Private/Graph/Nodes/MovieGraphApplyCVarPresetNode.cpp b/Engine/Plugins/MovieScene/MovieRenderPipeline/Source/MovieRenderPipelineCore/Private/Graph/Nodes/MovieGraphApplyCVarPresetNode.cpp
--- a/Engine/Plugins/MovieScene/MovieRenderPipeline/Source/MovieRenderPipelineCore/Private/Graph/Nodes/MovieGraphApplyCVarPresetNode.cpp
+++ b/Engine/Plugins/MovieScene/MovieRenderPipeline/Source/MovieRenderPipelineCore/Private/Graph/Nodes/MovieGraphApplyCVarPresetNode.cpp
@@ -3,6 +3,7 @@
 #include "Graph/Nodes/MovieGraphApplyCVarPresetNode.h"
 
 #include "Graph/MovieGraphConfig.h"
+#include "Sections/MovieSceneConsoleVariableTrackInterface.h"
 #include "Styling/AppStyle.h"
 
 FString UMovieGraphApplyCVarPresetNode::GetNodeInstanceName() const
@@ -23,13 +24,17 @@ EMovieGraphBranchRestriction UMovieGraphApplyCVarPresetNode::GetBranchRestrictio
 #if WITH_EDITOR
 FText UMovieGraphApplyCVarPresetNode::GetNodeTitle(const bool bGetDescriptive) const
 {
-	static const FText ApplyCVarPresetNodeName = NSLOCTEXT("MovieGraphNodes", "NodeName_ApplyCVarPreset", "Apply CVar Preset");
-	static const FText ApplyCVarPresetNodeDescription = NSLOCTEXT("MovieGraphNodes", "NodeDescription_ApplyCVarPreset", "Apply CVar Preset\n{0}");
+	static const FText ApplyCVarPresetNodeName = NSLOCTEXT("MovieGraphNodes", "NodeName_ApplyCVarPreset", "Apply Console Variable Preset");
+	static const FText ApplyCVarPresetNodeDescription = NSLOCTEXT("MovieGraphNodes", "NodeDescription_ApplyCVarPreset", "Apply Console Variable Preset\n{0} ({1} CVars)");
 
 	const FString InstanceName = GetNodeInstanceName();
 	if (bGetDescriptive && !InstanceName.IsEmpty())
 	{
-		return FText::Format(ApplyCVarPresetNodeDescription, FText::FromString(InstanceName));
+		constexpr bool bOnlyIncludeChecked = true;
+		TArray<TTuple<FString, FString>> ConsoleVariables;
+		ConsoleVariablePreset->GetConsoleVariablesForTrack(bOnlyIncludeChecked, ConsoleVariables);
+		
+		return FText::Format(ApplyCVarPresetNodeDescription, FText::FromString(InstanceName), FText::FromString(LexToString(ConsoleVariables.Num())));
 	}
 	
 	return ApplyCVarPresetNodeName;
diff --git a/Engine/Plugins/MovieScene/MovieRenderPipeline/Source/MovieRenderPipelineCore/Private/Graph/Nodes/MovieGraphSetCVarValueNode.cpp b/Engine/Plugins/MovieScene/MovieRenderPipeline/Source/MovieRenderPipelineCore/Private/Graph/Nodes/MovieGraphSetCVarValueNode.cpp
--- a/Engine/Plugins/MovieScene/MovieRenderPipeline/Source/MovieRenderPipelineCore/Private/Graph/Nodes/MovieGraphSetCVarValueNode.cpp
+++ b/Engine/Plugins/MovieScene/MovieRenderPipeline/Source/MovieRenderPipelineCore/Private/Graph/Nodes/MovieGraphSetCVarValueNode.cpp
@@ -3,6 +3,8 @@
 #include "Graph/Nodes/MovieGraphSetCVarValueNode.h"
 
 #include "Graph/MovieGraphConfig.h"
+#include "Internationalization/FastDecimalFormat.h"
+#include "Math/BasicMathExpressionEvaluator.h"
 #include "Styling/AppStyle.h"
 
 FString UMovieGraphSetCVarValueNode::GetNodeInstanceName() const
@@ -18,12 +20,19 @@ EMovieGraphBranchRestriction UMovieGraphSetCVarValueNode::GetBranchRestriction()
 #if WITH_EDITOR
 FText UMovieGraphSetCVarValueNode::GetNodeTitle(const bool bGetDescriptive) const
 {
-	static const FText SetCVarNodeName = NSLOCTEXT("MovieGraphNodes", "NodeName_SetCVar", "Set CVar Value");
-	static const FText SetCVarNodeDescription = NSLOCTEXT("MovieGraphNodes", "NodeDescription_SetCVar", "Set CVar Value\n{0}");
+	static const FText SetCVarNodeName = NSLOCTEXT("MovieGraphNodes", "NodeName_SetCVar", "Set Console Variable");
+	static const FText SetCVarNodeDescription = NSLOCTEXT("MovieGraphNodes", "NodeDescription_SetCVar", "Set Console Variable\n{0} ({1})");
 
 	if (bGetDescriptive && !Name.IsEmpty())
 	{
-		return FText::Format(SetCVarNodeDescription, FText::FromString(Name));
+		// Format the float value so it looks nice for display purposes
+		const FNumberFormattingOptions NumberFormattingOptions = FNumberFormattingOptions()
+			.SetUseGrouping(false)
+			.SetMinimumFractionalDigits(0)
+			.SetMaximumFractionalDigits(4);
+		const FString CVarValue = FastDecimalFormat::NumberToString(Value, ExpressionParser::GetLocalizedNumberFormattingRules(), NumberFormattingOptions);
+		
+		return FText::Format(SetCVarNodeDescription, FText::FromString(Name), FText::FromString(CVarValue));
 	}
 	
 	return SetCVarNodeName;
@@ -58,7 +67,14 @@ void UMovieGraphSetCVarValueNode::PostEditChangeProperty(FPropertyChangedEvent&
 {
 	Super::PostEditChangeProperty(PropertyChangedEvent);
 	
-	if (PropertyChangedEvent.GetPropertyName() == GET_MEMBER_NAME_CHECKED(UMovieGraphSetCVarValueNode, Name))
+	// Skip rapid updates from properties; only refresh on commit
+	if (PropertyChangedEvent.ChangeType == EPropertyChangeType::Interactive)
+	{
+		return;
+	}
+	
+	if (PropertyChangedEvent.GetPropertyName() == GET_MEMBER_NAME_CHECKED(UMovieGraphSetCVarValueNode, Name) ||
+		PropertyChangedEvent.GetPropertyName() == GET_MEMBER_NAME_CHECKED(UMovieGraphSetCVarValueNode, Value))
 	{
 		OnNodeChangedDelegate.Broadcast(this);
 	}
diff --git a/Engine/Plugins/MovieScene/MovieRenderPipeline/Source/MovieRenderPipelineEditor/Private/Customizations/ConsoleVariableSettingCustomization.h b/Engine/Plugins/MovieScene/MovieRenderPipeline/Source/MovieRenderPipelineEditor/Private/Customizations/ConsoleVariableSettingCustomization.h
--- a/Engine/Plugins/MovieScene/MovieRenderPipeline/Source/MovieRenderPipelineEditor/Private/Customizations/ConsoleVariableSettingCustomization.h
+++ b/Engine/Plugins/MovieScene/MovieRenderPipeline/Source/MovieRenderPipelineEditor/Private/Customizations/ConsoleVariableSettingCustomization.h
@@ -17,6 +17,9 @@
 
 #define LOCTEXT_NAMESPACE "SMoviePipelineEditor"
 
+/** Delegate that gets the menu widget associated with a specific cvar. */
+DECLARE_DELEGATE_RetVal_OneParam(TSharedRef<SWidget>, FGetCVarMenu, const FString&);
+
 /** Customize how properties in UMoviePipelineConsoleVariableSetting appear in the details panel. */
 class FConsoleVariablesSettingDetailsCustomization : public IDetailCustomization
 {
@@ -26,6 +29,63 @@
 		return MakeShared<FConsoleVariablesSettingDetailsCustomization>();
 	}
 
+	static void AddConsoleVariablePresetRowsToGroup(const TScriptInterface<IMovieSceneConsoleVariableTrackInterface>& InCVarPreset, IDetailGroup& InGroup, IDetailLayoutBuilder* InDetailBuilder, const TAttribute<bool>& InRowEnabledAttr, const FGetCVarMenu& InGetCVarMenuDelegate)
+	{
+		static const FText ConsoleVariableDisabledText = LOCTEXT("DisabledConsoleVariable", "This console variable is disabled.");
+		
+		constexpr bool bOnlyIncludeChecked = false;
+		TArray<TTuple<FString, FString>> CVars;
+		InCVarPreset->GetConsoleVariablesForTrack(bOnlyIncludeChecked, CVars);
+
+		// Show every console variable that's included in this preset; each console variable gets its own row
+		for (const TTuple<FString, FString>& CVar : CVars)
+		{
+			InGroup.AddWidgetRow()
+			.IsEnabled(InRowEnabledAttr)
+			.WholeRowContent()
+			[
+				SNew(SHorizontalBox)
+        		
+				+ SHorizontalBox::Slot()
+				.Padding(5, 0)
+				.FillWidth(0.75f)
+				.VAlign(VAlign_Center)
+				[
+					SNew(STextBlock)
+					.IsEnabled(InCVarPreset->IsConsoleVariableEnabled(CVar.Key))
+					.Font(InDetailBuilder->GetDetailFont())
+					.Text(FText::FromString(CVar.Key))
+					.ToolTipText(!InCVarPreset->IsConsoleVariableEnabled(CVar.Key) ? ConsoleVariableDisabledText : FText::GetEmpty())
+				]
+        		
+				+ SHorizontalBox::Slot()
+				.Padding(0, 0)
+				.FillWidth(0.25f)
+				.VAlign(VAlign_Center)
+				[
+					SNew(SEditableTextBox)
+					.IsEnabled(false)
+					.Font(InDetailBuilder->GetDetailFont())
+					.Text(FText::FromString(CVar.Value))
+				]
+        		
+				+ SHorizontalBox::Slot()
+				.Padding(5, 0)
+				.AutoWidth()
+				[
+					SNew(SComboButton)
+					.ComboButtonStyle(FAppStyle::Get(), "SimpleComboButton")
+					.ForegroundColor(FSlateColor::UseForeground())
+					.HasDownArrow(true)
+					.MenuContent()
+					[
+						InGetCVarMenuDelegate.Execute(CVar.Key)
+					]
+				]
+			];
+		}
+	}
+
 protected:
 	//~ Begin IDetailCustomization interface
 	virtual void CustomizeDetails(IDetailLayoutBuilder& DetailBuilder) override
@@ -97,66 +157,21 @@
 			return;
 		}
 
-		const bool bOnlyIncludeChecked = false;
-		TArray<TTuple<FString, FString>> CVars;
-		CVarPreset->GetConsoleVariablesForTrack(bOnlyIncludeChecked, CVars);
-
-		// Show every console variable that's included in this preset; each console variable gets its own row
-		for (const TTuple<FString, FString>& CVar : CVars)
+		constexpr bool bRowEnabledAttr = true;
+		AddConsoleVariablePresetRowsToGroup(CVarPreset, Group, DetailLayout, bRowEnabledAttr, FGetCVarMenu::CreateLambda([CVarSetting](const FString& InCVarName)
 		{
 			// Each console variable gets a menu that allows the user to create an override outside of the preset
 			FMenuBuilder CreateOverrideMenu(true, nullptr, nullptr, true);
-			FUIAction AddOverrideAction(
-				FExecuteAction::CreateLambda([CVar, CVarSetting]()
+			const FUIAction AddOverrideAction(
+				FExecuteAction::CreateLambda([InCVarName, CVarSetting]()
 				{
-					CVarSetting->AddConsoleVariable(CVar.Key, 0.f);
+					CVarSetting->AddConsoleVariable(InCVarName, 0.f);
 				})
 			);
 			CreateOverrideMenu.AddMenuEntry(ConsoleVariableOverrideText, FText::GetEmpty(), FSlateIcon(), AddOverrideAction);
-			
-			Group.AddWidgetRow()
-			.WholeRowContent()
-			[
-				SNew(SHorizontalBox)
-				
-				+ SHorizontalBox::Slot()
-				.Padding(5, 0)
-				.FillWidth(0.75f)
-				.VAlign(VAlign_Center)
-				[
-					SNew(STextBlock)
-					.IsEnabled(CVarPreset->IsConsoleVariableEnabled(CVar.Key))
-					.Font(DetailLayout->GetDetailFont())
-					.Text(FText::FromString(CVar.Key))
-					.ToolTipText(!CVarPreset->IsConsoleVariableEnabled(CVar.Key) ? ConsoleVariableDisabledText : FText::GetEmpty())
-				]
-				
-				+ SHorizontalBox::Slot()
-				.Padding(0, 0)
-				.FillWidth(0.25f)
-				.VAlign(VAlign_Center)
-				[
-					SNew(SEditableTextBox)
-					.IsEnabled(false)
-					.Font(DetailLayout->GetDetailFont())
-					.Text(FText::FromString(CVar.Value))
-				]
-				
-				+ SHorizontalBox::Slot()
-				.Padding(5, 0)
-				.AutoWidth()
-				[
-					SNew(SComboButton)
-					.ComboButtonStyle(FAppStyle::Get(), "SimpleComboButton")
-					.ForegroundColor(FSlateColor::UseForeground())
-					.HasDownArrow(true)
-					.MenuContent()
-					[
-						CreateOverrideMenu.MakeWidget()
-					]
-				]
-			];
-		}
+
+			return CreateOverrideMenu.MakeWidget();
+		}));
 	}
 };
 
diff --git a/Engine/Plugins/MovieScene/MovieRenderPipeline/Source/MovieRenderPipelineEditor/Private/Customizations/Graph/MovieGraphApplyCVarPresetNodeCustomization.h b/Engine/Plugins/MovieScene/MovieRenderPipeline/Source/MovieRenderPipelineEditor/Private/Customizations/Graph/MovieGraphApplyCVarPresetNodeCustomization.h
new file mode 100644
--- /dev/null
+++ b/Engine/Plugins/MovieScene/MovieRenderPipeline/Source/MovieRenderPipelineEditor/Private/Customizations/Graph/MovieGraphApplyCVarPresetNodeCustomization.h
@@ -0,0 +1,178 @@
+// Copyright Epic Games, Inc. All Rights Reserved.
+
+#pragma once
+
+#include "Customizations/ConsoleVariableSettingCustomization.h"
+#include "DetailLayoutBuilder.h"
+#include "DetailWidgetRow.h"
+#include "Graph/MovieEdGraphNode.h"
+#include "Graph/MovieGraphSchema.h"
+#include "Graph/Nodes/MovieGraphApplyCVarPresetNode.h"
+#include "Graph/Nodes/MovieGraphSetCVarValueNode.h"
+#include "IDetailCustomization.h"
+#include "IDetailGroup.h"
+#include "PropertyHandle.h"
+#include "Sections/MovieSceneConsoleVariableTrackInterface.h"
+#include "ScopedTransaction.h"
+#include "Widgets/Input/SEditableTextBox.h"
+#include "Widgets/Text/STextBlock.h"
+
+#define LOCTEXT_NAMESPACE "MoviePipelineEditor"
+
+/** Customize how the Apply Console Variable Preset node appears in the details panel. */
+class FMovieGraphApplyCVarPresetNodeCustomization : public IDetailCustomization
+{
+public:
+	static TSharedRef<IDetailCustomization> MakeInstance()
+	{
+		return MakeShared<FMovieGraphApplyCVarPresetNodeCustomization>();
+	}
+
+protected:
+	/** Creates a new Set CVar Value node (overriding the cvar with the provided name) connected to the given preset node. */
+	static void CreateOverrideViaSetCVarValueNode(const TWeakObjectPtr<UMovieGraphApplyCVarPresetNode> InPresetNode, const FString& InCVarName)
+	{
+		// Putting logic like this into a customization is a little awkward. However, it seems like it's still an OK choice after considering the
+		// alternatives. Putting the logic in a new ed graph node would require instantiating a new ed node class JUST to get this behavior. Doesn't
+		// belong in the schema because it's too use-specific. Also doesn't belong on the non-ed graph node because Core cannot access Editor.
+		
+		if (!InPresetNode.IsValid())
+		{
+			return;
+		}
+		
+		// When creating the new action, since it's only being used to create a node, the category, display name, and tooltip can just be empty
+		constexpr int32 Grouping = 0;
+		const FText Keywords = InPresetNode->GetKeywords();
+		const TSharedPtr<FMovieGraphSchemaAction_NewNode> NewNodeAction =
+			MakeShared<FMovieGraphSchemaAction_NewNode>(FText::GetEmpty(), FText::GetEmpty(), FText::GetEmpty(), Grouping, Keywords);
+		NewNodeAction->NodeClass = UMovieGraphSetCVarValueNode::StaticClass();
+
+		// Put the new node to the right of this node
+		const TObjectPtr<UEdGraphNode> GraphNode = InPresetNode->GraphNode;
+		const FVector2d NewLocation(GraphNode->NodePosX + 300, GraphNode->NodePosY);
+
+		// Create the new node (providing FromPin will trigger the action to connect the new node and the preset node)
+		UEdGraphPin* FromPin = GraphNode->FindPin(NAME_None, EGPD_Output);
+		UEdGraphNode* NewNode = NewNodeAction->PerformAction(GraphNode->GetGraph(), FromPin, NewLocation);
+
+		// Set the node to use the correct cvar
+		UMovieGraphNode* RuntimeNode = Cast<UMoviePipelineEdGraphNode>(NewNode)->GetRuntimeNode();
+		UMovieGraphSetCVarValueNode* SetCVarNode = Cast<UMovieGraphSetCVarValueNode>(RuntimeNode);
+		SetCVarNode->Name = InCVarName;
+
+		// Populate the default value of the override from the cvar in the preset
+		if (const TScriptInterface<IMovieSceneConsoleVariableTrackInterface>& PresetAsset = InPresetNode->ConsoleVariablePreset)
+		{
+			constexpr bool bOnlyIncludeChecked = false;
+			TArray<TTuple<FString, FString>> ConsoleVariables;
+			PresetAsset->GetConsoleVariablesForTrack(bOnlyIncludeChecked, ConsoleVariables);
+
+			const TTuple<FString, FString>* CVarPair =
+				ConsoleVariables.FindByPredicate([&InCVarName](const TTuple<FString, FString>& InPair) { return InPair.Key == InCVarName; });
+			if (CVarPair)
+			{
+				LexFromString(SetCVarNode->Value, *CVarPair->Value);
+			}
+		}
+
+		// Mark the properties as overridden so they are used in evaluation
+		SetCVarNode->bOverride_Name = true;
+		SetCVarNode->bOverride_Value = true;
+	}
+
+	/** Creates the override menu for a specific cvar row. */
+	static TSharedRef<SWidget> CreateCVarOverrideMenu(const FString& InCVarName, TWeakObjectPtr<UMovieGraphApplyCVarPresetNode> InPresetNode)
+	{
+		static const FText OverrideWithNewNodeText = LOCTEXT("OverrideWithSetConsoleVariableNode", "Override with Set Console Variable Node");
+		
+		FMenuBuilder CVarMenu(true, nullptr, nullptr, true);
+		const FUIAction AddOverrideAction(
+			FExecuteAction::CreateLambda([InPresetNode, InCVarName]()
+			{
+				if (InPresetNode.IsValid())
+				{
+					CreateOverrideViaSetCVarValueNode(InPresetNode, InCVarName);
+				}
+			})
+		);
+		CVarMenu.AddMenuEntry(OverrideWithNewNodeText, FText::GetEmpty(), FSlateIcon(), AddOverrideAction);
+
+		return CVarMenu.MakeWidget();
+	}
+
+	//~ Begin IDetailCustomization interface
+	virtual void CustomizeDetails(const TSharedPtr<IDetailLayoutBuilder>& DetailBuilder) override
+	{
+		CustomizeDetails(*DetailBuilder);
+	}
+	
+	virtual void CustomizeDetails(IDetailLayoutBuilder& DetailBuilder) override
+	{
+		for (const TWeakObjectPtr<UMovieGraphApplyCVarPresetNode>& ApplyPresetNode : DetailBuilder.GetObjectsOfTypeBeingCustomized<UMovieGraphApplyCVarPresetNode>())
+		{
+			IDetailCategoryBuilder& SettingsCategory = DetailBuilder.EditCategory(TEXT("Settings"));
+			const TSharedRef<IPropertyHandle> ConsoleVariablePresetHandle = DetailBuilder.GetProperty(GET_MEMBER_NAME_CHECKED(UMovieGraphApplyCVarPresetNode, ConsoleVariablePreset));
+
+			// Regenerate the preset details when the preset is updated
+			ConsoleVariablePresetHandle->SetOnPropertyValueChanged(FSimpleDelegate::CreateLambda([&DetailBuilder]()
+			{
+				DetailBuilder.ForceRefreshDetails();
+			}));
+
+			// Hide the default preset property since we'll be making a group that contains it instead (and the group's rows will be the individual
+			// console variables within the preset)
+			ConsoleVariablePresetHandle->MarkHiddenByCustomization();
+
+			const FStringView PresetPropName = ConsoleVariablePresetHandle->GetPropertyPath();
+			IDetailGroup& PresetGroup = SettingsCategory.AddGroup(FName(PresetPropName), FText::FromStringView(PresetPropName));
+
+			// Since we're making a group to replace the default property for the preset, the EditCondition checkbox won't come along for the ride for
+			// free. Supplying an attr which provides the EditCondition is required.
+			TAttribute<bool> PresetEnabledAttr = TAttribute<bool>::Create([ApplyPresetNode]()
+			{
+				return ApplyPresetNode.IsValid() ? ApplyPresetNode->bOverride_ConsoleVariablePreset > 0 : false;
+			});
+
+			// The group's header should look like the default preset property w/ the asset picker widget
+			PresetGroup.HeaderRow()
+			.PropertyHandleList({ConsoleVariablePresetHandle})
+			.ShouldAutoExpand(true)
+			.EditCondition(
+				PresetEnabledAttr,
+				FOnBooleanValueChanged::CreateLambda([ApplyPresetNode, &DetailBuilder](const bool bNewValue)
+				{
+					if (ApplyPresetNode.IsValid())
+					{
+						FScopedTransaction Transaction(LOCTEXT("PresetEditConditionChanged", "Edit Condition Changed"));
+						ApplyPresetNode->Modify();
+						ApplyPresetNode->bOverride_ConsoleVariablePreset = bNewValue;
+
+						// This is kinda hack-ish. The details panel doesn't fully refresh in some circumstances when the edit condition is toggled.
+						// However, when set the normal way, a refresh seems to be triggered anyway, so maybe there's no better way to accomplish this.
+						// Using the IPropertyHandle or FBoolProperty to change the value does not result in different behavior.
+						DetailBuilder.ForceRefreshDetails();
+					}
+				})
+			)
+			.NameContent()
+			[
+				ConsoleVariablePresetHandle->CreatePropertyNameWidget()
+			]
+			.ValueContent()
+			[
+				ConsoleVariablePresetHandle->CreatePropertyValueWidget()
+			];
+
+			// Add all of the cvars as rows under the preset group
+			if (TScriptInterface<IMovieSceneConsoleVariableTrackInterface> CVarPreset = ApplyPresetNode->ConsoleVariablePreset)
+			{
+				FConsoleVariablesSettingDetailsCustomization::AddConsoleVariablePresetRowsToGroup(CVarPreset, PresetGroup, &DetailBuilder, PresetEnabledAttr,
+					FGetCVarMenu::CreateStatic(&CreateCVarOverrideMenu, ApplyPresetNode));
+			}
+		}
+	}
+	//~ End IDetailCustomization interface
+};
+
+#undef LOCTEXT_NAMESPACE
\ No newline at end of file
diff --git a/Engine/Plugins/MovieScene/MovieRenderPipeline/Source/MovieRenderPipelineEditor/Private/Customizations/Graph/MovieGraphSetCVarValueNodeCustomization.h b/Engine/Plugins/MovieScene/MovieRenderPipeline/Source/MovieRenderPipelineEditor/Private/Customizations/Graph/MovieGraphSetCVarValueNodeCustomization.h
new file mode 100644
--- /dev/null
+++ b/Engine/Plugins/MovieScene/MovieRenderPipeline/Source/MovieRenderPipelineEditor/Private/Customizations/Graph/MovieGraphSetCVarValueNodeCustomization.h
@@ -0,0 +1,92 @@
+// Copyright Epic Games, Inc. All Rights Reserved.
+
+#pragma once
+
+#include "Customizations/ConsoleVariableCustomization.h"
+#include "DetailLayoutBuilder.h"
+#include "DetailWidgetRow.h"
+#include "Graph/Nodes/MovieGraphSetCVarValueNode.h"
+#include "IDetailCustomization.h"
+#include "PropertyHandle.h"
+#include "Widgets/Input/SEditableTextBox.h"
+#include "Widgets/Text/STextBlock.h"
+
+#define LOCTEXT_NAMESPACE "MoviePipelineEditor"
+
+/** Customize how the Set Console Variable Value node appears in the details panel. */
+class FMovieGraphSetCVarValueNodeCustomization : public IDetailCustomization
+{
+public:
+	static TSharedRef<IDetailCustomization> MakeInstance()
+	{
+		return MakeShared<FMovieGraphSetCVarValueNodeCustomization>();
+	}
+
+protected:	
+	FText GetConsoleVariableText(TSharedRef<IPropertyHandle> PropertyHandle) const
+	{
+		FString Value;
+		if (PropertyHandle->IsValidHandle() && (PropertyHandle->GetValue(Value) != FPropertyAccess::Fail))
+		{
+			return FText::FromString(Value);
+		}
+
+		return FText();
+	}
+
+	//~ Begin IDetailCustomization interface
+	virtual void CustomizeDetails(const TSharedPtr<IDetailLayoutBuilder>& DetailBuilder) override
+	{
+		CustomizeDetails(*DetailBuilder);
+	}
+	
+	virtual void CustomizeDetails(IDetailLayoutBuilder& DetailBuilder) override
+	{
+		for (TWeakObjectPtr<UMovieGraphSetCVarValueNode>& SetCVarValueNode : DetailBuilder.GetObjectsOfTypeBeingCustomized<UMovieGraphSetCVarValueNode>())
+		{
+			if (!SetCVarValueNode.IsValid())
+			{
+				continue;
+			}
+
+			TSharedRef<IPropertyHandle> CVarNameProperty = DetailBuilder.GetProperty(GET_MEMBER_NAME_CHECKED(UMovieGraphSetCVarValueNode, Name));
+
+			// Update the default widget for the cvar name to give it autocomplete
+			IDetailPropertyRow* CVarNameRow = DetailBuilder.EditDefaultProperty(CVarNameProperty);
+			CVarNameRow->CustomWidget(true)
+			.NameContent()
+			[
+				CVarNameProperty->CreatePropertyNameWidget()
+			]
+			.ValueContent()
+			[
+				SNew(SConsoleVariablesEditorCustomConsoleInputBox)
+				.Font(DetailBuilder.GetDetailFont())
+				.Text(this, &FMovieGraphSetCVarValueNodeCustomization::GetConsoleVariableText, CVarNameProperty)
+				.HideOnFocusLost(false)
+				.CommitOnFocusLost(true)
+				.ClearOnCommit(false)
+				.HintText(LOCTEXT("ConsoleVariableInputHintText", "Enter Console Variable"))
+				.OnTextChanged_Lambda([this, CVarNameProperty](const FText& NewText)
+				{
+					if (CVarNameProperty->IsValidHandle())
+					{
+						// Note that the change is marked as Interactive so the graph is not forced to refresh for every character typed
+						CVarNameProperty->SetValue(NewText.ToString(), EPropertyValueSetFlags::InteractiveChange);
+					}
+				})
+				.OnTextCommitted_Lambda([this, CVarNameProperty](const FText& NewText)
+				{
+					// Committing the text will trigger the graph to update
+					if (CVarNameProperty->IsValidHandle())
+					{
+						CVarNameProperty->SetValue(NewText.ToString());
+					}
+				})
+			];
+		}
+	}
+	//~ End IDetailCustomization interface
+};
+
+#undef LOCTEXT_NAMESPACE
\ No newline at end of file
diff --git a/Engine/Plugins/MovieScene/MovieRenderPipeline/Source/MovieRenderPipelineEditor/Private/Graph/MovieGraphAssetToolkit.cpp b/Engine/Plugins/MovieScene/MovieRenderPipeline/Source/MovieRenderPipelineEditor/Private/Graph/MovieGraphAssetToolkit.cpp
--- a/Engine/Plugins/MovieScene/MovieRenderPipeline/Source/MovieRenderPipelineEditor/Private/Graph/MovieGraphAssetToolkit.cpp
+++ b/Engine/Plugins/MovieScene/MovieRenderPipeline/Source/MovieRenderPipelineEditor/Private/Graph/MovieGraphAssetToolkit.cpp
@@ -2,6 +2,7 @@
 
 #include "MovieGraphAssetToolkit.h"
 
+#include "Customizations/Graph/MovieGraphApplyCVarPresetNodeCustomization.h"
 #include "Customizations/Graph/MovieGraphCollectionsCustomization.h"
 #include "Customizations/Graph/MovieGraphFormatTokenCustomization.h"
 #include "Customizations/Graph/MovieGraphMemberCustomization.h"
@@ -9,16 +10,20 @@
 #include "Customizations/Graph/MovieGraphNamedResolutionCustomization.h"
 #include "Customizations/Graph/MovieGraphNodeCustomization.h"
 #include "Customizations/Graph/MovieGraphSelectNodeCustomization.h"
+#include "Customizations/Graph/MovieGraphSetCVarValueNodeCustomization.h"
 #include "Customizations/Graph/MovieGraphShowFlagsCustomization.h"
 #include "Customizations/Graph/MovieGraphVersioningSettingsCustomization.h"
 #include "Graph/Renderers/MovieGraphShowFlags.h"
 
 #include "Graph/MovieGraphConfig.h"
+#include "Graph/Nodes/MovieGraphApplyCVarPresetNode.h"
 #include "Graph/Nodes/MovieGraphCollectionNode.h"
 #include "Graph/Nodes/MovieGraphCommandLineEncoderNode.h"
+#include "Graph/Nodes/MovieGraphDebugNode.h"
 #include "Graph/Nodes/MovieGraphFileOutputNode.h"
 #include "Graph/Nodes/MovieGraphModifierNode.h"
 #include "Graph/Nodes/MovieGraphSelectNode.h"
+#include "Graph/Nodes/MovieGraphSetCVarValueNode.h"
 #include "MovieEdGraphNode.h"
 #include "MovieGraphSchema.h"
 #include "MovieRenderPipelineSettings.h"
@@ -30,7 +35,6 @@
 #include "PropertyEditorModule.h"
 #include "Selection.h"
 #include "ToolMenus.h"
-#include "Graph/Nodes/MovieGraphDebugNode.h"
 #include "Widgets/Docking/SDockTab.h"
 #include "Widgets/Graph/SMovieGraphConfigPanel.h"
 
@@ -376,6 +380,14 @@ TSharedRef<SDockTab> FMovieGraphAssetToolkit::SpawnTab_RenderGraphDetails(const
 	SelectedGraphObjectsDetailsWidget->RegisterInstancedCustomPropertyLayout(
 		UMovieGraphSelectNode::StaticClass(),
 		FOnGetDetailCustomizationInstance::CreateStatic(&FMovieGraphSelectNodeCustomization::MakeInstance));
+	
+	SelectedGraphObjectsDetailsWidget->RegisterInstancedCustomPropertyLayout(
+		UMovieGraphSetCVarValueNode::StaticClass(),
+		FOnGetDetailCustomizationInstance::CreateStatic(&FMovieGraphSetCVarValueNodeCustomization::MakeInstance));
+
+	SelectedGraphObjectsDetailsWidget->RegisterInstancedCustomPropertyLayout(
+		UMovieGraphApplyCVarPresetNode::StaticClass(),
+		FOnGetDetailCustomizationInstance::CreateStatic(&FMovieGraphApplyCVarPresetNodeCustomization::MakeInstance));
 
 	SelectedGraphObjectsDetailsWidget->RegisterInstancedCustomPropertyTypeLayout(
 		UMovieGraphShowFlags::StaticClass()->GetFName(),
