MoviePipeline: Backed out change that updated the queue to hide unconnected graph variables.

Currently, variables are the only way to pass data to pre/post scripts in the graph. Hiding unconnected variables prevents these variables from having their values changed in the UI if they're only used within scripts.

#rnx

diff --git a/Engine/Plugins/MovieScene/MovieRenderPipeline/Source/MovieRenderPipelineCore/Private/MovieJobVariableAssignmentContainer.cpp b/Engine/Plugins/MovieScene/MovieRenderPipeline/Source/MovieRenderPipelineCore/Private/MovieJobVariableAssignmentContainer.cpp
--- a/Engine/Plugins/MovieScene/MovieRenderPipeline/Source/MovieRenderPipelineCore/Private/MovieJobVariableAssignmentContainer.cpp
+++ b/Engine/Plugins/MovieScene/MovieRenderPipeline/Source/MovieRenderPipelineCore/Private/MovieJobVariableAssignmentContainer.cpp
@@ -4,7 +4,6 @@
 
 #include "Graph/MovieGraphCommon.h"
 #include "Graph/MovieGraphConfig.h"
-#include "Graph/Nodes/MovieGraphVariableNode.h"
 
 // The prefix for property descriptions that are in charge of variable assignment EditCondition enable/disable state
 static const TCHAR* EditConditionPrefix = TEXT("bOverride_");
@@ -307,7 +306,7 @@ bool UMovieJobVariableAssignmentContainer::GenerateVariableOverride(const UMovie
 }
 
 #if WITH_EDITOR
-void UMovieJobVariableAssignmentContainer::UpdateGraphVariableOverrides(const bool bOnlyIncludeConnectedVariables/*= false*/)
+void UMovieJobVariableAssignmentContainer::UpdateGraphVariableOverrides()
 {
 	auto GetVariableGuidFromDesc = [](const FPropertyBagPropertyDesc& Desc) -> FGuid
 	{
@@ -446,26 +445,6 @@ void UMovieJobVariableAssignmentContainer::UpdateGraphVariableOverrides(const bo
 		}
 	}
 
-	// Optionally, remove any descs that don't correspond to a variable node with an outgoing connection whose Guid matches the desc's Guid
-	if (bOnlyIncludeConnectedVariables)
-	{
-		const int32 NumDescsRemoved = ModifiedDescs.RemoveAll([&GetVariableGuidFromDesc, this](const FPropertyBagPropertyDesc& Desc)
-		{
-			const FGuid DescGuid = GetVariableGuidFromDesc(Desc);
-			return !GraphPreset->GetNodes().ContainsByPredicate([&DescGuid](UMovieGraphNode* Node)
-				{
-					const UMovieGraphVariableNode* VariableNode = Cast<UMovieGraphVariableNode>(Node);
-					return VariableNode && VariableNode->GetFirstConnectedOutputPin() && 
-						VariableNode->GetVariable() && (VariableNode->GetVariable()->GetGuid() == DescGuid);
-				});
-		});
-
-		if (NumDescsRemoved > 0) 
-		{
-			bNeedsToRegenerate = true;
-		}
-	}
-
 	if (bNeedsToRegenerate)
 	{
 		const UPropertyBag* NewBagStruct = UPropertyBag::GetOrCreateFromDescs(ModifiedDescs);
diff --git a/Engine/Plugins/MovieScene/MovieRenderPipeline/Source/MovieRenderPipelineCore/Private/MoviePipelineQueue.cpp b/Engine/Plugins/MovieScene/MovieRenderPipeline/Source/MovieRenderPipelineCore/Private/MoviePipelineQueue.cpp
--- a/Engine/Plugins/MovieScene/MovieRenderPipeline/Source/MovieRenderPipelineCore/Private/MoviePipelineQueue.cpp
+++ b/Engine/Plugins/MovieScene/MovieRenderPipeline/Source/MovieRenderPipelineCore/Private/MoviePipelineQueue.cpp
@@ -19,15 +19,13 @@ namespace UE::MovieGraph::Private
 	TObjectPtr<UMovieJobVariableAssignmentContainer> GetOrCreateJobVariableAssignmentsForGraph(
 		const UMovieGraphConfig* InGraph, TArray<TObjectPtr<UMovieJobVariableAssignmentContainer>>& InVariableAssignments, UObject* InAssignmentsOwner)
 	{
-		constexpr bool bOnlyIncludeConnectedVariables = true;
-		
 		for (TObjectPtr<UMovieJobVariableAssignmentContainer>& VariableAssignment : InVariableAssignments)
 		{
 			const TSoftObjectPtr<UMovieGraphConfig> SoftGraphConfig = VariableAssignment->GetGraphConfig();
 			if (SoftGraphConfig.Get() == InGraph)
 			{
 #if WITH_EDITOR
-				VariableAssignment->UpdateGraphVariableOverrides(bOnlyIncludeConnectedVariables);
+				VariableAssignment->UpdateGraphVariableOverrides();
 #endif
 				
 				return VariableAssignment;
@@ -40,7 +38,7 @@ namespace UE::MovieGraph::Private
 		NewVariableAssignments->SetGraphConfig(InGraph);
 
 #if WITH_EDITOR
-		NewVariableAssignments->UpdateGraphVariableOverrides(bOnlyIncludeConnectedVariables);
+		NewVariableAssignments->UpdateGraphVariableOverrides();
 #endif
 	
 		return NewVariableAssignments;
diff --git a/Engine/Plugins/MovieScene/MovieRenderPipeline/Source/MovieRenderPipelineCore/Public/MovieJobVariableAssignmentContainer.h b/Engine/Plugins/MovieScene/MovieRenderPipeline/Source/MovieRenderPipelineCore/Public/MovieJobVariableAssignmentContainer.h
--- a/Engine/Plugins/MovieScene/MovieRenderPipeline/Source/MovieRenderPipelineCore/Public/MovieJobVariableAssignmentContainer.h
+++ b/Engine/Plugins/MovieScene/MovieRenderPipeline/Source/MovieRenderPipelineCore/Public/MovieJobVariableAssignmentContainer.h
@@ -230,10 +230,9 @@ public:
 	 * the graph variable name, value type, object type, and container type. Additionally, stale overrides that have no
 	 * corresponding graph variable will be removed, and overrides will be created for graph variables which do not have
 	 * existing overrides.
-	 * @param bOnlyIncludeConnectedVariables Optionally, if the stored variable overrides should only include variables that 1) have nodes referencing them in the graph, and 2) at least one of those variable nodes has a connection to another node, set bOnlyIncludeConnectedVariables to true. This will ensure that variable overrides that have no chance of being utilized in the graph are not available to be changed.
 	 */
 	UFUNCTION(BlueprintCallable, Category = "Experimental")
-	void UpdateGraphVariableOverrides(const bool bOnlyIncludeConnectedVariables = false);
+	void UpdateGraphVariableOverrides();
 #endif
 
 protected:
