MoviePipeline: Did a pass over collections and modifiers in the graph to add support for undo/redo operations.

Fixed undo/redo for the following:
* Add/remove collections from modifiers.
* Update modifier properties.
* Add/remove condition groups.
* Add/remove condition group queries.
* Change condition group query type.
* Update condition group query properties.

#rnx

diff --git a/Engine/Plugins/MovieScene/MovieRenderPipeline/Source/MovieRenderPipelineCore/Private/Graph/MovieGraphRenderLayerSubsystem.cpp b/Engine/Plugins/MovieScene/MovieRenderPipeline/Source/MovieRenderPipelineCore/Private/Graph/MovieGraphRenderLayerSubsystem.cpp
--- a/Engine/Plugins/MovieScene/MovieRenderPipeline/Source/MovieRenderPipelineCore/Private/Graph/MovieGraphRenderLayerSubsystem.cpp
+++ b/Engine/Plugins/MovieScene/MovieRenderPipeline/Source/MovieRenderPipelineCore/Private/Graph/MovieGraphRenderLayerSubsystem.cpp
@@ -35,6 +35,7 @@
 #include "SceneOutlinerModule.h"
 #include "SceneOutlinerPublicTypes.h"
 #include "SClassViewer.h"
+#include "ScopedTransaction.h"
 #include "Selection.h"
 #endif
 
@@ -588,6 +589,9 @@ TArray<TSharedRef<SWidget>> UMovieGraphConditionGroupQuery_Actor::GetWidgets()
 			.OnGetRowIcon_Static(&GetRowIcon)
 			.OnDelete_Lambda([this](const TSharedPtr<TSoftObjectPtr<AActor>> InActor)
 			{
+				const FScopedTransaction Transaction(LOCTEXT("RemoveActorsFromCollection", "Remove Actors from Collection"));
+				Modify();
+				
 				ListDataSource.Remove(InActor);
 				ActorsToMatch.Remove(*InActor.Get());
 				ActorsList->Refresh();
@@ -631,6 +635,9 @@ TSharedRef<SWidget> UMovieGraphConditionGroupQuery_Actor::GetAddMenuContents(con
 
 	auto AddActorToList = [this](const AActor* InActor, const FMovieGraphConditionGroupQueryContentsChanged& OnAddFinished)
 	{
+		const FScopedTransaction Transaction(LOCTEXT("AddActorsToCollection", "Add Actors to Collection"));
+		Modify();
+		
 		ActorsToMatch.Add(InActor);
 		ListDataSource.Add(MakeShared<TSoftObjectPtr<AActor>>(ActorsToMatch.Last()));
 		OnAddFinished.ExecuteIfBound();
@@ -655,6 +662,8 @@ TSharedRef<SWidget> UMovieGraphConditionGroupQuery_Actor::GetAddMenuContents(con
 			FUIAction(
 				FExecuteAction::CreateLambda([this, OnAddFinished, RefreshActorPickerFilterAndList, AddActorToList]()
 				{
+					const FScopedTransaction Transaction(LOCTEXT("AddSelectedActorsToCollection", "Add Selected Actors to Collection"));
+					
 					USelection* CurrentSelection = GEditor->GetSelectedActors();
 					TArray<AActor*> SelectedActors;
 					CurrentSelection->GetSelectedObjects<AActor>(SelectedActors);
@@ -801,7 +810,13 @@ TArray<TSharedRef<SWidget>> UMovieGraphConditionGroupQuery_ActorTagName::GetWidg
 		[
 			SNew(SMultiLineEditableTextBox)
 			.Text_Lambda([this]() { return FText::FromString(TagsToMatch); })
-			.OnTextChanged_Lambda([this](const FText& InText) { TagsToMatch = InText.ToString(); })
+			.OnTextCommitted_Lambda([this](const FText& InText, ETextCommit::Type TextCommitType)
+			{
+				const FScopedTransaction Transaction(LOCTEXT("UpdateActorTagNamesInCollection", "Update Actor Tag Names in Collection"));
+				Modify();
+				
+				TagsToMatch = InText.ToString();
+			})
 			.HintText(LOCTEXT("MovieGraphActorTagNameQueryHintText", "The actor must match one or more tags. Wildcards allowed.\nEnter each tag on a separate line."))
 		]
 	);
@@ -863,7 +878,13 @@ TArray<TSharedRef<SWidget>> UMovieGraphConditionGroupQuery_ActorName::GetWidgets
 		[
 			SNew(SMultiLineEditableTextBox)
 			.Text_Lambda([this]() { return FText::FromString(WildcardSearch); })
-			.OnTextChanged_Lambda([this](const FText& InText) { WildcardSearch = InText.ToString(); })
+			.OnTextCommitted_Lambda([this](const FText& InText, ETextCommit::Type TextCommitType)
+			{
+				const FScopedTransaction Transaction(LOCTEXT("UpdateActorNamesInCollection", "Update Actor Names in Collection"));
+				Modify();
+
+				WildcardSearch = InText.ToString();
+			})
 			.HintText(LOCTEXT("MovieGraphActorNameQueryHintText", "Actor names to query. Wildcards allowed.\nEnter each actor name on a separate line."))
 		]
 	);
@@ -917,6 +938,9 @@ TArray<TSharedRef<SWidget>> UMovieGraphConditionGroupQuery_ActorType::GetWidgets
 			.OnGetRowIcon_Static(&GetRowIcon)
 			.OnDelete_Lambda([this](UClass* InActorClass)
 			{
+				const FScopedTransaction Transaction(LOCTEXT("RemoveActorTypesFromCollection", "Remove Actor Types from Collection"));
+				Modify();
+				
 				ActorTypes.Remove(InActorClass);
 				ActorTypesList->Refresh();
 			})
@@ -948,6 +972,9 @@ TSharedRef<SWidget> UMovieGraphConditionGroupQuery_ActorType::GetAddMenuContents
 		Options,
 		FOnClassPicked::CreateLambda([this, OnAddFinished](UClass* InNewClass)
 		{
+			const FScopedTransaction Transaction(LOCTEXT("AddActorTypesToCollection", "Add Actor Types to Collection"));
+			Modify();
+			
 			FSlateApplication::Get().DismissAllMenus();
 			
 			ActorTypes.Add(InNewClass);
@@ -1069,7 +1096,13 @@ TArray<TSharedRef<SWidget>> UMovieGraphConditionGroupQuery_ComponentTagName::Get
 		[
 			SNew(SMultiLineEditableTextBox)
 			.Text_Lambda([this]() { return FText::FromString(TagsToMatch); })
-			.OnTextChanged_Lambda([this](const FText& InText) { TagsToMatch = InText.ToString(); })
+			.OnTextCommitted_Lambda([this](const FText& InText, ETextCommit::Type TextCommitType)
+			{
+				const FScopedTransaction Transaction(LOCTEXT("UpdateComponentTagNamesInCollection", "Update Component Tag Names in Collection"));
+				Modify();
+				
+				TagsToMatch = InText.ToString();
+			})
 			.HintText(LOCTEXT("MovieGraphComponentTagNameQueryHintText", "A component on the actor must match one or more component tags.\nWildcards allowed. Enter each tag on a separate line."))
 		]
 	);
@@ -1128,6 +1161,9 @@ TArray<TSharedRef<SWidget>> UMovieGraphConditionGroupQuery_ComponentType::GetWid
 			.OnGetRowIcon_Static(&GetRowIcon)
 			.OnDelete_Lambda([this](UClass* InComponentType)
 			{
+				const FScopedTransaction Transaction(LOCTEXT("RemoveComponentTypesFromCollection", "Remove Component Types from Collection"));
+				Modify();
+				
 				ComponentTypes.Remove(InComponentType);
 				ComponentTypesList->Refresh();
 			})			
@@ -1160,6 +1196,9 @@ TSharedRef<SWidget> UMovieGraphConditionGroupQuery_ComponentType::GetAddMenuCont
 		Options,
 		FOnClassPicked::CreateLambda([this, OnAddFinished](UClass* InNewClass)
 		{
+			const FScopedTransaction Transaction(LOCTEXT("AddComponentTypeToCollection", "Add Component Type to Collection"));
+			Modify();
+			
 			FSlateApplication::Get().DismissAllMenus();
 			
 			ComponentTypes.Add(InNewClass);
@@ -1242,6 +1281,9 @@ TArray<TSharedRef<SWidget>> UMovieGraphConditionGroupQuery_EditorFolder::GetWidg
     		.OnGetRowIcon_Static(&GetRowIcon)
     		.OnDelete_Lambda([this](FName InFolderPath)
     		{
+    			const FScopedTransaction Transaction(LOCTEXT("RemoveEditorFoldersFromCollection", "Remove Editor Folders from Collection"));
+				Modify();
+    			
     			FolderPaths.Remove(InFolderPath);
     			FolderPathsList->Refresh();
     			FolderPickerWidget->FullRefresh();
@@ -1271,6 +1313,9 @@ TSharedRef<SWidget> UMovieGraphConditionGroupQuery_EditorFolder::GetAddMenuConte
 				{
 					return;
 				}
+
+				const FScopedTransaction Transaction(LOCTEXT("AddEditorFolderToCollection", "Add Editor Folder to Collection"));
+				Modify();
 				
 				FolderPaths.AddUnique(FolderPath);
 				OnAddFinished.ExecuteIfBound();
@@ -1394,6 +1439,9 @@ TArray<TSharedRef<SWidget>> UMovieGraphConditionGroupQuery_Sublevel::GetWidgets(
 			.OnGetRowIcon_Static(&GetRowIcon)
 			.OnDelete_Lambda([this](const TSharedPtr<TSoftObjectPtr<UWorld>> InSublevel)
 			{
+				const FScopedTransaction Transaction(LOCTEXT("RemoveSublevelsFromCollection", "Remove Sublevels from Collection"));
+				Modify();
+				
 				ListDataSource.Remove(InSublevel);
 				Sublevels.Remove(*InSublevel.Get());
 				
@@ -1439,6 +1487,9 @@ TSharedRef<SWidget> UMovieGraphConditionGroupQuery_Sublevel::GetAddMenuContents(
 		SublevelPickerConfig.Filter.ClassPaths.Add(UWorld::StaticClass()->GetClassPathName());
 		SublevelPickerConfig.OnAssetSelected = FOnAssetSelected::CreateLambda([this, OnAddFinished](const FAssetData& InLevelAsset)
 		{
+			const FScopedTransaction Transaction(LOCTEXT("AddSublevelsToCollection", "Add Sublevels to Collection"));
+			Modify();
+			
 			FSlateApplication::Get().DismissAllMenus();
 			
 			Sublevels.AddUnique(InLevelAsset.GetAsset());
@@ -1590,7 +1641,11 @@ TSet<AActor*> UMovieGraphConditionGroup::Evaluate(const UWorld* InWorld) const
 
 UMovieGraphConditionGroupQueryBase* UMovieGraphConditionGroup::AddQuery(const TSubclassOf<UMovieGraphConditionGroupQueryBase>& InQueryType, const int32 InsertIndex)
 {
-	UMovieGraphConditionGroupQueryBase* NewQueryObj = NewObject<UMovieGraphConditionGroupQueryBase>(this, InQueryType.Get());
+	UMovieGraphConditionGroupQueryBase* NewQueryObj = NewObject<UMovieGraphConditionGroupQueryBase>(this, InQueryType.Get(), NAME_None, RF_Transactional);
+
+#if WITH_EDITOR
+	Modify();
+#endif
 
 	if (InsertIndex < 0)
 	{
@@ -1612,6 +1667,10 @@ const TArray<UMovieGraphConditionGroupQueryBase*>& UMovieGraphConditionGroup::Ge
 
 bool UMovieGraphConditionGroup::RemoveQuery(UMovieGraphConditionGroupQueryBase* InQuery)
 {
+#if WITH_EDITOR
+	Modify();
+#endif
+	
 	return Queries.RemoveSingle(InQuery) == 1;
 }
 
@@ -1709,7 +1768,12 @@ TSet<AActor*> UMovieGraphCollection::Evaluate(const UWorld* InWorld) const
 
 UMovieGraphConditionGroup* UMovieGraphCollection::AddConditionGroup()
 {
-	UMovieGraphConditionGroup* NewConditionGroup = NewObject<UMovieGraphConditionGroup>(this);
+	UMovieGraphConditionGroup* NewConditionGroup = NewObject<UMovieGraphConditionGroup>(this, NAME_None, RF_Transactional);
+
+#if WITH_EDITOR
+	Modify();
+#endif
+	
 	ConditionGroups.Add(NewConditionGroup);
 	return NewConditionGroup;
 }
@@ -1721,6 +1785,10 @@ const TArray<UMovieGraphConditionGroup*>& UMovieGraphCollection::GetConditionGro
 
 bool UMovieGraphCollection::RemoveConditionGroup(UMovieGraphConditionGroup* InConditionGroup)
 {
+#if WITH_EDITOR
+	Modify();
+#endif
+	
 	return ConditionGroups.RemoveSingle(InConditionGroup) == 1;
 }
 
diff --git a/Engine/Plugins/MovieScene/MovieRenderPipeline/Source/MovieRenderPipelineCore/Private/Graph/Nodes/MovieGraphModifierNode.cpp b/Engine/Plugins/MovieScene/MovieRenderPipeline/Source/MovieRenderPipelineCore/Private/Graph/Nodes/MovieGraphModifierNode.cpp
--- a/Engine/Plugins/MovieScene/MovieRenderPipeline/Source/MovieRenderPipelineCore/Private/Graph/Nodes/MovieGraphModifierNode.cpp
+++ b/Engine/Plugins/MovieScene/MovieRenderPipeline/Source/MovieRenderPipelineCore/Private/Graph/Nodes/MovieGraphModifierNode.cpp
@@ -181,9 +181,13 @@ UMovieGraphCollectionModifier* UMovieGraphModifierNode::AddModifier(TSubclassOf<
 		return nullptr;
 	}
 
+#if WITH_EDITOR
+	Modify();
+#endif
+
 	if (const UClass* ModifierClass = ModifierType.Get())
 	{
-		UMovieGraphCollectionModifier* NewModifier = NewObject<UMovieGraphCollectionModifier>(this, ModifierType, ModifierClass->GetFName());
+		UMovieGraphCollectionModifier* NewModifier = NewObject<UMovieGraphCollectionModifier>(this, ModifierType, ModifierClass->GetFName(), RF_Transactional);
 		ModifiersContainer->Modifiers.Add(NewModifier);
 
 		return NewModifier;
@@ -194,6 +198,10 @@ UMovieGraphCollectionModifier* UMovieGraphModifierNode::AddModifier(TSubclassOf<
 
 bool UMovieGraphModifierNode::RemoveModifier(TSubclassOf<UMovieGraphCollectionModifier> ModifierType)
 {
+#if WITH_EDITOR
+	Modify();
+#endif
+	
 	const int32 NumRemoved = ModifiersContainer->Modifiers.RemoveAll([&ModifierType](const TObjectPtr<UMovieGraphCollectionModifier>& Modifier)
 	{
 		return Modifier && (Modifier->GetClass() == ModifierType);
@@ -204,6 +212,10 @@ bool UMovieGraphModifierNode::RemoveModifier(TSubclassOf<UMovieGraphCollectionMo
 
 void UMovieGraphModifierNode::AddCollection(const FName& InCollectionName)
 {
+#if WITH_EDITOR
+	Modify();
+#endif
+	
 	if (InCollectionName == NAME_None)
 	{
 		return;
@@ -214,6 +226,10 @@ void UMovieGraphModifierNode::AddCollection(const FName& InCollectionName)
 
 bool UMovieGraphModifierNode::RemoveCollection(const FName& InCollectionName)
 {
+#if WITH_EDITOR
+	Modify();
+#endif
+	
 	return Collections.Remove(InCollectionName) > 0;
 }
 
diff --git a/Engine/Plugins/MovieScene/MovieRenderPipeline/Source/MovieRenderPipelineEditor/Private/Customizations/Graph/MovieGraphCollectionsCustomization.cpp b/Engine/Plugins/MovieScene/MovieRenderPipeline/Source/MovieRenderPipelineEditor/Private/Customizations/Graph/MovieGraphCollectionsCustomization.cpp
--- a/Engine/Plugins/MovieScene/MovieRenderPipeline/Source/MovieRenderPipelineEditor/Private/Customizations/Graph/MovieGraphCollectionsCustomization.cpp
+++ b/Engine/Plugins/MovieScene/MovieRenderPipeline/Source/MovieRenderPipelineEditor/Private/Customizations/Graph/MovieGraphCollectionsCustomization.cpp
@@ -320,6 +320,7 @@ void FMovieGraphConditionGroupQueryBuilder::GenerateChildContent(IDetailChildren
 							{
 								if (WeakConditionGroup.IsValid() && WeakQuery.IsValid())
 								{
+									const FScopedTransaction Transaction(LOCTEXT("RemoveConditionGroupQuery", "Remove Condition Group Query"));
 									WeakConditionGroup->RemoveQuery(WeakQuery.Get());
 								}
 							}),
@@ -524,6 +525,8 @@ void FMovieGraphConditionGroupBuilder::GenerateHeaderRowContent(FDetailWidgetRow
 					const TArray<UMovieGraphConditionGroup*>& ConditionGroups = WeakCollection->GetConditionGroups();
 					if (ConditionGroups.IsValidIndex(ConditionGroupIndex))
 					{
+						const FScopedTransaction Transaction(LOCTEXT("AddConditionGroupQuery", "Add Condition Group Query"));
+						
 						// Use the Actor Name query as the default
 						ConditionGroups[ConditionGroupIndex]->AddQuery(UMovieGraphConditionGroupQuery_ActorName::StaticClass());
 					}
@@ -565,6 +568,8 @@ void FMovieGraphConditionGroupBuilder::GenerateHeaderRowContent(FDetailWidgetRow
 								const TArray<UMovieGraphConditionGroup*>& ConditionGroups = WeakCollection->GetConditionGroups();
 								if (ConditionGroups.IsValidIndex(ConditionGroupIndex))
 								{
+									const FScopedTransaction Transaction(LOCTEXT("RemoveConditionGroup", "Remove Condition Group"));
+									
 									WeakCollection->RemoveConditionGroup(ConditionGroups[ConditionGroupIndex]);
 								}
 							}
@@ -939,6 +944,8 @@ void SMovieGraphCollectionTreeQueryTypeSelectorWidget::SetQueryType(UClass* InNe
 	{
 		return;
 	}
+
+	const FScopedTransaction Transaction(LOCTEXT("ChangeConditionGroupQueryType", "Change Condition Group Query Type"));
 	
 	const int32 ExistingQueryIndex = WeakConditionGroup->GetQueries().Find(WeakQuery.Get());
 
@@ -998,6 +1005,8 @@ void FMovieGraphCollectionsCustomization::CustomizeDetails(IDetailLayoutBuilder&
 			{
 				if (WeakCollection.IsValid())
 				{
+					const FScopedTransaction Transaction(LOCTEXT("AddConditionGroup", "Add Condition Group"));
+					
 					UMovieGraphConditionGroup* NewConditionGroup = WeakCollection->AddConditionGroup();
 					NewConditionGroup->AddQuery(UMovieGraphConditionGroupQuery_ActorName::StaticClass());
 				}
diff --git a/Engine/Plugins/MovieScene/MovieRenderPipeline/Source/MovieRenderPipelineEditor/Private/Customizations/Graph/MovieGraphModifiersCustomization.cpp b/Engine/Plugins/MovieScene/MovieRenderPipeline/Source/MovieRenderPipelineEditor/Private/Customizations/Graph/MovieGraphModifiersCustomization.cpp
--- a/Engine/Plugins/MovieScene/MovieRenderPipeline/Source/MovieRenderPipelineEditor/Private/Customizations/Graph/MovieGraphModifiersCustomization.cpp
+++ b/Engine/Plugins/MovieScene/MovieRenderPipeline/Source/MovieRenderPipelineEditor/Private/Customizations/Graph/MovieGraphModifiersCustomization.cpp
@@ -10,6 +10,7 @@
 #include "Graph/MovieGraphRenderLayerSubsystem.h"
 #include "Graph/Nodes/MovieGraphCollectionNode.h"
 #include "Graph/Nodes/MovieGraphModifierNode.h"
+#include "ScopedTransaction.h"
 #include "Widgets/Images/SLayeredImage.h"
 #include "Widgets/Layout/SWidgetSwitcher.h"
 
@@ -256,6 +257,8 @@ void FMovieGraphModifiersCustomization::CustomizeDetails(IDetailLayoutBuilder& I
 						{
 							if (ModifierNode.IsValid())
 							{
+								const FScopedTransaction Transaction(LOCTEXT("AddCollectionToModifier", "Add Collection to Modifier"));
+								
 								ModifierNode->AddCollection(PickedCollectionName);
 								ListDataSource = ModifierNode->GetCollections();
 								CollectionsList->Refresh();
@@ -287,6 +290,8 @@ void FMovieGraphModifiersCustomization::CustomizeDetails(IDetailLayoutBuilder& I
 		{
 			if (ModifierNode.IsValid())
 			{
+				const FScopedTransaction Transaction(LOCTEXT("RemoveCollectionFromModifier", "Remove Collection from Modifier"));
+				
 				ModifierNode.Get()->RemoveCollection(DeletedCollectionName);
 				ListDataSource = ModifierNode->GetCollections();
 				CollectionsList->Refresh();
