MoviePipeline: Fix localization build warning.

#rnx

diff --git a/Engine/Plugins/MovieScene/MovieRenderPipeline/Source/MovieRenderPipelineEditor/Private/Customizations/Graph/MovieGraphCollectionsCustomization.cpp b/Engine/Plugins/MovieScene/MovieRenderPipeline/Source/MovieRenderPipelineEditor/Private/Customizations/Graph/MovieGraphCollectionsCustomization.cpp
--- a/Engine/Plugins/MovieScene/MovieRenderPipeline/Source/MovieRenderPipelineEditor/Private/Customizations/Graph/MovieGraphCollectionsCustomization.cpp
+++ b/Engine/Plugins/MovieScene/MovieRenderPipeline/Source/MovieRenderPipelineEditor/Private/Customizations/Graph/MovieGraphCollectionsCustomization.cpp
@@ -516,7 +516,7 @@ void FMovieGraphConditionGroupBuilder::GenerateHeaderRowContent(FDetailWidgetRow
 		[
 			SNew(SButton)
 			.HAlign(HAlign_Center)
-			.ToolTip(SNew(SToolTip).Text(LOCTEXT("AddConditionGroupQuery", "Add a condition to this condition group.")))
+			.ToolTip(SNew(SToolTip).Text(LOCTEXT("AddConditionGroupQuery_Tooltip", "Add a condition to this condition group.")))
 			.ButtonStyle(FAppStyle::Get(), "SimpleButton")
 			.OnClicked_Lambda([this]()
 			{
