MoviePipeline: Made the evaluated graph available to blueprints (eg, for burn-ins).

#rnx

diff --git a/Engine/Plugins/MovieScene/MovieRenderPipeline/Source/MovieRenderPipelineCore/Private/Graph/Nodes/MovieGraphBurnInNode.cpp b/Engine/Plugins/MovieScene/MovieRenderPipeline/Source/MovieRenderPipelineCore/Private/Graph/Nodes/MovieGraphBurnInNode.cpp
--- a/Engine/Plugins/MovieScene/MovieRenderPipeline/Source/MovieRenderPipelineCore/Private/Graph/Nodes/MovieGraphBurnInNode.cpp
+++ b/Engine/Plugins/MovieScene/MovieRenderPipeline/Source/MovieRenderPipelineCore/Private/Graph/Nodes/MovieGraphBurnInNode.cpp
@@ -126,7 +126,7 @@ void UMovieGraphBurnInNode::FMovieGraphBurnInPass::Render(const FMovieGraphTrave
 	if (const TObjectPtr<UMovieGraphBurnInWidget> BurnInWidget = GetBurnInWidget())
 	{
 		UMovieGraphPipeline* Pipeline = Renderer->GetOwningGraph();
-		BurnInWidget->UpdateForGraph(Pipeline);
+		BurnInWidget->UpdateForGraph(Pipeline, Pipeline->GetTimeStepInstance()->GetCalculatedTimeData().EvaluatedConfig);
 	}
 	
 	FMovieGraphWidgetPass::Render(InFrameTraversalContext, InTimeData);
diff --git a/Engine/Plugins/MovieScene/MovieRenderPipeline/Source/MovieRenderPipelineCore/Public/Graph/MovieGraphBurnInWidget.h b/Engine/Plugins/MovieScene/MovieRenderPipeline/Source/MovieRenderPipelineCore/Public/Graph/MovieGraphBurnInWidget.h
--- a/Engine/Plugins/MovieScene/MovieRenderPipeline/Source/MovieRenderPipelineCore/Public/Graph/MovieGraphBurnInWidget.h
+++ b/Engine/Plugins/MovieScene/MovieRenderPipeline/Source/MovieRenderPipelineCore/Public/Graph/MovieGraphBurnInWidget.h
@@ -18,7 +18,8 @@ public:
 	/** 
 	* Called on the first temporal and first spatial sample of each output frame.
 	* @param	InGraphPipeline		The graph pipeline the burn-in is for. This will be consistent throughout a burn-in widget's life.
+	* @param	InEvaluatedConfig	The evaluated graph that was used to generate this output frame.
 	*/
 	UFUNCTION(BlueprintImplementableEvent)
-	void UpdateForGraph(UMovieGraphPipeline* InGraphPipeline);
+	void UpdateForGraph(UMovieGraphPipeline* InGraphPipeline, UMovieGraphEvaluatedConfig* InEvaluatedConfig);
 };
diff --git a/Engine/Plugins/MovieScene/MovieRenderPipeline/Source/MovieRenderPipelineCore/Public/Graph/MovieGraphConfig.h b/Engine/Plugins/MovieScene/MovieRenderPipeline/Source/MovieRenderPipelineCore/Public/Graph/MovieGraphConfig.h
--- a/Engine/Plugins/MovieScene/MovieRenderPipeline/Source/MovieRenderPipelineCore/Public/Graph/MovieGraphConfig.h
+++ b/Engine/Plugins/MovieScene/MovieRenderPipeline/Source/MovieRenderPipelineCore/Public/Graph/MovieGraphConfig.h
@@ -431,12 +431,13 @@ public:
 * looking at a named branch). You can use the functions to fetch a node by type from a given
 * branch and it will return the right object (or the CDO if the node is NOT in the config).
 */
-UCLASS()
+UCLASS(BlueprintType)
 class MOVIERENDERPIPELINECORE_API UMovieGraphEvaluatedConfig : public UObject
 {
 	GENERATED_BODY()
 public:
 
+	UFUNCTION(BlueprintCallable, Category = "Movie Graph")
 	const TArray<FName> GetBranchNames() const
 	{
 		TArray<FName> OutKeys;
@@ -444,6 +445,7 @@ public:
  		return OutKeys;
 	}
 
+	UFUNCTION(BlueprintCallable, Category = "Movie Graph")
 	UMovieGraphSettingNode* GetSettingForBranch(UClass* InClass, const FName InBranchName, bool bIncludeCDOs = true, bool bExactMatch = false) const
 	{
 		TArray<UMovieGraphSettingNode*> AllSettings = GetSettingsForBranch(InClass, InBranchName, bIncludeCDOs, bExactMatch);
@@ -455,6 +457,7 @@ public:
 		return nullptr;
 	}
 
+	UFUNCTION(BlueprintCallable, Category = "Movie Graph")
 	TArray<UMovieGraphSettingNode*> GetSettingsForBranch(UClass* InClass, const FName InBranchName, bool bIncludeCDOs = true, bool bExactMatch = false) const
 	{
 		const FMovieGraphEvaluatedBranchConfig* BranchConfig = BranchConfigMapping.Find(InBranchName);
diff --git a/Engine/Plugins/MovieScene/MovieRenderPipeline/Source/MovieRenderPipelineCore/Public/Graph/MovieGraphDataTypes.h b/Engine/Plugins/MovieScene/MovieRenderPipeline/Source/MovieRenderPipelineCore/Public/Graph/MovieGraphDataTypes.h
--- a/Engine/Plugins/MovieScene/MovieRenderPipeline/Source/MovieRenderPipelineCore/Public/Graph/MovieGraphDataTypes.h
+++ b/Engine/Plugins/MovieScene/MovieRenderPipeline/Source/MovieRenderPipelineCore/Public/Graph/MovieGraphDataTypes.h
@@ -105,6 +105,7 @@ public:
 	* called at the end of the frame when we kick off the renders for the frame. Should return data
 	* needed to calculate the correct rendering timestep.
 	*/
+	UFUNCTION(BlueprintCallable, Category = "Movie Graph")
 	virtual FMovieGraphTimeStepData GetCalculatedTimeData() const { return FMovieGraphTimeStepData(); }
 	
 	/**
diff --git a/Engine/Plugins/MovieScene/MovieRenderPipeline/Source/MovieRenderPipelineCore/Public/Graph/MovieGraphPipeline.h b/Engine/Plugins/MovieScene/MovieRenderPipeline/Source/MovieRenderPipelineCore/Public/Graph/MovieGraphPipeline.h
--- a/Engine/Plugins/MovieScene/MovieRenderPipeline/Source/MovieRenderPipelineCore/Public/Graph/MovieGraphPipeline.h
+++ b/Engine/Plugins/MovieScene/MovieRenderPipeline/Source/MovieRenderPipelineCore/Public/Graph/MovieGraphPipeline.h
@@ -95,6 +95,7 @@ public:
 	void TeardownShot(const TObjectPtr<UMoviePipelineExecutorShot>& InShot);
 	
 	/** Used occasionally to cross-reference other components. Don't call this unless you know what you're doing. */
+	UFUNCTION(BlueprintCallable, Category = "Movie Graph")
 	UMovieGraphTimeStepBase* GetTimeStepInstance() const;
 	
 	/** Used occasionally to cross-reference other components. Don't call this unless you know what you're doing. */
