MoviePipeline: Added the FlushDiskWritesPerShot property to the Global Output Settings node. This lets a graph flag that it needs the pipeline to stall between shots so the post-shot scripting callbacks can be called. This setting should remain off unless absolutely needed.

#rnx

diff --git a/Engine/Plugins/MovieScene/MovieRenderPipeline/Source/MovieRenderPipelineCore/Private/Graph/MovieGraphPipeline.cpp b/Engine/Plugins/MovieScene/MovieRenderPipeline/Source/MovieRenderPipelineCore/Private/Graph/MovieGraphPipeline.cpp
--- a/Engine/Plugins/MovieScene/MovieRenderPipeline/Source/MovieRenderPipelineCore/Private/Graph/MovieGraphPipeline.cpp
+++ b/Engine/Plugins/MovieScene/MovieRenderPipeline/Source/MovieRenderPipelineCore/Private/Graph/MovieGraphPipeline.cpp
@@ -1716,6 +1716,21 @@ bool UMovieGraphPipeline::IsPostShotCallbackNeeded() const
 		}
 	}
 
+	// If a script instance didn't specify per-shot callbacks are needed, check to see if a script flagged bFlushDiskWritesPerShot on the Global Output
+	// Settings node. This accomplishes the same thing as IsPerShotCallbackNeeded() on scripting nodes.
+	if (!bAnyScriptNeedsCallbacks)
+	{
+		if (const TObjectPtr<UMovieGraphEvaluatedConfig> EvaluatedConfig = GetTimeStepInstance()->GetCalculatedTimeData().EvaluatedConfig)
+		{
+			constexpr bool bIncludeCDOs = true;
+			constexpr bool bExactMatch = true;
+			const UMovieGraphGlobalOutputSettingNode* OutputSettingsNode =
+				EvaluatedConfig->GetSettingForBranch<UMovieGraphGlobalOutputSettingNode>(UMovieGraphNode::GlobalsPinName, bIncludeCDOs, bExactMatch);
+		
+			bAnyScriptNeedsCallbacks = OutputSettingsNode->bFlushDiskWritesPerShot;
+		}
+	}
+
 	return bAnyScriptNeedsCallbacks;
 }
 
diff --git a/Engine/Plugins/MovieScene/MovieRenderPipeline/Source/MovieRenderPipelineCore/Private/Graph/Nodes/MovieGraphGlobalOutputSettingNode.cpp b/Engine/Plugins/MovieScene/MovieRenderPipeline/Source/MovieRenderPipelineCore/Private/Graph/Nodes/MovieGraphGlobalOutputSettingNode.cpp
--- a/Engine/Plugins/MovieScene/MovieRenderPipeline/Source/MovieRenderPipelineCore/Private/Graph/Nodes/MovieGraphGlobalOutputSettingNode.cpp
+++ b/Engine/Plugins/MovieScene/MovieRenderPipeline/Source/MovieRenderPipelineCore/Private/Graph/Nodes/MovieGraphGlobalOutputSettingNode.cpp
@@ -11,6 +11,10 @@ UMovieGraphGlobalOutputSettingNode::UMovieGraphGlobalOutputSettingNode()
 	, bOverwriteExistingOutput(true)
 	, ZeroPadFrameNumbers(4)
 	, FrameNumberOffset(0)
+	, HandleFrameCount(0)
+	, CustomPlaybackRangeStartFrame(0)
+	, CustomPlaybackRangeEndFrame(0)
+	, bFlushDiskWritesPerShot(false)
 {
 	// We prefer a 1080p resolution by default, but users may not have a preset that matches that. So we'll 
 	// look for a matching resolution if we can find one, otherwise we go to Custom set to 1920x1080.
diff --git a/Engine/Plugins/MovieScene/MovieRenderPipeline/Source/MovieRenderPipelineCore/Public/Graph/Nodes/MovieGraphGlobalOutputSettingNode.h b/Engine/Plugins/MovieScene/MovieRenderPipeline/Source/MovieRenderPipelineCore/Public/Graph/Nodes/MovieGraphGlobalOutputSettingNode.h
--- a/Engine/Plugins/MovieScene/MovieRenderPipeline/Source/MovieRenderPipelineCore/Public/Graph/Nodes/MovieGraphGlobalOutputSettingNode.h
+++ b/Engine/Plugins/MovieScene/MovieRenderPipeline/Source/MovieRenderPipelineCore/Public/Graph/Nodes/MovieGraphGlobalOutputSettingNode.h
@@ -81,6 +81,9 @@ public:
 	UPROPERTY(EditAnywhere, BlueprintReadWrite, Category = Overrides, meta = (InlineEditConditionToggle))
 	uint8 bOverride_VersioningSettings : 1;
 
+	UPROPERTY(EditAnywhere, BlueprintReadWrite, Category = Overrides, meta = (InlineEditConditionToggle))
+	uint8 bOverride_bFlushDiskWritesPerShot : 1;
+
 	/** What directory should all of our output files be relative to. */
 	UPROPERTY(EditAnywhere, BlueprintReadWrite, Category = "File Output", meta=(EditCondition="bOverride_OutputDirectory"))
 	FDirectoryPath OutputDirectory;
@@ -128,4 +131,11 @@ public:
 	 */
 	UPROPERTY(EditAnywhere, BlueprintReadWrite, Category = "Versioning", meta = (EditCondition = "bOverride_VersioningSettings"))
 	FMovieGraphVersioningSettings VersioningSettings;
+
+	/**
+	* If true, the game thread will stall at the end of each shot to flush the rendering queue, and then flush any outstanding writes to disk, finalizing any
+	* outstanding videos and generally completing the work. This is only relevant for scripting where scripts may do post-shot callback work.
+	*/
+	UPROPERTY(EditAnywhere, BlueprintReadWrite, Category = "Scripting", meta = (EditCondition = "bOverride_bFlushDiskWritesPerShot"))
+	bool bFlushDiskWritesPerShot;
 };
diff --git a/Engine/Plugins/MovieScene/MovieRenderPipeline/Source/MovieRenderPipelineCore/Public/MoviePipelineBase.h b/Engine/Plugins/MovieScene/MovieRenderPipeline/Source/MovieRenderPipelineCore/Public/MoviePipelineBase.h
--- a/Engine/Plugins/MovieScene/MovieRenderPipeline/Source/MovieRenderPipelineCore/Public/MoviePipelineBase.h
+++ b/Engine/Plugins/MovieScene/MovieRenderPipeline/Source/MovieRenderPipelineCore/Public/MoviePipelineBase.h
@@ -92,10 +92,13 @@ public:
 
 
 	/**
-	* Only called if `IsPostShotCallbackNeeded()` returns true!
+	* This callback will not be called by default due to performance reasons. You need to opt into this (via scripting
+	* in MoviePipeline or in the node in Movie Graph) by setting FlushDiskWritesPerShot to true in the output setting
+	* for this job's configuration.
+	*
 	* Called after each shot is finished and files have been flushed to disk. The returned data in
 	* the params struct will have only the per-shot metadata for the just finished shot. Use
-	* OnMoviePipelineFinished() if you need all of the metadata.
+	* OnMoviePipelineWorkFinishedDelegate if you need all of the metadata.
 	*/
 	UPROPERTY(BlueprintAssignable, Category = "Movie Render Pipeline")
 	FMoviePipelineWorkFinished OnMoviePipelineShotWorkFinishedDelegate;
