MoviePipeline: Added drag/drop operations for most condition group queries in the graph.

diff --git a/Engine/Plugins/MovieScene/MovieRenderPipeline/Source/MovieRenderPipelineCore/MovieRenderPipelineCore.Build.cs b/Engine/Plugins/MovieScene/MovieRenderPipeline/Source/MovieRenderPipelineCore/MovieRenderPipelineCore.Build.cs
--- a/Engine/Plugins/MovieScene/MovieRenderPipeline/Source/MovieRenderPipelineCore/MovieRenderPipelineCore.Build.cs
+++ b/Engine/Plugins/MovieScene/MovieRenderPipeline/Source/MovieRenderPipelineCore/MovieRenderPipelineCore.Build.cs
@@ -46,6 +46,8 @@ public class MovieRenderPipelineCore : ModuleRules
 		        new string[]
 		        {
 			        "ClassViewer",
+			        "ContentBrowserData",
+			        "EditorWidgets",
 			        "SceneOutliner",
 			        "UnrealEd",
 		        });
diff --git a/Engine/Plugins/MovieScene/MovieRenderPipeline/Source/MovieRenderPipelineCore/Private/Graph/MovieGraphRenderLayerSubsystem.cpp b/Engine/Plugins/MovieScene/MovieRenderPipeline/Source/MovieRenderPipelineCore/Private/Graph/MovieGraphRenderLayerSubsystem.cpp
--- a/Engine/Plugins/MovieScene/MovieRenderPipeline/Source/MovieRenderPipelineCore/Private/Graph/MovieGraphRenderLayerSubsystem.cpp
+++ b/Engine/Plugins/MovieScene/MovieRenderPipeline/Source/MovieRenderPipelineCore/Private/Graph/MovieGraphRenderLayerSubsystem.cpp
@@ -25,8 +25,13 @@
 #include "ClassViewerFilter.h"
 #include "ClassViewerModule.h"
 #include "ComponentTreeItem.h"
+#include "ContentBrowserDataDragDropOp.h"
 #include "ContentBrowserDataSource.h"
 #include "ContentBrowserModule.h"
+#include "DragAndDrop/ActorDragDropGraphEdOp.h"
+#include "DragAndDrop/AssetDragDropOp.h"
+#include "DragAndDrop/FolderDragDropOp.h"
+#include "DragAndDrop/LevelDragDropOp.h"
 #include "Editor.h"
 #include "EditorActorFolders.h"
 #include "Graph/MovieGraphSharedWidgets.h"
@@ -36,6 +41,7 @@
 #include "SceneOutlinerPublicTypes.h"
 #include "SClassViewer.h"
 #include "ScopedTransaction.h"
+#include "SDropTarget.h"
 #include "Selection.h"
 #endif
 
@@ -78,6 +84,15 @@ namespace UE::MovieGraph::Private
 		/** Classes must have this base class to pass the filter. */
 		UClass* RequiredBaseClass = nullptr;
 	};
+
+	/** Gets all actors from a scene drag-drop operation (which is assumed to be dragging a folder). */
+	void GetActorsFromSceneDragDropOp(const TSharedPtr<FSceneOutlinerDragDropOp> InSceneDragDropOp, TArray<AActor*>& OutActors)
+	{
+		if (const TSharedPtr<FFolderDragDropOp> FolderOp = InSceneDragDropOp->GetSubOp<FFolderDragDropOp>())
+		{
+			FActorFolders::GetActorsFromFolders(*FolderOp->World.Get(), FolderOp->Folders, OutActors);
+		}
+	}
 #endif
 }
 
@@ -581,7 +596,33 @@ TArray<TSharedRef<SWidget>> UMovieGraphConditionGroupQuery_Actor::GetWidgets()
 	}
 
 	Widgets.Add(
-		SAssignNew(ActorsList, SMovieGraphSimpleList<TSharedPtr<TSoftObjectPtr<AActor>>>)
+		SNew(SDropTarget)
+		.OnAllowDrop_Lambda([](TSharedPtr<FDragDropOperation> InDragOperation)
+		{
+			// Support dragging both actors and folders from the Outliner (dragging a folder will add all actors in the folder)
+			return InDragOperation->IsOfType<FActorDragDropGraphEdOp>() || InDragOperation->IsOfType<FSceneOutlinerDragDropOp>();
+		})
+		.OnDropped_Lambda([this](const FGeometry& Geometry, const FDragDropEvent& DragDropEvent)
+		{
+			TArray<AActor*> DroppedActors;
+			
+			if (const TSharedPtr<FActorDragDropGraphEdOp> ActorOperation = DragDropEvent.GetOperationAs<FActorDragDropGraphEdOp>())
+			{
+				Algo::Transform(ActorOperation->Actors, DroppedActors, [](const TWeakObjectPtr<AActor>& Actor) { return Actor.IsValid() ? Actor.Get() : nullptr; });
+			}
+
+			if (const TSharedPtr<FSceneOutlinerDragDropOp> SceneOperation = DragDropEvent.GetOperationAs<FSceneOutlinerDragDropOp>())
+			{
+				UE::MovieGraph::Private::GetActorsFromSceneDragDropOp(SceneOperation, DroppedActors);
+			}
+
+			const FMovieGraphConditionGroupQueryContentsChanged OnAddFinished = nullptr;
+			AddActors(DroppedActors, OnAddFinished);
+
+			return FReply::Handled();
+		})
+		[
+			SAssignNew(ActorsList, SMovieGraphSimpleList<TSharedPtr<TSoftObjectPtr<AActor>>>)
 			.DataSource(&ListDataSource)
 			.DataType(FText::FromString("Actor"))
 			.DataTypePlural(FText::FromString("Actors"))
@@ -596,6 +637,7 @@ TArray<TSharedRef<SWidget>> UMovieGraphConditionGroupQuery_Actor::GetWidgets()
 				ActorsToMatch.Remove(*InActor.Get());
 				ActorsList->Refresh();
 			})
+		]
 	);
 
 	return Widgets;
@@ -632,26 +674,6 @@ TSharedRef<SWidget> UMovieGraphConditionGroupQuery_Actor::GetAddMenuContents(con
 	const FSceneOutlinerModule& SceneOutlinerModule = FModuleManager::LoadModuleChecked<FSceneOutlinerModule>("SceneOutliner");
 
 	FMenuBuilder MenuBuilder(false, MakeShared<FUICommandList>());
-
-	auto AddActorToList = [this](const AActor* InActor, const FMovieGraphConditionGroupQueryContentsChanged& OnAddFinished)
-	{
-		const FScopedTransaction Transaction(LOCTEXT("AddActorsToCollection", "Add Actors to Collection"));
-		Modify();
-		
-		ActorsToMatch.Add(InActor);
-		ListDataSource.Add(MakeShared<TSoftObjectPtr<AActor>>(ActorsToMatch.Last()));
-		OnAddFinished.ExecuteIfBound();
-	};
-	
-	auto RefreshActorPickerFilterAndList = [this]()
-	{
-		// Ensure that the actor picker filter runs again so duplicate actors cannot be selected
-		if (ActorPickerWidget.IsValid())
-		{
-			ActorPickerWidget->FullRefresh();
-			ActorsList->Refresh();
-		}
-	};
 	
 	MenuBuilder.BeginSection("AddActor", LOCTEXT("AddActor", "Add Actor"));
 	{
@@ -660,30 +682,18 @@ TSharedRef<SWidget> UMovieGraphConditionGroupQuery_Actor::GetAddMenuContents(con
 			LOCTEXT("AddSelectedInOutlinerTooltip", "Add actors currently selected in the level editor's scene outliner."),
 			FSlateIcon(FAppStyle::GetAppStyleSetName(),"FoliageEditMode.SetSelect"),
 			FUIAction(
-				FExecuteAction::CreateLambda([this, OnAddFinished, RefreshActorPickerFilterAndList, AddActorToList]()
+				FExecuteAction::CreateLambda([this, OnAddFinished]()
 				{
 					const FScopedTransaction Transaction(LOCTEXT("AddSelectedActorsToCollection", "Add Selected Actors to Collection"));
 					
-					USelection* CurrentSelection = GEditor->GetSelectedActors();
 					TArray<AActor*> SelectedActors;
-					CurrentSelection->GetSelectedObjects<AActor>(SelectedActors);
-					for (const AActor* Actor : SelectedActors)
-					{
-						if (Actor && !ActorsToMatch.Contains(Actor))
-						{
-							AddActorToList(Actor, OnAddFinished);
-						}
-					}
-
-					RefreshActorPickerFilterAndList();
+					GEditor->GetSelectedActors()->GetSelectedObjects<AActor>(SelectedActors);
 
-					FSlateApplication::Get().DismissAllMenus();
-				}
-				),
+					AddActors(SelectedActors, OnAddFinished);
+				}),
 				FCanExecuteAction::CreateLambda([]()
 				{
 					return GEditor->GetSelectedActors()->Num() > 0;
-					
 				})
 			)
 		);
@@ -694,11 +704,9 @@ TSharedRef<SWidget> UMovieGraphConditionGroupQuery_Actor::GetAddMenuContents(con
 	{
 		ActorPickerWidget = SceneOutlinerModule.CreateActorPicker(
 			SceneOutlinerInitOptions,
-			FOnActorPicked::CreateLambda([this, OnAddFinished, RefreshActorPickerFilterAndList, AddActorToList](AActor* InActor)
+			FOnActorPicked::CreateLambda([this, OnAddFinished](AActor* InActor)
 			{
-				AddActorToList(InActor, OnAddFinished);
-
-				RefreshActorPickerFilterAndList();
+				AddActors({InActor}, OnAddFinished);
 			}));
 
 		const TSharedRef<SBox> ActorPickerWidgetBox =
@@ -743,6 +751,32 @@ FText UMovieGraphConditionGroupQuery_Actor::GetRowText(TSharedPtr<TSoftObjectPtr
 
 	return LOCTEXT("MovieGraphActorConditionGroupQuery_InvalidActor", "(invalid)");
 }
+
+void UMovieGraphConditionGroupQuery_Actor::AddActors(const TArray<AActor*>& InActors, const FMovieGraphConditionGroupQueryContentsChanged& InOnAddFinished)
+{
+	const FScopedTransaction Transaction(LOCTEXT("AddActorsToCollection", "Add Actors to Collection"));
+	Modify();
+	
+	for (const AActor* Actor : InActors)
+	{
+		if (Actor && !ActorsToMatch.Contains(Actor))
+		{
+			ActorsToMatch.Add(Actor);
+			ListDataSource.Add(MakeShared<TSoftObjectPtr<AActor>>(ActorsToMatch.Last()));
+			
+			InOnAddFinished.ExecuteIfBound();
+		}
+	}
+
+	// Ensure that the actor picker filter runs again so duplicate actors cannot be selected
+	if (ActorPickerWidget.IsValid())
+	{
+		ActorPickerWidget->FullRefresh();
+	}
+	
+	ActorsList->Refresh();
+	FSlateApplication::Get().DismissAllMenus();
+}
 #endif	// WITH_EDITOR
 
 void UMovieGraphConditionGroupQuery_ActorTagName::Evaluate(const TArray<AActor*>& InActorsToQuery, const UWorld* InWorld, TSet<AActor*>& OutMatchingActors) const
@@ -872,20 +906,60 @@ TArray<TSharedRef<SWidget>> UMovieGraphConditionGroupQuery_ActorName::GetWidgets
 	TArray<TSharedRef<SWidget>> Widgets;
 
 	Widgets.Add(
-		SNew(SBox)
-		.HAlign(HAlign_Fill)
-		.Padding(7.f, 2.f)
-		[
-			SNew(SMultiLineEditableTextBox)
-			.Text_Lambda([this]() { return FText::FromString(WildcardSearch); })
-			.OnTextCommitted_Lambda([this](const FText& InText, ETextCommit::Type TextCommitType)
+		SNew(SDropTarget)
+		.OnAllowDrop_Lambda([](TSharedPtr<FDragDropOperation> InDragOperation)
+		{
+			// Support dragging both actors and folders from the Outliner (dragging a folder will add all actors in the folder)
+			return InDragOperation->IsOfType<FActorDragDropGraphEdOp>() || InDragOperation->IsOfType<FSceneOutlinerDragDropOp>();
+		})
+		.OnDropped_Lambda([this](const FGeometry& Geometry, const FDragDropEvent& DragDropEvent)
+		{
+			TArray<AActor*> DroppedActors;
+			
+			if (const TSharedPtr<FActorDragDropGraphEdOp> ActorOperation = DragDropEvent.GetOperationAs<FActorDragDropGraphEdOp>())
 			{
-				const FScopedTransaction Transaction(LOCTEXT("UpdateActorNamesInCollection", "Update Actor Names in Collection"));
-				Modify();
+				Algo::Transform(ActorOperation->Actors, DroppedActors, [](const TWeakObjectPtr<AActor>& Actor) { return Actor.IsValid() ? Actor.Get() : nullptr; });
+			}
 
-				WildcardSearch = InText.ToString();
-			})
-			.HintText(LOCTEXT("MovieGraphActorNameQueryHintText", "Actor names to query. Wildcards allowed.\nEnter each actor name on a separate line."))
+			if (const TSharedPtr<FSceneOutlinerDragDropOp> SceneOperation = DragDropEvent.GetOperationAs<FSceneOutlinerDragDropOp>())
+			{
+				UE::MovieGraph::Private::GetActorsFromSceneDragDropOp(SceneOperation, DroppedActors);
+			}
+			
+			const FScopedTransaction Transaction(LOCTEXT("UpdateActorNamesInCollection", "Update Actor Names in Collection"));
+			Modify();
+
+			for (const AActor* DroppedActor : DroppedActors)
+			{
+				TArray<FString> ActorStrings;
+				WildcardSearch.ParseIntoArrayLines(ActorStrings);
+				
+				// Only add the actor if it's not in the list already
+				if (!ActorStrings.Contains(DroppedActor->GetActorLabel()))
+				{
+					const FString LineSeparator = WildcardSearch.IsEmpty() ? FString() : LINE_TERMINATOR;
+					WildcardSearch += LineSeparator + DroppedActor->GetActorLabel();
+				}
+			}
+
+			return FReply::Handled();
+		})
+		[
+			SNew(SBox)
+			.HAlign(HAlign_Fill)
+			.Padding(7.f, 2.f)
+			[
+				SNew(SMultiLineEditableTextBox)
+				.Text_Lambda([this]() { return FText::FromString(WildcardSearch); })
+				.OnTextCommitted_Lambda([this](const FText& InText, ETextCommit::Type TextCommitType)
+				{
+					const FScopedTransaction Transaction(LOCTEXT("UpdateActorNamesInCollection", "Update Actor Names in Collection"));
+					Modify();
+
+					WildcardSearch = InText.ToString();
+				})
+				.HintText(LOCTEXT("MovieGraphActorNameQueryHintText", "Actor names to query. Wildcards allowed.\nEnter each actor name on a separate line."))
+			]
 		]
 	);
 
@@ -930,7 +1004,36 @@ TArray<TSharedRef<SWidget>> UMovieGraphConditionGroupQuery_ActorType::GetWidgets
 	TArray<TSharedRef<SWidget>> Widgets;
 
 	Widgets.Add(
-		SAssignNew(ActorTypesList, SMovieGraphSimpleList<UClass*>)
+		SNew(SDropTarget)
+		.OnAllowDrop_Lambda([](TSharedPtr<FDragDropOperation> InDragOperation)
+		{
+			// Support dragging both actors and folders from the Outliner (dragging a folder will add all actor types in the folder)
+			return InDragOperation->IsOfType<FActorDragDropGraphEdOp>() || InDragOperation->IsOfType<FSceneOutlinerDragDropOp>();
+		})
+		.OnDropped_Lambda([this](const FGeometry& Geometry, const FDragDropEvent& DragDropEvent)
+		{
+			TArray<UClass*> DroppedActorClasses;
+			
+			if (const TSharedPtr<FActorDragDropGraphEdOp> ActorOperation = DragDropEvent.GetOperationAs<FActorDragDropGraphEdOp>())
+			{
+				Algo::Transform(ActorOperation->Actors, DroppedActorClasses, [](const TWeakObjectPtr<AActor>& Actor) { return Actor.IsValid() ? Actor.Get()->GetClass() : nullptr; });
+			}
+
+			if (const TSharedPtr<FSceneOutlinerDragDropOp> SceneOperation = DragDropEvent.GetOperationAs<FSceneOutlinerDragDropOp>())
+			{
+				TArray<AActor*> DroppedActors;
+				UE::MovieGraph::Private::GetActorsFromSceneDragDropOp(SceneOperation, DroppedActors);
+
+				Algo::Transform(DroppedActors, DroppedActorClasses, [](const AActor* Actor) { return Actor ? Actor->GetClass() : nullptr; });
+			}
+			
+			const FMovieGraphConditionGroupQueryContentsChanged OnAddFinished = nullptr;
+			AddActorTypes(DroppedActorClasses, OnAddFinished);
+
+			return FReply::Handled();
+		})
+		[
+			SAssignNew(ActorTypesList, SMovieGraphSimpleList<UClass*>)
 			.DataSource(&ActorTypes)
 			.DataType(FText::FromString("Actor Type"))
 			.DataTypePlural(FText::FromString("Actor Types"))
@@ -944,6 +1047,7 @@ TArray<TSharedRef<SWidget>> UMovieGraphConditionGroupQuery_ActorType::GetWidgets
 				ActorTypes.Remove(InActorClass);
 				ActorTypesList->Refresh();
 			})
+		]
 	);
 
 	return Widgets;
@@ -972,20 +1076,7 @@ TSharedRef<SWidget> UMovieGraphConditionGroupQuery_ActorType::GetAddMenuContents
 		Options,
 		FOnClassPicked::CreateLambda([this, OnAddFinished](UClass* InNewClass)
 		{
-			const FScopedTransaction Transaction(LOCTEXT("AddActorTypesToCollection", "Add Actor Types to Collection"));
-			Modify();
-			
-			FSlateApplication::Get().DismissAllMenus();
-			
-			ActorTypes.Add(InNewClass);
-			OnAddFinished.ExecuteIfBound();
-
-			// Ensure that the class filters run again so duplicate actor types cannot be selected
-			if (ClassViewerWidget.IsValid())
-			{
-				ClassViewerWidget->Refresh();
-				ActorTypesList->Refresh();
-			}
+			AddActorTypes({InNewClass}, OnAddFinished);
 		}));
 
 	ClassViewerWidget = StaticCastSharedPtr<SClassViewer>(ClassViewer.ToSharedPtr()); 
@@ -1012,6 +1103,29 @@ FText UMovieGraphConditionGroupQuery_ActorType::GetRowText(UClass* InActorType)
 
 	return LOCTEXT("MovieGraphActorTypeConditionGroupQuery_Invalid", "(invalid)");
 }
+
+void UMovieGraphConditionGroupQuery_ActorType::AddActorTypes(const TArray<UClass*>& InActorTypes, const FMovieGraphConditionGroupQueryContentsChanged& InOnAddFinished)
+{
+	const FScopedTransaction Transaction(LOCTEXT("AddActorTypesToCollection", "Add Actor Types to Collection"));
+	Modify();
+	
+	FSlateApplication::Get().DismissAllMenus();
+
+	for (UClass* ActorType : InActorTypes)
+	{
+		ActorTypes.AddUnique(ActorType);
+	}
+	
+	InOnAddFinished.ExecuteIfBound();
+
+	// Ensure that the class filters run again so duplicate actor types cannot be selected
+	if (ClassViewerWidget.IsValid())
+	{
+		ClassViewerWidget->Refresh();
+	}
+
+	ActorTypesList->Refresh();
+}
 #endif
 
 void UMovieGraphConditionGroupQuery_ComponentTagName::Evaluate(const TArray<AActor*>& InActorsToQuery, const UWorld* InWorld, TSet<AActor*>& OutMatchingActors) const
@@ -1166,7 +1280,7 @@ TArray<TSharedRef<SWidget>> UMovieGraphConditionGroupQuery_ComponentType::GetWid
 				
 				ComponentTypes.Remove(InComponentType);
 				ComponentTypesList->Refresh();
-			})			
+			})
 	);
 
 	return Widgets;
@@ -1271,23 +1385,55 @@ bool UMovieGraphConditionGroupQuery_EditorFolder::IsEditorOnlyQuery() const
 TArray<TSharedRef<SWidget>> UMovieGraphConditionGroupQuery_EditorFolder::GetWidgets()
 {
 	TArray<TSharedRef<SWidget>> Widgets;
-    
+
+	auto GetFolderDragOp = [](FDragDropOperation* InDragDropOp) -> TSharedPtr<FFolderDragDropOp>
+	{
+		if (InDragDropOp->IsOfType<FSceneOutlinerDragDropOp>())
+		{
+			FSceneOutlinerDragDropOp* SceneOutlinerOp = static_cast<FSceneOutlinerDragDropOp*>(InDragDropOp);
+			return SceneOutlinerOp->GetSubOp<FFolderDragDropOp>();
+		}
+
+		return nullptr;
+	};
+	
     Widgets.Add(
-    	SAssignNew(FolderPathsList, SMovieGraphSimpleList<FName>)
-    		.DataSource(&FolderPaths)
-    		.DataType(FText::FromString("Folder"))
-    		.DataTypePlural(FText::FromString("Folders"))
-    		.OnGetRowText_Static(&GetRowText)
-    		.OnGetRowIcon_Static(&GetRowIcon)
-    		.OnDelete_Lambda([this](FName InFolderPath)
-    		{
-    			const FScopedTransaction Transaction(LOCTEXT("RemoveEditorFoldersFromCollection", "Remove Editor Folders from Collection"));
+		SNew(SDropTarget)
+		.OnAllowDrop_Lambda([GetFolderDragOp](TSharedPtr<FDragDropOperation> InDragOperation)
+		{
+			return GetFolderDragOp(InDragOperation.Get()) != nullptr;
+		})
+		.OnDropped_Lambda([this, GetFolderDragOp](const FGeometry& Geometry, const FDragDropEvent& DragDropEvent)
+		{
+			if (const TSharedPtr<FFolderDragDropOp> FolderDragDropOp = GetFolderDragOp(DragDropEvent.GetOperation().Get()))
+			{
+				const FMovieGraphConditionGroupQueryContentsChanged OnAddFinished = nullptr;
+				AddFolders(FolderDragDropOp->Folders, OnAddFinished);
+			}
+
+			return FReply::Handled();
+		})
+		[
+			SAssignNew(FolderPathsList, SMovieGraphSimpleList<FName>)
+			.DataSource(&FolderPaths)
+			.DataType(FText::FromString("Folder"))
+			.DataTypePlural(FText::FromString("Folders"))
+			.OnGetRowText_Static(&GetRowText)
+			.OnGetRowIcon_Static(&GetRowIcon)
+			.OnDelete_Lambda([this](FName InFolderPath)
+			{
+				const FScopedTransaction Transaction(LOCTEXT("RemoveEditorFoldersFromCollection", "Remove Editor Folders from Collection"));
 				Modify();
-    			
-    			FolderPaths.Remove(InFolderPath);
-    			FolderPathsList->Refresh();
-    			FolderPickerWidget->FullRefresh();
-    		})			
+
+				FolderPaths.Remove(InFolderPath);
+				FolderPathsList->Refresh();
+
+				if (FolderPickerWidget.IsValid())
+				{
+					FolderPickerWidget->FullRefresh();
+				}
+			})
+		]
     );
 
     return Widgets;
@@ -1307,24 +1453,8 @@ TSharedRef<SWidget> UMovieGraphConditionGroupQuery_EditorFolder::GetAddMenuConte
 			if (FolderItem->IsValid())
 			{
 				const FName& FolderPath = FolderItem->GetPath();
-					
-				// Don't allow duplicate folder paths
-				if (FolderPaths.Contains(FolderPath))
-				{
-					return;
-				}
-
-				const FScopedTransaction Transaction(LOCTEXT("AddEditorFolderToCollection", "Add Editor Folder to Collection"));
-				Modify();
 				
-				FolderPaths.AddUnique(FolderPath);
-				OnAddFinished.ExecuteIfBound();
-				
-				if (FolderPickerWidget.IsValid())
-				{
-					FolderPickerWidget->FullRefresh();
-					FolderPathsList->Refresh();
-				}
+				AddFolders({FolderPath}, OnAddFinished);
 			}
 		}
 	});
@@ -1372,6 +1502,26 @@ FText UMovieGraphConditionGroupQuery_EditorFolder::GetRowText(FName InFolderPath
 {
 	return FText::FromString(InFolderPath.ToString());
 }
+
+void UMovieGraphConditionGroupQuery_EditorFolder::AddFolders(const TArray<FName>& InFolderPaths, const FMovieGraphConditionGroupQueryContentsChanged& InOnAddFinished)
+{
+	const FScopedTransaction Transaction(LOCTEXT("AddEditorFolderToCollection", "Add Editor Folder to Collection"));
+	Modify();
+	
+	for (const FName& FolderPath : InFolderPaths)
+	{
+		FolderPaths.AddUnique(FolderPath);
+	}
+	
+	InOnAddFinished.ExecuteIfBound();
+				
+	if (FolderPickerWidget.IsValid())
+	{
+		FolderPickerWidget->FullRefresh();
+	}
+
+	FolderPathsList->Refresh();
+}
 #endif	// WITH_EDITOR
 
 void UMovieGraphConditionGroupQuery_Sublevel::Evaluate(const TArray<AActor*>& InActorsToQuery, const UWorld* InWorld, TSet<AActor*>& OutMatchingActors) const
@@ -1431,7 +1581,65 @@ TArray<TSharedRef<SWidget>> UMovieGraphConditionGroupQuery_Sublevel::GetWidgets(
 	}
 
 	Widgets.Add(
-		SAssignNew(SublevelsList, SMovieGraphSimpleList<TSharedPtr<TSoftObjectPtr<UWorld>>>)
+		SNew(SDropTarget)
+		.OnAllowDrop_Lambda([](TSharedPtr<FDragDropOperation> InDragOperation)
+		{
+			// Support drag-n-drop from the Content browser
+			if (InDragOperation->IsOfType<FContentBrowserDataDragDropOp>())
+			{
+				const FContentBrowserDataDragDropOp* ContentBrowserOp = static_cast<FContentBrowserDataDragDropOp*>(InDragOperation.Get());
+				for (const FAssetData& DraggedAsset : ContentBrowserOp->GetAssets())
+				{
+					if (DraggedAsset.AssetClassPath == UWorld::StaticClass()->GetClassPathName())
+					{
+						return true;
+					}
+				}
+			}
+
+			// Support drag-n-drop from the Levels editor
+			if (InDragOperation->IsOfType<FLevelDragDropOp>())
+			{
+				return true;
+			}
+
+			return false;
+		})
+		.OnDropped_Lambda([this](const FGeometry& Geometry, const FDragDropEvent& DragDropEvent)
+		{
+			if (const TSharedPtr<FContentBrowserDataDragDropOp> ContentBrowserOp = DragDropEvent.GetOperationAs<FContentBrowserDataDragDropOp>())
+			{
+				TArray<UWorld*> DroppedLevels;
+				
+				for (const FAssetData& DraggedAsset : ContentBrowserOp->GetAssets())
+				{
+					if (DraggedAsset.AssetClassPath == UWorld::StaticClass()->GetClassPathName())
+					{
+						DroppedLevels.Add(Cast<UWorld>(DraggedAsset.GetAsset()));
+					}
+				}
+
+				const FMovieGraphConditionGroupQueryContentsChanged OnAddFinished = nullptr;
+				AddLevels(DroppedLevels, OnAddFinished);
+			}
+
+			if (const TSharedPtr<FLevelDragDropOp> LevelEditorOp = DragDropEvent.GetOperationAs<FLevelDragDropOp>())
+			{
+				TArray<UWorld*> Levels;
+				Algo::Transform(LevelEditorOp->LevelsToDrop, Levels, [](const TWeakObjectPtr<ULevel>& Level) { return Level.IsValid() ? Level->GetWorld() : nullptr; });
+
+				TArray<UWorld*> StreamingLevels;
+				Algo::Transform(LevelEditorOp->StreamingLevelsToDrop, StreamingLevels, [](const TWeakObjectPtr<ULevelStreaming>& LevelStreaming) { return LevelStreaming.IsValid() ? LevelStreaming->GetWorldAsset().Get() : nullptr; });
+				
+				const FMovieGraphConditionGroupQueryContentsChanged OnAddFinished = nullptr;
+				AddLevels(Levels, OnAddFinished);
+				AddLevels(StreamingLevels, OnAddFinished);
+			}
+
+			return FReply::Handled();
+		})
+		[
+			SAssignNew(SublevelsList, SMovieGraphSimpleList<TSharedPtr<TSoftObjectPtr<UWorld>>>)
 			.DataSource(&ListDataSource)
 			.DataType(FText::FromString("Sublevel"))
 			.DataTypePlural(FText::FromString("Sublevels"))
@@ -1450,6 +1658,7 @@ TArray<TSharedRef<SWidget>> UMovieGraphConditionGroupQuery_Sublevel::GetWidgets(
 				constexpr bool bUpdateSources = true;
 				RefreshLevelPicker.ExecuteIfBound(bUpdateSources);
 			})
+		]
 	);
 
 	return Widgets;
@@ -1487,22 +1696,7 @@ TSharedRef<SWidget> UMovieGraphConditionGroupQuery_Sublevel::GetAddMenuContents(
 		SublevelPickerConfig.Filter.ClassPaths.Add(UWorld::StaticClass()->GetClassPathName());
 		SublevelPickerConfig.OnAssetSelected = FOnAssetSelected::CreateLambda([this, OnAddFinished](const FAssetData& InLevelAsset)
 		{
-			const FScopedTransaction Transaction(LOCTEXT("AddSublevelsToCollection", "Add Sublevels to Collection"));
-			Modify();
-			
-			FSlateApplication::Get().DismissAllMenus();
-			
-			Sublevels.AddUnique(InLevelAsset.GetAsset());
-			ListDataSource.AddUnique(MakeShared<TSoftObjectPtr<UWorld>>(InLevelAsset.GetAsset()));
-			OnAddFinished.ExecuteIfBound();
-
-			if (SublevelsList.IsValid())
-			{
-				SublevelsList->Refresh();
-				
-				constexpr bool bUpdateSources = false;
-				RefreshLevelPicker.ExecuteIfBound(bUpdateSources);
-			}
+			AddLevels({Cast<UWorld>(InLevelAsset.GetAsset())}, OnAddFinished);
 		});
 		SublevelPickerConfig.OnShouldFilterAsset = FOnShouldFilterAsset::CreateLambda([this](const FAssetData& InLevelAsset)
 		{
@@ -1542,6 +1736,33 @@ FText UMovieGraphConditionGroupQuery_Sublevel::GetRowText(TSharedPtr<TSoftObject
 
 	return LOCTEXT("MovieGraphSublevelConditionGroupQuery_InvalidLevel", "(invalid)");
 }
+
+void UMovieGraphConditionGroupQuery_Sublevel::AddLevels(const TArray<UWorld*>& InLevels, const FMovieGraphConditionGroupQueryContentsChanged& InOnAddFinished)
+{
+	const FScopedTransaction Transaction(LOCTEXT("AddSublevelsToCollection", "Add Sublevels to Collection"));
+	Modify();
+	
+	FSlateApplication::Get().DismissAllMenus();
+
+	for (UWorld* Level : InLevels)
+	{
+		if (Level && !Sublevels.Contains(Level))
+		{
+			Sublevels.Add(Level);
+			ListDataSource.Add(MakeShared<TSoftObjectPtr<UWorld>>(Level));
+		}
+	}
+	
+	InOnAddFinished.ExecuteIfBound();
+
+	if (SublevelsList.IsValid())
+	{
+		SublevelsList->Refresh();
+	}
+
+	constexpr bool bUpdateSources = false;
+	RefreshLevelPicker.ExecuteIfBound(bUpdateSources);
+}
 #endif	// WITH_EDITOR
 
 UMovieGraphConditionGroup::UMovieGraphConditionGroup()
diff --git a/Engine/Plugins/MovieScene/MovieRenderPipeline/Source/MovieRenderPipelineCore/Public/Graph/MovieGraphRenderLayerSubsystem.h b/Engine/Plugins/MovieScene/MovieRenderPipeline/Source/MovieRenderPipelineCore/Public/Graph/MovieGraphRenderLayerSubsystem.h
--- a/Engine/Plugins/MovieScene/MovieRenderPipeline/Source/MovieRenderPipelineCore/Public/Graph/MovieGraphRenderLayerSubsystem.h
+++ b/Engine/Plugins/MovieScene/MovieRenderPipeline/Source/MovieRenderPipelineCore/Public/Graph/MovieGraphRenderLayerSubsystem.h
@@ -163,6 +163,9 @@ private:
 #if WITH_EDITOR
 	static const FSlateBrush* GetRowIcon(TSharedPtr<TSoftObjectPtr<AActor>> InActor);
 	static FText GetRowText(TSharedPtr<TSoftObjectPtr<AActor>> InActor);
+
+	/** Adds the provided actors to the query, updating the UI as needed. Calls InOnAddFinished when done. */
+	void AddActors(const TArray<AActor*>& InActors, const FMovieGraphConditionGroupQueryContentsChanged& InOnAddFinished);
 #endif
 
 public:
@@ -277,6 +280,9 @@ private:
 #if WITH_EDITOR
 	static const FSlateBrush* GetRowIcon(UClass* InActorType);
 	static FText GetRowText(UClass* InActorType);
+	
+	/** Adds the provided actor types to the query, updating the UI as needed. Calls InOnAddFinished when done. */
+	void AddActorTypes(const TArray<UClass*>& InActorTypes, const FMovieGraphConditionGroupQueryContentsChanged& InOnAddFinished);
 
 	/** Displays the actor types which have been chosen. */
 	TSharedPtr<SMovieGraphSimpleList<UClass*>> ActorTypesList;
@@ -377,6 +383,9 @@ private:
 	static const FSlateBrush* GetRowIcon(FName InFolderPath);
 	static FText GetRowText(FName InFolderPath);
 
+	/** Adds the provided folders to the query, updating the UI as needed. Calls InOnAddFinished when done. */
+	void AddFolders(const TArray<FName>& InFolderPaths, const FMovieGraphConditionGroupQueryContentsChanged& InOnAddFinished);
+
 	/** Displays the paths of folders which have been chosen. */
 	TSharedPtr<SMovieGraphSimpleList<FName>> FolderPathsList;
 
@@ -412,6 +421,9 @@ private:
 	static const FSlateBrush* GetRowIcon(TSharedPtr<TSoftObjectPtr<UWorld>> InSublevel);
 	static FText GetRowText(TSharedPtr<TSoftObjectPtr<UWorld>> InSublevel);
 
+	/** Adds the provided levels to the query, updating the UI as needed. Calls InOnAddFinished when done. */
+	void AddLevels(const TArray<UWorld*>& InLevels, const FMovieGraphConditionGroupQueryContentsChanged& InOnAddFinished);
+
 	/** Displays the names of sublevels which have been chosen. */
 	TSharedPtr<SMovieGraphSimpleList<TSharedPtr<TSoftObjectPtr<UWorld>>>> SublevelsList;
 
