MoviePipeline: Fixed a crash that occurred in the queue, as well as several other minor bugs.

* Fixed an issue where the config preset menu for shots was populating with primary configs rather than shot configs. Selecting a primary config for a shot caused a crash.
* Resetting a job/shot to use a config when a graph was previously assigned was not properly tracked in the undo system.
* Resetting a job to use a config let child shots keep their graph configs; fixed to remove graph configs from shots.
* If no shot configs are available for a shot in the preset menu, the menu previously incorrectly displayed "No Render Graphs Found".

#rnx

diff --git a/Engine/Plugins/MovieScene/MovieRenderPipeline/Source/MovieRenderPipelineCore/Private/MoviePipelineBlueprintLibrary.cpp b/Engine/Plugins/MovieScene/MovieRenderPipeline/Source/MovieRenderPipelineCore/Private/MoviePipelineBlueprintLibrary.cpp
--- a/Engine/Plugins/MovieScene/MovieRenderPipeline/Source/MovieRenderPipelineCore/Private/MoviePipelineBlueprintLibrary.cpp
+++ b/Engine/Plugins/MovieScene/MovieRenderPipeline/Source/MovieRenderPipelineCore/Private/MoviePipelineBlueprintLibrary.cpp
@@ -752,6 +752,8 @@ void UMoviePipelineBlueprintLibrary::UpdateJobShotListFromSequence(ULevelSequenc
 			if (!NewShot)
 			{
 				NewShot = NewObject<UMoviePipelineExecutorShot>(InJob);
+				NewShot->SetFlags(RF_Transactional);
+				
 				UE_LOG(LogMovieRenderPipeline, Log, TEXT("Generated new ShotInfo for Inner: %s Outer: %s (No existing shot found in the job)."), *Entity.Name.Get<0>(), *Entity.Name.Get<1>());
 			}
 			else
diff --git a/Engine/Plugins/MovieScene/MovieRenderPipeline/Source/MovieRenderPipelineCore/Private/MoviePipelineQueue.cpp b/Engine/Plugins/MovieScene/MovieRenderPipeline/Source/MovieRenderPipelineCore/Private/MoviePipelineQueue.cpp
--- a/Engine/Plugins/MovieScene/MovieRenderPipeline/Source/MovieRenderPipelineCore/Private/MoviePipelineQueue.cpp
+++ b/Engine/Plugins/MovieScene/MovieRenderPipeline/Source/MovieRenderPipelineCore/Private/MoviePipelineQueue.cpp
@@ -310,8 +310,22 @@ void UMoviePipelineExecutorJob::RefreshAllVariableAssignments()
 
 void UMoviePipelineExecutorJob::SetGraphPreset(const UMovieGraphConfig* InGraphPreset, const bool bUpdateVariableAssignments)
 {
+#if WITH_EDITOR
+	Modify();
+#endif
+	
 	GraphPreset = InGraphPreset;
 
+	// If the graph is being cleared out, also clear out all graphs on the job's shots. A shot cannot have a graph assigned to it while
+	// the parent job is using a legacy config.
+	if (InGraphPreset == nullptr)
+	{
+		for (const TObjectPtr<UMoviePipelineExecutorShot>& Shot : ShotInfo)
+		{
+			Shot->SetGraphPreset(nullptr);
+		}
+	}
+
 	if (bUpdateVariableAssignments)
 	{
 		RefreshAllVariableAssignments();
@@ -391,6 +405,10 @@ void UMoviePipelineExecutorJob::OnGraphPreSave(UObject* InObject, FObjectPreSave
 
 void UMoviePipelineExecutorShot::SetGraphPreset(const UMovieGraphConfig* InGraphPreset, const bool bUpdateVariableAssignments)
 {
+#if WITH_EDITOR
+	Modify();
+#endif
+	
 	GraphPreset = InGraphPreset;
 
 	if (bUpdateVariableAssignments)
diff --git a/Engine/Plugins/MovieScene/MovieRenderPipeline/Source/MovieRenderPipelineEditor/Private/Widgets/SMoviePipelineQueueEditor.cpp b/Engine/Plugins/MovieScene/MovieRenderPipeline/Source/MovieRenderPipelineEditor/Private/Widgets/SMoviePipelineQueueEditor.cpp
--- a/Engine/Plugins/MovieScene/MovieRenderPipeline/Source/MovieRenderPipelineEditor/Private/Widgets/SMoviePipelineQueueEditor.cpp
+++ b/Engine/Plugins/MovieScene/MovieRenderPipeline/Source/MovieRenderPipelineEditor/Private/Widgets/SMoviePipelineQueueEditor.cpp
@@ -335,6 +335,16 @@ public:
 		}
 	}
 
+	void OnClearGraph() const
+	{
+		if (UMoviePipelineExecutorJob* Job = WeakJob.Get())
+		{
+			FScopedTransaction Transaction(LOCTEXT("ClearJobGraph_Transaction", "Replace Graph with Config"));
+			
+			Job->SetGraphPreset(nullptr);
+		}
+	}
+
 	EVisibility GetPrimaryConfigModifiedVisibility() const
 	{
 		UMoviePipelineExecutorJob* Job = WeakJob.Get();
@@ -435,7 +445,6 @@ public:
 
 		UClass* ConfigType = Job->IsUsingGraphConfiguration() ? UMovieGraphConfig::StaticClass() : UMoviePipelinePrimaryConfig::StaticClass();
 		constexpr bool bIsShot = false;
-		const FExecuteAction ClearGraphAction = nullptr;
 		
 		return OnGenerateConfigPresetPickerMenuFromClass(
 			ConfigType,
@@ -445,7 +454,7 @@ public:
 			FExecuteAction::CreateRaw(this, &FMoviePipelineQueueJobTreeItem::OnPickNewPreset),
 			FExecuteAction::CreateRaw(this, &FMoviePipelineQueueJobTreeItem::OnReplaceWithRenderGraph),
 			FExecuteAction::CreateRaw(this, &FMoviePipelineQueueJobTreeItem::OnCreateNewGraphAndAssign),
-			ClearGraphAction
+			FExecuteAction::CreateRaw(this, &FMoviePipelineQueueJobTreeItem::OnClearGraph)
 			);
 	}
 
@@ -458,8 +467,8 @@ public:
 
 		FAssetPickerConfig AssetPickerConfig;
 		{
-			const FText NoAssetsFoundWarning = InClass->IsChildOf(UMoviePipelinePrimaryConfig::StaticClass())
-				? LOCTEXT("NoConfigs_Warning", "No Primary Configurations Found")
+			const FText NoAssetsFoundWarning = InClass->IsChildOf(UMoviePipelineConfigBase::StaticClass())
+				? LOCTEXT("NoConfigs_Warning", "No Configurations Found")
 				: LOCTEXT("NoGraphs_Warning", "No Render Graphs Found");
 			
 			AssetPickerConfig.SelectionMode = ESelectionMode::Single;
@@ -552,7 +561,7 @@ public:
 						LOCTEXT("ReplaceWithPreset_Label", "Replace with Preset"),
 						LOCTEXT("ReplaceWithPreset_Tooltip", "Replaces the current configuration with a new default (non-graph) config."),
 						FSlateIcon(),
-						FUIAction(InNewConfig),
+						FUIAction(InClearGraph),
 						NAME_None,
 						EUserInterfaceActionType::Button
 					);
@@ -1087,6 +1096,8 @@ struct FMoviePipelineShotItem : IMoviePipelineQueueTreeItem
 	{
 		if (UMoviePipelineExecutorShot* Shot = WeakShot.Get())
 		{
+			FScopedTransaction Transaction(LOCTEXT("ClearShotGraph_Transaction", "Replace Graph with Config"));
+			
 			Shot->SetGraphPreset(nullptr);
 		}
 	}
@@ -1117,7 +1128,7 @@ struct FMoviePipelineShotItem : IMoviePipelineQueueTreeItem
 	TSharedRef<SWidget> OnGenerateShotConfigPresetPickerMenu()
 	{
 		const UMoviePipelineExecutorShot* Shot = WeakShot.Get();
-		UClass* ConfigType = Shot && Shot->IsUsingGraphConfiguration() ? UMovieGraphConfig::StaticClass() : UMoviePipelinePrimaryConfig::StaticClass();
+		UClass* ConfigType = Shot && Shot->IsUsingGraphConfiguration() ? UMovieGraphConfig::StaticClass() : UMoviePipelineShotConfig::StaticClass();
 		constexpr bool bIsShot = true;
 		const FExecuteAction ReplaceWithGraph = nullptr;
 		
