MovieGraph: Warm Up Settings now default to off and the default graph has been updated to turn them back on. This makes the Warm Up Setting consistent with other nodes where if if a node is not included, it has no effect on the outcome.

diff --git a/Engine/Plugins/MovieScene/MovieRenderPipeline/Source/MovieRenderPipelineCore/Private/Graph/Nodes/MovieGraphWarmUpSettingNode.cpp b/Engine/Plugins/MovieScene/MovieRenderPipeline/Source/MovieRenderPipelineCore/Private/Graph/Nodes/MovieGraphWarmUpSettingNode.cpp
--- a/Engine/Plugins/MovieScene/MovieRenderPipeline/Source/MovieRenderPipelineCore/Private/Graph/Nodes/MovieGraphWarmUpSettingNode.cpp
+++ b/Engine/Plugins/MovieScene/MovieRenderPipeline/Source/MovieRenderPipelineCore/Private/Graph/Nodes/MovieGraphWarmUpSettingNode.cpp
@@ -4,8 +4,8 @@
 #include "Styling/AppStyle.h"
 
 UMovieGraphWarmUpSettingNode::UMovieGraphWarmUpSettingNode()
-	: NumWarmUpFrames(64)
-	, bEmulateMotionBlur(true)
+	: NumWarmUpFrames(0)
+	, bEmulateMotionBlur(false)
 {}
 
 EMovieGraphBranchRestriction UMovieGraphWarmUpSettingNode::GetBranchRestriction() const 
