MoviePipeline: Fixed an issue where the "Expose Property as Pin" feature in the graph would incorrectly surface properties that should not be exposable.

Specifically, hid the ShowFlags property on the renderer nodes, the Collection property on the Collection node, the ModifiersContainer property on the Modifier node, and the SamplingMethodClass on the Sampling Method node. Promoting the ShowFlags property to a variable could also cause a crash, and this is a temporary fix for that.

#rnx

diff --git a/Engine/Plugins/MovieScene/MovieRenderPipeline/Source/MovieRenderPipelineCore/Private/Graph/MovieGraphNode.cpp b/Engine/Plugins/MovieScene/MovieRenderPipeline/Source/MovieRenderPipelineCore/Private/Graph/MovieGraphNode.cpp
--- a/Engine/Plugins/MovieScene/MovieRenderPipeline/Source/MovieRenderPipelineCore/Private/Graph/MovieGraphNode.cpp
+++ b/Engine/Plugins/MovieScene/MovieRenderPipeline/Source/MovieRenderPipelineCore/Private/Graph/MovieGraphNode.cpp
@@ -481,6 +481,15 @@ TArray<FMovieGraphPropertyInfo> UMovieGraphNode::GetOverrideablePropertyInfo() c
 	
 	for (TFieldIterator<FProperty> PropertyIterator(GetClass()); PropertyIterator; ++PropertyIterator)
 	{
+		// Only allow overriding properties that are shown in the details panel (ie, has EditAnywhere/VisibleAnywhere property flags -- CPF_Edit). However,
+		// if the property has VisibleAnywhere (CPF_EditConst) then hide the property from being overrideable. Setting the property as VisibleAnywhere
+		// gives details customizations the opportunity to pick up the property for further processing, but because of the complicated nature of these
+		// properties, VisibleAnywhere serves as a filtering mechanism to prevent it from being overridden.
+		if (!PropertyIterator->HasAnyPropertyFlags(CPF_Edit) || PropertyIterator->HasAnyPropertyFlags(CPF_EditConst))
+		{
+			continue;
+		}
+		
 		if (UMovieGraphConfig::FindOverridePropertyForRealProperty(GetClass(), *PropertyIterator))
 		{
 			FMovieGraphPropertyInfo Info;
diff --git a/Engine/Plugins/MovieScene/MovieRenderPipeline/Source/MovieRenderPipelineCore/Public/Graph/Nodes/MovieGraphCollectionNode.h b/Engine/Plugins/MovieScene/MovieRenderPipeline/Source/MovieRenderPipelineCore/Public/Graph/Nodes/MovieGraphCollectionNode.h
--- a/Engine/Plugins/MovieScene/MovieRenderPipeline/Source/MovieRenderPipelineCore/Public/Graph/Nodes/MovieGraphCollectionNode.h
+++ b/Engine/Plugins/MovieScene/MovieRenderPipeline/Source/MovieRenderPipelineCore/Public/Graph/Nodes/MovieGraphCollectionNode.h
@@ -37,6 +37,6 @@ public:
 	UPROPERTY(EditAnywhere, BlueprintReadWrite, Category = Overrides, meta = (InlineEditConditionToggle))
 	uint8 bOverride_Collection : 1 = 1;	// The collection is customized in the details panel, so the override should always be enabled
 
-	UPROPERTY(BlueprintReadWrite, EditAnywhere, Category = "General", meta=(EditCondition="bOverride_Collection"))
+	UPROPERTY(BlueprintReadWrite, VisibleAnywhere, Category = "General", meta=(EditCondition="bOverride_Collection"))
 	TObjectPtr<UMovieGraphCollection> Collection;
 };
\ No newline at end of file
diff --git a/Engine/Plugins/MovieScene/MovieRenderPipeline/Source/MovieRenderPipelineRenderPasses/Public/Graph/Nodes/MovieGraphImagePassBaseNode.h b/Engine/Plugins/MovieScene/MovieRenderPipeline/Source/MovieRenderPipelineRenderPasses/Public/Graph/Nodes/MovieGraphImagePassBaseNode.h
--- a/Engine/Plugins/MovieScene/MovieRenderPipeline/Source/MovieRenderPipelineRenderPasses/Public/Graph/Nodes/MovieGraphImagePassBaseNode.h
+++ b/Engine/Plugins/MovieScene/MovieRenderPipeline/Source/MovieRenderPipelineRenderPasses/Public/Graph/Nodes/MovieGraphImagePassBaseNode.h
@@ -69,7 +69,7 @@ protected:
 	uint8 bOverride_ShowFlags : 1 = 1; //-V570
 
 	/** The show flags that should be active during a render for this node. */
-	UPROPERTY(EditAnywhere, BlueprintReadWrite, Category = "Show Flags")
+	UPROPERTY(VisibleAnywhere, BlueprintReadWrite, Category = "Show Flags")
 	TObjectPtr<UMovieGraphShowFlags> ShowFlags;
 
 	TArray<TUniquePtr<UE::MovieGraph::Rendering::FMovieGraphImagePassBase>> CurrentInstances;
