MoviePipeline: Fixed file name format autocomplete being partially broken (in the legacy UI as well as the graph), which was caused by a regression when modifying it to work in the graph.

Removed a redundant method which was causing some of the autocomplete logic to fire twice, resulting in unexpected behavior.

#rnx

diff --git a/Engine/Plugins/MovieScene/MovieRenderPipeline/Source/MovieRenderPipelineEditor/Private/Widgets/SMoviePipelineFormatTokenAutoCompleteBox.cpp b/Engine/Plugins/MovieScene/MovieRenderPipeline/Source/MovieRenderPipelineEditor/Private/Widgets/SMoviePipelineFormatTokenAutoCompleteBox.cpp
--- a/Engine/Plugins/MovieScene/MovieRenderPipeline/Source/MovieRenderPipelineEditor/Private/Widgets/SMoviePipelineFormatTokenAutoCompleteBox.cpp
+++ b/Engine/Plugins/MovieScene/MovieRenderPipeline/Source/MovieRenderPipelineEditor/Private/Widgets/SMoviePipelineFormatTokenAutoCompleteBox.cpp
@@ -34,7 +34,6 @@ void SMoviePipelineFormatTokenAutoCompleteBox::Construct(const FArguments& InArg
 					.ListItemsSource(&Suggestions)
 					.SelectionMode(ESelectionMode::Single)
 					.OnGenerateRow(this, &SMoviePipelineFormatTokenAutoCompleteBox::HandleSuggestionListViewGenerateRow)
-					.OnSelectionChanged(this, &SMoviePipelineFormatTokenAutoCompleteBox::HandleSuggestionListViewSelectionChanged)
 					.OnMouseButtonClick(this, &SMoviePipelineFormatTokenAutoCompleteBox::OnItemClicked)
 				]
 			]
@@ -111,9 +110,10 @@ void SMoviePipelineFormatTokenAutoCompleteBox::SetText(const FText& InText)
 	TextBox->GoTo(OriginalCursorLocation);
 }
 
-void SMoviePipelineFormatTokenAutoCompleteBox::OnItemClicked(TSharedPtr<FString> Item) const
+void SMoviePipelineFormatTokenAutoCompleteBox::OnItemClicked(TSharedPtr<FString> Item)
 {
 	ReplaceRelevantTextWithSuggestion(*Item);
+	CloseMenuAndReset();
 }
 
 void SMoviePipelineFormatTokenAutoCompleteBox::FindAutoCompletableTextAtPos(const FString& InWholeString, int32 InCursorPos, FString& OutStr, bool& bShowAutoComplete)
@@ -285,16 +285,4 @@ TSharedRef<ITableRow> SMoviePipelineFormatTokenAutoCompleteBox::HandleSuggestion
 				.Text(FText::FromString(SuggestionText))
 			]
 		];
-}
-
-void SMoviePipelineFormatTokenAutoCompleteBox::HandleSuggestionListViewSelectionChanged(TSharedPtr<FString> NewValue, ESelectInfo::Type SelectInfo)
-{
-	// This is called when clicking on an item and when navigating the menu via arrow keys, but we already
-	// handle arrow keys in OnKeyDown so we only want to handle mouse click here.
-	if (SelectInfo == ESelectInfo::Type::OnMouseClick)
-	{
-		// Trigger the auto-complete for the highlighted suggestion
-		ReplaceRelevantTextWithSuggestion(*NewValue);
-		CloseMenuAndReset();
-	}
-}
+}
\ No newline at end of file
diff --git a/Engine/Plugins/MovieScene/MovieRenderPipeline/Source/MovieRenderPipelineEditor/Private/Widgets/SMoviePipelineFormatTokenAutoCompleteBox.h b/Engine/Plugins/MovieScene/MovieRenderPipeline/Source/MovieRenderPipelineEditor/Private/Widgets/SMoviePipelineFormatTokenAutoCompleteBox.h
--- a/Engine/Plugins/MovieScene/MovieRenderPipeline/Source/MovieRenderPipelineEditor/Private/Widgets/SMoviePipelineFormatTokenAutoCompleteBox.h
+++ b/Engine/Plugins/MovieScene/MovieRenderPipeline/Source/MovieRenderPipelineEditor/Private/Widgets/SMoviePipelineFormatTokenAutoCompleteBox.h
@@ -32,7 +32,7 @@ public:
 
 	void SetText(const FText& InText);
 
-	void OnItemClicked(TSharedPtr<FString>) const;
+	void OnItemClicked(TSharedPtr<FString>);
 
 	static void FindAutoCompletableTextAtPos(const FString& InWholeString, int32 InCursorPos, FString& OutStr, bool& bShowAutoComplete);
 
@@ -48,8 +48,6 @@ public:
 
 	TSharedRef<ITableRow> HandleSuggestionListViewGenerateRow(TSharedPtr<FString> Text, const TSharedRef<STableViewBase>& OwnerTable) const;
 
-	void HandleSuggestionListViewSelectionChanged(TSharedPtr<FString> NewValue, ESelectInfo::Type SelectInfo);
-
 private:
 	TSharedPtr<SListView<TSharedPtr<FString>>> SuggestionListView;
 	TSharedPtr<SMultiLineEditableTextBox> TextBox;
