MovieGraph: Fixed FMovieGraphNamedResolution customization not supporting undo/redo. This fixes an issue where if you changed the Output Resolution on the UMovieGraphOutputSettings node and then undid, the change would not be undone.

diff --git a/Engine/Plugins/MovieScene/MovieRenderPipeline/Source/MovieRenderPipelineCore/Private/Graph/MovieGraphBlueprintLibrary.cpp b/Engine/Plugins/MovieScene/MovieRenderPipeline/Source/MovieRenderPipelineCore/Private/Graph/MovieGraphBlueprintLibrary.cpp
--- a/Engine/Plugins/MovieScene/MovieRenderPipeline/Source/MovieRenderPipelineCore/Private/Graph/MovieGraphBlueprintLibrary.cpp
+++ b/Engine/Plugins/MovieScene/MovieRenderPipeline/Source/MovieRenderPipelineCore/Private/Graph/MovieGraphBlueprintLibrary.cpp
@@ -361,12 +361,19 @@ FIntPoint UMovieGraphBlueprintLibrary::GetEffectiveOutputResolution(UMovieGraphE
 		RescaledOverscan = FMath::Clamp(CameraSetting->OverscanPercentage / 100.f, 0.f, 1.f);
 	}
 
-	
+	// We need to look at the Project Settings for the latest value for a given profile
+	FMovieGraphNamedResolution NamedResolution;
+	if (UMovieGraphBlueprintLibrary::IsNamedResolutionValid(OutputSetting->OutputResolution.ProfileName))
+	{
+		NamedResolution = UMovieGraphBlueprintLibrary::NamedResolutionFromProfile(OutputSetting->OutputResolution.ProfileName);
+	}
+	else
+	{
+		// Otherwise if it's not in the output settings as a valid profile, we use our internally stored one.
+		NamedResolution = OutputSetting->OutputResolution;
+	}
 
-	return UMoviePipelineBlueprintLibrary::Utility_GetEffectiveOutputResolution(
-		RescaledOverscan,
-		OutputSetting->GetSyncedOutputResolution()
-	);
+	return UMoviePipelineBlueprintLibrary::Utility_GetEffectiveOutputResolution(RescaledOverscan, NamedResolution.Resolution);
 }
 
 FText UMovieGraphBlueprintLibrary::GetJobName(const UMovieGraphPipeline* InMovieGraphPipeline)
diff --git a/Engine/Plugins/MovieScene/MovieRenderPipeline/Source/MovieRenderPipelineCore/Private/Graph/Nodes/MovieGraphGlobalOutputSettingNode.cpp b/Engine/Plugins/MovieScene/MovieRenderPipeline/Source/MovieRenderPipelineCore/Private/Graph/Nodes/MovieGraphGlobalOutputSettingNode.cpp
--- a/Engine/Plugins/MovieScene/MovieRenderPipeline/Source/MovieRenderPipelineCore/Private/Graph/Nodes/MovieGraphGlobalOutputSettingNode.cpp
+++ b/Engine/Plugins/MovieScene/MovieRenderPipeline/Source/MovieRenderPipelineCore/Private/Graph/Nodes/MovieGraphGlobalOutputSettingNode.cpp
@@ -16,31 +16,6 @@ UMovieGraphGlobalOutputSettingNode::UMovieGraphGlobalOutputSettingNode()
 	, CustomPlaybackRangeEndFrame(0)
 	, bFlushDiskWritesPerShot(false)
 {
-	// We prefer a 1080p resolution by default, but users may not have a preset that matches that. So we'll 
-	// look for a matching resolution if we can find one, otherwise we go to Custom set to 1920x1080.
-	const FMovieGraphNamedResolution* FoundResolution = nullptr;
-	if (const UMovieGraphProjectSettings* MovieGraphProjectSettings =
-		GetDefault<UMovieGraphProjectSettings>())
-	{
-		FoundResolution = Algo::FindByPredicate(
-			MovieGraphProjectSettings->DefaultNamedResolutions,
-			[](const FMovieGraphNamedResolution& Other)
-			{
-				return Other.Resolution == FIntPoint(1920, 1080);
-			});
-	}
-
-	// If we found one that was 1080p, regardless of name, use it.
-	if (FoundResolution)
-	{
-		OutputResolution = *FoundResolution;
-	}
-	else
-	{
-		// We didn't find one that was 1080p, just force a custom resolution.
-		OutputResolution = UMovieGraphBlueprintLibrary::NamedResolutionFromSize(1920, 1080);
-	}
-	
 	OutputDirectory.Path = TEXT("{project_dir}/Saved/MovieRenders/");
 }
 
@@ -50,14 +25,24 @@ void UMovieGraphGlobalOutputSettingNode::GetFormatResolveArgs(FMovieGraphResolve
 	OutMergedFormatArgs.FilenameArguments.Add(TEXT("project_dir"), ResolvedProjectDir);
 	OutMergedFormatArgs.FileMetadata.Add(TEXT("unreal/project_dir"), ResolvedProjectDir);
 
-	FIntPoint OutputResolutionAsIntPoint = GetSyncedOutputResolution();
+	// We need to look at the Project Settings for the latest value for a given profile
+	FMovieGraphNamedResolution NamedResolution;
+	if (UMovieGraphBlueprintLibrary::IsNamedResolutionValid(OutputResolution.ProfileName))
+	{
+		NamedResolution = UMovieGraphBlueprintLibrary::NamedResolutionFromProfile(OutputResolution.ProfileName);
+	}
+	else
+	{
+		// Otherwise if it's not in the output settings as a valid profile, we use our internally stored one.
+		NamedResolution = OutputResolution;
+	}
 	
 	// Resolution Arguments
 	{
-		FString Resolution = FString::Printf(TEXT("%d_%d"), OutputResolutionAsIntPoint.X, OutputResolutionAsIntPoint.Y);
+		FString Resolution = FString::Printf(TEXT("%d_%d"), NamedResolution.Resolution.X, NamedResolution.Resolution.Y);
 		OutMergedFormatArgs.FilenameArguments.Add(TEXT("output_resolution"), Resolution);
-		OutMergedFormatArgs.FilenameArguments.Add(TEXT("output_width"), FString::FromInt(OutputResolutionAsIntPoint.X));
-		OutMergedFormatArgs.FilenameArguments.Add(TEXT("output_height"), FString::FromInt(OutputResolutionAsIntPoint.Y));
+		OutMergedFormatArgs.FilenameArguments.Add(TEXT("output_width"), FString::FromInt(NamedResolution.Resolution.X));
+		OutMergedFormatArgs.FilenameArguments.Add(TEXT("output_height"), FString::FromInt(NamedResolution.Resolution.Y));
 	}
 
 	// We don't resolve the version here because that's handled on a per-file/shot basis
@@ -95,19 +80,3 @@ EMovieGraphBranchRestriction UMovieGraphGlobalOutputSettingNode::GetBranchRestri
 }
 #endif // WITH_EDITOR
 
-FIntPoint UMovieGraphGlobalOutputSettingNode::GetSyncedOutputResolution() const
-{
-	// Try to find a matching entry from Project Settings to stay in sync
-	const UMovieGraphProjectSettings* MovieGraphProjectSettings = GetDefault<UMovieGraphProjectSettings>();
-	if (ensureAlwaysMsgf(MovieGraphProjectSettings, TEXT("%hs: Failed to find UMovieGraphProjectSettings!"), __FUNCTION__))
-	{
-		const FMovieGraphNamedResolution* FoundResolution = MovieGraphProjectSettings->FindNamedResolutionForOption(OutputResolution.ProfileName);
-		if (FoundResolution)
-		{
-			return FoundResolution->Resolution;
-		}
-	}
-
-	// Otherwise return what we have saved locally
-	return OutputResolution.Resolution;
-}
diff --git a/Engine/Plugins/MovieScene/MovieRenderPipeline/Source/MovieRenderPipelineCore/Public/Graph/MovieGraphNamedResolution.h b/Engine/Plugins/MovieScene/MovieRenderPipeline/Source/MovieRenderPipelineCore/Public/Graph/MovieGraphNamedResolution.h
--- a/Engine/Plugins/MovieScene/MovieRenderPipeline/Source/MovieRenderPipelineCore/Public/Graph/MovieGraphNamedResolution.h
+++ b/Engine/Plugins/MovieScene/MovieRenderPipeline/Source/MovieRenderPipelineCore/Public/Graph/MovieGraphNamedResolution.h
@@ -56,4 +56,11 @@ public:
 	 * Predefined name for the 'custom' resolution option in the combobox
 	 */
 	inline static FName CustomEntryName = TEXT("Custom");
-};
\ No newline at end of file
+};
+
+
+/** Convert a FMovieGraphNamedResolution into a string */
+inline FString LexToString(const FMovieGraphNamedResolution InResolution)
+{
+	return FString::Printf(TEXT("%s [%d, %d]"), *InResolution.ProfileName.ToString(), InResolution.Resolution.X, InResolution.Resolution.Y);
+}
diff --git a/Engine/Plugins/MovieScene/MovieRenderPipeline/Source/MovieRenderPipelineCore/Public/Graph/Nodes/MovieGraphGlobalOutputSettingNode.h b/Engine/Plugins/MovieScene/MovieRenderPipeline/Source/MovieRenderPipelineCore/Public/Graph/Nodes/MovieGraphGlobalOutputSettingNode.h
--- a/Engine/Plugins/MovieScene/MovieRenderPipeline/Source/MovieRenderPipelineCore/Public/Graph/Nodes/MovieGraphGlobalOutputSettingNode.h
+++ b/Engine/Plugins/MovieScene/MovieRenderPipeline/Source/MovieRenderPipelineCore/Public/Graph/Nodes/MovieGraphGlobalOutputSettingNode.h
@@ -48,8 +48,6 @@ public:
 #endif
 	// ~UMovieGraphSettingNode Interface
 
-	FIntPoint GetSyncedOutputResolution() const;
-
 public:
 	UPROPERTY(EditAnywhere, BlueprintReadWrite, Category = Overrides, meta = (InlineEditConditionToggle))
 	uint8 bOverride_OutputDirectory : 1;
diff --git a/Engine/Plugins/MovieScene/MovieRenderPipeline/Source/MovieRenderPipelineEditor/Private/Customizations/Graph/MovieGraphNamedResolutionCustomization.cpp b/Engine/Plugins/MovieScene/MovieRenderPipeline/Source/MovieRenderPipelineEditor/Private/Customizations/Graph/MovieGraphNamedResolutionCustomization.cpp
new file mode 100644
--- /dev/null
+++ b/Engine/Plugins/MovieScene/MovieRenderPipeline/Source/MovieRenderPipelineEditor/Private/Customizations/Graph/MovieGraphNamedResolutionCustomization.cpp
@@ -0,0 +1,518 @@
+// Copyright Epic Games, Inc. All Rights Reserved.
+#include "MovieGraphNamedResolutionCustomization.h"
+
+#include "Graph/MovieGraphNamedResolution.h"
+#include "Graph/MovieGraphProjectSettings.h"
+#include "Graph/MovieGraphBlueprintLibrary.h"
+
+#include "MovieRenderPipelineCoreModule.h"
+
+#include "DetailLayoutBuilder.h"
+#include "DetailWidgetRow.h"
+#include "IDetailChildrenBuilder.h"
+#include "IPropertyTypeCustomization.h"
+#include "Widgets/Input/SButton.h"
+#include "Widgets/Input/SComboBox.h"
+#include "Widgets/Input/SNumericEntryBox.h"
+#include "Widgets/SBoxPanel.h"
+#include "ScopedTransaction.h"
+
+#define LOCTEXT_NAMESPACE "FMovieGraphNamedResolutionCustomization"
+
+FMovieGraphNamedResolutionCustomization::~FMovieGraphNamedResolutionCustomization()
+{
+	UMovieGraphProjectSettings* MovieGraphProjectSettings = GetMutableDefault<UMovieGraphProjectSettings>();
+	if (ensureAlwaysMsgf(MovieGraphProjectSettings, TEXT("%hs: Failed to find UMovieGraphProjectSettings!"), __FUNCTION__))
+	{
+		MovieGraphProjectSettings->OnSettingChanged().Remove(OnProjectSettingsModifiedHandle);
+	}
+}
+
+void FMovieGraphNamedResolutionCustomization::CustomizeHeader(TSharedRef<IPropertyHandle> InStructPropertyHandle, FDetailWidgetRow& HeaderRow, IPropertyTypeCustomizationUtils& CustomizationUtils)
+{
+	StructPropertyHandle = InStructPropertyHandle;
+
+	const TSharedPtr<IPropertyHandle> ResolutionPropertyHandle = InStructPropertyHandle->GetChildHandle(GET_MEMBER_NAME_CHECKED(FMovieGraphNamedResolution, Resolution));
+
+	// Cache the property handles as we use them to do all of our math.
+	ProfileNamePropertyHandle = InStructPropertyHandle->GetChildHandle(GET_MEMBER_NAME_CHECKED(FMovieGraphNamedResolution, ProfileName));
+	ResolutionXPropertyHandle = ResolutionPropertyHandle->GetChildHandle(GET_MEMBER_NAME_CHECKED(FIntPoint, X));
+	ResolutionYPropertyHandle = ResolutionPropertyHandle->GetChildHandle(GET_MEMBER_NAME_CHECKED(FIntPoint, Y));
+
+	// Set up option list for first time. If LoadedOption doesn't exist, use the first option in the list.
+	UpdateComboBoxOptions();
+
+	ComboBoxWidget = CreateComboBoxWidget();
+
+	// We start with aspect-ratio locking on (it's not stored in the struct but is just a UI element). This could
+	// be stored in a config later. We don't update the lock image here because it doesn't exist until the CustomizeChildren is called.
+	LockedAspectRatio = CalculateAspectRatio();
+
+	// Repopulate options when the project settings change
+	BindToOnProjectSettingsModified();
+
+	HeaderRow
+	.NameContent()
+	[
+		InStructPropertyHandle->CreatePropertyNameWidget() 	// Show "Resolution" as the row name
+	]
+	.ValueContent()
+	.HAlign(HAlign_Fill)
+	[
+		SNew(SHorizontalBox)
+
+		+ SHorizontalBox::Slot()
+		.Padding(0, 2, 8, 2)
+		.HAlign(HAlign_Fill)
+		[
+			ComboBoxWidget.ToSharedRef()
+		]
+	];
+}
+
+void FMovieGraphNamedResolutionCustomization::CustomizeChildren(TSharedRef<IPropertyHandle> InStructPropertyHandle, IDetailChildrenBuilder& StructBuilder, IPropertyTypeCustomizationUtils& StructCustomizationUtils)
+{
+	// Override the default CustomizeChildren to only create children for the X/Y (skipping ProfileName and Description)
+	AddCustomRowForResolutionAxis(EAxis::X, StructBuilder, ResolutionXPropertyHandle->GeneratePathToProperty(), ResolutionXPropertyHandle->CreatePropertyNameWidget(), MakeAspectRatioLockWidget());
+	AddCustomRowForResolutionAxis(EAxis::Y, StructBuilder, ResolutionYPropertyHandle->GeneratePathToProperty(), ResolutionYPropertyHandle->CreatePropertyNameWidget());
+	UpdateAspectRatioLockImage();
+}
+
+TSharedRef<SWidget> FMovieGraphNamedResolutionCustomization::MakeWidgetForOptionInDropdown(FName InOption) const
+{
+	// This is called when the drop-down is opened to create a widget for each row in the drop-down
+	const bool bSuppressResolution = InOption == FMovieGraphNamedResolution::CustomEntryName;
+
+	return
+		SNew(STextBlock)
+		.Text(this, &FMovieGraphNamedResolutionCustomization::GetWidgetTextForOption, InOption, bSuppressResolution)
+		.ToolTipText(this, &FMovieGraphNamedResolutionCustomization::GetWidgetTooltipTextForOption, InOption)
+		.Font(this, &FMovieGraphNamedResolutionCustomization::GetWidgetFontForOptions);
+}
+
+TSharedRef<SWidget> FMovieGraphNamedResolutionCustomization::MakeWidgetForSelectedItem() const
+{
+	// This is called once when the customization is created to create the widget shown at the "root"
+	// of the combo-box (not the individual items in the dropdown).
+	return
+		SNew(STextBlock)
+		.Text(this, &FMovieGraphNamedResolutionCustomization::GetWidgetTextForCurrent)
+		.ToolTipText(this, &FMovieGraphNamedResolutionCustomization::GetWidgetTooltipTextForCurrent)
+		.Font(this, &FMovieGraphNamedResolutionCustomization::GetWidgetFontForOptions);
+}
+
+TSharedRef<SWidget> FMovieGraphNamedResolutionCustomization::CreateComboBoxWidget()
+{
+	FName LoadedOption = NAME_None;
+	ProfileNamePropertyHandle->GetValue(LoadedOption);
+
+	return
+		SNew(SComboBox<FName>)
+		.OnComboBoxOpening(this, &FMovieGraphNamedResolutionCustomization::UpdateComboBoxOptions)
+		.OptionsSource(&ComboBoxOptions)
+		.OnSelectionChanged(this, &FMovieGraphNamedResolutionCustomization::OnComboBoxSelectionChanged)
+		.OnGenerateWidget(this, &FMovieGraphNamedResolutionCustomization::MakeWidgetForOptionInDropdown)
+		.InitiallySelectedItem(LoadedOption)
+		[
+			MakeWidgetForSelectedItem()
+		];
+}
+
+FText FMovieGraphNamedResolutionCustomization::GetWidgetTextForOption(const FName InOption, const bool bSuppressResolution) const
+{
+	// We always try to look up the latest named resolution from the project settings to ensure
+	// what is displayed in the UI is in sync with the project settings.
+	FMovieGraphNamedResolution NamedResolution;
+
+	if (UMovieGraphBlueprintLibrary::IsNamedResolutionValid(InOption))
+	{
+		NamedResolution = UMovieGraphBlueprintLibrary::NamedResolutionFromProfile(InOption);
+	}
+	else
+	{
+		FIntPoint CurrentInternalSize;
+		ResolutionXPropertyHandle->GetValue(CurrentInternalSize.X);
+		ResolutionYPropertyHandle->GetValue(CurrentInternalSize.Y);
+
+		// Doing this via the BlueprintLibrary just for consistency since it creates with the correct "Custom" name.
+		NamedResolution = UMovieGraphBlueprintLibrary::NamedResolutionFromSize(CurrentInternalSize.X, CurrentInternalSize.Y);
+	}
+
+	FNumberFormattingOptions FormattingOptions;
+	FormattingOptions.UseGrouping = false;
+
+	if (bSuppressResolution)
+	{
+		return FText::Format(LOCTEXT("NamedResolutionComboboxItemWidgetText_NoResolution", "{0}"),
+			FText::FromName(NamedResolution.ProfileName));
+	}
+	else
+	{
+		return FText::Format(LOCTEXT("NamedResolutionComboboxItemWidgetText", "{0} - {1}x{2}"),
+			FText::FromName(NamedResolution.ProfileName),
+			FText::AsNumber(NamedResolution.Resolution.X, &FormattingOptions),
+			FText::AsNumber(NamedResolution.Resolution.Y, &FormattingOptions));
+	}
+}
+
+FText FMovieGraphNamedResolutionCustomization::GetWidgetTooltipTextForOption(const FName InOption) const
+{
+	// We always try to look up the latest named resolution from the project settings to ensure
+	// what is displayed in the UI is in sync with the project settings.
+	FMovieGraphNamedResolution NamedResolution;
+
+	if (UMovieGraphBlueprintLibrary::IsNamedResolutionValid(InOption))
+	{
+		NamedResolution = UMovieGraphBlueprintLibrary::NamedResolutionFromProfile(InOption);
+	}
+	else
+	{
+		FIntPoint CurrentInternalSize;
+		ResolutionXPropertyHandle->GetValue(CurrentInternalSize.X);
+		ResolutionYPropertyHandle->GetValue(CurrentInternalSize.Y);
+
+		// Doing this via the BlueprintLibrary just for consistency since it creates with the correct "Custom" name.
+		NamedResolution = UMovieGraphBlueprintLibrary::NamedResolutionFromSize(CurrentInternalSize.X, CurrentInternalSize.Y);
+	}
+
+	return FText::FromString(NamedResolution.Description);
+}
+
+FText FMovieGraphNamedResolutionCustomization::GetWidgetTextForCurrent() const
+{
+	FName LoadedOption = NAME_None;
+	if (ProfileNamePropertyHandle->GetValue(LoadedOption) == FPropertyAccess::Result::Success)
+	{
+		// Will automatically handle missing/invalid options.
+		const bool bSuppressResolution = false;
+		return GetWidgetTextForOption(LoadedOption, bSuppressResolution);
+	}
+
+	return FText();
+}
+
+FText FMovieGraphNamedResolutionCustomization::GetWidgetTooltipTextForCurrent() const
+{
+	FName LoadedOption = NAME_None;
+	if (ProfileNamePropertyHandle->GetValue(LoadedOption) == FPropertyAccess::Result::Success)
+	{
+		// Will automatically handle missing/invalid options.
+		return GetWidgetTooltipTextForOption(LoadedOption);
+	}
+
+	return FText();
+}
+
+void FMovieGraphNamedResolutionCustomization::UpdateComboBoxOptions()
+{
+	ComboBoxOptions.Reset(ComboBoxOptions.Num());
+
+	const UMovieGraphProjectSettings* MovieGraphProjectSettings = GetDefault<UMovieGraphProjectSettings>();
+	if (ensureAlwaysMsgf(MovieGraphProjectSettings, TEXT("%hs: Failed to find UMovieGraphProjectSettings!"), __FUNCTION__))
+	{
+		for (const FMovieGraphNamedResolution& Resolution : MovieGraphProjectSettings->DefaultNamedResolutions)
+		{
+			ComboBoxOptions.AddUnique(Resolution.ProfileName);
+		}
+	}
+
+	// We need to always add "Custom" so there's something there if the project settings are empty,
+	// and because without it, after changing the value to "Custom" (in the underlying resolution)
+	// the drop-down doesn't work the first time you click on it due to a ComboBox issue.
+	ComboBoxOptions.AddUnique(FMovieGraphNamedResolution::CustomEntryName);
+}
+
+void FMovieGraphNamedResolutionCustomization::BindToOnProjectSettingsModified()
+{
+	UMovieGraphProjectSettings* MovieGraphProjectSettings = GetMutableDefault<UMovieGraphProjectSettings>();
+	if (ensureAlwaysMsgf(MovieGraphProjectSettings, TEXT("%hs: Unable to get UMovieGraphProjectSettings."), __FUNCTION__))
+	{
+		OnProjectSettingsModifiedHandle = MovieGraphProjectSettings->OnSettingChanged().AddSP(this, &FMovieGraphNamedResolutionCustomization::OnProjectSettingsChanged);
+	}
+}
+
+void FMovieGraphNamedResolutionCustomization::UpdateAspectRatioLockImage()
+{
+	// The widget will only exist if the dropdown is expanded, so it's okay
+	// to call this without that. But it should be called again after the widget is created.
+	if (AspectRatioLockImage)
+	{
+		AspectRatioLockImage->SetImage(GetAspectRatioLockBrush());
+		AspectRatioLockImage->SetToolTipText(GetAspectRatioLockTooltipText());
+	}
+}
+
+void FMovieGraphNamedResolutionCustomization::ToggleAspectRatioLock()
+{
+	if (LockedAspectRatio.IsSet())
+	{
+		LockedAspectRatio.Reset();
+	}
+	else
+	{
+		LockedAspectRatio = CalculateAspectRatio();
+	}
+
+	UpdateAspectRatioLockImage();
+}
+
+const FSlateBrush* FMovieGraphNamedResolutionCustomization::GetAspectRatioLockBrush() const
+{
+	return LockedAspectRatio.IsSet() ? FAppStyle::GetBrush(TEXT("Icons.Link")) : FAppStyle::GetBrush(TEXT("Icons.Unlink"));
+}
+
+FText FMovieGraphNamedResolutionCustomization::GetAspectRatioLockTooltipText() const
+{
+	return LockedAspectRatio.IsSet() ?
+		FText::Format(LOCTEXT("LockedAspectRatioTooltipFormat", "Click to unlock Aspect Ratio ({0})"), FText::AsNumber(LockedAspectRatio.GetValue())) :
+		LOCTEXT("UnlockedAspectRatioTooltip", "Click to lock the Aspect Ratio");
+}
+
+TSharedRef<SWidget> FMovieGraphNamedResolutionCustomization::MakeAspectRatioLockWidget()
+{
+	return SNew(SButton)
+	.OnClicked_Lambda([this]()
+	{
+		ToggleAspectRatioLock();
+		
+		return FReply::Handled();
+	})
+	.ContentPadding(FMargin(0, 0, 4, 0))
+	.ButtonStyle(FAppStyle::Get(), "NoBorder")
+	[
+		SAssignNew(AspectRatioLockImage, SImage)
+		.ColorAndOpacity(FSlateColor::UseForeground())
+	];
+}
+
+float FMovieGraphNamedResolutionCustomization::CalculateAspectRatio()
+{
+	float X = 1;
+	float Y = 1;
+	
+	FIntPoint EffectiveResolution = GetEffectiveResolution();
+	X = (float)EffectiveResolution.X;
+	Y = (float)EffectiveResolution.Y;
+	
+	return X / Y;
+}
+
+namespace MovieGraph::Private
+{
+	int32 RoundToNearestEvenNumber(double InNumberToRound)
+	{
+		// Round to nearest even
+		const int32 Floored = FMath::FloorToInt(InNumberToRound);
+		const int32 Ceiled = FMath::CeilToInt(InNumberToRound);
+
+		return Floored % 2 == 0 ? Floored : Ceiled;
+	}
+}
+
+FIntPoint FMovieGraphNamedResolutionCustomization::GetEffectiveResolution() const
+{
+	// To get the resolution that we display, we need to look it up in the Project Settings first,
+	// because they might have changed the project setting value since the last time the struct was changed.
+	// If we can't find it in the project setting (or it is Custom) then we return the values inside the struct.
+	FIntPoint Resolution;
+	FName ProfileName;
+	ProfileNamePropertyHandle->GetValue(ProfileName);
+
+	if (UMovieGraphBlueprintLibrary::IsNamedResolutionValid(ProfileName))
+	{
+		FMovieGraphNamedResolution ResolutionFromPreset = UMovieGraphBlueprintLibrary::NamedResolutionFromProfile(ProfileName);
+		Resolution = ResolutionFromPreset.Resolution;
+	}
+	else
+	{
+		// If it's not a valid profile, we use the stored value.
+		ResolutionXPropertyHandle->GetValue(Resolution.X);
+		ResolutionYPropertyHandle->GetValue(Resolution.Y);
+	}
+
+	return Resolution;
+}
+
+int32 FMovieGraphNamedResolutionCustomization::GetCurrentlySelectedResolutionByAxis(const EAxis::Type Axis)
+{
+	if (!ResolutionXPropertyHandle || !ResolutionYPropertyHandle || !ProfileNamePropertyHandle)
+	{
+		return INDEX_NONE;
+	}
+
+	FIntPoint Resolution = GetEffectiveResolution();
+	return Axis == EAxis::X ? Resolution.X : Resolution.Y;
+}
+
+void FMovieGraphNamedResolutionCustomization::OnCustomSliderBeginMovement()
+{
+	// Wrap all interactive changes into a single commit while dragging.
+	InteractiveResolutionEditTransaction = MakeUnique<FScopedTransaction>(LOCTEXT("ChangeResolution_Transaction", "Edit Resolution"));
+}
+
+void FMovieGraphNamedResolutionCustomization::OnCustomSliderEndMovement(uint32 NewValue)
+{
+	// Reset our open transaction to commit it now that they're done interactively changing values.
+	InteractiveResolutionEditTransaction.Reset();
+}
+
+void FMovieGraphNamedResolutionCustomization::OnCustomSliderValueChanged(uint32 NewValue, EAxis::Type Axis)
+{
+	if (!ResolutionXPropertyHandle || !ResolutionYPropertyHandle || !ProfileNamePropertyHandle)
+	{
+		return;
+	}
+
+	// If they directly edit the value (via typing it in and hitting enter) OnCustomSliderValueChanged
+	// will be called, so we need to make a transaction here. This function is also called when being
+	// dragged, so we use OnCustomSliderBeginMovement/OnCustomSliderEndMovement to create a transaction
+	// which wraps all the smaller transactions to group them together into one undo/redo event.
+	FScopedTransaction Transaction(LOCTEXT("ChangeResolution_Transaction", "Edit Resolution"));
+
+	// Check to see if we need to convert from an existing profile resolution to a custom one,
+	// otherwise the UI won't update with our new values.
+	FName ProfileName;
+	if (ProfileNamePropertyHandle->GetValue(ProfileName) == FPropertyAccess::Success)
+	{
+		if (ProfileName != FMovieGraphNamedResolution::CustomEntryName)
+		{
+			// Convert us from our existing profile to be "Custom".
+			ProfileNamePropertyHandle->SetValue(FMovieGraphNamedResolution::CustomEntryName);
+		}
+	}
+
+	// Get our XY resolution out of the properties
+	FIntPoint NewResolution;
+	ResolutionXPropertyHandle->GetValue(NewResolution.X);
+	ResolutionYPropertyHandle->GetValue(NewResolution.Y);
+	
+	// Update the correct axis on the Named Resolution struct, potentially constrained by aspect ratio.
+	switch (Axis)
+	{
+		case EAxis::X: 
+		{ 
+			NewResolution.X = NewValue;
+			if (LockedAspectRatio.IsSet())
+			{
+				const double NewYResolution = (double)NewResolution.X / LockedAspectRatio.GetValue();
+				NewResolution.Y = MovieGraph::Private::RoundToNearestEvenNumber(NewYResolution);
+			}
+			break;
+		}
+		case EAxis::Y: 
+		{ 
+			NewResolution.Y = NewValue;
+			if (LockedAspectRatio.IsSet())
+			{
+				const double NewXResolution = (double)NewResolution.Y * LockedAspectRatio.GetValue();
+				NewResolution.X = MovieGraph::Private::RoundToNearestEvenNumber(NewXResolution);
+			}
+			break; 
+		}
+	}
+
+	// Update both X and Y as they both may have changed due to aspect ratio locking.
+	ResolutionXPropertyHandle->SetValue(NewResolution.X);
+	ResolutionYPropertyHandle->SetValue(NewResolution.Y);
+}
+
+void FMovieGraphNamedResolutionCustomization::AddCustomRowForResolutionAxis(EAxis::Type Axis, IDetailChildrenBuilder& StructBuilder, const FString& FilterString, TSharedRef<SWidget> NameContentWidget, TSharedPtr<SWidget> AspectRatioLockExtensionWidget)
+{
+	const FSlateFontInfo PropertyFontStyle = FAppStyle::GetFontStyle( TEXT("PropertyWindow.NormalFont") );
+	const TOptional<uint32> PropertyMaxValue = TOptional<uint32>();
+
+	FDetailWidgetRow& CustomRow = StructBuilder.AddCustomRow(FText::FromString(FilterString))
+	.NameContent()
+	[
+		NameContentWidget
+	]
+	.ValueContent()
+	.HAlign(HAlign_Fill)
+	[
+		// We make our own widget so that we can edit two variables at once in the callback (for aspect-ratio linked X and Y)
+		SNew(SBox)
+		// Maintain spacing when there would be no extension widget
+		.Padding(0, 0, Axis == EAxis::Y ? 24 /* Assuming image is 16x16 + 8 padding */ : 0, 0) 
+		[
+			SNew(SNumericEntryBox<uint32>)
+			.Font(PropertyFontStyle)
+			.AllowSpin(true)
+			.MinSliderValue(2)
+			.MinValue(2)
+			.MaxValue(PropertyMaxValue)
+			.MaxSliderValue(PropertyMaxValue)
+			.Delta(2)
+			.Value_Lambda([this, Axis] (){ return GetCurrentlySelectedResolutionByAxis(Axis); }) // Lambda to bypass const requirement
+			.OnValueChanged(this, &FMovieGraphNamedResolutionCustomization::OnCustomSliderValueChanged, Axis)
+			.OnBeginSliderMovement(this, &FMovieGraphNamedResolutionCustomization::OnCustomSliderBeginMovement)
+			.OnEndSliderMovement(this, &FMovieGraphNamedResolutionCustomization::OnCustomSliderEndMovement)
+		]
+	]
+	.OverrideResetToDefault(FResetToDefaultOverride::Hide()); // Prevent showing Reset to Default
+
+	if (AspectRatioLockExtensionWidget.IsValid())
+	{
+		CustomRow
+		.ExtensionContent()
+		[
+			AspectRatioLockExtensionWidget.ToSharedRef()
+		];
+	}
+}
+
+
+void FMovieGraphNamedResolutionCustomization::OnProjectSettingsChanged(UObject*, FPropertyChangedEvent& PropertyChangedEvent)
+{
+	if (PropertyChangedEvent.MemberProperty->GetFName() == GET_MEMBER_NAME_CHECKED(UMovieGraphProjectSettings, DefaultNamedResolutions))
+	{
+		UpdateComboBoxOptions();
+	}
+}
+
+void FMovieGraphNamedResolutionCustomization::OnComboBoxSelectionChanged(const FName NewValue, ESelectInfo::Type InType)
+{
+	// Ensure the value selected was valid, if so, update our internal copy to match.
+	FMovieGraphNamedResolution NamedResolution;
+	if (UMovieGraphBlueprintLibrary::IsNamedResolutionValid(NewValue))
+	{
+		// Fetch the latest values from the Project Settings to update the internal copy of the resolution with.
+		NamedResolution = UMovieGraphBlueprintLibrary::NamedResolutionFromProfile(NewValue);
+	}
+	else
+	{
+		// If it's not a valid resolution, we still want to run this code to switch it to "Custom", otherwise clicking "Custom"
+		// in the dropdown doesn't do anything.
+		FIntPoint CurrentInternalSize;
+		ResolutionXPropertyHandle->GetValue(CurrentInternalSize.X);
+		ResolutionYPropertyHandle->GetValue(CurrentInternalSize.Y);
+
+		// Doing this via the BlueprintLibrary just for consistency since it creates with the correct "Custom" name.
+		NamedResolution = UMovieGraphBlueprintLibrary::NamedResolutionFromSize(CurrentInternalSize.X, CurrentInternalSize.Y);
+	}
+
+	const TSharedPtr<IPropertyHandle> PinnedStructPropertyHandle = StructPropertyHandle.Pin();
+	if (!PinnedStructPropertyHandle)
+	{
+		return;
+	}
+
+	void* StructData = nullptr;
+	if (PinnedStructPropertyHandle->GetValueData(StructData) == FPropertyAccess::Success && StructData != nullptr)
+	{
+		// We have to call all of these by hand because we're directly editing the memory and not using property handles.
+		FScopedTransaction ScopedTransaction(LOCTEXT("ChangeResolution_Transaction", "Assign Preset"));
+		PinnedStructPropertyHandle->NotifyPreChange();
+
+		PinnedStructPropertyHandle->EnumerateRawData([NamedResolution](void* RawData, const int32 /*DataIndex*/, const int32 /*NumDatas*/)
+			{
+				FMovieGraphNamedResolution* OldData = static_cast<FMovieGraphNamedResolution*>(RawData);
+				FMovieGraphNamedResolution::StaticStruct()->CopyScriptStruct(RawData, &NamedResolution);
+				return true;
+			});
+
+		PinnedStructPropertyHandle->NotifyPostChange(EPropertyChangeType::ValueSet);
+		PinnedStructPropertyHandle->NotifyFinishedChangingProperties();
+	}
+}
+
+#undef LOCTEXT_NAMESPACE
diff --git a/Engine/Plugins/MovieScene/MovieRenderPipeline/Source/MovieRenderPipelineEditor/Private/Customizations/Graph/MovieGraphNamedResolutionCustomization.h b/Engine/Plugins/MovieScene/MovieRenderPipeline/Source/MovieRenderPipelineEditor/Private/Customizations/Graph/MovieGraphNamedResolutionCustomization.h
--- a/Engine/Plugins/MovieScene/MovieRenderPipeline/Source/MovieRenderPipelineEditor/Private/Customizations/Graph/MovieGraphNamedResolutionCustomization.h
+++ b/Engine/Plugins/MovieScene/MovieRenderPipeline/Source/MovieRenderPipelineEditor/Private/Customizations/Graph/MovieGraphNamedResolutionCustomization.h
@@ -2,20 +2,11 @@
 #pragma once
 
 #include "Graph/MovieGraphNamedResolution.h"
-#include "Graph/MovieGraphProjectSettings.h"
-
-#include "MovieRenderPipelineCoreModule.h"
-
-#include "DetailLayoutBuilder.h"
-#include "DetailWidgetRow.h"
 #include "IDetailChildrenBuilder.h"
+#include "DetailLayoutBuilder.h"
 #include "IPropertyTypeCustomization.h"
-#include "Widgets/Input/SButton.h"
-#include "Widgets/Input/SComboBox.h"
-#include "Widgets/Input/SNumericEntryBox.h"
-#include "Widgets/SBoxPanel.h"
 
-#define LOCTEXT_NAMESPACE "FMovieGraphNamedResolutionCustomization"
+class FScopedTransaction;
 
 /* Customizes how named resolutions are displayed in the details pane. */
 class MOVIERENDERPIPELINEEDITOR_API FMovieGraphNamedResolutionCustomization : public IPropertyTypeCustomization
@@ -26,560 +17,126 @@ public:
 		return MakeShared<FMovieGraphNamedResolutionCustomization>();
 	}
 
-	virtual ~FMovieGraphNamedResolutionCustomization() override
-	{
-		UMovieGraphProjectSettings* MovieGraphProjectSettings = GetMutableDefault<UMovieGraphProjectSettings>();
-		if (ensureAlwaysMsgf(MovieGraphProjectSettings, TEXT("%hs: Failed to find UMovieGraphProjectSettings!"), __FUNCTION__))
-		{
-			MovieGraphProjectSettings->OnSettingChanged().Remove(OnProjectSettingsModifiedHandle);
-		}
-	}
+	virtual ~FMovieGraphNamedResolutionCustomization() override;
 
 protected:
 
-	bool IsCurrentOptionValid() const
-	{
-		return ComboBoxOptions.Contains(CurrentOption);
-	}
-	
-	/**
-	 * Checks to see if CurrentOption is None, and if so calls RepopulateOptions() which sets CurrentOption to the first option.
-	 * Otherwise just returns CurrentOption.
-	 * @return The current option selected in the combobox
-	 */
-	const FName& GetCurrentOptionAssigningIfNeeded()
-	{
-		if (!IsCurrentOptionValid())
-		{
-			RepopulateOptions();
-		}
-		return CurrentOption;
-	}
-
-	void SetCurrentOptionAndCacheTooltipText(const FName& InOption)
-	{
-		if (ComboBox)
-		{
-			ComboBox->SetSelectedItem(InOption);
-		}
-		CurrentOption = InOption;
-		CachedSelectedOptionTooltipText = GetOptionTooltipText(CurrentOption);
-	}
-
 	/**
-	 * Adds all ProfileNames from DefaultNamedResolutions to the combobox,
-	 * and sets CurrentOption to the first option if it's None or isn't a known option.
-	 */
-	void RepopulateOptions()
-	{
-		ComboBoxOptions.Reset(ComboBoxOptions.Num());
-
-		const UMovieGraphProjectSettings* MovieGraphProjectSettings = GetDefault<UMovieGraphProjectSettings>();
-		if (ensureAlwaysMsgf(MovieGraphProjectSettings, TEXT("%hs: Failed to find UMovieGraphProjectSettings!"), __FUNCTION__))
-		{
-			for (const FMovieGraphNamedResolution& Resolution : MovieGraphProjectSettings->DefaultNamedResolutions)
-			{
-				ComboBoxOptions.AddUnique(Resolution.ProfileName);
-			}
-		}
-
-		ComboBoxOptions.AddUnique(FMovieGraphNamedResolution::CustomEntryName);
-
-		if (!IsCurrentOptionValid())
-		{
-			SetCurrentOptionAndCacheTooltipText(FMovieGraphNamedResolution::CustomEntryName);
-		}
-	}
+	* Update ComboBoxOptions with the latest values in the Project Settings,
+	* and adds the "Custom" entry as well.
+	*/
+	void UpdateComboBoxOptions();
 
-	/**
-	 * Returns the tooltip text for the given option by looking up its Description
-	 * in DefaultNamedResolutions from Project Settings.
-	 * If a match isn't found, the tooltip text will be the name of the option.
-	 * @param InOption The name of the option.
-	 * @return The tooltip text for the option.
-	 */
-	FText GetOptionTooltipText(FName InOption) const
-	{
-		if (InOption.IsEqual(CurrentOption))
-		{
-			return CachedSelectedOptionTooltipText;
-		}
-		if (const FMovieGraphNamedResolution* Match = FindNamedResolutionForOption(InOption))
-		{
-			return FText::FromString(Match->Description);
-		}
-		return FText::FromName(InOption);
-	}
+	/*
+	* Gets the text that should be displayed in the widget for the specified option.
+	* In the form of "{name} - {x}x{y}" unless bSuppressResolution is true, at which
+	* point it will just be in the form of "{name}"
+	*/
+	FText GetWidgetTextForOption(const FName InOption, const bool bSuppressResolution) const;
+	/** The above, but for the currently selected (ie: "root" of combo box). */
+	FText GetWidgetTextForCurrent() const;
+	/** Gets the description property for the specified option */
+	FText GetWidgetTooltipTextForOption(const FName InOption) const;
+	/** The above, but for the currently selected (ie: "root" of combo box). */
+	FText GetWidgetTooltipTextForCurrent() const;
 
-	/**
-	 * Tries to find a match for the given option in DefaultNamedResolutions from Project Settings.
-	 * @param InOption The name of the option.
-	 * @return A raw pointer to an FMovieGraphNamedResolution if a match is found, nullptr otherwise.
-	 */
-	const FMovieGraphNamedResolution* FindNamedResolutionForOption(const FName& InOption) const
-	{
-		// Return the custom entry
-		if (InOption.IsEqual(FMovieGraphNamedResolution::CustomEntryName))
-		{
-			return &CustomEntry;
-		}
-
-		// Find a matching custom entry from Project Settings
-		const UMovieGraphProjectSettings* MovieGraphProjectSettings = GetDefault<UMovieGraphProjectSettings>();
-		if (ensureAlwaysMsgf(MovieGraphProjectSettings, TEXT("%hs: Failed to find UMovieGraphProjectSettings!"), __FUNCTION__))
-		{
-			return MovieGraphProjectSettings->FindNamedResolutionForOption(InOption);
-		}
-
-		UE_LOG(LogMovieRenderPipeline, Error, TEXT("%hs: No match found for FMovieGraphNamedResolution %s!"), __FUNCTION__, *InOption.ToString());
-		return nullptr;
-	}
-
-	FText GetWidgetTextForOption(const FName& InOption) const
-	{
-		const FMovieGraphNamedResolution* NamedResolution = FindNamedResolutionForOption(InOption);
-		if (!ensureAlwaysMsgf(NamedResolution, TEXT("%hs: Failed to find FMovieGraphNamedResolution for option %s!"), __FUNCTION__, *InOption.ToString()))
-		{
-			return LOCTEXT("UnknownNamedResolutionComboboxItemWidgetText", "Unknown Named Resolution");
-		}
-
-
-		FNumberFormattingOptions FormattingOptions;
-		FormattingOptions.UseGrouping = false;
-		
-		return FText::Format(LOCTEXT("NamedResolutionComboboxItemWidgetText", "{0} - {1}x{2}"),
-				FText::FromName(InOption),
-				FText::AsNumber(NamedResolution->Resolution.X, &FormattingOptions),
-				FText::AsNumber(NamedResolution->Resolution.Y, &FormattingOptions));
-	}
 
 	FSlateFontInfo GetWidgetFontForOptions() const
 	{
 		return IDetailLayoutBuilder::GetDetailFont();
 	}
 
-	/**
-	 * Generates a widget for the given option with the option name as STextBlock with tooltip text .
-	 * @param InOption The name of the option.
-	 * @return A shared reference to the generated widget.
-	 */
-	TSharedRef<SWidget> MakeWidgetForOption(FName InOption) const
-	{
-		return
-			SNew(STextBlock)
-			.Text(GetWidgetTextForOption(InOption))
-			.ToolTipText(this, &FMovieGraphNamedResolutionCustomization::GetOptionTooltipText, InOption)
-			.Font(this, &FMovieGraphNamedResolutionCustomization::GetWidgetFontForOptions);
-		;
-	}
+	/** Called to generate a widget for each entry in the dropdown when you open the dropdown. */
+	TSharedRef<SWidget> MakeWidgetForOptionInDropdown(FName InOption) const;
+	/** Called to generate the widget for the "root" of the combo box that is always visible. */
+	TSharedRef<SWidget> MakeWidgetForSelectedItem() const;
 
-	TSharedRef<SWidget> CreateOptionComboBox()
-	{
-		ComboBox = 
-			SNew(SComboBox<FName>)
-			.OnComboBoxOpening(this, &FMovieGraphNamedResolutionCustomization::RepopulateOptions)
-			.OptionsSource(&ComboBoxOptions)
-			.OnSelectionChanged(this, &FMovieGraphNamedResolutionCustomization::OnComboBoxSelectionChanged)
-			.OnGenerateWidget(this, &FMovieGraphNamedResolutionCustomization::MakeWidgetForOption)
-			.InitiallySelectedItem(GetCurrentOptionAssigningIfNeeded())
-			[
-				// We specifically don't use MakeWidgetForOption here for performance reasons
-				// as it will need to be called each frame to stay current (especially for 'Custom')
-				SNew(STextBlock)
-				.Text_Lambda([this]()
-				{
-					return GetWidgetTextForOption(GetCurrentOptionAssigningIfNeeded());
-				})
-				.ToolTipText_Lambda([this]()
-				{
-					return CachedSelectedOptionTooltipText;
-				})
-				.Font(this, &FMovieGraphNamedResolutionCustomization::GetWidgetFontForOptions)
-			];
-		return ComboBox.ToSharedRef(); 
-	}
+	/** Create the actual ComboBox Widget */
+	TSharedRef<SWidget> CreateComboBoxWidget();
 
-	void BindToOnProjectSettingsModified()
-	{
-		UMovieGraphProjectSettings* MovieGraphProjectSettings = GetMutableDefault<UMovieGraphProjectSettings>();
-		if (ensureAlwaysMsgf(MovieGraphProjectSettings, TEXT("%hs: Unable to get UMovieGraphProjectSettings."), __FUNCTION__))
-		{
-			OnProjectSettingsModifiedHandle =
-				MovieGraphProjectSettings->OnSettingChanged().AddSP(
-					this, &FMovieGraphNamedResolutionCustomization::OnProjectSettingsChanged);
-		}
-	}
+	/** Binds to the Project Settings to know when a user has added or removed a preset. */
+	void BindToOnProjectSettingsModified();
 
-	const FSlateBrush* GetLockBrush() const
-	{
-		return bLockedAspectRatio ? FAppStyle::GetBrush(TEXT("Icons.Link")) : FAppStyle::GetBrush(TEXT("Icons.Unlink"));
-	}
+	/** Get the icon to use for the Aspect Ratio Locked widget, could be a locked or unlocked icon. */
+	const FSlateBrush* GetAspectRatioLockBrush() const;
 
-	FText GetLockTooltipText() const
-	{
-		return bLockedAspectRatio ?
-			FText::Format(LOCTEXT("LockedAspectRatioTooltipFormat", "Aspect ratio is locked ({0})"), FText::AsNumber(CurrentAspectRatio)) :
-			LOCTEXT("UnlockedAspectRatioTooltip", "Aspect Ratio is unlocked");
-	}
+	/** Get the current tooltip for the aspect ratio lock */
+	FText GetAspectRatioLockTooltipText() const;
 	
 	//~ Begin IPropertyTypeCustomization interface
 	virtual void CustomizeHeader(
 		TSharedRef<IPropertyHandle> InStructPropertyHandle,
-		FDetailWidgetRow& HeaderRow, IPropertyTypeCustomizationUtils& CustomizationUtils) override
-	{
-		StructPropertyHandle = InStructPropertyHandle;
-		CreateOptionComboBox();
-		
-		const TSharedRef<IPropertyHandle> NameProp = InStructPropertyHandle->GetChildHandle(
-			GET_MEMBER_NAME_CHECKED(FMovieGraphNamedResolution, ProfileName)).ToSharedRef();
-		
-		// Try to 'load' selected option from the struct handle
-		FName LoadedOption = NAME_None;
-		if (NameProp->GetValue(LoadedOption) == FPropertyAccess::Result::Success && !LoadedOption.IsNone())
-		{
-			SetCurrentOptionAndCacheTooltipText(LoadedOption);
-		}
-		
-		// Set up option list for first time. If LoadedOption doesn't exist, use the first option in the list.
-		RepopulateOptions();
-
-		// Repopulate options when the project settings change
-		BindToOnProjectSettingsModified();
-
-		// Prevent reset to default from being shown
-		HeaderRow.OverrideResetToDefault(FResetToDefaultOverride::Hide());
-			
-		HeaderRow
-		.NameContent()
-		[
-			// Show "Resolution" as the row name
-			InStructPropertyHandle->CreatePropertyNameWidget()
-		]
-		.ValueContent()
-		.HAlign(HAlign_Fill)
-		[
-			SNew(SHorizontalBox)
-			
-			+ SHorizontalBox::Slot()
-			.Padding(0, 2, 8, 2)
-			.HAlign(HAlign_Fill)
-			[
-				ComboBox.ToSharedRef()
-			]
-		];
-	}
-
-	void UpdateAspectRatioLockImage()
-	{
-		LockImage->SetImage(GetLockBrush());
-		LockImage->SetToolTipText(GetLockTooltipText());
-	}
-
-	void ToggleAspectRatioLock()
-	{
-		bLockedAspectRatio = !bLockedAspectRatio;
-		if (bLockedAspectRatio) { CacheAspectRatio(); }
-
-		UpdateAspectRatioLockImage();
-	}
+		FDetailWidgetRow& HeaderRow, IPropertyTypeCustomizationUtils& CustomizationUtils) override;
+	virtual void CustomizeChildren(
+		TSharedRef<IPropertyHandle> InStructPropertyHandle,
+		IDetailChildrenBuilder& StructBuilder, IPropertyTypeCustomizationUtils& StructCustomizationUtils) override;
+	//~ End IPropertyTypeCustomization interface
 
-	TSharedRef<SWidget> MakeLockExtensionWidget()
-	{
-		return SNew(SButton)
-		.OnClicked_Lambda([this]()
-		{
-			ToggleAspectRatioLock();
-			
-			return FReply::Handled();
-		})
-		.ContentPadding(FMargin(0, 0, 4, 0))
-		.ButtonStyle(FAppStyle::Get(), "NoBorder")
-		[
-			SAssignNew(LockImage, SImage)
-			.ColorAndOpacity(FSlateColor::UseForeground())
-		];
-	}
+	/** Update the aspect ratio lock image based on current lock setting. */
+	void UpdateAspectRatioLockImage();
 
-	void CacheAspectRatio()
-	{
-		float X = (float)CustomEntry.Resolution.X;
-		float Y = (float)CustomEntry.Resolution.Y;
-		
-		if (const FMovieGraphNamedResolution* SelectedResolution = FindNamedResolutionForOption(GetCurrentOptionAssigningIfNeeded()))
-		{
-			X = (float)SelectedResolution->Resolution.X;
-			Y = (float)SelectedResolution->Resolution.Y;
-		}
-		CurrentAspectRatio = X / Y;
-	}
+	/** On click, toggle the current lock state. */
+	void ToggleAspectRatioLock();
 
-	void UpdateCustomEntryValues()
-	{
-		if (!ensureMsgf(ResolutionXPropertyHandle, TEXT("%hs: `ResolutionXPropertyHandle` is null."), __FUNCTION__))
-		{
-			return;
-		}
-		if (!ensureMsgf(ResolutionYPropertyHandle, TEXT("%hs: `ResolutionYPropertyHandle` is null."), __FUNCTION__))
-		{
-			return;
-		}
-			
-		// Set the custom resolution to the preset resolution
-		int32 ResolutionX = 0;
-		ResolutionXPropertyHandle->GetValue(ResolutionX);
-
-		CustomEntry.Resolution.X = ResolutionX;
-
-		int32 ResolutionY = 0;
-		ResolutionYPropertyHandle->GetValue(ResolutionY);
-
-		CustomEntry.Resolution.Y = ResolutionY;
-
-		// Cache the current aspect ratio
-		CacheAspectRatio();
-	}
+	/** Create the Aspect Ratio Lock widget that is added in the customization. */
+	TSharedRef<SWidget> MakeAspectRatioLockWidget();
 
-	int32 RoundToNearestEvenNumber(double InNumberToRound) const
-	{
-		// Round to nearest even
-		const int32 Floored = FMath::FloorToInt(InNumberToRound);
-		const int32 Ceiled = FMath::CeilToInt(InNumberToRound);
+	/** Calculates the aspect ratio of the current resolution. */
+	float CalculateAspectRatio();
 
-		return Floored % 2 == 0 ? Floored : Ceiled;
-	}
-	
-	int32 GetCurrentlySelectedResolutionByAxis(const EAxis::Type Axis)
-	{
-		if (const FMovieGraphNamedResolution* const CurrentNamedResolution = FindNamedResolutionForOption(GetCurrentOptionAssigningIfNeeded()))
-		{
-			return Axis == EAxis::Y ? CurrentNamedResolution->Resolution.Y : CurrentNamedResolution->Resolution.X;
-		}
+	/** 
+	* Gets the effective resolution for the current resolution. For profiles this looks it up in the Project Settings, 
+	* so that if the project setting has been changed since the resolution was last updated, we reflect the correct
+	* value. MRQ will also look it up from the Project Setting at runtime.
+	*/
+	FIntPoint GetEffectiveResolution() const;
 
-		return INDEX_NONE;
-	}
+	/** Called when dragging the slider starts. Used to open a transaction to encapsulate all changes into one. */
+	void OnCustomSliderBeginMovement();
+	/** Called when the slider stops dragging. Closes the transaction. */
+	void OnCustomSliderEndMovement(uint32 NewValue);
 
-	void OnCustomSliderValueChanged(uint32 NewValue, EAxis::Type Axis)
-	{
-		// If the custom resolution is changed, switch to 'Custom' option
-		if (!CurrentOption.IsEqual(FMovieGraphNamedResolution::CustomEntryName))
-		{
-			UpdateCustomEntryValues();
-
-			// Switch to 'Custom' option
-			SetCurrentOptionAndCacheTooltipText(FMovieGraphNamedResolution::CustomEntryName);
-		}
-
-		// Enforce aspect ratio if desired
-		if (Axis == EAxis::Y)
-		{
-			CustomEntry.Resolution.Y = NewValue;
-			
-			if (bLockedAspectRatio)
-			{
-				const double NewXResolution = (double)CustomEntry.Resolution.Y * CurrentAspectRatio;
-				CustomEntry.Resolution.X = RoundToNearestEvenNumber(NewXResolution);
-			}
-		}
-		else // X Changed
-		{
-			CustomEntry.Resolution.X = NewValue;
-			
-			if (bLockedAspectRatio)
-			{
-				const double NewYResolution = (double)CustomEntry.Resolution.X / CurrentAspectRatio;
-				CustomEntry.Resolution.Y = RoundToNearestEvenNumber(NewYResolution);
-			}
-		}
-	}
+	/** Gets the resolution for the given axis. */
+	int32 GetCurrentlySelectedResolutionByAxis(const EAxis::Type Axis);
 
-	void OnCustomSliderValueCommitted(uint32 NewValue, ETextCommit::Type, EAxis::Type Axis)
-	{
-		OnCustomSliderValueChanged(NewValue, Axis);
-		
-		UpdateOwningStruct();
-	}
+	/** Called when dragging the slider, or after typing in a number and hitting enter. */
+	void OnCustomSliderValueChanged(uint32 NewValue, EAxis::Type Axis);
 
+	/** Adds a row to the customization for the given axis. */
 	void AddCustomRowForResolutionAxis(
 		EAxis::Type Axis, IDetailChildrenBuilder& StructBuilder, const FString& FilterString,
-		TSharedRef<SWidget> NameContentWidget, TSharedPtr<SWidget> AspectRatioLockExtensionWidget = nullptr)
-	{
-		const FSlateFontInfo PropertyFontStyle = FAppStyle::GetFontStyle( TEXT("PropertyWindow.NormalFont") );
-		const TOptional<uint32> PropertyMaxValue = TOptional<uint32>();
-	
-		FDetailWidgetRow& CustomRow = StructBuilder.AddCustomRow(FText::FromString(FilterString))
-		.NameContent()
-		[
-			NameContentWidget
-		]
-		.ValueContent()
-		.HAlign(HAlign_Fill)
-		[
-			// We make our own value widget because we don't want the widget tied to the property handle
-			// as the property handle gets re-instanced every time a blueprint is edited.
-			SNew(SBox)
-			// Maintain spacing when there would be no extension widget
-			.Padding(0, 0, Axis == EAxis::Y ? 24 /* Assuming image is 16x16 + 8 padding */ : 0, 0) 
-			[
-				SNew(SNumericEntryBox<uint32>)
-				.Font(PropertyFontStyle)
-				.AllowSpin(true)
-				.MinSliderValue(2)
-				.MinValue(2)
-				.MaxValue(PropertyMaxValue)
-				.MaxSliderValue(PropertyMaxValue)
-				.Delta(2)
-				.Value_Lambda([this, Axis] (){ return GetCurrentlySelectedResolutionByAxis(Axis); }) // Lambda to bypass const requirement
-				.OnValueChanged(this, &FMovieGraphNamedResolutionCustomization::OnCustomSliderValueChanged, Axis)
-				.OnValueCommitted(this, &FMovieGraphNamedResolutionCustomization::OnCustomSliderValueCommitted, Axis)
-			]
-		]
-		.OverrideResetToDefault(FResetToDefaultOverride::Hide()); // Prevent showing Reset to Default
-
-		if (AspectRatioLockExtensionWidget.IsValid())
-		{
-			CustomRow
-			.ExtensionContent()
-			[
-				AspectRatioLockExtensionWidget.ToSharedRef()
-			];
-		}
-	}
+		TSharedRef<SWidget> NameContentWidget, TSharedPtr<SWidget> AspectRatioLockExtensionWidget = nullptr);
 
-	virtual void CustomizeChildren(
-		TSharedRef<IPropertyHandle> InStructPropertyHandle,
-		IDetailChildrenBuilder& StructBuilder, IPropertyTypeCustomizationUtils& StructCustomizationUtils) override
-	{
-		const TSharedPtr<IPropertyHandle> ResolutionPropertyHandle =
-			InStructPropertyHandle->GetChildHandle(GET_MEMBER_NAME_CHECKED(FMovieGraphNamedResolution, Resolution));
-		
-		if (const TSharedPtr<IPropertyHandle> PropertyXHandle = ResolutionPropertyHandle->GetChildHandle(0))
-		{
-			ResolutionXPropertyHandle = PropertyXHandle;
-
-			AddCustomRowForResolutionAxis(
-				EAxis::X, StructBuilder, PropertyXHandle->GeneratePathToProperty(), PropertyXHandle->CreatePropertyNameWidget(), MakeLockExtensionWidget());
-			
-			UpdateAspectRatioLockImage();
-		}
-
-		if (const TSharedPtr<IPropertyHandle> PropertyYHandle = ResolutionPropertyHandle->GetChildHandle(1))
-		{			
-			ResolutionYPropertyHandle = PropertyYHandle;
-			
-			AddCustomRowForResolutionAxis(
-				EAxis::Y, StructBuilder, PropertyYHandle->GeneratePathToProperty(), PropertyYHandle->CreatePropertyNameWidget());
-		}
-
-		CacheAspectRatio();
-	}
-	//~ End IPropertyTypeCustomization interface
-	
-	void OnProjectSettingsChanged(UObject*, FPropertyChangedEvent& PropertyChangedEvent)
-	{
-		if (PropertyChangedEvent.MemberProperty->GetFName().IsEqual(GET_MEMBER_NAME_CHECKED(UMovieGraphProjectSettings, DefaultNamedResolutions)))
-		{
-			RepopulateOptions();
-		}
-	}
 
-	/**
-	 * Updates the struct which owns this customization with the data from the matched named resolution.
-	 */
-	void UpdateOwningStruct()
-	{
-		const TSharedPtr<IPropertyHandle> PinnedStructPropertyHandle = StructPropertyHandle.Pin();
-		if (!PinnedStructPropertyHandle)
-		{
-			UE_LOG(LogMovieRenderPipeline, Error, TEXT("%hs: Unable to pin StructPropertyHandle."), __FUNCTION__);
-			return;
-		}
-
-		const FMovieGraphNamedResolution* NamedResolution = FindNamedResolutionForOption(CurrentOption);
-		if (!NamedResolution)
-		{
-			return;
-		}
-		
-		void* StructData = nullptr;
-		if (PinnedStructPropertyHandle->GetValueData(StructData) == FPropertyAccess::Success && StructData != nullptr)
-		{
-			// Copy the data from our matched NamedResolution to the owning struct
-			FMemory::Memcpy(StructData, NamedResolution, sizeof(FMovieGraphNamedResolution));
-
-			// Transact on outer objects
-			TArray<UObject*> OutObjects;
-			PinnedStructPropertyHandle->GetOuterObjects(OutObjects);
-
-			for (UObject* Object : OutObjects)
-			{
-				if (Object)
-				{
-					constexpr bool bAlwaysMarkDirty = true;	
-					Object->Modify(bAlwaysMarkDirty);
-				}
-			}
-		}
-		else
-		{
-			UE_LOG(LogMovieRenderPipeline, Error, TEXT("%hs: Unable to get struct data from StructPropertyHandle."), __FUNCTION__);
-		}
-	}
-
-	void OnComboBoxSelectionChanged(const FName NewValue, ESelectInfo::Type)
-	{
-		SetCurrentOptionAndCacheTooltipText(NewValue);
+	/** Callback bound when project setting changes. */
+	void OnProjectSettingsChanged(UObject*, FPropertyChangedEvent& PropertyChangedEvent);
 
-		UpdateOwningStruct();
+	/** Called when the user selects a new preset from the combo box dropdown. */
+	void OnComboBoxSelectionChanged(const FName NewValue, ESelectInfo::Type InType);
 
-		CacheAspectRatio();
-	}
-
-	/**
-	 * The widget that allows the end user to select the output resolution
-	 */
-	TSharedPtr<SComboBox<FName>> ComboBox;
+	/** The widget that allows the end user to select the output resolution */
+	TSharedPtr<SWidget> ComboBoxWidget;
 
-	/**
-	 * The image that indicates whether the aspect ratio is locked
-	 */
-	TSharedPtr<SImage> LockImage;
+	/** The image that indicates whether the aspect ratio is locked. */
+	TSharedPtr<SImage> AspectRatioLockImage;
 
-	/**
-	 * Weak ptr to the handle of the struct which owns this customization
-	 */
+	/** Weak ptr to the handle of the struct which owns this customization */
 	TWeakPtr<IPropertyHandle> StructPropertyHandle;
 
-	/**
-	 * Handle of the widget representing the X dimension of the 'Custom' resolution
-	 */
+	/** Handle to the "Resolution.X" property on the edited object. */
 	TSharedPtr<IPropertyHandle> ResolutionXPropertyHandle;
-	/**
-	 * Handle of the widget representing the Y dimension of the 'Custom' resolution
-	 */
+	/** Handle to the "Resolution.Y" property on the edited object. */
 	TSharedPtr<IPropertyHandle> ResolutionYPropertyHandle;
-
-	/**
-	 * The option currently selected in the combobox
-	 */
-	FName CurrentOption = NAME_None;
-	
-	/**
-	 * The tooltip text for the currently selected option
-	 */
-	FText CachedSelectedOptionTooltipText;
-
-	/**
-	 * If true, changing one dimension of the custom resolution will also
-	 * change the other dimension to maintain the aspect ratio.
-	 */
-	bool bLockedAspectRatio = true;
+	/** Handle to the "ProfileName" property on the edited object. */
+	TSharedPtr<IPropertyHandle> ProfileNamePropertyHandle;
 
 	/**
 	 * The aspect ratio of the currently selected resolution.
 	 * Updates when the aspect ratio lock is applied, a named resolution is selected,
-	 * or when the 'Custom' resolution values are changed.
+	 * or when the 'Custom' resolution values are changed. Used to track if Aspect
+	 * ratio lock is enabled or not.
 	 */
-	float CurrentAspectRatio = 0.0f;
+	TOptional<float> LockedAspectRatio;
 
 	/**
 	 * The options for the combo box, populated from UMovieGraphProjectSettings::DefaultNamedResolutions.
@@ -592,12 +149,10 @@ protected:
 	 */
 	FDelegateHandle OnProjectSettingsModifiedHandle;
 
-
 	/**
-	 * Default FMovieGraphNamedResolution for the 'custom' option
-	 */
-	inline static FMovieGraphNamedResolution CustomEntry =
-		FMovieGraphNamedResolution(FMovieGraphNamedResolution::CustomEntryName, FIntPoint(1920, 1080), FMovieGraphNamedResolution::CustomEntryName.ToString());
+	* If valid, this means that a slide operation is currently in progress and that the changes will be encapsulated into this transaction. 
+	*/
+	TUniquePtr<FScopedTransaction> InteractiveResolutionEditTransaction;
 };
 
 #undef LOCTEXT_NAMESPACE
