MoviePipeline: Final icon updates for the initial release of MRG.

#rnx

diff --git a/Engine/Plugins/MovieScene/MovieRenderPipeline/Source/MovieRenderPipelineCore/Private/Graph/MovieGraphRenderLayerSubsystem.cpp b/Engine/Plugins/MovieScene/MovieRenderPipeline/Source/MovieRenderPipelineCore/Private/Graph/MovieGraphRenderLayerSubsystem.cpp
--- a/Engine/Plugins/MovieScene/MovieRenderPipeline/Source/MovieRenderPipelineCore/Private/Graph/MovieGraphRenderLayerSubsystem.cpp
+++ b/Engine/Plugins/MovieScene/MovieRenderPipeline/Source/MovieRenderPipelineCore/Private/Graph/MovieGraphRenderLayerSubsystem.cpp
@@ -821,8 +821,7 @@ void UMovieGraphConditionGroupQuery_ActorTagName::Evaluate(const TArray<AActor*>
 
 const FSlateIcon& UMovieGraphConditionGroupQuery_ActorTagName::GetIcon() const
 {
-	// TODO: This icon is wrong
-	static const FSlateIcon ActorTagIcon = FSlateIcon(FAppStyle::GetAppStyleSetName(), "LevelEditor.Tabs.Debug");
+	static const FSlateIcon ActorTagIcon = FSlateIcon(FAppStyle::GetAppStyleSetName(), "MainFrame.OpenIssueTracker");
 	return ActorTagIcon;
 }
 
@@ -1187,8 +1186,7 @@ void UMovieGraphConditionGroupQuery_ComponentTagName::Evaluate(const TArray<AAct
 
 const FSlateIcon& UMovieGraphConditionGroupQuery_ComponentTagName::GetIcon() const
 {
-	// TODO: This icon is wrong
-	static const FSlateIcon ActorTagIcon = FSlateIcon(FAppStyle::GetAppStyleSetName(), "ClassIcon.DataAsset");
+	static const FSlateIcon ActorTagIcon = FSlateIcon(FAppStyle::GetAppStyleSetName(), "MainFrame.OpenIssueTracker");
 	return ActorTagIcon;
 }
 
diff --git a/Engine/Plugins/MovieScene/MovieRenderPipeline/Source/MovieRenderPipelineCore/Private/Graph/Nodes/MovieGraphBurnInNode.cpp b/Engine/Plugins/MovieScene/MovieRenderPipeline/Source/MovieRenderPipelineCore/Private/Graph/Nodes/MovieGraphBurnInNode.cpp
--- a/Engine/Plugins/MovieScene/MovieRenderPipeline/Source/MovieRenderPipelineCore/Private/Graph/Nodes/MovieGraphBurnInNode.cpp
+++ b/Engine/Plugins/MovieScene/MovieRenderPipeline/Source/MovieRenderPipelineCore/Private/Graph/Nodes/MovieGraphBurnInNode.cpp
@@ -32,10 +32,10 @@ FLinearColor UMovieGraphBurnInNode::GetNodeTitleColor() const
 
 FSlateIcon UMovieGraphBurnInNode::GetIconAndTint(FLinearColor& OutColor) const
 {
-	static const FSlateIcon DeferredRendererIcon = FSlateIcon(FAppStyle::GetAppStyleSetName(), "SequenceRecorder.TabIcon");
+	static const FSlateIcon BurnInIcon = FSlateIcon(FAppStyle::GetAppStyleSetName(), "SequenceRecorder.TabIcon");
 
 	OutColor = FLinearColor::White;
-	return DeferredRendererIcon;
+	return BurnInIcon;
 }
 #endif	// WITH_EDITOR
 
diff --git a/Engine/Plugins/MovieScene/MovieRenderPipeline/Source/MovieRenderPipelineCore/Private/Graph/Nodes/MovieGraphDebugNode.cpp b/Engine/Plugins/MovieScene/MovieRenderPipeline/Source/MovieRenderPipelineCore/Private/Graph/Nodes/MovieGraphDebugNode.cpp
--- a/Engine/Plugins/MovieScene/MovieRenderPipeline/Source/MovieRenderPipelineCore/Private/Graph/Nodes/MovieGraphDebugNode.cpp
+++ b/Engine/Plugins/MovieScene/MovieRenderPipeline/Source/MovieRenderPipelineCore/Private/Graph/Nodes/MovieGraphDebugNode.cpp
@@ -25,7 +25,7 @@ FLinearColor UMovieGraphDebugSettingNode::GetNodeTitleColor() const
 
 FSlateIcon UMovieGraphDebugSettingNode::GetIconAndTint(FLinearColor& OutColor) const
 {
-	static const FSlateIcon Icon = FSlateIcon(FAppStyle::GetAppStyleSetName(), "DebugTools.TabIcon");
+	static const FSlateIcon Icon = FSlateIcon(FAppStyle::GetAppStyleSetName(), "LevelEditor.Tabs.Debug");
 
 	OutColor = FLinearColor::White;
 	return Icon;
diff --git a/Engine/Plugins/MovieScene/MovieRenderPipeline/Source/MovieRenderPipelineCore/Private/Graph/Nodes/MovieGraphExecuteScriptNode.cpp b/Engine/Plugins/MovieScene/MovieRenderPipeline/Source/MovieRenderPipelineCore/Private/Graph/Nodes/MovieGraphExecuteScriptNode.cpp
--- a/Engine/Plugins/MovieScene/MovieRenderPipeline/Source/MovieRenderPipelineCore/Private/Graph/Nodes/MovieGraphExecuteScriptNode.cpp
+++ b/Engine/Plugins/MovieScene/MovieRenderPipeline/Source/MovieRenderPipelineCore/Private/Graph/Nodes/MovieGraphExecuteScriptNode.cpp
@@ -23,10 +23,10 @@ FLinearColor UMovieGraphExecuteScriptNode::GetNodeTitleColor() const
 
 FSlateIcon UMovieGraphExecuteScriptNode::GetIconAndTint(FLinearColor& OutColor) const
 {
-	static const FSlateIcon DeferredRendererIcon = FSlateIcon(FAppStyle::GetAppStyleSetName(), "SequenceRecorder.TabIcon");
+	static const FSlateIcon ExecuteScriptIcon = FSlateIcon(FAppStyle::GetAppStyleSetName(), "Log.TabIcon");
 
 	OutColor = FLinearColor::White;
-	return DeferredRendererIcon;
+	return ExecuteScriptIcon;
 }
 #endif	// WITH_EDITOR
 
diff --git a/Engine/Plugins/MovieScene/MovieRenderPipeline/Source/MovieRenderPipelineCore/Private/Graph/Nodes/MovieGraphModifierNode.cpp b/Engine/Plugins/MovieScene/MovieRenderPipeline/Source/MovieRenderPipelineCore/Private/Graph/Nodes/MovieGraphModifierNode.cpp
--- a/Engine/Plugins/MovieScene/MovieRenderPipeline/Source/MovieRenderPipelineCore/Private/Graph/Nodes/MovieGraphModifierNode.cpp
+++ b/Engine/Plugins/MovieScene/MovieRenderPipeline/Source/MovieRenderPipelineCore/Private/Graph/Nodes/MovieGraphModifierNode.cpp
@@ -132,7 +132,7 @@ FLinearColor UMovieGraphModifierNode::GetNodeTitleColor() const
 
 FSlateIcon UMovieGraphModifierNode::GetIconAndTint(FLinearColor& OutColor) const
 {
-	static const FSlateIcon ModifierIcon = FSlateIcon(FAppStyle::GetAppStyleSetName(), "ContentBrowser.ReferenceViewer");
+	static const FSlateIcon ModifierIcon = FSlateIcon(FAppStyle::GetAppStyleSetName(), "Icons.Edit");
 
 	OutColor = FLinearColor::White;
 	return ModifierIcon;
diff --git a/Engine/Plugins/MovieScene/MovieRenderPipeline/Source/MovieRenderPipelineCore/Private/Graph/Nodes/MovieGraphSetCVarValueNode.cpp b/Engine/Plugins/MovieScene/MovieRenderPipeline/Source/MovieRenderPipelineCore/Private/Graph/Nodes/MovieGraphSetCVarValueNode.cpp
--- a/Engine/Plugins/MovieScene/MovieRenderPipeline/Source/MovieRenderPipelineCore/Private/Graph/Nodes/MovieGraphSetCVarValueNode.cpp
+++ b/Engine/Plugins/MovieScene/MovieRenderPipeline/Source/MovieRenderPipelineCore/Private/Graph/Nodes/MovieGraphSetCVarValueNode.cpp
@@ -57,10 +57,10 @@ FLinearColor UMovieGraphSetCVarValueNode::GetNodeTitleColor() const
 
 FSlateIcon UMovieGraphSetCVarValueNode::GetIconAndTint(FLinearColor& OutColor) const
 {
-	static const FSlateIcon RenderLayerIcon = FSlateIcon(FAppStyle::GetAppStyleSetName(), "MainFrame.BrowseCVars");
+	static const FSlateIcon SetCVarValueIcon = FSlateIcon(FAppStyle::GetAppStyleSetName(), "MainFrame.BrowseCVars");
 
 	OutColor = FLinearColor::White;
-	return RenderLayerIcon;
+	return SetCVarValueIcon;
 }
 
 void UMovieGraphSetCVarValueNode::PostEditChangeProperty(FPropertyChangedEvent& PropertyChangedEvent)
diff --git a/Engine/Plugins/MovieScene/MovieRenderPipeline/Source/MovieRenderPipelineCore/Private/Graph/Nodes/MovieGraphUIRendererNode.cpp b/Engine/Plugins/MovieScene/MovieRenderPipeline/Source/MovieRenderPipelineCore/Private/Graph/Nodes/MovieGraphUIRendererNode.cpp
--- a/Engine/Plugins/MovieScene/MovieRenderPipeline/Source/MovieRenderPipelineCore/Private/Graph/Nodes/MovieGraphUIRendererNode.cpp
+++ b/Engine/Plugins/MovieScene/MovieRenderPipeline/Source/MovieRenderPipelineCore/Private/Graph/Nodes/MovieGraphUIRendererNode.cpp
@@ -24,10 +24,10 @@ FLinearColor UMovieGraphUIRendererNode::GetNodeTitleColor() const
 
 FSlateIcon UMovieGraphUIRendererNode::GetIconAndTint(FLinearColor& OutColor) const
 {
-	static const FSlateIcon DeferredRendererIcon = FSlateIcon(FAppStyle::GetAppStyleSetName(), "SequenceRecorder.TabIcon");
+	static const FSlateIcon UIRendererIcon = FSlateIcon(FAppStyle::GetAppStyleSetName(), "SequenceRecorder.TabIcon");
 
 	OutColor = FLinearColor::White;
-	return DeferredRendererIcon;
+	return UIRendererIcon;
 }
 #endif	// WITH_EDITOR
 
diff --git a/Engine/Plugins/MovieScene/MovieRenderPipeline/Source/MovieRenderPipelineRenderPasses/Private/Graph/Nodes/MovieGraphDeferredPassNode.cpp b/Engine/Plugins/MovieScene/MovieRenderPipeline/Source/MovieRenderPipelineRenderPasses/Private/Graph/Nodes/MovieGraphDeferredPassNode.cpp
--- a/Engine/Plugins/MovieScene/MovieRenderPipeline/Source/MovieRenderPipelineRenderPasses/Private/Graph/Nodes/MovieGraphDeferredPassNode.cpp
+++ b/Engine/Plugins/MovieScene/MovieRenderPipeline/Source/MovieRenderPipelineRenderPasses/Private/Graph/Nodes/MovieGraphDeferredPassNode.cpp
@@ -41,6 +41,14 @@ FText UMovieGraphDeferredRenderPassNode::GetNodeTitle(const bool bGetDescriptive
 {
 	return NSLOCTEXT("MovieGraphNodes", "DeferredRenderPassGraphNode_Description", "Deferred Renderer");
 }
+
+FSlateIcon UMovieGraphDeferredRenderPassNode::GetIconAndTint(FLinearColor& OutColor) const
+{
+	static const FSlateIcon DeferredRendererIcon = FSlateIcon(FAppStyle::GetAppStyleSetName(), "LevelViewport.HighResScreenshot");
+	
+	OutColor = FLinearColor::White;
+	return DeferredRendererIcon;
+}
 #endif
 
 FString UMovieGraphDeferredRenderPassNode::GetRendererNameImpl() const
diff --git a/Engine/Plugins/MovieScene/MovieRenderPipeline/Source/MovieRenderPipelineRenderPasses/Private/Graph/Nodes/MovieGraphPathTracerPassNode.cpp b/Engine/Plugins/MovieScene/MovieRenderPipeline/Source/MovieRenderPipelineRenderPasses/Private/Graph/Nodes/MovieGraphPathTracerPassNode.cpp
--- a/Engine/Plugins/MovieScene/MovieRenderPipeline/Source/MovieRenderPipelineRenderPasses/Private/Graph/Nodes/MovieGraphPathTracerPassNode.cpp
+++ b/Engine/Plugins/MovieScene/MovieRenderPipeline/Source/MovieRenderPipelineRenderPasses/Private/Graph/Nodes/MovieGraphPathTracerPassNode.cpp
@@ -39,6 +39,14 @@ FText UMovieGraphPathTracerRenderPassNode::GetNodeTitle(const bool bGetDescripti
 {
 	return NSLOCTEXT("MovieGraphNodes", "PathTracedRenderPassGraphNode_Description", "Path Traced Renderer");
 }
+
+FSlateIcon UMovieGraphPathTracerRenderPassNode::GetIconAndTint(FLinearColor& OutColor) const
+{
+	static const FSlateIcon PathTracerRendererIcon = FSlateIcon(FAppStyle::GetAppStyleSetName(), "EditorViewport.PathTracingMode");
+	
+    OutColor = FLinearColor::White;
+    return PathTracerRendererIcon;
+}
 #endif
 
 void UMovieGraphPathTracerRenderPassNode::SetupImpl(const FMovieGraphRenderPassSetupData& InSetupData)
diff --git a/Engine/Plugins/MovieScene/MovieRenderPipeline/Source/MovieRenderPipelineRenderPasses/Public/Graph/Nodes/MovieGraphDeferredPassNode.h b/Engine/Plugins/MovieScene/MovieRenderPipeline/Source/MovieRenderPipelineRenderPasses/Public/Graph/Nodes/MovieGraphDeferredPassNode.h
--- a/Engine/Plugins/MovieScene/MovieRenderPipeline/Source/MovieRenderPipelineRenderPasses/Public/Graph/Nodes/MovieGraphDeferredPassNode.h
+++ b/Engine/Plugins/MovieScene/MovieRenderPipeline/Source/MovieRenderPipelineRenderPasses/Public/Graph/Nodes/MovieGraphDeferredPassNode.h
@@ -20,6 +20,7 @@ public:
 
 #if WITH_EDITOR
 	virtual FText GetNodeTitle(const bool bGetDescriptive = false) const override;
+	virtual FSlateIcon GetIconAndTint(FLinearColor& OutColor) const override;
 #endif
 
 protected:
diff --git a/Engine/Plugins/MovieScene/MovieRenderPipeline/Source/MovieRenderPipelineRenderPasses/Public/Graph/Nodes/MovieGraphPathTracerPassNode.h b/Engine/Plugins/MovieScene/MovieRenderPipeline/Source/MovieRenderPipelineRenderPasses/Public/Graph/Nodes/MovieGraphPathTracerPassNode.h
--- a/Engine/Plugins/MovieScene/MovieRenderPipeline/Source/MovieRenderPipelineRenderPasses/Public/Graph/Nodes/MovieGraphPathTracerPassNode.h
+++ b/Engine/Plugins/MovieScene/MovieRenderPipeline/Source/MovieRenderPipelineRenderPasses/Public/Graph/Nodes/MovieGraphPathTracerPassNode.h
@@ -17,6 +17,7 @@ public:
 
 #if WITH_EDITOR
 	virtual FText GetNodeTitle(const bool bGetDescriptive = false) const override;
+	virtual FSlateIcon GetIconAndTint(FLinearColor& OutColor) const override;
 #endif
 
 	// UMovieGraphRenderPassNode Interface
