CIS: Fixed a Non-Unity Build Issue

diff --git a/Engine/Plugins/MovieScene/MovieRenderPipeline/Source/MovieRenderPipelineEditor/Private/Customizations/Graph/MovieGraphNamedResolutionCustomization.cpp b/Engine/Plugins/MovieScene/MovieRenderPipeline/Source/MovieRenderPipelineEditor/Private/Customizations/Graph/MovieGraphNamedResolutionCustomization.cpp
--- a/Engine/Plugins/MovieScene/MovieRenderPipeline/Source/MovieRenderPipelineEditor/Private/Customizations/Graph/MovieGraphNamedResolutionCustomization.cpp
+++ b/Engine/Plugins/MovieScene/MovieRenderPipeline/Source/MovieRenderPipelineEditor/Private/Customizations/Graph/MovieGraphNamedResolutionCustomization.cpp
@@ -15,7 +15,6 @@
 #include "Widgets/Input/SComboBox.h"
 #include "Widgets/Input/SNumericEntryBox.h"
 #include "Widgets/SBoxPanel.h"
-#include "ScopedTransaction.h"
 
 #define LOCTEXT_NAMESPACE "FMovieGraphNamedResolutionCustomization"
 
diff --git a/Engine/Plugins/MovieScene/MovieRenderPipeline/Source/MovieRenderPipelineEditor/Private/Customizations/Graph/MovieGraphNamedResolutionCustomization.h b/Engine/Plugins/MovieScene/MovieRenderPipeline/Source/MovieRenderPipelineEditor/Private/Customizations/Graph/MovieGraphNamedResolutionCustomization.h
--- a/Engine/Plugins/MovieScene/MovieRenderPipeline/Source/MovieRenderPipelineEditor/Private/Customizations/Graph/MovieGraphNamedResolutionCustomization.h
+++ b/Engine/Plugins/MovieScene/MovieRenderPipeline/Source/MovieRenderPipelineEditor/Private/Customizations/Graph/MovieGraphNamedResolutionCustomization.h
@@ -5,8 +5,7 @@
 #include "IDetailChildrenBuilder.h"
 #include "DetailLayoutBuilder.h"
 #include "IPropertyTypeCustomization.h"
-
-class FScopedTransaction;
+#include "ScopedTransaction.h"
 
 /* Customizes how named resolutions are displayed in the details pane. */
 class MOVIERENDERPIPELINEEDITOR_API FMovieGraphNamedResolutionCustomization : public IPropertyTypeCustomization
