MovieGraph: Fixed an issue where scripting could not access the Job's variable overrides, which meant you could not set overrides on jobs, or their shots from scripting.

diff --git a/Engine/Plugins/MovieScene/MovieRenderPipeline/Source/MovieRenderPipelineCore/Private/MoviePipelineQueue.cpp b/Engine/Plugins/MovieScene/MovieRenderPipeline/Source/MovieRenderPipelineCore/Private/MoviePipelineQueue.cpp
--- a/Engine/Plugins/MovieScene/MovieRenderPipeline/Source/MovieRenderPipelineCore/Private/MoviePipelineQueue.cpp
+++ b/Engine/Plugins/MovieScene/MovieRenderPipeline/Source/MovieRenderPipelineCore/Private/MoviePipelineQueue.cpp
@@ -356,7 +356,7 @@ void UMoviePipelineExecutorJob::SetSequence(FSoftObjectPath InSequence)
 	}
 }
 
-TObjectPtr<UMovieJobVariableAssignmentContainer> UMoviePipelineExecutorJob::GetOrCreateJobVariableAssignmentsForGraph(const UMovieGraphConfig* InGraph)
+UMovieJobVariableAssignmentContainer* UMoviePipelineExecutorJob::GetOrCreateJobVariableAssignmentsForGraph(const UMovieGraphConfig* InGraph)
 {
 	if (InGraph)
 	{
@@ -419,7 +419,7 @@ void UMoviePipelineExecutorShot::SetGraphPreset(const UMovieGraphConfig* InGraph
 	OnShotGraphPresetChanged.Broadcast(this, GraphPreset.Get());
 }
 
-TObjectPtr<UMovieJobVariableAssignmentContainer> UMoviePipelineExecutorShot::GetOrCreateJobVariableAssignmentsForGraph(const UMovieGraphConfig* InGraph, const bool bIsForPrimaryOverrides)
+UMovieJobVariableAssignmentContainer* UMoviePipelineExecutorShot::GetOrCreateJobVariableAssignmentsForGraph(const UMovieGraphConfig* InGraph, const bool bIsForPrimaryOverrides)
 {
 	if (InGraph)
 	{
diff --git a/Engine/Plugins/MovieScene/MovieRenderPipeline/Source/MovieRenderPipelineCore/Public/MoviePipelineQueue.h b/Engine/Plugins/MovieScene/MovieRenderPipeline/Source/MovieRenderPipelineCore/Public/MoviePipelineQueue.h
--- a/Engine/Plugins/MovieScene/MovieRenderPipeline/Source/MovieRenderPipelineCore/Public/MoviePipelineQueue.h
+++ b/Engine/Plugins/MovieScene/MovieRenderPipeline/Source/MovieRenderPipelineCore/Public/MoviePipelineQueue.h
@@ -160,8 +160,29 @@ public:
 	UFUNCTION(BlueprintCallable, Category = "Movie Render Pipeline")
 	void SetGraphPreset(const UMovieGraphConfig* InGraphPreset, const bool bUpdateVariableAssignments = true);
 
-	/** Gets the job variable assignments for a specific graph. Creates a new variable assignments container if one was not found for the given graph. */
-	TObjectPtr<UMovieJobVariableAssignmentContainer> GetOrCreateJobVariableAssignmentsForGraph(const UMovieGraphConfig* InGraph, const bool bIsForPrimaryOverrides = false);
+	/**
+	* This method will return you the object which contains variable overrides for either the Job's Primary or the Shot's GraphPreset. UMoviePipelineExecutorShot
+	* has two separate sets of overrides. You can use the shot to override a variable on the Primary Graph (ie: the one assigned to the whole job),
+	* but a graph can also have an entirely separate UMovieGraph config asset to run (though at runtime some variables will only be read from the
+	* Primary Graph, ie: Custom Frame Range due to it applying to the entire sequence). 
+	* 
+	* If you specify true for bIsForPrimaryOverrides it returns an object that allows this shot to override a variable that comes from the primary 
+	* graph. If you return false, then it returns an object that allows overriding a variable for this shot's override config (see: GetGraphPreset).
+	* See UMoviePipelineExecutorJob's version of this functoin for more details.
+	* 
+	* @param InGraph - The graph asset to return the config for. If this shot has its own Graph Preset override, you should return GetGraphPreset()
+	* or one of it's sub-graph pointers. If this shot is just trying to override the Primary Graph from the parent UMoviePipelineExecutorJob then
+	* you should return a pointer to the Job's GetGraphPreset() (or one of it's sub-graphs). Each graph/sub-graph gets its own set of overrides
+	* since sub-graphs can have different variables than the parents, so you have to provide the pointer to the one you want to override variables for.
+	* 
+	* @param bIsForPrimaryOverride - Default false. If true, tries to override variables on the parent UMoviePipelineExecutorJob's graphs. If false,
+	* tries to override variables on the Graph Preset assigned to this shot.
+	* 
+	* @return A container object which holds a copy of the variables for the specified Graph Asset that can be used to override their values
+	* on jobs without actually editing the default asset.
+	*/
+	UFUNCTION(BlueprintCallable, Category = "Movie Render Pipeline", DisplayName = "Get or Create Variable Overrides", meta = (ScriptName = "GetOrCreateVariableOverrides"))
+	UMovieJobVariableAssignmentContainer* GetOrCreateJobVariableAssignmentsForGraph(const UMovieGraphConfig* InGraph, const bool bIsForPrimaryOverrides = false);
 
 	/** Returns whether this should should be rendered */
 	UFUNCTION(BlueprintPure, Category = "Movie Render Pipeline")
@@ -476,8 +497,22 @@ public:
 	UFUNCTION(BlueprintSetter, Category = "Movie Render Pipeline")
 	void SetSequence(FSoftObjectPath InSequence);
 
-	/** Gets the job variable assignments for a specific graph. Creates a new variable assignments container if one was not found for the given graph. */
-	TObjectPtr<UMovieJobVariableAssignmentContainer> GetOrCreateJobVariableAssignmentsForGraph(const UMovieGraphConfig* InGraph);
+	/**
+	* This method will return you the object which contains variable overrides for the Primary Graph assigned to this job. You need to provide
+	* a pointer to the exact graph you want (as the Primary Graph may contain sub-graphs, and those sub-graphs can have their own variables),
+	* though this will be the same as GetGraphPreset() if you're not using any sub-graphs, or your variables only exist on the Primary graph.
+	* 
+	* If you want to override a variable on the primary graph but only for a specific shot, you should get the UMoviePipelineExecutorShot and
+	* call that classes version of this function, except passing True for the extra boolean.  See comment on that function for more details.
+	* 
+	* @param InGraph - The graph asset to get the Job Override values for. Should be the graph the variables you want to edit are defined on,
+	* which can either be the primary graph (GetGraphPreset()) or one of the sub-graphs it points to (as sub-graphs can contain their own
+	* variables which are all shown at the top level job in the Editor UI).
+	* @return A container object which holds a copy of the variables for the specified Graph Asset that can be used to override their values
+	* on jobs without actually editing the default asset.
+	*/
+	UFUNCTION(BlueprintCallable, Category = "Movie Render Pipeline", DisplayName = "Get or Create Variable Overrides", meta=(ScriptName ="GetOrCreateVariableOverrides"))
+	UMovieJobVariableAssignmentContainer* GetOrCreateJobVariableAssignmentsForGraph(const UMovieGraphConfig* InGraph);
 
 public:
 	// UObject Interface
