Localization: Fixed a duplicate localization key in FMovieGraphNamedResolutionCustomization.

diff --git a/Engine/Plugins/MovieScene/MovieRenderPipeline/Source/MovieRenderPipelineEditor/Private/Customizations/Graph/MovieGraphNamedResolutionCustomization.cpp b/Engine/Plugins/MovieScene/MovieRenderPipeline/Source/MovieRenderPipelineEditor/Private/Customizations/Graph/MovieGraphNamedResolutionCustomization.cpp
--- a/Engine/Plugins/MovieScene/MovieRenderPipeline/Source/MovieRenderPipelineEditor/Private/Customizations/Graph/MovieGraphNamedResolutionCustomization.cpp
+++ b/Engine/Plugins/MovieScene/MovieRenderPipeline/Source/MovieRenderPipelineEditor/Private/Customizations/Graph/MovieGraphNamedResolutionCustomization.cpp
@@ -346,7 +346,7 @@ int32 FMovieGraphNamedResolutionCustomization::GetCurrentlySelectedResolutionByA
 void FMovieGraphNamedResolutionCustomization::OnCustomSliderBeginMovement()
 {
 	// Wrap all interactive changes into a single commit while dragging.
-	InteractiveResolutionEditTransaction = MakeUnique<FScopedTransaction>(LOCTEXT("ChangeResolution_Transaction", "Edit Resolution"));
+	InteractiveResolutionEditTransaction = MakeUnique<FScopedTransaction>(LOCTEXT("Transaction_EditResolution", "Edit Resolution"));
 }
 
 void FMovieGraphNamedResolutionCustomization::OnCustomSliderEndMovement(uint32 NewValue)
@@ -366,7 +366,7 @@ void FMovieGraphNamedResolutionCustomization::OnCustomSliderValueChanged(uint32
 	// will be called, so we need to make a transaction here. This function is also called when being
 	// dragged, so we use OnCustomSliderBeginMovement/OnCustomSliderEndMovement to create a transaction
 	// which wraps all the smaller transactions to group them together into one undo/redo event.
-	FScopedTransaction Transaction(LOCTEXT("ChangeResolution_Transaction", "Edit Resolution"));
+	FScopedTransaction Transaction(LOCTEXT("Transaction_EditResolution", "Edit Resolution"));
 
 	// Check to see if we need to convert from an existing profile resolution to a custom one,
 	// otherwise the UI won't update with our new values.
@@ -499,7 +499,7 @@ void FMovieGraphNamedResolutionCustomization::OnComboBoxSelectionChanged(const F
 	if (PinnedStructPropertyHandle->GetValueData(StructData) == FPropertyAccess::Success && StructData != nullptr)
 	{
 		// We have to call all of these by hand because we're directly editing the memory and not using property handles.
-		FScopedTransaction ScopedTransaction(LOCTEXT("ChangeResolution_Transaction", "Assign Preset"));
+		FScopedTransaction ScopedTransaction(LOCTEXT("Transaction_AssignPreset", "Assign Preset"));
 		PinnedStructPropertyHandle->NotifyPreChange();
 
 		PinnedStructPropertyHandle->EnumerateRawData([NamedResolution](void* RawData, const int32 /*DataIndex*/, const int32 /*NumDatas*/)
diff --git a/Engine/Plugins/MovieScene/MovieRenderPipeline/Source/MovieRenderPipelineEditor/Private/Customizations/Graph/MovieGraphNamedResolutionCustomization.h b/Engine/Plugins/MovieScene/MovieRenderPipeline/Source/MovieRenderPipelineEditor/Private/Customizations/Graph/MovieGraphNamedResolutionCustomization.h
--- a/Engine/Plugins/MovieScene/MovieRenderPipeline/Source/MovieRenderPipelineEditor/Private/Customizations/Graph/MovieGraphNamedResolutionCustomization.h
+++ b/Engine/Plugins/MovieScene/MovieRenderPipeline/Source/MovieRenderPipelineEditor/Private/Customizations/Graph/MovieGraphNamedResolutionCustomization.h
@@ -153,5 +153,3 @@ protected:
 	*/
 	TUniquePtr<FScopedTransaction> InteractiveResolutionEditTransaction;
 };
-
-#undef LOCTEXT_NAMESPACE
