MovieGraph: Fixed missing Export macros for FMovieGraphImagePassBase and associated data structures.
Also fixed the ViewActor not being set in the Movie Graph which was a regression compared to the Movie Render Queue.

PR #11696

diff --git a/Engine/Plugins/MovieScene/MovieRenderPipeline/Source/MovieRenderPipelineCore/Public/Graph/MovieGraphDefaultRenderer.h b/Engine/Plugins/MovieScene/MovieRenderPipeline/Source/MovieRenderPipelineCore/Public/Graph/MovieGraphDefaultRenderer.h
--- a/Engine/Plugins/MovieScene/MovieRenderPipeline/Source/MovieRenderPipelineCore/Public/Graph/MovieGraphDefaultRenderer.h
+++ b/Engine/Plugins/MovieScene/MovieRenderPipeline/Source/MovieRenderPipelineCore/Public/Graph/MovieGraphDefaultRenderer.h
@@ -52,7 +52,10 @@ namespace UE::MovieGraph::DefaultRenderer
 			, DoFSensorScale(1.0f)
 		{}
 
+		/** Minimal view info describing the position and orientation in the world. */
 		FMinimalViewInfo ViewInfo;
+
+		/** The actor that is considered the "owner" of the view. Used by the renderer to respect bOnlyOwnerSee/bOwnerNoSee flags. */
 		class AActor* ViewActor;
 
 		/** The name to use for the {camera_name} token. Filled out by the renderer where possible. */
diff --git a/Engine/Plugins/MovieScene/MovieRenderPipeline/Source/MovieRenderPipelineRenderPasses/Private/Graph/Renderers/MovieGraphDeferredPass.cpp b/Engine/Plugins/MovieScene/MovieRenderPipeline/Source/MovieRenderPipelineRenderPasses/Private/Graph/Renderers/MovieGraphDeferredPass.cpp
--- a/Engine/Plugins/MovieScene/MovieRenderPipeline/Source/MovieRenderPipelineRenderPasses/Private/Graph/Renderers/MovieGraphDeferredPass.cpp
+++ b/Engine/Plugins/MovieScene/MovieRenderPipeline/Source/MovieRenderPipelineRenderPasses/Private/Graph/Renderers/MovieGraphDeferredPass.cpp
@@ -190,6 +190,7 @@ void FMovieGraphDeferredPass::Render(const FMovieGraphTraversalContext& InFrameT
 		if (LocalPlayerController->PlayerCameraManager)
 		{
 			CameraInfo.ViewInfo = LocalPlayerController->PlayerCameraManager->GetCameraCacheView();
+			CameraInfo.ViewActor = LocalPlayerController->GetViewTarget();
 		}
 
 
diff --git a/Engine/Plugins/MovieScene/MovieRenderPipeline/Source/MovieRenderPipelineRenderPasses/Public/Graph/Renderers/MovieGraphDeferredPass.h b/Engine/Plugins/MovieScene/MovieRenderPipeline/Source/MovieRenderPipelineRenderPasses/Public/Graph/Renderers/MovieGraphDeferredPass.h
--- a/Engine/Plugins/MovieScene/MovieRenderPipeline/Source/MovieRenderPipelineRenderPasses/Public/Graph/Renderers/MovieGraphDeferredPass.h
+++ b/Engine/Plugins/MovieScene/MovieRenderPipeline/Source/MovieRenderPipelineRenderPasses/Public/Graph/Renderers/MovieGraphDeferredPass.h
@@ -5,7 +5,7 @@
 
 namespace UE::MovieGraph::Rendering
 {
-	struct FMovieGraphDeferredPass : public FMovieGraphImagePassBase
+	struct MOVIERENDERPIPELINERENDERPASSES_API FMovieGraphDeferredPass : public FMovieGraphImagePassBase
 	{
 		// FMovieGraphImagePassBase Interface
 		virtual void Setup(TWeakObjectPtr<UMovieGraphDefaultRenderer> InRenderer, TWeakObjectPtr<UMovieGraphImagePassBaseNode> InRenderPassNode, const FMovieGraphRenderPassLayerData& InLayer) override;
diff --git a/Engine/Plugins/MovieScene/MovieRenderPipeline/Source/MovieRenderPipelineRenderPasses/Public/Graph/Renderers/MovieGraphImagePassBase.h b/Engine/Plugins/MovieScene/MovieRenderPipeline/Source/MovieRenderPipelineRenderPasses/Public/Graph/Renderers/MovieGraphImagePassBase.h
--- a/Engine/Plugins/MovieScene/MovieRenderPipeline/Source/MovieRenderPipelineRenderPasses/Public/Graph/Renderers/MovieGraphImagePassBase.h
+++ b/Engine/Plugins/MovieScene/MovieRenderPipeline/Source/MovieRenderPipelineRenderPasses/Public/Graph/Renderers/MovieGraphImagePassBase.h
@@ -10,7 +10,7 @@ class UMovieGraphImagePassBaseNode;
 
 namespace UE::MovieGraph::Rendering
 {
-	struct FMovieGraphRenderDataAccumulationArgs
+	struct MOVIERENDERPIPELINERENDERPASSES_API FMovieGraphRenderDataAccumulationArgs
 	{
 	public:
 		TWeakPtr<FImageOverlappedAccumulator, ESPMode::ThreadSafe> ImageAccumulator;
@@ -25,7 +25,7 @@ namespace UE::MovieGraph::Rendering
 	// Forward Declare
 	void AccumulateSample_TaskThread(TUniquePtr<FImagePixelData>&& InPixelData, const ::UE::MovieGraph::FMovieGraphSampleState InSampleState, const FMovieGraphRenderDataAccumulationArgs& InAccumulationParams);
 
-	struct FViewFamilyInitData
+	struct MOVIERENDERPIPELINERENDERPASSES_API FViewFamilyInitData
 	{
 		FViewFamilyInitData()
 			: RenderTarget(nullptr)
@@ -52,7 +52,7 @@ namespace UE::MovieGraph::Rendering
 	
 
 
-	struct FMovieGraphImagePassBase
+	struct MOVIERENDERPIPELINERENDERPASSES_API FMovieGraphImagePassBase
 	{
 		FMovieGraphImagePassBase() = default;
 		virtual ~FMovieGraphImagePassBase() = default;
diff --git a/Engine/Plugins/MovieScene/MovieRenderPipeline/Source/MovieRenderPipelineRenderPasses/Public/Graph/Renderers/MovieGraphPathTracerPass.h b/Engine/Plugins/MovieScene/MovieRenderPipeline/Source/MovieRenderPipelineRenderPasses/Public/Graph/Renderers/MovieGraphPathTracerPass.h
--- a/Engine/Plugins/MovieScene/MovieRenderPipeline/Source/MovieRenderPipelineRenderPasses/Public/Graph/Renderers/MovieGraphPathTracerPass.h
+++ b/Engine/Plugins/MovieScene/MovieRenderPipeline/Source/MovieRenderPipelineRenderPasses/Public/Graph/Renderers/MovieGraphPathTracerPass.h
@@ -5,7 +5,7 @@
 
 namespace UE::MovieGraph::Rendering
 {
-	struct FMovieGraphPathTracerPass : public FMovieGraphDeferredPass
+	struct MOVIERENDERPIPELINERENDERPASSES_API FMovieGraphPathTracerPass : public FMovieGraphDeferredPass
 	{
 		virtual UMovieGraphImagePassBaseNode* GetParentNode(UMovieGraphEvaluatedConfig* InConfig) const override;
 		virtual void ApplyMovieGraphOverridesToSceneView(TSharedRef<FSceneViewFamilyContext> InOutFamily, const FViewFamilyInitData& InInitData, const UE::MovieGraph::DefaultRenderer::FCameraInfo& InCameraInfo) const override;
