MoviePipeline: Fixed an issue that caused the Show Flags UI to go blank the second time it is displayed in the details panel.

For some reason the show flag Display Name is empty the second time show flags are fetched from the show flag API. This is a temporary fix to work around this issue.

#rnx

diff --git a/Engine/Plugins/MovieScene/MovieRenderPipeline/Source/MovieRenderPipelineEditor/Private/Customizations/Graph/MovieGraphShowFlagsCustomization.h b/Engine/Plugins/MovieScene/MovieRenderPipeline/Source/MovieRenderPipelineEditor/Private/Customizations/Graph/MovieGraphShowFlagsCustomization.h
--- a/Engine/Plugins/MovieScene/MovieRenderPipeline/Source/MovieRenderPipelineEditor/Private/Customizations/Graph/MovieGraphShowFlagsCustomization.h
+++ b/Engine/Plugins/MovieScene/MovieRenderPipeline/Source/MovieRenderPipelineEditor/Private/Customizations/Graph/MovieGraphShowFlagsCustomization.h
@@ -76,9 +76,12 @@ protected:
 			for (const FShowFlagData& ShowFlag : *ShowFlags)
 			{
 				const uint32 ShowFlagIndex = ShowFlag.EngineShowFlagIndex;
+
+				// Temporary workaround: ShowFlag.DisplayName can sometimes be blank, so the non-display name is used instead.
+				const FText ShowFlagText = FText::FromName(ShowFlag.ShowFlagName);
 				
 				FlagGroup.AddWidgetRow()
-				.FilterString(ShowFlag.DisplayName)
+				.FilterString(ShowFlagText)
 				.NameContent()
 				[
 					SNew(SHorizontalBox)
@@ -105,7 +108,7 @@ protected:
 					[
 						SNew(STextBlock)
 						.IsEnabled_Lambda([ShowFlagObject, ShowFlagIndex]() { return ShowFlagObject->IsShowFlagOverridden(ShowFlagIndex); })
-						.Text(ShowFlag.DisplayName)
+						.Text(ShowFlagText)
 						.Font(CustomizationUtils.GetRegularFont())
 					]
 				]
