MovieRenderPipeline: Fix broken channel names on linux due to wrong format specifier

The channel name is a plain const char* string, so requires %hs with the UE printf functions.

Reported via UDN

diff --git a/Engine/Plugins/MovieScene/MovieRenderPipeline/Source/MovieRenderPipelineRenderPasses/Private/MoviePipelineEXROutput.cpp b/Engine/Plugins/MovieScene/MovieRenderPipeline/Source/MovieRenderPipelineRenderPasses/Private/MoviePipelineEXROutput.cpp
--- a/Engine/Plugins/MovieScene/MovieRenderPipeline/Source/MovieRenderPipelineRenderPasses/Private/MoviePipelineEXROutput.cpp
+++ b/Engine/Plugins/MovieScene/MovieRenderPipeline/Source/MovieRenderPipelineRenderPasses/Private/MoviePipelineEXROutput.cpp
@@ -300,7 +300,7 @@ static FString GetChannelName(const FString& InLayerName, const int32 InChannelI
 
 	if (InLayerName.Len() > 0)
 	{
-		return FString::Printf(TEXT("%s.%s"), *InLayerName, ChannelNames[InChannelIndex]);
+		return FString::Printf(TEXT("%s.%hs"), *InLayerName, ChannelNames[InChannelIndex]);
 	}
 	else
 	{
