MoviePipeline: Added ability to enable/disable individual collections within a Modifier node.

diff --git a/Engine/Plugins/MovieScene/MovieRenderPipeline/Source/MovieRenderPipelineCore/Private/Graph/MovieGraphPipeline.cpp b/Engine/Plugins/MovieScene/MovieRenderPipeline/Source/MovieRenderPipelineCore/Private/Graph/MovieGraphPipeline.cpp
--- a/Engine/Plugins/MovieScene/MovieRenderPipeline/Source/MovieRenderPipelineCore/Private/Graph/MovieGraphPipeline.cpp
+++ b/Engine/Plugins/MovieScene/MovieRenderPipeline/Source/MovieRenderPipelineCore/Private/Graph/MovieGraphPipeline.cpp
@@ -452,7 +452,13 @@ void UMovieGraphPipeline::UpdateLayerContentsInRenderLayerSubsystem(const UMovie
 						if (CollectionNode->Collection->GetCollectionName() == ModifiedCollectionName)
 						{
 							bFoundModifiedCollection = true;
-							ModifierCollections.Add(CollectionNode->Collection);
+
+							// Collections can be disabled within a modifier; only include it if the collection is enabled
+							if (ModifierNode->IsCollectionEnabled(ModifiedCollectionName))
+							{
+								ModifierCollections.Add(CollectionNode->Collection);
+							}
+							
 							break;
 						}
 					}
diff --git a/Engine/Plugins/MovieScene/MovieRenderPipeline/Source/MovieRenderPipelineCore/Private/Graph/Nodes/MovieGraphModifierNode.cpp b/Engine/Plugins/MovieScene/MovieRenderPipeline/Source/MovieRenderPipelineCore/Private/Graph/Nodes/MovieGraphModifierNode.cpp
--- a/Engine/Plugins/MovieScene/MovieRenderPipeline/Source/MovieRenderPipelineCore/Private/Graph/Nodes/MovieGraphModifierNode.cpp
+++ b/Engine/Plugins/MovieScene/MovieRenderPipeline/Source/MovieRenderPipelineCore/Private/Graph/Nodes/MovieGraphModifierNode.cpp
@@ -229,6 +229,9 @@ bool UMovieGraphModifierNode::RemoveCollection(const FName& InCollectionName)
 #if WITH_EDITOR
 	Modify();
 #endif
+
+	// Remove from disabled collections (in case this collection was disabled when it was deleted)
+	DisabledCollections.Remove(InCollectionName);
 	
 	return Collections.Remove(InCollectionName) > 0;
 }
@@ -238,4 +241,25 @@ const TArray<FName>& UMovieGraphModifierNode::GetCollections() const
 	return Collections;
 }
 
+void UMovieGraphModifierNode::SetCollectionEnabled(const FName& InCollectionName, const bool bIsCollectionEnabled)
+{
+#if WITH_EDITOR
+	Modify();
+#endif
+	
+	if (bIsCollectionEnabled)
+	{
+		DisabledCollections.Remove(InCollectionName);
+	}
+	else
+	{
+		DisabledCollections.Add(InCollectionName);
+	}
+}
+
+bool UMovieGraphModifierNode::IsCollectionEnabled(const FName& InCollectionName) const
+{
+	return !DisabledCollections.Contains(InCollectionName);
+}
+
 #undef LOCTEXT_NAMESPACE // "MovieGraph"
\ No newline at end of file
diff --git a/Engine/Plugins/MovieScene/MovieRenderPipeline/Source/MovieRenderPipelineCore/Public/Graph/MovieGraphSharedWidgets.h b/Engine/Plugins/MovieScene/MovieRenderPipeline/Source/MovieRenderPipelineCore/Public/Graph/MovieGraphSharedWidgets.h
--- a/Engine/Plugins/MovieScene/MovieRenderPipeline/Source/MovieRenderPipelineCore/Public/Graph/MovieGraphSharedWidgets.h
+++ b/Engine/Plugins/MovieScene/MovieRenderPipeline/Source/MovieRenderPipelineCore/Public/Graph/MovieGraphSharedWidgets.h
@@ -4,12 +4,15 @@
 
 #include "CoreMinimal.h"
 #include "Framework/Application/SlateApplication.h"
-#include "Widgets/SCompoundWidget.h"
+#include "Widgets/Input/SCheckBox.h"
 #include "Widgets/Input/SComboButton.h"
+#include "Widgets/SCompoundWidget.h"
 #include "Widgets/Text/STextBlock.h"
 #include "Widgets/Views/ITableRow.h"
 #include "Widgets/Views/SListView.h"
 
+#define LOCTEXT_NAMESPACE "MovieGraphSharedWidgets"
+
 // NOTE: There should not be widgetry defined in core. To fix this, the condition group queries in the render layer subsystem need to be refactored to
 // expose their custom widgetry in MovieRenderPipelineEditor.
 
@@ -24,8 +27,11 @@ public:
 	DECLARE_DELEGATE_RetVal_OneParam(const FSlateBrush*, FGetRowIcon, ListType);
 	DECLARE_DELEGATE_RetVal_OneParam(FText, FGetRowText, ListType);
 	DECLARE_DELEGATE_OneParam(FOnDelete, ListType);
+	DECLARE_DELEGATE_RetVal_OneParam(bool, FGetRowEnableState, ListType);
+	DECLARE_DELEGATE_TwoParams(FSetRowEnableState, ListType, bool);
 	
 	SLATE_BEGIN_ARGS(SMovieGraphSimpleList<ListType>)
+		: _ShowEnableDisable(false)
 		{}
 		/** The source of data that the list will display. */
 		SLATE_ATTRIBUTE(TArray<ListType>*, DataSource)
@@ -44,6 +50,15 @@ public:
 
 		/** Invoked when a delete operation is performed. */
 		SLATE_EVENT(FOnDelete, OnDelete)
+
+		/** Whether the enable/disable checkbox should be displayed. */
+		SLATE_ATTRIBUTE(bool, ShowEnableDisable)
+
+		/** Gets the enable state of the row. Called if ShowEnableDisable is turned on. */
+		SLATE_EVENT(FGetRowEnableState, OnGetRowEnableState)
+
+		/** Sets the enable state of the row. */
+		SLATE_EVENT(FSetRowEnableState, OnSetRowEnableState)
 	SLATE_END_ARGS()
 	
 	void Construct(const FArguments& InArgs)
@@ -54,6 +69,9 @@ public:
 		OnGetRowIcon = InArgs._OnGetRowIcon;
 		OnGetRowText = InArgs._OnGetRowText;
 		OnDelete = InArgs._OnDelete;
+		bShowEnableDisable = InArgs._ShowEnableDisable.Get();
+		OnGetRowEnableState = InArgs._OnGetRowEnableState;
+		OnSetRowEnableState = InArgs._OnSetRowEnableState;
 
 		ChildSlot
 		.HAlign(HAlign_Fill)
@@ -122,6 +140,17 @@ private:
 		return FText();
 	}
 
+	/** Determines if the row that displays the given data should be enabled. */
+	bool IsRowEnabled(const ListType& InListData) const
+	{
+		if (OnGetRowEnableState.IsBound())
+		{
+			return OnGetRowEnableState.Execute(InListData);
+		}
+
+		return true;
+	}
+
 	/** Handles the delete operation. */
 	FReply HandleDelete(const FGeometry& InGeometry, const FKeyEvent& InKeyEvent) const
 	{
@@ -143,6 +172,8 @@ private:
 	/** Generates a row in the list for the specified data. */
 	TSharedRef<ITableRow> GenerateRow(ListType InListData, const TSharedRef<STableViewBase>& InOwnerTable) const
 	{
+		const FSlateBrush* RowIcon = GetRowIcon(InListData);
+		
 		return
 			SNew(STableRow<ListType>, InOwnerTable)
 			.Style(FAppStyle::Get(), "TableView.AlternatingRow")
@@ -151,11 +182,42 @@ private:
 				SNew(SHorizontalBox)
 				+ SHorizontalBox::Slot()
 				.VAlign(VAlign_Center)
+				.AutoWidth()
 				.Padding(7.f, 5.f, 7.f, 5.f)
+				[
+					SNew(SCheckBox)
+					.ToolTipText_Lambda([this, InListData]()
+					{
+						return IsRowEnabled(InListData)
+							? LOCTEXT("CollectionEnabled", "This collection is enabled and will be modified by this Modifier node.")
+							: LOCTEXT("CollectionDisabled", "This collection is disabled and will not be modified by this Modifier node.");
+					})
+					.Visibility_Lambda([this]()
+					{
+						return bShowEnableDisable ? EVisibility::Visible : EVisibility::Collapsed;
+					})
+					.IsChecked_Lambda([this, InListData]()
+					{
+						return IsRowEnabled(InListData) ? ECheckBoxState::Checked : ECheckBoxState::Unchecked;
+					})
+					.OnCheckStateChanged_Lambda([this, InListData](ECheckBoxState NewState)
+					{
+						if (OnSetRowEnableState.IsBound())
+						{
+							OnSetRowEnableState.Execute(InListData, NewState == ECheckBoxState::Checked);
+						}
+					})
+				]
+				
+				+ SHorizontalBox::Slot()
+				.VAlign(VAlign_Center)
+				.Padding(0.f, 5.f, 7.f, 5.f)
 				.AutoWidth()
 				[
 					SNew(SImage)
-					.Image(GetRowIcon(InListData))
+					.IsEnabled_Lambda([this, InListData]() { return IsRowEnabled(InListData); })
+					.Visibility_Lambda([RowIcon]() { return RowIcon ? EVisibility::Visible : EVisibility::Collapsed; })
+					.Image(RowIcon)
 				]
 				
 				+ SHorizontalBox::Slot()
@@ -163,6 +225,7 @@ private:
 				.HAlign(HAlign_Fill)
 				[
 					SNew(STextBlock)
+					.IsEnabled_Lambda([this, InListData]() { return IsRowEnabled(InListData); })
 					.Font(FAppStyle::GetFontStyle(TEXT("PropertyWindow.NormalFont")))
 					.Text(GetRowText(InListData))
 				]
@@ -174,7 +237,12 @@ private:
 	FGetRowIcon OnGetRowIcon;
 	FGetRowText OnGetRowText;
 	FOnDelete OnDelete;
+	bool bShowEnableDisable = false;
+	FGetRowEnableState OnGetRowEnableState;
+	FSetRowEnableState OnSetRowEnableState;
 	FText DataType;
 	FText DataTypePlural;
 	TArray<ListType>* DataSource = nullptr;
-};
\ No newline at end of file
+};
+
+#undef LOCTEXT_NAMESPACE
\ No newline at end of file
diff --git a/Engine/Plugins/MovieScene/MovieRenderPipeline/Source/MovieRenderPipelineCore/Public/Graph/Nodes/MovieGraphModifierNode.h b/Engine/Plugins/MovieScene/MovieRenderPipeline/Source/MovieRenderPipelineCore/Public/Graph/Nodes/MovieGraphModifierNode.h
--- a/Engine/Plugins/MovieScene/MovieRenderPipeline/Source/MovieRenderPipelineCore/Public/Graph/Nodes/MovieGraphModifierNode.h
+++ b/Engine/Plugins/MovieScene/MovieRenderPipeline/Source/MovieRenderPipelineCore/Public/Graph/Nodes/MovieGraphModifierNode.h
@@ -84,6 +84,17 @@ public:
 	UFUNCTION(BlueprintCallable, Category = "Modifiers")
 	const TArray<FName>& GetCollections() const;
 
+	/**
+	 * Sets the enable state (within this modifier) of the collection with the given name. Disabled collections will not be modified by this modifier
+	 * node. Collections that are added to the modifier are enabled by default.
+	 */
+	UFUNCTION(BlueprintCallable, Category = "Modifiers")
+	void SetCollectionEnabled(const FName& InCollectionName, const bool bIsCollectionEnabled);
+
+	/** Gets the enable state (within this modifier) of the collection with the given name. */
+	UFUNCTION(BlueprintCallable, Category = "Modifiers")
+	bool IsCollectionEnabled(const FName& InCollectionName) const;
+
 public:
 	UPROPERTY()
 	uint8 bOverride_ModifierName : 1 = 1;	// Always merge the modifier name, no need for the user to do this explicitly
@@ -93,11 +104,16 @@ public:
 	FString ModifierName;
 
 private:
+	// These private override properties exist so that the associated non-override properties are merged properly
+	// by UMovieGraphMergeableModifierContainer
 	UPROPERTY()
 	uint8 bOverride_Collections : 1 = 1; //-V570
 
 	UPROPERTY()
 	uint8 bOverride_ModifiersContainer : 1 = 1; //-V570
+
+	UPROPERTY()
+	uint8 bOverride_DisabledCollections : 1 = 1; //-V570
 	
 	/** The names of collections being modified. */
 	UPROPERTY()
@@ -106,4 +122,8 @@ private:
 	/** The modifiers this node should run. */
 	UPROPERTY(meta=(DisplayName="Modifiers"))
 	TObjectPtr<UMovieGraphMergeableModifierContainer> ModifiersContainer;
+
+	/** The collections on this node that have been disabled. */
+	UPROPERTY()
+	TSet<FName> DisabledCollections;
 };
\ No newline at end of file
diff --git a/Engine/Plugins/MovieScene/MovieRenderPipeline/Source/MovieRenderPipelineEditor/Private/Customizations/Graph/MovieGraphModifiersCustomization.cpp b/Engine/Plugins/MovieScene/MovieRenderPipeline/Source/MovieRenderPipelineEditor/Private/Customizations/Graph/MovieGraphModifiersCustomization.cpp
--- a/Engine/Plugins/MovieScene/MovieRenderPipeline/Source/MovieRenderPipelineEditor/Private/Customizations/Graph/MovieGraphModifiersCustomization.cpp
+++ b/Engine/Plugins/MovieScene/MovieRenderPipeline/Source/MovieRenderPipelineEditor/Private/Customizations/Graph/MovieGraphModifiersCustomization.cpp
@@ -297,8 +297,21 @@ void FMovieGraphModifiersCustomization::CustomizeDetails(IDetailLayoutBuilder& I
 				CollectionsList->Refresh();
 			}
 		})
-		.OnGetRowIcon_Static(&GetCollectionRowIcon)
 		.OnGetRowText_Static(&GetCollectionRowText)
+		.ShowEnableDisable(true)
+		.OnGetRowEnableState_Lambda([ModifierNode](FName InCollectionName)
+		{
+			return ModifierNode.IsValid() ? ModifierNode.Get()->IsCollectionEnabled(InCollectionName) : true;
+		})
+		.OnSetRowEnableState_Lambda([ModifierNode](FName InCollectionName, bool bNewEnableState)
+		{
+			if (ModifierNode.IsValid())
+			{
+				const FScopedTransaction Transaction(LOCTEXT("ChangeCollectionEnableState", "Change Collection Enable State"));
+				
+				ModifierNode.Get()->SetCollectionEnabled(InCollectionName, bNewEnableState);
+			}
+		})
 	];
 
 	// For all modifiers added to the node, add a category for each, and add each modifier's EditAnywhere properties to the category
