MoviePipeline: Filter the autowire context menu in the graph to restrict it to nodes that can only be connected on that branch type or pin type.

diff --git a/Engine/Plugins/MovieScene/MovieRenderPipeline/Source/MovieRenderPipelineEditor/Private/Graph/MovieGraphSchema.cpp b/Engine/Plugins/MovieScene/MovieRenderPipeline/Source/MovieRenderPipelineEditor/Private/Graph/MovieGraphSchema.cpp
--- a/Engine/Plugins/MovieScene/MovieRenderPipeline/Source/MovieRenderPipelineEditor/Private/Graph/MovieGraphSchema.cpp
+++ b/Engine/Plugins/MovieScene/MovieRenderPipeline/Source/MovieRenderPipelineEditor/Private/Graph/MovieGraphSchema.cpp
@@ -144,8 +144,11 @@ bool UMovieGraphSchema::IsConnectionToBranchAllowed(const UEdGraphPin* InputPin,
 
 void UMovieGraphSchema::AddExtraMenuActions(FGraphActionMenuBuilder& ActionMenuBuilder) const
 {
-	// Comment action
-	ActionMenuBuilder.AddAction(CreateCommentMenuAction());
+	// Comment action. Only add if there's no FromPin (ie, no connection is currently being built).
+	if (!ActionMenuBuilder.FromPin)
+	{
+		ActionMenuBuilder.AddAction(CreateCommentMenuAction());
+	}
 }
 
 TSharedRef<FMovieGraphSchemaAction_NewComment> UMovieGraphSchema::CreateCommentMenuAction() const
@@ -182,17 +185,29 @@ void UMovieGraphSchema::GetGraphContextActions(FGraphContextMenuBuilder& Context
 			// Add variable actions separately
 			continue;
 		}
-		if(PipelineNodeClass == UMovieGraphInputNode::StaticClass() ||
+		
+		if (PipelineNodeClass == UMovieGraphInputNode::StaticClass() ||
 			PipelineNodeClass == UMovieGraphOutputNode::StaticClass())
 		{
 			// Can't place Input and Output nodes manually.
 			continue;
 		}
 
-		// This can be used to sort whether or not an option shows up. For now there's no restrictions
-		// on where nodes can be made, but eventually we might check which branch they're on (if from pin)
-		// to filter out incompatible nodes.
-		// if (!ContextMenuBuilder.FromPin || ContextMenuBuilder.FromPin->Direction == EGPD_Input)
+		// Determine if this node type can be created in the branch that FromPin is in. FromPin is non-null if the node is being created and connected
+		// to an existing pin in one step (ie, the user is currently creating a connection).
+		bool bCanAppearInMenu = true;
+		if (ContextMenuBuilder.FromPin)
+		{
+			if (const UMovieGraphNode* FromNode = UE::MovieGraph::Private::GetGraphNodeFromEdPin(ContextMenuBuilder.FromPin))
+			{
+				const bool bBranchRestrictionIsOk = (PipelineNode->GetBranchRestriction() == EMovieGraphBranchRestriction::Any) ||
+					(FromNode->GetBranchRestriction() == PipelineNode->GetBranchRestriction());
+				
+				bCanAppearInMenu = bBranchRestrictionIsOk && (ContextMenuBuilder.FromPin->PinType.PinCategory == PC_Branch);
+			}
+		}
+		
+		if (bCanAppearInMenu)
 		{
 			const FText Name = PipelineNode->GetNodeTitle();
 			const FText Category = PipelineNode->GetMenuCategory();
@@ -208,7 +223,7 @@ void UMovieGraphSchema::GetGraphContextActions(FGraphContextMenuBuilder& Context
 	}
 
 	// Create an accessor node action for each variable the graph has
-	const bool bIncludeGlobal = true;
+	constexpr bool bIncludeGlobal = true;
 	for (const UMovieGraphVariable* Variable : RuntimeGraph->GetVariables(bIncludeGlobal))
 	{
 		const FText Name = FText::Format(LOCTEXT("CreateVariable_Name", "Get {0}"), FText::FromString(Variable->GetMemberName()));
@@ -218,7 +233,21 @@ void UMovieGraphSchema::GetGraphContextActions(FGraphContextMenuBuilder& Context
 		TSharedPtr<FMovieGraphSchemaAction> NewAction = MakeShared<FMovieGraphSchemaAction_NewVariableNode>(Category, Name, Variable->GetGuid(), Tooltip);
 		NewAction->NodeClass = UMovieGraphVariableNode::StaticClass();
 		
-		ContextMenuBuilder.AddAction(NewAction);
+		// Determine if this node can be created and connected to FromPin
+		bool bCanAppearInMenu = true;
+		if (ContextMenuBuilder.FromPin)
+		{
+			if (const UMovieGraphPin* FromPin = UE::MovieGraph::Private::GetGraphPinFromEdPin(ContextMenuBuilder.FromPin))
+			{
+				// Variable type and pin type must match
+				bCanAppearInMenu = (FromPin->Properties.Type == Variable->GetValueType()) && (FromPin->Properties.TypeObject == Variable->GetValueTypeObject());
+			}
+		}
+
+		if (bCanAppearInMenu)
+		{
+			ContextMenuBuilder.AddAction(NewAction);
+		}
 	}
 
 	AddExtraMenuActions(ContextMenuBuilder);
