MoviePipeline: Prevent undo from resetting the selection in the graph.

diff --git a/Engine/Plugins/MovieScene/MovieRenderPipeline/Source/MovieRenderPipelineEditor/Private/Widgets/Graph/SMovieGraphConfigPanel.cpp b/Engine/Plugins/MovieScene/MovieRenderPipeline/Source/MovieRenderPipelineEditor/Private/Widgets/Graph/SMovieGraphConfigPanel.cpp
--- a/Engine/Plugins/MovieScene/MovieRenderPipeline/Source/MovieRenderPipelineEditor/Private/Widgets/Graph/SMovieGraphConfigPanel.cpp
+++ b/Engine/Plugins/MovieScene/MovieRenderPipeline/Source/MovieRenderPipelineEditor/Private/Widgets/Graph/SMovieGraphConfigPanel.cpp
@@ -662,7 +662,6 @@ void SMoviePipelineGraphPanel::PostUndo(bool bSuccess)
 			// This unfortunately causes the graph to flicker for one frame, but all Blueprint Graphs do this.
 			// This is needed so that when a node is pasted, then undone, the widgets on the graph for the now
 			// deleted nodes also get removed, otherwise we try to draw them and they're pointed to invalid uobjects.
-			GraphEditorWidget->ClearSelectionSet();
 			GraphEditorWidget->NotifyGraphChanged();
 
 			FSlateApplication::Get().DismissAllMenus();
