Change return value of UMoviePipelineQueue::CopyFrom to return a MoviePipelineQueue* to align with rest of similar functions and make chaining functions easier in python by returning an object

diff --git a/Engine/Plugins/MovieScene/MovieRenderPipeline/Source/MovieRenderPipelineCore/Private/MoviePipelineQueue.cpp b/Engine/Plugins/MovieScene/MovieRenderPipeline/Source/MovieRenderPipelineCore/Private/MoviePipelineQueue.cpp
--- a/Engine/Plugins/MovieScene/MovieRenderPipeline/Source/MovieRenderPipelineCore/Private/MoviePipelineQueue.cpp
+++ b/Engine/Plugins/MovieScene/MovieRenderPipeline/Source/MovieRenderPipelineCore/Private/MoviePipelineQueue.cpp
@@ -161,12 +161,12 @@ UMoviePipelineExecutorJob* UMoviePipelineQueue::DuplicateJob(UMoviePipelineExecu
 	return NewJob;
 }
 
-void UMoviePipelineQueue::CopyFrom(UMoviePipelineQueue* InQueue)
+UMoviePipelineQueue* UMoviePipelineQueue::CopyFrom(UMoviePipelineQueue* InQueue)
 {
 	if (!InQueue)
 	{
 		UE_LOG(LogMovieRenderPipeline, Warning, TEXT("Cannot copy the contents of a null queue."));
-		return;
+		return nullptr;
 	}
 
 	// The copy should reflect the input queue's origin (ie, the queue asset it was originally based off of). Setting
@@ -186,6 +186,7 @@ void UMoviePipelineQueue::CopyFrom(UMoviePipelineQueue* InQueue)
 	// Ensure the serial number gets bumped at least once so the UI refreshes in case
 	// the queue we are copying from was empty.
 	QueueSerialNumber++;
+	return this;
 }
 
 void UMoviePipelineQueue::SetJobIndex(UMoviePipelineExecutorJob* InJob, int32 Index)
diff --git a/Engine/Plugins/MovieScene/MovieRenderPipeline/Source/MovieRenderPipelineCore/Public/MoviePipelineQueue.h b/Engine/Plugins/MovieScene/MovieRenderPipeline/Source/MovieRenderPipelineCore/Public/MoviePipelineQueue.h
--- a/Engine/Plugins/MovieScene/MovieRenderPipeline/Source/MovieRenderPipelineCore/Public/MoviePipelineQueue.h
+++ b/Engine/Plugins/MovieScene/MovieRenderPipeline/Source/MovieRenderPipelineCore/Public/MoviePipelineQueue.h
@@ -696,9 +696,10 @@ public:
 
 	/** 
 	* Replace the contents of this queue with a copy of the contents from another queue. 
+	* Returns a pointer to this queue if the copy was successful, else nullptr.
 	*/
 	UFUNCTION(BlueprintCallable, Category = "Movie Render Pipeline|Queue")
-	void CopyFrom(UMoviePipelineQueue* InQueue);
+	UMoviePipelineQueue* CopyFrom(UMoviePipelineQueue* InQueue);
 	
 	/* Set the index of the given job */
 	UFUNCTION(BlueprintCallable, Category = "Movie Render Pipeline|Queue")
