MoviePipeline: Added Actor Layer and Data Layer queries for collections in the graph.

Also added a new shared widget, SMovieGraphSimplePicker, which is now used in both the Modifier node and the Actor Layer query's layer picker.

diff --git a/Engine/Plugins/MovieScene/MovieRenderPipeline/Source/MovieRenderPipelineCore/MovieRenderPipelineCore.Build.cs b/Engine/Plugins/MovieScene/MovieRenderPipeline/Source/MovieRenderPipelineCore/MovieRenderPipelineCore.Build.cs
--- a/Engine/Plugins/MovieScene/MovieRenderPipeline/Source/MovieRenderPipelineCore/MovieRenderPipelineCore.Build.cs
+++ b/Engine/Plugins/MovieScene/MovieRenderPipeline/Source/MovieRenderPipelineCore/MovieRenderPipelineCore.Build.cs
@@ -47,7 +47,9 @@ public class MovieRenderPipelineCore : ModuleRules
 		        {
 			        "ClassViewer",
 			        "ContentBrowserData",
+			        "DataLayerEditor",
 			        "EditorWidgets",
+			        "Layers",
 			        "SceneOutliner",
 			        "UnrealEd",
 		        });
diff --git a/Engine/Plugins/MovieScene/MovieRenderPipeline/Source/MovieRenderPipelineCore/Private/Graph/MovieGraphRenderLayerSubsystem.cpp b/Engine/Plugins/MovieScene/MovieRenderPipeline/Source/MovieRenderPipelineCore/Private/Graph/MovieGraphRenderLayerSubsystem.cpp
--- a/Engine/Plugins/MovieScene/MovieRenderPipeline/Source/MovieRenderPipelineCore/Private/Graph/MovieGraphRenderLayerSubsystem.cpp
+++ b/Engine/Plugins/MovieScene/MovieRenderPipeline/Source/MovieRenderPipelineCore/Private/Graph/MovieGraphRenderLayerSubsystem.cpp
@@ -16,6 +16,8 @@
 #include "Styling/SlateIconFinder.h"
 #include "UObject/Package.h"
 #include "Widgets/Input/SMultiLineEditableTextBox.h"
+#include "WorldPartition/DataLayer/DataLayerAsset.h"
+#include "WorldPartition/DataLayer/DataLayerManager.h"
 
 #if WITH_EDITOR
 #include "ActorFolderPickingMode.h"
@@ -28,6 +30,7 @@
 #include "ContentBrowserDataDragDropOp.h"
 #include "ContentBrowserDataSource.h"
 #include "ContentBrowserModule.h"
+#include "DataLayer/DataLayerDragDropOp.h"
 #include "DragAndDrop/ActorDragDropGraphEdOp.h"
 #include "DragAndDrop/AssetDragDropOp.h"
 #include "DragAndDrop/FolderDragDropOp.h"
@@ -37,6 +40,8 @@
 #include "Graph/MovieGraphSharedWidgets.h"
 #include "IContentBrowserSingleton.h"
 #include "ISceneOutliner.h"
+#include "LayersDragDropOp.h"
+#include "Layers/LayersSubsystem.h"
 #include "SceneOutlinerModule.h"
 #include "SceneOutlinerPublicTypes.h"
 #include "SClassViewer.h"
@@ -1763,6 +1768,354 @@ void UMovieGraphConditionGroupQuery_Sublevel::AddLevels(const TArray<UWorld*>& I
 }
 #endif	// WITH_EDITOR
 
+void UMovieGraphConditionGroupQuery_ActorLayer::Evaluate(const TArray<AActor*>& InActorsToQuery, const UWorld* InWorld, TSet<AActor*>& OutMatchingActors) const
+{
+	for (AActor* Actor : InActorsToQuery)
+	{
+		for (const FName& LayerName : LayerNames)
+		{
+			if (Actor->Layers.Contains(LayerName))
+			{
+				OutMatchingActors.Add(Actor);
+				break;
+			}
+		}
+	}
+}
+
+const FSlateIcon& UMovieGraphConditionGroupQuery_ActorLayer::GetIcon() const
+{
+	static const FSlateIcon ActorLayerIcon = FSlateIcon(FAppStyle::GetAppStyleSetName(), "LevelEditor.Tabs.DataLayers");
+	return ActorLayerIcon;
+}
+
+const FText& UMovieGraphConditionGroupQuery_ActorLayer::GetDisplayName() const
+{
+	static const FText DisplayName = LOCTEXT("ConditionGroupQueryDisplayName_ActorLayer", "Actor Layer");
+	return DisplayName;
+}
+
+bool UMovieGraphConditionGroupQuery_ActorLayer::IsEditorOnly() const
+{
+	// Actor Layers are only available in the editor
+	return true;
+}
+
+#if WITH_EDITOR
+TArray<TSharedRef<SWidget>> UMovieGraphConditionGroupQuery_ActorLayer::GetWidgets()
+{
+	TArray<TSharedRef<SWidget>> Widgets;
+
+	Widgets.Add(
+		SNew(SDropTarget)
+		.OnAllowDrop_Lambda([](TSharedPtr<FDragDropOperation> InDragOperation)
+		{
+			return InDragOperation->IsOfType<FLayersDragDropOp>();
+		})
+		.OnDropped_Lambda([this](const FGeometry& Geometry, const FDragDropEvent& DragDropEvent)
+		{
+			if (const TSharedPtr<FLayersDragDropOp> LayerOperation = DragDropEvent.GetOperationAs<FLayersDragDropOp>())
+			{
+				const FMovieGraphConditionGroupQueryContentsChanged OnAddFinished = nullptr;
+				AddActorLayers(LayerOperation->Layers, OnAddFinished);
+			}
+			
+			return FReply::Handled();
+		})
+		[
+			SAssignNew(LayerNamesList, SMovieGraphSimpleList<FName>)
+			.DataSource(&LayerNames)
+			.DataType(FText::FromString("Actor Layer"))
+			.DataTypePlural(FText::FromString("Actor Layers"))
+			.OnGetRowText_Static(&GetRowText)
+			.OnGetRowIcon_Static(&GetRowIcon)
+			.OnDelete_Lambda([this](FName InLayerName)
+			{
+				const FScopedTransaction Transaction(LOCTEXT("RemoveActorLayerFromCollection", "Remove Actor Layers from Collection"));
+				Modify();
+
+				LayerNames.Remove(InLayerName);
+				LayerNamesList->Refresh();
+			})
+		]
+	);
+
+	return Widgets;
+}
+
+bool UMovieGraphConditionGroupQuery_ActorLayer::HasAddMenu() const
+{
+	return true;
+}
+
+TSharedRef<SWidget> UMovieGraphConditionGroupQuery_ActorLayer::GetAddMenuContents(const FMovieGraphConditionGroupQueryContentsChanged& OnAddFinished)
+{
+	// Refresh the list's data source
+	LayerPickerDataSource.Reset();
+	if (const ULayersSubsystem* LayersSubsystem = GEditor->GetEditorSubsystem<ULayersSubsystem>())
+	{
+		LayersSubsystem->AddAllLayerNamesTo(LayerPickerDataSource);
+
+		// Don't include layers that have already been picked
+		LayerPickerDataSource.RemoveAll([this](const FName InLayerName) { return LayerNames.Contains(InLayerName); });
+	}
+	
+	return
+		SNew(SMovieGraphSimplePicker<FName>)
+		.Title(LOCTEXT("PickActorLayerHelpText", "Pick an Actor Layer"))
+		.DataSourceEmptyMessage(LOCTEXT("NoActorLayersFoundWarning", "No actor layers found."))
+		.DataSource(LayerPickerDataSource)
+		.OnGetRowIcon_Lambda([](FName ListItem)
+		{
+			return FAppStyle::Get().GetBrush("Layer.Icon16x");
+		})
+		.OnGetRowText_Lambda([](FName ListItem)
+		{
+			return FText::FromName(ListItem);
+		})
+		.OnItemPicked_Lambda([this, OnAddFinished](FName InLayerName)
+		{
+			AddActorLayers({InLayerName}, OnAddFinished);
+		});
+}
+
+const FSlateBrush* UMovieGraphConditionGroupQuery_ActorLayer::GetRowIcon(FName InLayerName)
+{
+	return FAppStyle::Get().GetBrush("Layer.Icon16x");
+}
+
+FText UMovieGraphConditionGroupQuery_ActorLayer::GetRowText(FName InLayerName)
+{
+	if (const ULayersSubsystem* LayersSubsystem = GEditor->GetEditorSubsystem<ULayersSubsystem>())
+	{
+		if (LayersSubsystem->GetLayer(InLayerName))
+		{
+			return FText::FromName(InLayerName);
+		}
+	}
+
+	return FText::Format(LOCTEXT("InvalidActorLayer", "{0} (invalid)"), FText::FromName(InLayerName));
+}
+
+void UMovieGraphConditionGroupQuery_ActorLayer::AddActorLayers(const TArray<FName>& InActorLayers, const FMovieGraphConditionGroupQueryContentsChanged& InOnAddFinished)
+{
+	const FScopedTransaction Transaction(LOCTEXT("AddActorLayersToCollection", "Add Actor Layers to Collection"));
+	Modify();
+
+	for (const FName& ActorLayer : InActorLayers)
+	{
+		LayerNames.AddUnique(ActorLayer);
+	}
+	
+	LayerNamesList->Refresh();
+
+	InOnAddFinished.ExecuteIfBound();
+}
+#endif	// WITH_EDITOR
+
+void UMovieGraphConditionGroupQuery_DataLayer::Evaluate(const TArray<AActor*>& InActorsToQuery, const UWorld* InWorld, TSet<AActor*>& OutMatchingActors) const
+{
+	for (AActor* InActor : InActorsToQuery)
+	{
+		for (const TSoftObjectPtr<UDataLayerAsset>& DataLayer : DataLayers)
+		{
+			if (InActor->ContainsDataLayer(DataLayer.Get()))
+			{
+				OutMatchingActors.Add(InActor);
+				break;
+			}
+		}
+	}
+}
+
+const FSlateIcon& UMovieGraphConditionGroupQuery_DataLayer::GetIcon() const
+{
+	static const FSlateIcon DataLayerIcon = FSlateIcon(FAppStyle::GetAppStyleSetName(), "DataLayer.Editor");
+	return DataLayerIcon;
+}
+
+const FText& UMovieGraphConditionGroupQuery_DataLayer::GetDisplayName() const
+{
+	static const FText DisplayName = LOCTEXT("ConditionGroupQueryDisplayName_DataLayer", "Data Layer");
+	return DisplayName;
+}
+
+#if WITH_EDITOR
+TArray<TSharedRef<SWidget>> UMovieGraphConditionGroupQuery_DataLayer::GetWidgets()
+{
+	TArray<TSharedRef<SWidget>> Widgets;
+
+	// Create the data source for the list view
+	ListDataSource.Empty();
+	for (TSoftObjectPtr<UDataLayerAsset>& DataLayer : DataLayers)
+	{
+		ListDataSource.Add(MakeShared<TSoftObjectPtr<UDataLayerAsset>>(DataLayer));
+	}
+
+	Widgets.Add(
+		SNew(SDropTarget)
+		.OnAllowDrop_Lambda([](TSharedPtr<FDragDropOperation> InDragOperation)
+		{
+			return InDragOperation->IsOfType<FDataLayerDragDropOp>() || InDragOperation->IsOfType<FAssetDragDropOp>();
+		})
+		.OnDropped_Lambda([this](const FGeometry& Geometry, const FDragDropEvent& DragDropEvent)
+		{
+			TArray<const UDataLayerAsset*> DroppedLayers;
+
+			// Drag-n-drop from the Data Layers editor
+			if (const TSharedPtr<FDataLayerDragDropOp> LayerOperation = DragDropEvent.GetOperationAs<FDataLayerDragDropOp>())
+			{
+				for (const TWeakObjectPtr<UDataLayerInstance>& DroppedLayer : LayerOperation->DataLayerInstances)
+				{
+					if (const UDataLayerAsset* DroppedLayerAsset = DroppedLayer->GetAsset())
+					{
+						DroppedLayers.AddUnique(DroppedLayerAsset);
+					}
+				}
+			}
+
+			// Drag-n-drop from the Content Browser
+			else if (const TSharedPtr<FAssetDragDropOp> AssetOperation = DragDropEvent.GetOperationAs<FAssetDragDropOp>())
+			{
+				for (const FAssetData& AssetData : AssetOperation->GetAssets())
+				{
+					if (const UDataLayerAsset* DataLayer = Cast<const UDataLayerAsset>(AssetData.GetAsset()))
+					{
+						DroppedLayers.AddUnique(DataLayer);
+					}
+				}
+			}
+			
+			if (!DroppedLayers.IsEmpty())
+			{
+				const FMovieGraphConditionGroupQueryContentsChanged OnAddFinished = nullptr;
+				AddDataLayers(DroppedLayers, OnAddFinished);
+			}
+			
+			return FReply::Handled();
+		})
+		[
+			SAssignNew(DataLayersList, SMovieGraphSimpleList<TSharedPtr<TSoftObjectPtr<UDataLayerAsset>>>)
+			.DataSource(&ListDataSource)
+			.DataType(FText::FromString("Data Layer"))
+			.DataTypePlural(FText::FromString("Data Layers"))
+			.OnGetRowText_Static(&GetRowText)
+			.OnGetRowIcon_Static(&GetRowIcon)
+			.OnDelete_Lambda([this](TSharedPtr<TSoftObjectPtr<UDataLayerAsset>> InLayer)
+			{
+				const FScopedTransaction Transaction(LOCTEXT("RemoveDataLayerFromCollection", "Remove Data Layers from Collection"));
+				Modify();
+
+				ListDataSource.Remove(InLayer);
+				DataLayers.Remove(*InLayer.Get());
+				DataLayersList->Refresh();
+
+				constexpr bool bUpdateSources = true;
+				RefreshDataLayerPicker.ExecuteIfBound(bUpdateSources);
+			})
+		]
+	);
+
+	return Widgets;
+}
+
+bool UMovieGraphConditionGroupQuery_DataLayer::HasAddMenu() const
+{
+	return true;
+}
+
+TSharedRef<SWidget> UMovieGraphConditionGroupQuery_DataLayer::GetAddMenuContents(const FMovieGraphConditionGroupQueryContentsChanged& OnAddFinished)
+{
+	FAssetPickerConfig DataLayerPickerConfig;
+	{
+		DataLayerPickerConfig.SelectionMode = ESelectionMode::Single;
+		DataLayerPickerConfig.SaveSettingsName = TEXT("MovieRenderGraphDataLayerPicker");
+		DataLayerPickerConfig.RefreshAssetViewDelegates.Add(&RefreshDataLayerPicker);
+		DataLayerPickerConfig.InitialAssetViewType = EAssetViewType::Column;
+		DataLayerPickerConfig.bFocusSearchBoxWhenOpened = true;
+		DataLayerPickerConfig.bAllowNullSelection = false;
+		DataLayerPickerConfig.bShowBottomToolbar = true;
+		DataLayerPickerConfig.bAutohideSearchBar = false;
+		DataLayerPickerConfig.bAllowDragging = false;
+		DataLayerPickerConfig.bCanShowClasses = false;
+		DataLayerPickerConfig.bShowPathInColumnView = true;
+		DataLayerPickerConfig.bShowTypeInColumnView = false;
+		DataLayerPickerConfig.bSortByPathInColumnView = false;
+		DataLayerPickerConfig.HiddenColumnNames = {
+			ContentBrowserItemAttributes::ItemDiskSize.ToString(),
+			ContentBrowserItemAttributes::VirtualizedData.ToString(),
+			TEXT("PrimaryAssetType"),
+			TEXT("PrimaryAssetName")
+		};
+		DataLayerPickerConfig.AssetShowWarningText = LOCTEXT("ConditionGroupQuery_NoDataLayersFound", "No Data Layers Found");
+		DataLayerPickerConfig.Filter.ClassPaths.Add(UDataLayerAsset::StaticClass()->GetClassPathName());
+		DataLayerPickerConfig.OnAssetSelected = FOnAssetSelected::CreateLambda([this, OnAddFinished](const FAssetData& InDataLayerAsset)
+		{
+			AddDataLayers({Cast<UDataLayerAsset>(InDataLayerAsset.GetAsset())}, OnAddFinished);
+		});
+		DataLayerPickerConfig.OnShouldFilterAsset = FOnShouldFilterAsset::CreateLambda([this](const FAssetData& InDataLayerAsset)
+		{
+			// Don't show data layers which have already been picked
+			UDataLayerAsset* DataLayer = Cast<UDataLayerAsset>(InDataLayerAsset.GetAsset());
+			return !DataLayer || DataLayers.Contains(DataLayer);
+		});
+	}
+
+	IContentBrowserSingleton& ContentBrowser = FModuleManager::LoadModuleChecked<FContentBrowserModule>("ContentBrowser").Get();
+
+	return
+		SNew(SBox)
+		.Padding(0, 10.f, 0, 0)
+		.WidthOverride(400.f)
+		.HeightOverride(300.f)
+		[
+			ContentBrowser.CreateAssetPicker(DataLayerPickerConfig)
+		];
+}
+
+const FSlateBrush* UMovieGraphConditionGroupQuery_DataLayer::GetRowIcon(TSharedPtr<TSoftObjectPtr<UDataLayerAsset>> InDataLayer)
+{
+	return FAppStyle::Get().GetBrush("DataLayer.Editor");
+}
+
+FText UMovieGraphConditionGroupQuery_DataLayer::GetRowText(TSharedPtr<TSoftObjectPtr<UDataLayerAsset>> InDataLayer)
+{
+	if (InDataLayer.IsValid())
+	{
+		if (InDataLayer.Get()->IsValid())
+		{
+			// The first Get() returns the TSoftObjectPtr, the second Get() dereferences the TSoftObjectPtr
+			return FText::FromString(InDataLayer.Get()->Get()->GetName());
+		}
+	}
+
+	return LOCTEXT("MovieGraphDataLayerConditionGroupQuery_InvalidDataLayer", "(invalid or unloaded)");
+}
+
+void UMovieGraphConditionGroupQuery_DataLayer::AddDataLayers(const TArray<const UDataLayerAsset*>& InDataLayers, const FMovieGraphConditionGroupQueryContentsChanged& InOnAddFinished)
+{
+	const FScopedTransaction Transaction(LOCTEXT("AddDataLayersToCollection", "Add Data Layers to Collection"));
+	Modify();
+	
+	for (const UDataLayerAsset* DataLayer : InDataLayers)
+	{
+		if (DataLayer && !DataLayers.Contains(DataLayer))
+		{
+			DataLayers.Add(DataLayer);
+			ListDataSource.Add(MakeShared<TSoftObjectPtr<UDataLayerAsset>>(DataLayer));
+		}
+	}
+	
+	DataLayersList->Refresh();
+	FSlateApplication::Get().DismissAllMenus();
+
+	InOnAddFinished.ExecuteIfBound();
+
+	constexpr bool bUpdateSources = false;
+	RefreshDataLayerPicker.ExecuteIfBound(bUpdateSources);
+}
+#endif	// WITH_EDITOR
+
 UMovieGraphConditionGroup::UMovieGraphConditionGroup()
 	: OpType(EMovieGraphConditionGroupOpType::Add)
 {
diff --git a/Engine/Plugins/MovieScene/MovieRenderPipeline/Source/MovieRenderPipelineCore/Public/Graph/MovieGraphRenderLayerSubsystem.h b/Engine/Plugins/MovieScene/MovieRenderPipeline/Source/MovieRenderPipelineCore/Public/Graph/MovieGraphRenderLayerSubsystem.h
--- a/Engine/Plugins/MovieScene/MovieRenderPipeline/Source/MovieRenderPipelineCore/Public/Graph/MovieGraphRenderLayerSubsystem.h
+++ b/Engine/Plugins/MovieScene/MovieRenderPipeline/Source/MovieRenderPipelineCore/Public/Graph/MovieGraphRenderLayerSubsystem.h
@@ -14,6 +14,7 @@
 #include "MovieGraphRenderLayerSubsystem.generated.h"
 
 class SWidget;
+class UDataLayerAsset;
 
 /** Operation types available on condition groups. */
 UENUM(BlueprintType)
@@ -441,6 +442,87 @@ public:
 	TArray<TSoftObjectPtr<UWorld>> Sublevels;
 };
 
+/** Query type which filters actors via Actor Layers. */
+UCLASS(BlueprintType)
+class MOVIERENDERPIPELINECORE_API UMovieGraphConditionGroupQuery_ActorLayer final : public UMovieGraphConditionGroupQueryBase
+{
+	GENERATED_BODY()
+
+public:
+	virtual void Evaluate(const TArray<AActor*>& InActorsToQuery, const UWorld* InWorld, TSet<AActor*>& OutMatchingActors) const override;
+	virtual const FSlateIcon& GetIcon() const override;
+	virtual const FText& GetDisplayName() const override;
+	virtual bool IsEditorOnly() const override;
+
+#if WITH_EDITOR
+	virtual TArray<TSharedRef<SWidget>> GetWidgets() override;
+	virtual bool HasAddMenu() const override;
+	virtual TSharedRef<SWidget> GetAddMenuContents(const FMovieGraphConditionGroupQueryContentsChanged& OnAddFinished) override;
+#endif
+
+private:
+#if WITH_EDITOR
+	static const FSlateBrush* GetRowIcon(FName InLayerName);
+	static FText GetRowText(FName InLayerName);
+
+	/** Adds the provided actor layers to the query, updating the UI as needed. Calls InOnAddFinished when done. */
+	void AddActorLayers(const TArray<FName>& InActorLayers, const FMovieGraphConditionGroupQueryContentsChanged& InOnAddFinished);
+
+	/** Displays the layers which have been chosen. */
+	TSharedPtr<SMovieGraphSimpleList<FName>> LayerNamesList;
+
+	/** The data source for the layer picker widget (contains all layers which are available and not yet picked). */
+	TArray<FName> LayerPickerDataSource;
+#endif
+
+public:
+	/** The actor must be in one of the actor layers with these names in order to be a match. */
+	UPROPERTY(EditAnywhere, Category="General")
+	TArray<FName> LayerNames;
+};
+
+/** Query type which filters actors via World Partition Data Layers. */
+UCLASS(BlueprintType)
+class MOVIERENDERPIPELINECORE_API UMovieGraphConditionGroupQuery_DataLayer final : public UMovieGraphConditionGroupQueryBase
+{
+	GENERATED_BODY()
+
+public:
+	virtual void Evaluate(const TArray<AActor*>& InActorsToQuery, const UWorld* InWorld, TSet<AActor*>& OutMatchingActors) const override;
+	virtual const FSlateIcon& GetIcon() const override;
+	virtual const FText& GetDisplayName() const override;
+
+#if WITH_EDITOR
+	virtual TArray<TSharedRef<SWidget>> GetWidgets() override;
+	virtual bool HasAddMenu() const override;
+	virtual TSharedRef<SWidget> GetAddMenuContents(const FMovieGraphConditionGroupQueryContentsChanged& OnAddFinished) override;
+#endif
+
+private:
+#if WITH_EDITOR
+	static const FSlateBrush* GetRowIcon(TSharedPtr<TSoftObjectPtr<UDataLayerAsset>> InDataLayer);
+	static FText GetRowText(TSharedPtr<TSoftObjectPtr<UDataLayerAsset>> InDataLayer);
+	
+	/** Adds the provided data layers to the query, updating the UI as needed. Calls InOnAddFinished when done. */
+	void AddDataLayers(const TArray<const UDataLayerAsset*>& InDataLayers, const FMovieGraphConditionGroupQueryContentsChanged& InOnAddFinished);
+
+	/** Displays the layers which have been chosen. */
+	TSharedPtr<SMovieGraphSimpleList<TSharedPtr<TSoftObjectPtr<UDataLayerAsset>>>> DataLayersList;
+	
+	// Not ideal to store a duplicate of DataLayersList, but SListView requires TSharedPtr<...> as the data source, and UPROPERTY does not
+	// support TSharedPtr<...>
+	TArray<TSharedPtr<TSoftObjectPtr<UDataLayerAsset>>> ListDataSource;
+
+	/** Refreshes the contents of the data layer picker widget when called. */
+	FRefreshAssetViewDelegate RefreshDataLayerPicker;
+#endif
+
+public:
+	/** The actor must be in one of the these data layer assets in order to be a match. */
+	UPROPERTY(EditAnywhere, Category="General")
+	TArray<TSoftObjectPtr<UDataLayerAsset>> DataLayers;
+};
+
 /** A group of queries which can be added to a collection. */
 UCLASS(BlueprintType)
 class MOVIERENDERPIPELINECORE_API UMovieGraphConditionGroup : public UObject
diff --git a/Engine/Plugins/MovieScene/MovieRenderPipeline/Source/MovieRenderPipelineCore/Public/Graph/MovieGraphSharedWidgets.h b/Engine/Plugins/MovieScene/MovieRenderPipeline/Source/MovieRenderPipelineCore/Public/Graph/MovieGraphSharedWidgets.h
--- a/Engine/Plugins/MovieScene/MovieRenderPipeline/Source/MovieRenderPipelineCore/Public/Graph/MovieGraphSharedWidgets.h
+++ b/Engine/Plugins/MovieScene/MovieRenderPipeline/Source/MovieRenderPipelineCore/Public/Graph/MovieGraphSharedWidgets.h
@@ -7,6 +7,7 @@
 #include "Widgets/Input/SCheckBox.h"
 #include "Widgets/Input/SComboButton.h"
 #include "Widgets/SCompoundWidget.h"
+#include "Widgets/Layout/SWidgetSwitcher.h"
 #include "Widgets/Text/STextBlock.h"
 #include "Widgets/Views/ITableRow.h"
 #include "Widgets/Views/SListView.h"
@@ -245,4 +246,171 @@ private:
 	TArray<ListType>* DataSource = nullptr;
 };
 
+/**
+ * A simple widget that displays a list that can be picked from. Meant for use within pop-ups. Shows a title above the list, and an optional message
+ * to be displayed (instead of the list) if the data source is empty.
+ */
+template<typename DataType>
+class SMovieGraphSimplePicker : public SCompoundWidget
+{
+public:
+	DECLARE_DELEGATE_OneParam(FOnItemPicked, DataType);
+	DECLARE_DELEGATE_RetVal_OneParam(const FSlateBrush*, FGetRowIcon, DataType);
+	DECLARE_DELEGATE_RetVal_OneParam(FText, FGetRowText, DataType);
+	DECLARE_DELEGATE_RetVal_OneParam(bool, FOnFilter, DataType);
+
+	SLATE_BEGIN_ARGS(SMovieGraphSimplePicker)
+		{}
+
+		/** Called when an item is picked in the list. */
+		SLATE_EVENT(FOnItemPicked, OnItemPicked)
+
+		/** Gets the icon for a row in the list. */
+		SLATE_EVENT(FGetRowIcon, OnGetRowIcon);
+
+		/** Gets the text for a row in the list. */
+		SLATE_EVENT(FGetRowText, OnGetRowText);
+
+		/** The title text shown within the picker widget. */
+		SLATE_ARGUMENT(FText, Title)
+
+		/** The message displayed if the data source is empty. */
+		SLATE_ARGUMENT(FText, DataSourceEmptyMessage)
+
+		/** The data that is displayed within the picker. */
+		SLATE_ARGUMENT(TArray<DataType>, DataSource);
+	SLATE_END_ARGS()
+
+	void Construct(const FArguments& InArgs)
+	{
+		OnItemPicked = InArgs._OnItemPicked;
+		OnGetRowIcon = InArgs._OnGetRowIcon;
+		OnGetRowText = InArgs._OnGetRowText;
+		Title = InArgs._Title;
+		DataSourceEmptyMessage = InArgs._DataSourceEmptyMessage;
+		DataSource = InArgs._DataSource;
+		
+		ChildSlot
+		.HAlign(HAlign_Fill)
+		.VAlign(VAlign_Fill)
+		[
+			SNew(SBox)
+			.WidthOverride(200.f)
+			.HeightOverride(200.f)
+			[
+				SNew(SVerticalBox)
+				+ SVerticalBox::Slot()
+				.Padding(5.f)
+				.AutoHeight()
+				[
+					SNew(STextBlock)
+					.Text(Title)
+					.Font(FAppStyle::GetFontStyle("PropertyWindow.NormalFont"))
+				]
+				
+				+ SVerticalBox::Slot()
+				.VAlign(VAlign_Fill)
+				[
+					SNew(SWidgetSwitcher)
+					.WidgetIndex_Lambda([this] { return DataSource.IsEmpty() ? 0 : 1; })
+
+					+ SWidgetSwitcher::Slot()
+					.VAlign(VAlign_Center)
+					.HAlign(HAlign_Center)
+					[
+						SNew(STextBlock)
+						.Text(DataSourceEmptyMessage)
+					]
+
+					+ SWidgetSwitcher::Slot()
+					[
+						SNew(SListView<DataType>)
+						.ListItemsSource(&DataSource)
+						.SelectionMode(ESelectionMode::Single)
+						.OnSelectionChanged(this, &SMovieGraphSimplePicker::OnItemSelected)
+						.OnGenerateRow(this, &SMovieGraphSimplePicker::GenerateRow)
+					]
+				]
+			]
+		];
+	}
+
+private:
+	/** Handles an item selected event. */
+	void OnItemSelected(const DataType Item, ESelectInfo::Type Type) const
+	{
+		if (OnItemPicked.IsBound())
+		{
+			OnItemPicked.Execute(Item);
+		}
+
+		FSlateApplication::Get().DismissAllMenus();
+	}
+	
+	/** Gets the row icon associated with the given item. */
+	const FSlateBrush* GetRowIcon(const DataType InItem) const
+	{
+		if (OnGetRowIcon.IsBound())
+		{
+			return OnGetRowIcon.Execute(InItem);
+		}
+
+		return FAppStyle::GetBrush("Icons.FilledCircle");
+	}
+
+	/** Gets the row text associated with the given list data. */
+	FText GetRowText(const DataType InListData) const
+	{
+		if (OnGetRowText.IsBound())
+		{
+			return OnGetRowText.Execute(InListData);
+		}
+		
+		return FText();
+	}
+
+	/** Generates a row which displays a single item. */
+	TSharedRef<ITableRow> GenerateRow(const DataType InItem, const TSharedRef<STableViewBase>& InOwnerTable) const
+	{
+		return
+			SNew(STableRow<FName>, InOwnerTable)
+			.Style(FAppStyle::Get(), "TableView.AlternatingRow")
+			.ShowWires(false)
+			[
+				SNew(SHorizontalBox)
+				+ SHorizontalBox::Slot()
+				.VAlign(VAlign_Center)
+				.Padding(7.f, 5.f, 7.f, 5.f)
+				.AutoWidth()
+				[
+					SNew(SImage)
+					.Image(GetRowIcon(InItem))
+				]
+				
+				+ SHorizontalBox::Slot()
+				.VAlign(VAlign_Center)
+				.HAlign(HAlign_Fill)
+				[
+					SNew(STextBlock)
+					.Font(FAppStyle::GetFontStyle(TEXT("PropertyWindow.NormalFont")))
+					.Text(GetRowText(InItem))
+				]
+			];
+	}
+
+protected:
+	/**
+	 * The data source for the list view widget. Names of the items that can be picked. Widgets that inherit from this class can optionally control
+	 * this data member directly rather than providing it as an argument.
+	 */
+	TArray<DataType> DataSource;
+
+private:
+	FOnItemPicked OnItemPicked;
+	FGetRowIcon OnGetRowIcon;
+	FGetRowText OnGetRowText;
+	FText Title;
+	FText DataSourceEmptyMessage;
+};
+
 #undef LOCTEXT_NAMESPACE
\ No newline at end of file
diff --git a/Engine/Plugins/MovieScene/MovieRenderPipeline/Source/MovieRenderPipelineEditor/Private/Customizations/Graph/MovieGraphModifiersCustomization.cpp b/Engine/Plugins/MovieScene/MovieRenderPipeline/Source/MovieRenderPipelineEditor/Private/Customizations/Graph/MovieGraphModifiersCustomization.cpp
--- a/Engine/Plugins/MovieScene/MovieRenderPipeline/Source/MovieRenderPipelineEditor/Private/Customizations/Graph/MovieGraphModifiersCustomization.cpp
+++ b/Engine/Plugins/MovieScene/MovieRenderPipeline/Source/MovieRenderPipelineEditor/Private/Customizations/Graph/MovieGraphModifiersCustomization.cpp
@@ -12,12 +12,11 @@
 #include "Graph/Nodes/MovieGraphModifierNode.h"
 #include "ScopedTransaction.h"
 #include "Widgets/Images/SLayeredImage.h"
-#include "Widgets/Layout/SWidgetSwitcher.h"
 
 #define LOCTEXT_NAMESPACE "MovieGraphModifiersCustomization"
 
 /** Discovers collections that are pickable from a specific graph, and presents them in a list. */
-class SMovieGraphCollectionPicker final : public SCompoundWidget
+class SMovieGraphCollectionPicker final : public SMovieGraphSimplePicker<FName>
 {
 public:
 	DECLARE_DELEGATE_OneParam(FOnCollectionPicked, FName);
@@ -41,46 +40,15 @@ public:
 		OnCollectionPicked = InArgs._OnCollectionPicked;
 		OnFilter = InArgs._OnFilter;
 
-		UpdateDataSource();
-		
-		ChildSlot
-		.HAlign(HAlign_Fill)
-		.VAlign(VAlign_Fill)
-		[
-			SNew(SVerticalBox)
-			+ SVerticalBox::Slot()
-			.Padding(5.f)
-			.AutoHeight()
-			[
-				SNew(STextBlock)
-				.Text(LOCTEXT("PickCollectionHelpText", "Pick a Collection"))
-				.Font(FAppStyle::GetFontStyle("PropertyWindow.NormalFont"))
-			]
-			
-			+ SVerticalBox::Slot()
-			.VAlign(VAlign_Fill)
-			[
-				SNew(SWidgetSwitcher)
-				.WidgetIndex_Lambda([this] { return DataSource.IsEmpty() ? 0 : 1; })
-
-				+ SWidgetSwitcher::Slot()
-				.VAlign(VAlign_Center)
-				.HAlign(HAlign_Center)
-				[
-					SNew(STextBlock)
-					.Text(LOCTEXT("NoCollectionsFoundWarning", "No collections found."))
-				]
+		SMovieGraphSimplePicker::Construct(SMovieGraphSimplePicker::FArguments()
+			.OnGetRowText_Lambda([](FName InCollectionName) { return FText::FromName(InCollectionName); })
+			.Title(LOCTEXT("PickCollectionHelpText", "Pick a Collection"))
+			.DataSourceEmptyMessage(LOCTEXT("NoCollectionsFoundWarning", "No collections found."))
+			.OnItemPicked(OnCollectionPicked));
 
-				+ SWidgetSwitcher::Slot()
-				[
-					SNew(SListView<FName>)
-					.ListItemsSource(&DataSource)
-					.SelectionMode(ESelectionMode::Single)
-					.OnSelectionChanged(this, &SMovieGraphCollectionPicker::OnCollectionSelected)
-					.OnGenerateRow(this, &SMovieGraphCollectionPicker::GenerateRow)
-				]
-			]
-		];
+		// Update the data source *after* calling Construct() above because Construct() will populate DataSource based on the widget arguments, but
+		// we want to control/update it manually.
+		UpdateDataSource();
 	}
 
 private:
@@ -128,50 +96,7 @@ private:
 		return CollectionNodes;
 	}
 
-	/** Handles a collection selected event. */
-	void OnCollectionSelected(const FName CollectionName, ESelectInfo::Type Type) const
-	{
-		if (OnCollectionPicked.IsBound())
-		{
-			OnCollectionPicked.Execute(CollectionName);
-		}
-
-		FSlateApplication::Get().DismissAllMenus();
-	}
-
-	/** Generates a row which displays a single collection. */
-	TSharedRef<ITableRow> GenerateRow(const FName CollectionName, const TSharedRef<STableViewBase>& InOwnerTable) const
-	{
-		return
-			SNew(STableRow<FName>, InOwnerTable)
-			.Style(FAppStyle::Get(), "TableView.AlternatingRow")
-			.ShowWires(false)
-			[
-				SNew(SHorizontalBox)
-				+ SHorizontalBox::Slot()
-				.VAlign(VAlign_Center)
-				.Padding(7.f, 5.f, 7.f, 5.f)
-				.AutoWidth()
-				[
-					SNew(SImage)
-					.Image(FAppStyle::GetBrush("Icons.FilledCircle"))
-				]
-					
-				+ SHorizontalBox::Slot()
-				.VAlign(VAlign_Center)
-				.HAlign(HAlign_Fill)
-				[
-					SNew(STextBlock)
-					.Font(FAppStyle::GetFontStyle(TEXT("PropertyWindow.NormalFont")))
-					.Text(FText::FromName(CollectionName))
-				]
-			];
-	}
-
 private:
-	/** The data source for the list view widget. Names of the collection nodes that can be picked. */
-	TArray<FName> DataSource;
-
 	/** The current graph being viewed. The data source will be populated from this graph. */
 	UMovieGraphConfig* CurrentGraph = nullptr;
 	
@@ -238,33 +163,28 @@ void FMovieGraphModifiersCustomization::CustomizeDetails(IDetailLayoutBuilder& I
 			.OnGetMenuContent_Lambda([ModifierNode, this]()
 			{
 				return
-					SNew(SBox)
-					.WidthOverride(200.f)
-					.HeightOverride(200.f)
-					[
-						SNew(SMovieGraphCollectionPicker)
-						.Graph(ModifierNode->GetTypedOuter<UMovieGraphConfig>())
-						.OnFilter_Lambda([ModifierNode](const FName CollectionName)
+					SNew(SMovieGraphCollectionPicker)
+					.Graph(ModifierNode->GetTypedOuter<UMovieGraphConfig>())
+					.OnFilter_Lambda([ModifierNode](const FName CollectionName)
+					{
+						if (ModifierNode.IsValid())
 						{
-							if (ModifierNode.IsValid())
-							{
-								return !ModifierNode.Get()->GetCollections().Contains(CollectionName);
-							}
-
-							return false;
-						})
-						.OnCollectionPicked_Lambda([ModifierNode, this](const FName PickedCollectionName)
+							return !ModifierNode.Get()->GetCollections().Contains(CollectionName);
+						}
+
+						return false;
+					})
+					.OnCollectionPicked_Lambda([ModifierNode, this](const FName PickedCollectionName)
+					{
+						if (ModifierNode.IsValid())
 						{
-							if (ModifierNode.IsValid())
-							{
-								const FScopedTransaction Transaction(LOCTEXT("AddCollectionToModifier", "Add Collection to Modifier"));
-								
-								ModifierNode->AddCollection(PickedCollectionName);
-								ListDataSource = ModifierNode->GetCollections();
-								CollectionsList->Refresh();
-							}
-						})
-					];
+							const FScopedTransaction Transaction(LOCTEXT("AddCollectionToModifier", "Add Collection to Modifier"));
+							
+							ModifierNode->AddCollection(PickedCollectionName);
+							ListDataSource = ModifierNode->GetCollections();
+							CollectionsList->Refresh();
+						}
+					});
 			})
 			.ButtonContent()
 			[
