MoviePipeline: The Editor Folder condition group query should also traverse into subfolders.

diff --git a/Engine/Plugins/MovieScene/MovieRenderPipeline/Source/MovieRenderPipelineCore/Private/Graph/MovieGraphRenderLayerSubsystem.cpp b/Engine/Plugins/MovieScene/MovieRenderPipeline/Source/MovieRenderPipelineCore/Private/Graph/MovieGraphRenderLayerSubsystem.cpp
--- a/Engine/Plugins/MovieScene/MovieRenderPipeline/Source/MovieRenderPipelineCore/Private/Graph/MovieGraphRenderLayerSubsystem.cpp
+++ b/Engine/Plugins/MovieScene/MovieRenderPipeline/Source/MovieRenderPipelineCore/Private/Graph/MovieGraphRenderLayerSubsystem.cpp
@@ -1359,10 +1359,51 @@ void UMovieGraphConditionGroupQuery_EditorFolder::Evaluate(const TArray<AActor*>
 {
 #if WITH_EDITOR
 	// This const cast is unfortunate, but should be harmless
-	TArray<AActor*> ActorsInFolders;
-	FActorFolders::GetActorsFromFolders(*const_cast<UWorld*>(InWorld), FolderPaths, ActorsInFolders);
-	
-	OutMatchingActors.Append(ActorsInFolders);
+	const FFolder::FRootObject FolderRootObject = FFolder::GetWorldRootFolder(const_cast<UWorld*>(InWorld)).GetRootObject();
+
+	for (AActor* Actor : InActorsToQuery)
+	{
+		if (Actor->GetFolderRootObject() != FolderRootObject)
+		{
+			continue;
+		}
+		
+		const FName ActorFolderPath = Actor->GetFolderPath();
+		if (ActorFolderPath.IsNone())
+		{
+			continue;
+		}
+		
+		const FString ActorFolderPathString = ActorFolderPath.ToString();
+		const int32 ActorFolderLen = ActorFolderPath.GetStringLength();
+
+		for (const FName& ParentFolderPath : FolderPaths)
+		{
+			const int32 ParentFolderLen = ParentFolderPath.GetStringLength();
+			
+			// We shouldn't be looking at an empty folder path, but just in case.
+			if (ParentFolderLen == 0)
+			{
+				continue;
+			}
+
+			// The actor is a match if it's in a folder that's an exact match.
+			if (ActorFolderPath == ParentFolderPath)
+			{
+				OutMatchingActors.Add(Actor);
+				break;
+			}
+
+			// The actor is also a match if it's in a matching subfolder.
+			if ((ActorFolderLen > ParentFolderLen) &&
+				(ActorFolderPathString[ParentFolderLen] == '/') &&
+				(ActorFolderPathString.Left(ParentFolderLen) == ParentFolderPath))
+			{
+				OutMatchingActors.Add(Actor);
+				break;
+			}
+		}
+	}
 #endif	// WITH_EDITOR
 }
 
