MoviePipeline: Added a context menu to condition group query lists that has "Delete" in it. Deleting via the delete key isn't always obvious so the context menu should surface this functionality more.

diff --git a/Engine/Plugins/MovieScene/MovieRenderPipeline/Source/MovieRenderPipelineCore/Public/Graph/MovieGraphSharedWidgets.h b/Engine/Plugins/MovieScene/MovieRenderPipeline/Source/MovieRenderPipelineCore/Public/Graph/MovieGraphSharedWidgets.h
--- a/Engine/Plugins/MovieScene/MovieRenderPipeline/Source/MovieRenderPipelineCore/Public/Graph/MovieGraphSharedWidgets.h
+++ b/Engine/Plugins/MovieScene/MovieRenderPipeline/Source/MovieRenderPipelineCore/Public/Graph/MovieGraphSharedWidgets.h
@@ -4,6 +4,7 @@
 
 #include "CoreMinimal.h"
 #include "Framework/Application/SlateApplication.h"
+#include "Framework/Commands/GenericCommands.h"
 #include "Widgets/Input/SCheckBox.h"
 #include "Widgets/Input/SComboButton.h"
 #include "Widgets/SCompoundWidget.h"
@@ -74,6 +75,16 @@ public:
 		OnGetRowEnableState = InArgs._OnGetRowEnableState;
 		OnSetRowEnableState = InArgs._OnSetRowEnableState;
 
+		CommandList = MakeShared<FUICommandList>();
+
+		CommandList->MapAction(
+			FGenericCommands::Get().Delete,
+			FExecuteAction::CreateSP(this, &SMovieGraphSimpleList<ListType>::HandleDelete),
+			FCanExecuteAction::CreateLambda([this]()
+			{
+				return OnDelete.IsBound();
+			}));
+
 		ChildSlot
 		.HAlign(HAlign_Fill)
 		.VAlign(VAlign_Top)
@@ -84,9 +95,12 @@ public:
 			[
 				SAssignNew(ListView, SListView<ListType>)
 				.ListItemsSource(DataSource)
-				.SelectionMode(ESelectionMode::Single)
-				.OnKeyDownHandler(this, &SMovieGraphSimpleList<ListType>::HandleDelete)
+				.OnKeyDownHandler_Lambda([this](const FGeometry& InGeometry, const FKeyEvent& InKeyEvent)
+				{
+					return CommandList->ProcessCommandBindings(InKeyEvent) ? FReply::Handled() : FReply::Unhandled();
+				})
 				.OnGenerateRow(this, &SMovieGraphSimpleList<ListType>::GenerateRow)
+				.OnContextMenuOpening(this, &SMovieGraphSimpleList<ListType>::OnContextMenuOpening)
 			]
 
 			+ SVerticalBox::Slot()
@@ -152,22 +166,33 @@ private:
 		return true;
 	}
 
-	/** Handles the delete operation. */
-	FReply HandleDelete(const FGeometry& InGeometry, const FKeyEvent& InKeyEvent) const
+	/** Deletes selected items in the list. */
+	void HandleDelete() const
 	{
 		TArray<ListType> SelectedItems;
 		ListView->GetSelectedItems(SelectedItems);
-				
-		if ((InKeyEvent.GetKey() == EKeys::Delete) && (SelectedItems.Num() == 1))
+
+		if (!SelectedItems.IsEmpty())
 		{
 			if (OnDelete.IsBound())
 			{
 				OnDelete.Execute(SelectedItems[0]);
-				return FReply::Handled();
 			}
 		}
+	}
+
+	TSharedPtr<SWidget> OnContextMenuOpening()
+	{
+		constexpr bool bShouldCloseWindowAfterMenuSelection = true;
+		FMenuBuilder MenuBuilder(bShouldCloseWindowAfterMenuSelection, CommandList);
+
+		MenuBuilder.BeginSection("BasicOperations");
+		{
+			MenuBuilder.AddMenuEntry(FGenericCommands::Get().Delete);
+		}
+		MenuBuilder.EndSection();
 
-		return FReply::Unhandled();
+		return MenuBuilder.MakeWidget();
 	}
 
 	/** Generates a row in the list for the specified data. */
@@ -235,6 +260,7 @@ private:
 
 private:
 	TSharedPtr<SListView<ListType>> ListView;
+	TSharedPtr<FUICommandList> CommandList;
 	FGetRowIcon OnGetRowIcon;
 	FGetRowText OnGetRowText;
 	FOnDelete OnDelete;
