MoviePipeline: Added reset widgets for individual properties within Show Flags on the render nodes.

Also made operations on Show Flags properly undo-able.

diff --git a/Engine/Plugins/MovieScene/MovieRenderPipeline/Source/MovieRenderPipelineCore/Private/Graph/Renderers/MovieGraphShowFlags.cpp b/Engine/Plugins/MovieScene/MovieRenderPipeline/Source/MovieRenderPipelineCore/Private/Graph/Renderers/MovieGraphShowFlags.cpp
--- a/Engine/Plugins/MovieScene/MovieRenderPipeline/Source/MovieRenderPipelineCore/Private/Graph/Renderers/MovieGraphShowFlags.cpp
+++ b/Engine/Plugins/MovieScene/MovieRenderPipeline/Source/MovieRenderPipelineCore/Private/Graph/Renderers/MovieGraphShowFlags.cpp
@@ -43,6 +43,10 @@ void UMovieGraphShowFlags::SetShowFlagEnabled(const uint32 ShowFlagIndex, const
 		return;
 	}
 
+#if WITH_EDITOR
+	Modify();
+#endif
+
 	bool& EnableState = ShowFlagEnableState.FindOrAdd(ShowFlagIndex);
 	EnableState = bIsShowFlagEnabled;
 }
@@ -54,6 +58,10 @@ bool UMovieGraphShowFlags::IsShowFlagOverridden(const uint32 ShowFlagIndex) cons
 
 void UMovieGraphShowFlags::SetShowFlagOverridden(const uint32 ShowFlagIndex, const bool bIsOverridden)
 {
+#if WITH_EDITOR
+	Modify();
+#endif
+	
 	if (bIsOverridden)
 	{
 		OverriddenShowFlags.Add(ShowFlagIndex);
@@ -79,6 +87,23 @@ const TSet<uint32>& UMovieGraphShowFlags::GetOverriddenShowFlags() const
 	return OverriddenShowFlags;
 }
 
+bool UMovieGraphShowFlags::IsShowFlagSetToDefaultValue(const uint32 ShowFlagIndex) const
+{
+	const bool* ShowFlagEnableValue = ShowFlagEnableState.Find(ShowFlagIndex);
+	
+	return !ShowFlagEnableValue ||
+		(ShowFlagEnableValue && (ShowFlags.GetSingleFlag(ShowFlagIndex) == *ShowFlagEnableValue)); 
+}
+
+void UMovieGraphShowFlags::RevertShowFlagToDefaultValue(const uint32 ShowFlagIndex)
+{
+#if WITH_EDITOR
+	Modify();
+#endif
+	
+	ShowFlagEnableState.Remove(ShowFlagIndex);
+}
+
 void UMovieGraphShowFlags::ApplyDefaultShowFlagValue(const uint32 ShowFlagIndex, const bool bShowFlagState)
 {
 	// If there's no current override on this flag, then the default can be set
diff --git a/Engine/Plugins/MovieScene/MovieRenderPipeline/Source/MovieRenderPipelineCore/Public/Graph/Renderers/MovieGraphShowFlags.h b/Engine/Plugins/MovieScene/MovieRenderPipeline/Source/MovieRenderPipelineCore/Public/Graph/Renderers/MovieGraphShowFlags.h
--- a/Engine/Plugins/MovieScene/MovieRenderPipeline/Source/MovieRenderPipelineCore/Public/Graph/Renderers/MovieGraphShowFlags.h
+++ b/Engine/Plugins/MovieScene/MovieRenderPipeline/Source/MovieRenderPipelineCore/Public/Graph/Renderers/MovieGraphShowFlags.h
@@ -40,6 +40,12 @@ public:
 	/** Gets the indices of all show flags that have been overridden. */
 	const TSet<uint32>& GetOverriddenShowFlags() const;
 
+	/** Determines if the provided show flag is set to a non-default value. The override state does not impact the value returned. */
+	bool IsShowFlagSetToDefaultValue(const uint32 ShowFlagIndex) const;
+
+	/** Reverts the provided show flag to its default value. Does not impact the override state for the flag. */
+	void RevertShowFlagToDefaultValue(const uint32 ShowFlagIndex);
+
 	/**
 	 * Applies a default value to the show flags. However, the default will not be applied if a user has already
 	 * specified an override for the flag. This should generally only be called from renderers in order for them to
diff --git a/Engine/Plugins/MovieScene/MovieRenderPipeline/Source/MovieRenderPipelineEditor/Private/Customizations/Graph/MovieGraphShowFlagsCustomization.h b/Engine/Plugins/MovieScene/MovieRenderPipeline/Source/MovieRenderPipelineEditor/Private/Customizations/Graph/MovieGraphShowFlagsCustomization.h
--- a/Engine/Plugins/MovieScene/MovieRenderPipeline/Source/MovieRenderPipelineEditor/Private/Customizations/Graph/MovieGraphShowFlagsCustomization.h
+++ b/Engine/Plugins/MovieScene/MovieRenderPipeline/Source/MovieRenderPipelineEditor/Private/Customizations/Graph/MovieGraphShowFlagsCustomization.h
@@ -79,35 +79,41 @@ protected:
 
 				// Temporary workaround: ShowFlag.DisplayName can sometimes be blank, so the non-display name is used instead.
 				const FText ShowFlagText = FText::FromName(ShowFlag.ShowFlagName);
+
+				// Create a custom reset handler for the show flags
+				FIsResetToDefaultVisible IsResetVisible = FIsResetToDefaultVisible::CreateSPLambda(this, [ShowFlagObject, ShowFlagIndex](TSharedPtr<IPropertyHandle> InChildHandle)
+				{
+					return !ShowFlagObject->IsShowFlagSetToDefaultValue(ShowFlagIndex);
+				});
+				FResetToDefaultHandler ResetHandler = FResetToDefaultHandler::CreateSPLambda(this, [ShowFlagObject, ShowFlagIndex](TSharedPtr<IPropertyHandle> InChildHandle)
+				{
+					const FScopedTransaction Transaction(LOCTEXT("Transaction_ResetShowFlagValue", "Reset to Default"));
+					ShowFlagObject->RevertShowFlagToDefaultValue(ShowFlagIndex);
+				});
+				FResetToDefaultOverride ResetOverride = FResetToDefaultOverride::Create(IsResetVisible, ResetHandler);
 				
 				FlagGroup.AddWidgetRow()
+				.PropertyHandleList({PropertyHandle})
+				.OverrideResetToDefault(ResetOverride)
+				.EditCondition(
+					TAttribute<bool>::Create([ShowFlagObject, ShowFlagIndex]()
+					{
+						return ShowFlagObject->IsShowFlagOverridden(ShowFlagIndex);
+					}),
+					FOnBooleanValueChanged::CreateSPLambda(this, [ShowFlagObject, ShowFlagIndex](bool NewValue)
+					{
+						const FScopedTransaction Transaction(LOCTEXT("Transaction_EditShowFlagOverrideState", "Edit Show Flag Override State"));
+						ShowFlagObject->SetShowFlagOverridden(ShowFlagIndex, NewValue);
+					})
+				)
 				.FilterString(ShowFlagText)
 				.NameContent()
 				[
 					SNew(SHorizontalBox)
-					+ SHorizontalBox::Slot()
-					.VAlign(VAlign_Center)
-					.AutoWidth()
-					.Padding(0, 0, 1, 0)
-					[
-						SNew(SCheckBox)
-						.IsChecked_Lambda([ShowFlagObject, ShowFlagIndex]()
-						{
-							return ShowFlagObject->IsShowFlagOverridden(ShowFlagIndex)
-								? ECheckBoxState::Checked : ECheckBoxState::Unchecked;
-						})
-						.OnCheckStateChanged_Lambda([ShowFlagObject, ShowFlagIndex](const ECheckBoxState NewState)
-						{
-							const bool bIsOverridden = (NewState == ECheckBoxState::Checked);
-							ShowFlagObject->SetShowFlagOverridden(ShowFlagIndex, bIsOverridden);
-						})
-					]
-
 					+ SHorizontalBox::Slot()
 					.VAlign(VAlign_Center)
 					[
 						SNew(STextBlock)
-						.IsEnabled_Lambda([ShowFlagObject, ShowFlagIndex]() { return ShowFlagObject->IsShowFlagOverridden(ShowFlagIndex); })
 						.Text(ShowFlagText)
 						.Font(CustomizationUtils.GetRegularFont())
 					]
@@ -115,7 +121,6 @@ protected:
 				.ValueContent()
 				[
 					SNew(SCheckBox)
-					.IsEnabled_Lambda([ShowFlagObject, ShowFlagIndex]() { return ShowFlagObject->IsShowFlagOverridden(ShowFlagIndex); })
 					.IsChecked_Lambda([ShowFlagObject, ShowFlagIndex]()
 					{
 						return ShowFlagObject->IsShowFlagEnabled(ShowFlagIndex)
@@ -123,6 +128,8 @@ protected:
 					})
 					.OnCheckStateChanged_Lambda([ShowFlagObject, ShowFlagIndex](const ECheckBoxState NewState)
 					{
+						const FScopedTransaction Transaction(LOCTEXT("Transaction_EditShowFlagEnableState", "Edit Show Flag Enable State"));
+						
 						const bool bIsUsed = (NewState == ECheckBoxState::Checked);
 						ShowFlagObject->SetShowFlagEnabled(ShowFlagIndex, bIsUsed);
 					})
