MoviePipeline: Updated copy/paste of conditions within condition groups to paste them as new objects rather than references.

diff --git a/Engine/Plugins/MovieScene/MovieRenderPipeline/Source/MovieRenderPipelineCore/Public/Graph/MovieGraphRenderLayerSubsystem.h b/Engine/Plugins/MovieScene/MovieRenderPipeline/Source/MovieRenderPipelineCore/Public/Graph/MovieGraphRenderLayerSubsystem.h
--- a/Engine/Plugins/MovieScene/MovieRenderPipeline/Source/MovieRenderPipelineCore/Public/Graph/MovieGraphRenderLayerSubsystem.h
+++ b/Engine/Plugins/MovieScene/MovieRenderPipeline/Source/MovieRenderPipelineCore/Public/Graph/MovieGraphRenderLayerSubsystem.h
@@ -587,15 +587,15 @@ private:
 	EMovieGraphConditionGroupOpType OpType;
 
 	/** The queries that are contained within the condition group. */
-	UPROPERTY(EditAnywhere, Category="General")
-	TArray<TObjectPtr<UMovieGraphConditionGroupQueryBase>> Queries;
+	UPROPERTY(EditAnywhere, Category="General", Instanced)
+	TArray<TObjectPtr<UMovieGraphConditionGroupQueryBase>> Queries;	// Note: Marked as Instanced so conditions get duplicated during copy/paste (not referenced)
 
 	/** Persisted actor set which can be re-used for query evaluations across frames to prevent constantly re-allocating it. */
-	UPROPERTY(Transient)
+	UPROPERTY(Transient, DuplicateTransient)
 	mutable TSet<AActor*> QueryResult;
 
 	/** Persisted actor set which can be re-used for condition group evaluations across frames to prevent constantly re-allocating it. */
-	UPROPERTY(Transient)
+	UPROPERTY(Transient, DuplicateTransient)
 	mutable TSet<AActor*> EvaluationResult;
 };
 
