MoviePipeline: Added the ability to duplicate a condition group condition in the graph.

diff --git a/Engine/Plugins/MovieScene/MovieRenderPipeline/Source/MovieRenderPipelineCore/Private/Graph/MovieGraphRenderLayerSubsystem.cpp b/Engine/Plugins/MovieScene/MovieRenderPipeline/Source/MovieRenderPipelineCore/Private/Graph/MovieGraphRenderLayerSubsystem.cpp
--- a/Engine/Plugins/MovieScene/MovieRenderPipeline/Source/MovieRenderPipelineCore/Private/Graph/MovieGraphRenderLayerSubsystem.cpp
+++ b/Engine/Plugins/MovieScene/MovieRenderPipeline/Source/MovieRenderPipelineCore/Private/Graph/MovieGraphRenderLayerSubsystem.cpp
@@ -2287,6 +2287,32 @@ bool UMovieGraphConditionGroup::RemoveQuery(UMovieGraphConditionGroupQueryBase*
 	return Queries.RemoveSingle(InQuery) == 1;
 }
 
+UMovieGraphConditionGroupQueryBase* UMovieGraphConditionGroup::DuplicateQuery(const int32 QueryIndex)
+{
+	if (!Queries.IsValidIndex(QueryIndex))
+	{
+		UE_LOG(LogMovieRenderPipeline, Warning, TEXT("Invalid query index provided to DuplicateQuery()."));
+		return nullptr;
+	}
+	
+	const UMovieGraphConditionGroupQueryBase* SourceQuery = Queries[QueryIndex];
+
+	FObjectDuplicationParameters DuplicationParameters = InitStaticDuplicateObjectParams(SourceQuery, this);
+	UMovieGraphConditionGroupQueryBase* DuplicateQuery = Cast<UMovieGraphConditionGroupQueryBase>(StaticDuplicateObjectEx(DuplicationParameters));
+
+	if (DuplicateQuery)
+	{
+		Modify();
+		Queries.Add(DuplicateQuery);
+	}
+	else
+	{
+		UE_LOG(LogMovieRenderPipeline, Warning, TEXT("Failed to duplicate condition group query."));
+	}
+
+	return DuplicateQuery;
+}
+
 bool UMovieGraphConditionGroup::IsFirstConditionGroup() const
 {
 	const UMovieGraphCollection* ParentCollection = GetTypedOuter<UMovieGraphCollection>();
diff --git a/Engine/Plugins/MovieScene/MovieRenderPipeline/Source/MovieRenderPipelineCore/Public/Graph/MovieGraphRenderLayerSubsystem.h b/Engine/Plugins/MovieScene/MovieRenderPipeline/Source/MovieRenderPipelineCore/Public/Graph/MovieGraphRenderLayerSubsystem.h
--- a/Engine/Plugins/MovieScene/MovieRenderPipeline/Source/MovieRenderPipelineCore/Public/Graph/MovieGraphRenderLayerSubsystem.h
+++ b/Engine/Plugins/MovieScene/MovieRenderPipeline/Source/MovieRenderPipelineCore/Public/Graph/MovieGraphRenderLayerSubsystem.h
@@ -563,6 +563,13 @@ public:
 	UFUNCTION(BlueprintCallable, Category = "Settings")
 	bool RemoveQuery(UMovieGraphConditionGroupQueryBase* InQuery);
 
+	/**
+	 * Duplicates the condition group query at the specified index. The duplicate is placed at the end of the query list. Returns the duplicate
+	 * query on success, else nullptr.
+	 */
+	UFUNCTION(BlueprintCallable, Category = "Settings")
+	UMovieGraphConditionGroupQueryBase* DuplicateQuery(const int32 QueryIndex);
+
 	/** Determines if this is the first condition group under the parent collection. */
 	UFUNCTION(BlueprintCallable, Category = "Settings")
 	bool IsFirstConditionGroup() const;
diff --git a/Engine/Plugins/MovieScene/MovieRenderPipeline/Source/MovieRenderPipelineEditor/Private/Customizations/Graph/MovieGraphCollectionsCustomization.cpp b/Engine/Plugins/MovieScene/MovieRenderPipeline/Source/MovieRenderPipelineEditor/Private/Customizations/Graph/MovieGraphCollectionsCustomization.cpp
--- a/Engine/Plugins/MovieScene/MovieRenderPipeline/Source/MovieRenderPipelineEditor/Private/Customizations/Graph/MovieGraphCollectionsCustomization.cpp
+++ b/Engine/Plugins/MovieScene/MovieRenderPipeline/Source/MovieRenderPipelineEditor/Private/Customizations/Graph/MovieGraphCollectionsCustomization.cpp
@@ -238,6 +238,18 @@ void FMovieGraphConditionGroupQueryBuilder::GenerateChildContent(IDetailChildren
 		.ShowPropertyButtons(false)
 		.OverrideResetToDefault(FResetToDefaultOverride::Hide())
 		.CustomWidget()
+		.AddCustomContextMenuAction(FUIAction(
+			FExecuteAction::CreateLambda([this, QueryIndex]()
+			{
+				if (WeakConditionGroup.IsValid())
+				{
+					const FScopedTransaction Transaction(LOCTEXT("DuplicateConditionGroupQuery", "Duplicate Condition Group Query"));
+					WeakConditionGroup->DuplicateQuery(QueryIndex);
+				}
+			})),
+			LOCTEXT("DuplicateConditionGroupQuery", "Duplicate"),
+			LOCTEXT("DuplicateConditionGroupQueryTooltip", "Duplicate this condition."),
+			FSlateIcon(FAppStyle::GetAppStyleSetName(), "Icons.Duplicate"))
 		.FilterString(Query->GetClass()->GetDisplayNameText())
 		.NameContent()
 		.HAlign(HAlign_Fill)
