MoviePipeline: Fix CIS localization warning.

diff --git a/Engine/Plugins/MovieScene/MovieRenderPipeline/Source/MovieRenderPipelineEditor/Private/Customizations/Graph/MovieGraphCollectionsCustomization.cpp b/Engine/Plugins/MovieScene/MovieRenderPipeline/Source/MovieRenderPipelineEditor/Private/Customizations/Graph/MovieGraphCollectionsCustomization.cpp
--- a/Engine/Plugins/MovieScene/MovieRenderPipeline/Source/MovieRenderPipelineEditor/Private/Customizations/Graph/MovieGraphCollectionsCustomization.cpp
+++ b/Engine/Plugins/MovieScene/MovieRenderPipeline/Source/MovieRenderPipelineEditor/Private/Customizations/Graph/MovieGraphCollectionsCustomization.cpp
@@ -243,7 +243,7 @@ void FMovieGraphConditionGroupQueryBuilder::GenerateChildContent(IDetailChildren
 			{
 				if (WeakConditionGroup.IsValid())
 				{
-					const FScopedTransaction Transaction(LOCTEXT("DuplicateConditionGroupQuery", "Duplicate Condition Group Query"));
+					const FScopedTransaction Transaction(LOCTEXT("DuplicateConditionGroupQuery_Transaction", "Duplicate Condition Group Query"));
 					WeakConditionGroup->DuplicateQuery(QueryIndex);
 				}
 			})),
