Fixing single header compilation.

diff --git a/Engine/Plugins/Experimental/CustomDetailsView/Source/CustomDetailsView/Public/Items/ICustomDetailsViewCustomItem.h b/Engine/Plugins/Experimental/CustomDetailsView/Source/CustomDetailsView/Public/Items/ICustomDetailsViewCustomItem.h
--- a/Engine/Plugins/Experimental/CustomDetailsView/Source/CustomDetailsView/Public/Items/ICustomDetailsViewCustomItem.h
+++ b/Engine/Plugins/Experimental/CustomDetailsView/Source/CustomDetailsView/Public/Items/ICustomDetailsViewCustomItem.h
@@ -8,6 +8,7 @@ class ICustomDetailsViewItem;
 class FText;
 class SWidget;
 enum class EDetailNodeType;
+template<typename OptionalType> struct TOptional;
 
 class ICustomDetailsViewCustomItem
 {
diff --git a/Engine/Plugins/MovieScene/MovieRenderPipeline/Source/MovieRenderPipelineEditor/Private/Customizations/Graph/MovieGraphShowFlagsCustomization.h b/Engine/Plugins/MovieScene/MovieRenderPipeline/Source/MovieRenderPipelineEditor/Private/Customizations/Graph/MovieGraphShowFlagsCustomization.h
--- a/Engine/Plugins/MovieScene/MovieRenderPipeline/Source/MovieRenderPipelineEditor/Private/Customizations/Graph/MovieGraphShowFlagsCustomization.h
+++ b/Engine/Plugins/MovieScene/MovieRenderPipeline/Source/MovieRenderPipelineEditor/Private/Customizations/Graph/MovieGraphShowFlagsCustomization.h
@@ -8,6 +8,7 @@
 #include "IDetailChildrenBuilder.h"
 #include "IDetailGroup.h"
 #include "PropertyHandle.h"
+#include "ScopedTransaction.h"
 #include "Widgets/Input/SCheckBox.h"
 #include "Widgets/Text/STextBlock.h"
 
diff --git a/Engine/Source/Runtime/TypedElementFramework/Public/Elements/Interfaces/Capabilities/TypedElementUiStyleOverrideCapability.h b/Engine/Source/Runtime/TypedElementFramework/Public/Elements/Interfaces/Capabilities/TypedElementUiStyleOverrideCapability.h
--- a/Engine/Source/Runtime/TypedElementFramework/Public/Elements/Interfaces/Capabilities/TypedElementUiStyleOverrideCapability.h
+++ b/Engine/Source/Runtime/TypedElementFramework/Public/Elements/Interfaces/Capabilities/TypedElementUiStyleOverrideCapability.h
@@ -6,6 +6,8 @@
 #include "Misc/Attribute.h"
 #include "Math/Color.h"
 
+struct FSlateColor;
+
 /**
  * Interface to provide access to widgets that support working with style overrides.
  */
