MoviePipeline: Added tooltips in the graph that surface a variable's description (in the Members panel, the variable node and variable node pin, and the queue UI).

diff --git a/Engine/Plugins/MovieScene/MovieRenderPipeline/Source/MovieRenderPipelineCore/Private/MovieJobVariableAssignmentContainer.cpp b/Engine/Plugins/MovieScene/MovieRenderPipeline/Source/MovieRenderPipelineCore/Private/MovieJobVariableAssignmentContainer.cpp
--- a/Engine/Plugins/MovieScene/MovieRenderPipeline/Source/MovieRenderPipelineCore/Private/MovieJobVariableAssignmentContainer.cpp
+++ b/Engine/Plugins/MovieScene/MovieRenderPipeline/Source/MovieRenderPipelineCore/Private/MovieJobVariableAssignmentContainer.cpp
@@ -282,7 +282,8 @@ bool UMovieJobVariableAssignmentContainer::GenerateVariableOverride(const UMovie
 	NewProperty.Name = FName(InGraphVariable->GetMemberName());
 	NewProperty.ID = FGuid::NewGuid();
 #if WITH_EDITOR
-	NewProperty.MetaData.Add(FPropertyBagPropertyDescMetaData("VariableGUID", InGraphVariable->GetGuid().ToString()));
+	NewProperty.MetaData.Add(FPropertyBagPropertyDescMetaData(VariableGuidMetaDataKey, InGraphVariable->GetGuid().ToString()));
+	NewProperty.MetaData.Add(FPropertyBagPropertyDescMetaData(ToolTipMetaDataKey, InGraphVariable->Description));
 #endif
 
 	// Track a separate EditCondition property that can enable/disable the above property. Since the variable can be
@@ -291,7 +292,7 @@ bool UMovieJobVariableAssignmentContainer::GenerateVariableOverride(const UMovie
 	FPropertyBagPropertyDesc NewPropertyEditCondition = FPropertyBagPropertyDesc(FName(EditCondPropName), EPropertyBagPropertyType::Bool);
 	NewPropertyEditCondition.ID = FGuid::NewGuid();
 #if WITH_EDITOR
-	NewPropertyEditCondition.MetaData.Add(FPropertyBagPropertyDescMetaData("VariableGUID", InGraphVariable->GetGuid().ToString()));
+	NewPropertyEditCondition.MetaData.Add(FPropertyBagPropertyDescMetaData(VariableGuidMetaDataKey, InGraphVariable->GetGuid().ToString()));
 #endif
 
 #if WITH_EDITOR
@@ -312,7 +313,7 @@ void UMovieJobVariableAssignmentContainer::UpdateGraphVariableOverrides()
 	{
 		const FPropertyBagPropertyDescMetaData* Meta = Desc.MetaData.FindByPredicate([](const FPropertyBagPropertyDescMetaData& MetaData)
 		{
-			return MetaData.Key == "VariableGUID";
+			return MetaData.Key == VariableGuidMetaDataKey;
 		});
 		
 		if (Meta)
@@ -413,6 +414,17 @@ void UMovieJobVariableAssignmentContainer::UpdateGraphVariableOverrides()
 				bNeedsToRegenerate = true;
 				Desc.ContainerTypes = { static_cast<EPropertyBagContainerType>(Variable->GetValueContainerType()) };
 			}
+
+			FPropertyBagPropertyDescMetaData* ToolTipMetaDataEntry =
+				Desc.MetaData.FindByPredicate([](const FPropertyBagPropertyDescMetaData& InMetaDataEntry)
+				{
+					return InMetaDataEntry.Key == ToolTipMetaDataKey;
+				});
+			if (ToolTipMetaDataEntry && (ToolTipMetaDataEntry->Value != Variable->Description))
+			{
+				bNeedsToRegenerate = true;
+				ToolTipMetaDataEntry->Value = Variable->Description;
+			}
 		}
 	}
 
diff --git a/Engine/Plugins/MovieScene/MovieRenderPipeline/Source/MovieRenderPipelineCore/Public/MovieJobVariableAssignmentContainer.h b/Engine/Plugins/MovieScene/MovieRenderPipeline/Source/MovieRenderPipelineCore/Public/MovieJobVariableAssignmentContainer.h
--- a/Engine/Plugins/MovieScene/MovieRenderPipeline/Source/MovieRenderPipelineCore/Public/MovieJobVariableAssignmentContainer.h
+++ b/Engine/Plugins/MovieScene/MovieRenderPipeline/Source/MovieRenderPipelineCore/Public/MovieJobVariableAssignmentContainer.h
@@ -257,6 +257,12 @@ protected:
 	FName ConvertVariableToInternalName(const UMovieGraphVariable* InGraphVariable) const;
 
 private:
+	/* The metadata key on properties that is associated with a value which identifies a graph variable (by GUID). */
+	static inline FName VariableGuidMetaDataKey = FName(TEXT("VariableGUID"));
+	
+	/* The metadata key on properties that is associated with a tooltip value. */
+	static inline FName ToolTipMetaDataKey = FName(TEXT("ToolTip"));
+	
 	/** The properties managed by this object. */
 	UPROPERTY(EditAnywhere, meta=(ShowOnlyInnerProperties, FixedLayout), Category = "Value")
 	FInstancedPropertyBag Value;
diff --git a/Engine/Plugins/MovieScene/MovieRenderPipeline/Source/MovieRenderPipelineEditor/Private/Graph/MovieEdGraphNode.h b/Engine/Plugins/MovieScene/MovieRenderPipeline/Source/MovieRenderPipelineEditor/Private/Graph/MovieEdGraphNode.h
--- a/Engine/Plugins/MovieScene/MovieRenderPipeline/Source/MovieRenderPipelineEditor/Private/Graph/MovieEdGraphNode.h
+++ b/Engine/Plugins/MovieScene/MovieRenderPipeline/Source/MovieRenderPipelineEditor/Private/Graph/MovieEdGraphNode.h
@@ -53,7 +53,7 @@ protected:
 	void CreatePins(const TArray<UMovieGraphPin*>& InInputPins, const TArray<UMovieGraphPin*>& InOutputPins);
 	
 	/** Gets the tooltip for the given pin. */
-	FString GetPinTooltip(const UMovieGraphPin* InPin) const;
+	virtual FString GetPinTooltip(const UMovieGraphPin* InPin) const;
 
 	/** Recreate the pins on this node, discarding all existing pins. */
 	void ReconstructPins();
diff --git a/Engine/Plugins/MovieScene/MovieRenderPipeline/Source/MovieRenderPipelineEditor/Private/Graph/MovieEdGraphVariableNode.cpp b/Engine/Plugins/MovieScene/MovieRenderPipeline/Source/MovieRenderPipelineEditor/Private/Graph/MovieEdGraphVariableNode.cpp
--- a/Engine/Plugins/MovieScene/MovieRenderPipeline/Source/MovieRenderPipelineEditor/Private/Graph/MovieEdGraphVariableNode.cpp
+++ b/Engine/Plugins/MovieScene/MovieRenderPipeline/Source/MovieRenderPipelineEditor/Private/Graph/MovieEdGraphVariableNode.cpp
@@ -67,4 +67,46 @@ bool UMoviePipelineEdGraphVariableNode::CanPasteHere(const UEdGraph* TargetGraph
 	return true;
 }
 
+FText UMoviePipelineEdGraphVariableNode::GetTooltipText() const
+{
+	const FString VariableDescription = GetVariableDescription();
+	if (!VariableDescription.IsEmpty())
+	{
+		return FText::FromString(VariableDescription);
+	}
+	
+	return Super::GetTooltipText();
+}
+
+FString UMoviePipelineEdGraphVariableNode::GetPinTooltip(const UMovieGraphPin* InPin) const
+{
+	FString Tooltip = Super::GetPinTooltip(InPin);
+
+	// Add the variable description to the tooltip if available. There's only one pin on the variable node, so we don't need to check to see which
+	// pin is being hovered over.
+	const FString VariableDescription = GetVariableDescription();
+	if (!VariableDescription.IsEmpty())
+	{
+		Tooltip += "\n\n" + VariableDescription;
+	}
+
+	return Tooltip;
+}
+
+FString UMoviePipelineEdGraphVariableNode::GetVariableDescription() const
+{
+	if (const UMovieGraphVariableNode* VariableNode = Cast<UMovieGraphVariableNode>(RuntimeNode))
+	{
+		if (const UMovieGraphVariable* VariableMember = VariableNode->GetVariable())
+		{
+			if (!VariableMember->Description.IsEmpty())
+			{
+				return VariableMember->Description;
+			}
+		}
+	}
+
+	return FString();
+}
+
 #undef LOCTEXT_NAMESPACE
\ No newline at end of file
diff --git a/Engine/Plugins/MovieScene/MovieRenderPipeline/Source/MovieRenderPipelineEditor/Private/Graph/MovieEdGraphVariableNode.h b/Engine/Plugins/MovieScene/MovieRenderPipeline/Source/MovieRenderPipelineEditor/Private/Graph/MovieEdGraphVariableNode.h
--- a/Engine/Plugins/MovieScene/MovieRenderPipeline/Source/MovieRenderPipelineEditor/Private/Graph/MovieEdGraphVariableNode.h
+++ b/Engine/Plugins/MovieScene/MovieRenderPipeline/Source/MovieRenderPipelineEditor/Private/Graph/MovieEdGraphVariableNode.h
@@ -17,5 +17,15 @@ public:
 	virtual FText GetNodeTitle(ENodeTitleType::Type TitleType) const override;
 	virtual void AllocateDefaultPins() override;
 	virtual bool CanPasteHere(const UEdGraph* TargetGraph) const override;
+	virtual FText GetTooltipText() const override;
 	// ~End UEdGraphNode interface
+
+protected:
+	// ~Begin UMoviePipelineEdGraphNodeBase interface
+	virtual FString GetPinTooltip(const UMovieGraphPin* InPin) const override;
+	// ~End UMoviePipelineEdGraphNodeBase interface
+
+private:
+	/** Get the description (if any) for the variable this node represents. */
+	FString GetVariableDescription() const;
 };
\ No newline at end of file
diff --git a/Engine/Plugins/MovieScene/MovieRenderPipeline/Source/MovieRenderPipelineEditor/Private/Graph/SMovieGraphMembersTabContent.cpp b/Engine/Plugins/MovieScene/MovieRenderPipeline/Source/MovieRenderPipelineEditor/Private/Graph/SMovieGraphMembersTabContent.cpp
--- a/Engine/Plugins/MovieScene/MovieRenderPipeline/Source/MovieRenderPipelineEditor/Private/Graph/SMovieGraphMembersTabContent.cpp
+++ b/Engine/Plugins/MovieScene/MovieRenderPipeline/Source/MovieRenderPipelineEditor/Private/Graph/SMovieGraphMembersTabContent.cpp
@@ -137,7 +137,8 @@ TSharedRef<SWidget> SMovieGraphMembersTabContent::CreateActionWidget(FCreateWidg
 			[
 				SNew(STextBlock)
 				.Text(InCreateData->Action->GetMenuDescription())
-			];	
+				.ToolTipText(InCreateData->Action->GetTooltipDescription())
+			];
 	}
 	
 	return
@@ -235,7 +236,7 @@ void SMovieGraphMembersTabContent::CollectAllActions(FGraphActionListBuilderBase
 	auto AddToActionMenu = [&ActionMenuBuilder, this](UMovieGraphMember* ActionTarget, const EActionSection Section, const FText& Category) -> void
 	{
 		const FText MemberActionDesc = FText::FromString(ActionTarget->GetMemberName());
-		const FText MemberActionTooltip;
+		const FText MemberActionTooltip = FText::FromString(ActionTarget->Description);
 		const FText MemberActionKeywords;
 		const int32 MemberActionSectionID = static_cast<int32>(Section);
 		const TSharedPtr<FMovieGraphSchemaAction> MemberAction(new FMovieGraphSchemaAction(Category, MemberActionDesc, MemberActionTooltip, 0, MemberActionKeywords, MemberActionSectionID));
