MoviePipeline: Removed support for the "Buffer Visualization" view mode on the deferred renderer node in the graph -- it was causing a crash.

Buffer Visualization conflicts with Post Process Material support, and it's unlikely that the Buffer Visualization mode will actually be used in practice.

Also removed "Unknown" from the mode list because it's a value not needed by users.

diff --git a/Engine/Plugins/MovieScene/MovieRenderPipeline/Source/MovieRenderPipelineRenderPasses/Public/Graph/Nodes/MovieGraphDeferredPassNode.h b/Engine/Plugins/MovieScene/MovieRenderPipeline/Source/MovieRenderPipelineRenderPasses/Public/Graph/Nodes/MovieGraphDeferredPassNode.h
--- a/Engine/Plugins/MovieScene/MovieRenderPipeline/Source/MovieRenderPipelineRenderPasses/Public/Graph/Nodes/MovieGraphDeferredPassNode.h
+++ b/Engine/Plugins/MovieScene/MovieRenderPipeline/Source/MovieRenderPipelineRenderPasses/Public/Graph/Nodes/MovieGraphDeferredPassNode.h
@@ -108,7 +108,7 @@ public:
 	* have to be used in combination with certain Show Flags to produce a result similar to what you see in
 	* the viewport.
 	*/
-	UPROPERTY(EditAnywhere, BlueprintReadWrite, Category = "Settings", meta=(EditCondition="bOverride_ViewModeIndex", InvalidEnumValues = "VMI_PathTracing"))
+	UPROPERTY(EditAnywhere, BlueprintReadWrite, Category = "Settings", meta=(EditCondition="bOverride_ViewModeIndex", InvalidEnumValues = "VMI_PathTracing,VMI_VisualizeBuffer,VMI_Unknown"))
 	TEnumAsByte<EViewModeIndex> ViewModeIndex;
 
 	/**
