MoviePipeline: Fixed an issue where pins on conditional nodes could not have their non-branch pins promoted to variables.

diff --git a/Engine/Plugins/MovieScene/MovieRenderPipeline/Source/MovieRenderPipelineCore/Private/Graph/MovieGraphNode.cpp b/Engine/Plugins/MovieScene/MovieRenderPipeline/Source/MovieRenderPipelineCore/Private/Graph/MovieGraphNode.cpp
--- a/Engine/Plugins/MovieScene/MovieRenderPipeline/Source/MovieRenderPipelineCore/Private/Graph/MovieGraphNode.cpp
+++ b/Engine/Plugins/MovieScene/MovieRenderPipeline/Source/MovieRenderPipelineCore/Private/Graph/MovieGraphNode.cpp
@@ -515,6 +515,23 @@ TArray<FMovieGraphPropertyInfo> UMovieGraphNode::GetOverrideablePropertyInfo() c
 		}
 	}
 
+	// Some properties exist only via GetInputProperties(). Some of them may be non-branch-typed so they should be exposed as overrideable.
+	for (const FMovieGraphPinProperties& InputPinProperty : GetInputPinProperties())
+	{
+		if (InputPinProperty.bIsBranch)
+		{
+			continue;
+		}
+		
+		FMovieGraphPropertyInfo Info;
+		Info.Name = InputPinProperty.Label;
+		Info.ValueType = InputPinProperty.Type;
+		Info.ValueTypeObject = InputPinProperty.TypeObject;
+		Info.bIsPermanentlyExposed = InputPinProperty.bIsBuiltIn;
+
+		OverrideableProperties.Add(MoveTemp(Info));
+	}
+
 	return OverrideableProperties;
 }
 
diff --git a/Engine/Plugins/MovieScene/MovieRenderPipeline/Source/MovieRenderPipelineCore/Public/Graph/MovieGraphNode.h b/Engine/Plugins/MovieScene/MovieRenderPipeline/Source/MovieRenderPipelineCore/Public/Graph/MovieGraphNode.h
--- a/Engine/Plugins/MovieScene/MovieRenderPipeline/Source/MovieRenderPipelineCore/Public/Graph/MovieGraphNode.h
+++ b/Engine/Plugins/MovieScene/MovieRenderPipeline/Source/MovieRenderPipelineCore/Public/Graph/MovieGraphNode.h
@@ -57,6 +57,12 @@ struct FMovieGraphPropertyInfo
 	UPROPERTY()
 	TObjectPtr<const UObject> ValueTypeObject;
 
+	/**
+	 * Whether this property is permanently exposed on the node. If true, it cannot be toggled on/off (via Expose Property as Pin).
+	 */
+	UPROPERTY()
+	bool bIsPermanentlyExposed = false;
+
 	/**
 	 * Determines if this struct represents the same property as another instance of this struct.
 	 *
diff --git a/Engine/Plugins/MovieScene/MovieRenderPipeline/Source/MovieRenderPipelineEditor/Private/Graph/MovieEdGraphNode.cpp b/Engine/Plugins/MovieScene/MovieRenderPipeline/Source/MovieRenderPipelineEditor/Private/Graph/MovieEdGraphNode.cpp
--- a/Engine/Plugins/MovieScene/MovieRenderPipeline/Source/MovieRenderPipelineEditor/Private/Graph/MovieEdGraphNode.cpp
+++ b/Engine/Plugins/MovieScene/MovieRenderPipeline/Source/MovieRenderPipelineEditor/Private/Graph/MovieEdGraphNode.cpp
@@ -309,6 +309,12 @@ void UMoviePipelineEdGraphNode::GetPropertyPromotionContextMenuActions(UToolMenu
 	FToolMenuSection& ExposeAsPinSection = Menu->AddSection("MoviePipelineGraphExposeAsPin", LOCTEXT("ExposeAsPin", "Expose Property as Pin"));
 	for (const FMovieGraphPropertyInfo& PropertyInfo : OverrideablePropertyInfo)
 	{
+		// If a property is permanently exposed on the node, don't allow it to be toggled off
+		if (PropertyInfo.bIsPermanentlyExposed)
+		{
+			continue;
+		}
+		
 		ExposeAsPinSection.AddMenuEntry(
 			PropertyInfo.Name,
 			FText::FromName(PropertyInfo.Name),
