MoviePipeline: Exposed some more data on graph nodes to Python.

Made a small breaking change to the API: GetInputPins() and GetOutputPins() now return raw pointers instead of TObjectPtrs in order to get these methods UFUNCTION'd and exposed to Python.

diff --git a/Engine/Plugins/MovieScene/MovieRenderPipeline/Source/MovieRenderPipelineCore/Private/Graph/MovieGraphNode.cpp b/Engine/Plugins/MovieScene/MovieRenderPipeline/Source/MovieRenderPipelineCore/Private/Graph/MovieGraphNode.cpp
--- a/Engine/Plugins/MovieScene/MovieRenderPipeline/Source/MovieRenderPipelineCore/Private/Graph/MovieGraphNode.cpp
+++ b/Engine/Plugins/MovieScene/MovieRenderPipeline/Source/MovieRenderPipelineCore/Private/Graph/MovieGraphNode.cpp
@@ -575,7 +575,7 @@ TArray<UMovieGraphPin*> UMovieGraphNode::EvaluatePinsToFollow(FMovieGraphEvaluat
 	// By default we provide all Input Pins to this node that are the Branch Type.
 	// You should override this in downstream nodes that need custom logic, such
 	// as branch or switch nodes.
-	for (const TObjectPtr<UMovieGraphPin>& InputPin : GetInputPins())
+	for (UMovieGraphPin* InputPin : GetInputPins())
 	{
 		if (InputPin->Properties.bIsBranch)
 		{
diff --git a/Engine/Plugins/MovieScene/MovieRenderPipeline/Source/MovieRenderPipelineCore/Public/Graph/MovieGraphNode.h b/Engine/Plugins/MovieScene/MovieRenderPipeline/Source/MovieRenderPipelineCore/Public/Graph/MovieGraphNode.h
--- a/Engine/Plugins/MovieScene/MovieRenderPipeline/Source/MovieRenderPipelineCore/Public/Graph/MovieGraphNode.h
+++ b/Engine/Plugins/MovieScene/MovieRenderPipeline/Source/MovieRenderPipelineCore/Public/Graph/MovieGraphNode.h
@@ -86,6 +86,7 @@ struct FMovieGraphPropertyInfo
 };
 
 /** Describes a restriction on what kind of branch a node can be created in within the graph. */
+UENUM(BlueprintType)
 enum class EMovieGraphBranchRestriction : uint8
 {
 	Any,			///< The node can be created in any type of branch
@@ -113,14 +114,23 @@ public:
 	
 	UMovieGraphNode();
 
-	const TArray<TObjectPtr<UMovieGraphPin>>& GetInputPins() const { return InputPins; }
-	const TArray<TObjectPtr<UMovieGraphPin>>& GetOutputPins() const { return OutputPins; }
+	/** Gets all input pins on the node. Note that the returned array is const, so input pins cannot be added/removed from the node via this array. */
+	UFUNCTION(BlueprintCallable, Category = "Movie Graph")
+	const TArray<UMovieGraphPin*>& GetInputPins() const { return InputPins; }
+	
+	/** Gets all output pins on the node. Note that the returned array is const, so output pins cannot be added/removed from the node via this array. */
+	UFUNCTION(BlueprintCallable, Category = "Movie Graph")
+	const TArray<UMovieGraphPin*>& GetOutputPins() const { return OutputPins; }
 	
+	/** Gets the properties for all input pins. */
+	UFUNCTION(BlueprintCallable, Category = "Movie Graph")
 	virtual TArray<FMovieGraphPinProperties> GetInputPinProperties() const
 	{
 		return TArray<FMovieGraphPinProperties>();
 	}
 	
+	/** Gets the properties for all output pins. */
+	UFUNCTION(BlueprintCallable, Category = "Movie Graph")
 	virtual TArray<FMovieGraphPinProperties> GetOutputPinProperties() const
 	{
 		return TArray<FMovieGraphPinProperties>();
@@ -181,6 +191,7 @@ public:
 	void SetDynamicPropertyOverridden(const FName& InPropertyName, const bool bIsOverridden);
 
 	/** Gets the information about properties which can be exposed as a pin on the node. */
+	UFUNCTION(BlueprintCallable, Category = "Movie Graph")
 	virtual TArray<FMovieGraphPropertyInfo> GetOverrideablePropertyInfo() const;
 
 	/** Gets the information about properties which are currently exposed as pins on the node. */
@@ -220,6 +231,7 @@ public:
 	 * Determines if this node type can be added to the graph interactively by a user or via the API when constructing a graph.
 	 * @return true if the object can be added via the API, false otherwise
 	 */
+	UFUNCTION(BlueprintCallable, Category = "Movie Graph")
 	virtual bool CanBeAddedByUser() const { return true; }
 
 	/**
@@ -246,24 +258,30 @@ public:
 	UMovieGraphPin* GetOutputPin(const FName& InPinLabel) const;
 
 	/** Gets the first input pin on the node which has a connection, or nullptr if no pins are connected. */
+	UFUNCTION(BlueprintCallable, Category = "Movie Graph")
 	UMovieGraphPin* GetFirstConnectedInputPin() const;
 	
 	/** Gets the first output pin on the node which has a connection, or nullptr if no pins are connected. */
+	UFUNCTION(BlueprintCallable, Category = "Movie Graph")
 	UMovieGraphPin* GetFirstConnectedOutputPin() const;
 
 	/** Gets the GUID which uniquely identifies this node. */
 	const FGuid& GetGuid() const { return Guid; }
 	
 	/** Determines which types of branches the node can be created in. */
+	UFUNCTION(BlueprintCallable, Category = "Movie Graph")
 	virtual EMovieGraphBranchRestriction GetBranchRestriction() const { return EMovieGraphBranchRestriction::Any; }
 
 	/** Determines if this node can be disabled. */
+	UFUNCTION(BlueprintCallable, Category = "Movie Graph")
 	virtual bool CanBeDisabled() const;
 
 	/** Set whether this node is currently disabled. */
+	UFUNCTION(BlueprintCallable, Category = "Movie Graph")
 	void SetDisabled(const bool bNewDisableState);
 
 	/** Determines if this node is currently disabled. */
+	UFUNCTION(BlueprintCallable, Category = "Movie Graph")
 	bool IsDisabled() const;
 
 #if WITH_EDITOR
@@ -320,12 +338,15 @@ public:
 	 * Gets the node's title. Optionally gets a more descriptive, multi-line title for the node if bGetDescriptive is
 	 * set to true.
 	 */
+	UFUNCTION(BlueprintCallable, Category = "Movie Graph")
 	virtual FText GetNodeTitle(const bool bGetDescriptive = false) const PURE_VIRTUAL(UMovieGraphNode::GetNodeTitle, return FText(););
 
 	/** Gets the category that the node belongs under. */
+	UFUNCTION(BlueprintCallable, Category = "Movie Graph")
 	virtual FText GetMenuCategory() const PURE_VIRTUAL(UMovieGraphNode::GetMenuCategory, return FText(); );
 
 	/** Gets the keywords (space-separated) that will be searched in the node creation context menu. */
+	UFUNCTION(BlueprintCallable, Category = "Movie Graph")
 	virtual FText GetKeywords() const { return FText(); }
 #endif
 
diff --git a/Engine/Plugins/MovieScene/MovieRenderPipeline/Source/MovieRenderPipelineCore/Public/Graph/MovieGraphPin.h b/Engine/Plugins/MovieScene/MovieRenderPipeline/Source/MovieRenderPipelineCore/Public/Graph/MovieGraphPin.h
--- a/Engine/Plugins/MovieScene/MovieRenderPipeline/Source/MovieRenderPipelineCore/Public/Graph/MovieGraphPin.h
+++ b/Engine/Plugins/MovieScene/MovieRenderPipeline/Source/MovieRenderPipelineCore/Public/Graph/MovieGraphPin.h
@@ -105,17 +105,31 @@ public:
 	bool BreakEdgeTo(UMovieGraphPin* InOtherPin);
 	bool BreakAllEdges();
 	FPinConnectionResponse CanCreateConnection_PinConnectionResponse(const UMovieGraphPin* InOtherPin) const;
+
+	UFUNCTION(BlueprintCallable, Category = "Movie Graph")
 	bool CanCreateConnection(const UMovieGraphPin* InOtherPin) const;
+
+	UFUNCTION(BlueprintCallable, Category = "Movie Graph")
 	bool IsConnected() const;
+
+	UFUNCTION(BlueprintCallable, Category = "Movie Graph")
 	bool IsInputPin() const;
+	
+	UFUNCTION(BlueprintCallable, Category = "Movie Graph")
 	bool IsOutputPin() const;
+
+	UFUNCTION(BlueprintCallable, Category = "Movie Graph")
 	int32 EdgeCount() const;
+
+	UFUNCTION(BlueprintCallable, Category = "Movie Graph")
 	bool AllowsMultipleConnections() const;
 
 	/** Gets the first pin connected to this pin. Returns nullptr if no valid connection exists. */
+	UFUNCTION(BlueprintCallable, Category = "Movie Graph")
 	UMovieGraphPin* GetFirstConnectedPin() const;
 
 	/** Gets all connected pins. */
+	UFUNCTION(BlueprintCallable, Category = "Movie Graph")
 	TArray<UMovieGraphPin*> GetAllConnectedPins() const;
 
 	/**
@@ -130,7 +144,10 @@ public:
 	 * Determines if the connection between this pin and OtherPin follows branch restriction rules. OutError is populated
 	 * with an error if the connection should be rejected and the function will return false.
 	 */
+	UFUNCTION(BlueprintCallable, Category = "Movie Graph")
 	bool IsConnectionToBranchAllowed(const UMovieGraphPin* OtherPin, FText& OutError) const;
+
+	UFUNCTION(BlueprintCallable, Category = "Movie Graph")
 	bool IsPinDirectionCompatibleWith(const UMovieGraphPin* OtherPin) const;
 
 public:
diff --git a/Engine/Plugins/MovieScene/MovieRenderPipeline/Source/MovieRenderPipelineEditor/Private/Graph/MovieEdGraph.cpp b/Engine/Plugins/MovieScene/MovieRenderPipeline/Source/MovieRenderPipelineEditor/Private/Graph/MovieEdGraph.cpp
--- a/Engine/Plugins/MovieScene/MovieRenderPipeline/Source/MovieRenderPipelineEditor/Private/Graph/MovieEdGraph.cpp
+++ b/Engine/Plugins/MovieScene/MovieRenderPipeline/Source/MovieRenderPipelineEditor/Private/Graph/MovieEdGraph.cpp
@@ -123,7 +123,7 @@ void UMoviePipelineEdGraph::CreateLinks(UMoviePipelineEdGraphNodeBase* InGraphNo
 	const UMovieGraphNode* RuntimeNode = InGraphNode->GetRuntimeNode();
 	check(RuntimeNode);
 
-	auto CreateLinks = [&](const TArray<TObjectPtr<UMovieGraphPin>>& Pins, EEdGraphPinDirection PrimaryDirection)
+	auto CreateLinks = [&](const TArray<UMovieGraphPin*>& Pins, EEdGraphPinDirection PrimaryDirection)
 	{
 		for (const UMovieGraphPin* Pin : Pins)
 		{
diff --git a/Engine/Plugins/MovieScene/MovieRenderPipeline/Source/MovieRenderPipelineEditor/Private/Graph/MovieEdGraphNode.cpp b/Engine/Plugins/MovieScene/MovieRenderPipeline/Source/MovieRenderPipelineEditor/Private/Graph/MovieEdGraphNode.cpp
--- a/Engine/Plugins/MovieScene/MovieRenderPipeline/Source/MovieRenderPipelineEditor/Private/Graph/MovieEdGraphNode.cpp
+++ b/Engine/Plugins/MovieScene/MovieRenderPipeline/Source/MovieRenderPipelineEditor/Private/Graph/MovieEdGraphNode.cpp
@@ -449,11 +449,11 @@ void UMoviePipelineEdGraphNodeBase::AutowireNewNode(UEdGraphPin* FromPin)
 	}
 
 	const bool bFromPinIsInput = FromPin->Direction == EEdGraphPinDirection::EGPD_Input;
-	const TArray<TObjectPtr<UMovieGraphPin>>& OtherPinsList = bFromPinIsInput ? RuntimeNode->GetOutputPins() : RuntimeNode->GetInputPins();
+	const TArray<UMovieGraphPin*>& OtherPinsList = bFromPinIsInput ? RuntimeNode->GetOutputPins() : RuntimeNode->GetInputPins();
 
 	// Try to connect to the first compatible pin
 	bool bDidAutoconnect = false;
-	for (const TObjectPtr<UMovieGraphPin>& OtherPin : OtherPinsList)
+	for (const UMovieGraphPin* OtherPin : OtherPinsList)
 	{
 		check(OtherPin);
 
diff --git a/Engine/Plugins/MovieScene/MovieRenderPipeline/Source/MovieRenderPipelineEditor/Private/Graph/MovieEdGraphVariableNode.cpp b/Engine/Plugins/MovieScene/MovieRenderPipeline/Source/MovieRenderPipelineEditor/Private/Graph/MovieEdGraphVariableNode.cpp
--- a/Engine/Plugins/MovieScene/MovieRenderPipeline/Source/MovieRenderPipelineEditor/Private/Graph/MovieEdGraphVariableNode.cpp
+++ b/Engine/Plugins/MovieScene/MovieRenderPipeline/Source/MovieRenderPipelineEditor/Private/Graph/MovieEdGraphVariableNode.cpp
@@ -25,11 +25,11 @@ void UMoviePipelineEdGraphVariableNode::AllocateDefaultPins()
 {
 	if (const UMovieGraphVariableNode* VariableNode = Cast<UMovieGraphVariableNode>(RuntimeNode))
 	{
-		const TArray<TObjectPtr<UMovieGraphPin>>& OutputPins = RuntimeNode->GetOutputPins();
+		const TArray<UMovieGraphPin*>& OutputPins = RuntimeNode->GetOutputPins();
 		if (!OutputPins.IsEmpty())
 		{
-			UEdGraphPin* NewPin = CreatePin(EGPD_Output, GetPinType(OutputPins[0].Get()), FName(VariableNode->GetVariable()->GetMemberName()));
-			NewPin->PinToolTip = GetPinTooltip(OutputPins[0].Get());
+			UEdGraphPin* NewPin = CreatePin(EGPD_Output, GetPinType(OutputPins[0]), FName(VariableNode->GetVariable()->GetMemberName()));
+			NewPin->PinToolTip = GetPinTooltip(OutputPins[0]);
 		}
 	}
 }
diff --git a/Engine/Plugins/MovieScene/MovieRenderPipeline/Source/MovieRenderPipelineEditor/Private/Tests/MovieGraphTests.cpp b/Engine/Plugins/MovieScene/MovieRenderPipeline/Source/MovieRenderPipelineEditor/Private/Tests/MovieGraphTests.cpp
--- a/Engine/Plugins/MovieScene/MovieRenderPipeline/Source/MovieRenderPipelineEditor/Private/Tests/MovieGraphTests.cpp
+++ b/Engine/Plugins/MovieScene/MovieRenderPipeline/Source/MovieRenderPipelineEditor/Private/Tests/MovieGraphTests.cpp
@@ -174,7 +174,7 @@ namespace UE::MovieGraph::Private::Tests
 	{
 		for (const TObjectPtr<UMovieGraphNode>& NodeFrom : NodesWithOutputs)
 		{
-			for (const TObjectPtr<UMovieGraphPin>& OutputPin : NodeFrom->GetOutputPins())
+			for (UMovieGraphPin* OutputPin : NodeFrom->GetOutputPins())
 			{
 				for (const TObjectPtr<UMovieGraphNode>& NodeTo : NodesWithInputs)
 				{
@@ -183,7 +183,7 @@ namespace UE::MovieGraph::Private::Tests
 						continue;
 					}
 					
-					for (const TObjectPtr<UMovieGraphPin>& InputPin : NodeTo->GetInputPins())
+					for (UMovieGraphPin* InputPin : NodeTo->GetInputPins())
 					{
 						if (OutputPin->CanCreateConnection(InputPin))
 						{
