MovieGraph: Removed reliance on cvars for the (deferred renderer) 32bit setting on additional post-process materials.

Instead, we issue the request for 32-bit from the buffer pipes. This allows integration in MRG by preventing potential conflicts between multiple renders, where each layer could be fighting for a different cvar value resulting in undefined behavior. This also allows per-material granularity to be user-exposed.

diff --git a/Engine/Plugins/MovieScene/MovieRenderPipeline/Source/MovieRenderPipelineRenderPasses/Private/MoviePipelineDeferredPasses.cpp b/Engine/Plugins/MovieScene/MovieRenderPipeline/Source/MovieRenderPipelineRenderPasses/Private/MoviePipelineDeferredPasses.cpp
--- a/Engine/Plugins/MovieScene/MovieRenderPipeline/Source/MovieRenderPipelineRenderPasses/Private/MoviePipelineDeferredPasses.cpp
+++ b/Engine/Plugins/MovieScene/MovieRenderPipeline/Source/MovieRenderPipelineRenderPasses/Private/MoviePipelineDeferredPasses.cpp
@@ -253,8 +253,6 @@ void UMoviePipelineDeferredPassBase::SetupImpl(const MoviePipeline::FMoviePipeli
 	AccumulatorPool = MakeShared<TAccumulatorPool<FImageOverlappedAccumulator>, ESPMode::ThreadSafe>(PoolSize);
 	
 	PreviousCustomDepthValue.Reset();
-	PreviousDumpFramesValue.Reset();
-	PreviousColorFormatValue.Reset();
 
 	// This scene view extension will be released automatically as soon as Render Sequence is torn down.
 	// One Extension per sequence, since each sequence has its own OCIO settings.
@@ -278,23 +276,6 @@ void UMoviePipelineDeferredPassBase::SetupImpl(const MoviePipeline::FMoviePipeli
 			}
 		}
 	}
-	
-	if (bUse32BitPostProcessMaterials)
-	{
-		IConsoleVariable* DumpFramesCVar = IConsoleManager::Get().FindConsoleVariable(TEXT("r.BufferVisualizationDumpFramesAsHDR"));
-		if (DumpFramesCVar)
-		{
-			PreviousDumpFramesValue = DumpFramesCVar->GetInt();
-			DumpFramesCVar->Set(1, EConsoleVariableFlags::ECVF_SetByConsole);
-		}
-		
-		IConsoleVariable* ColorFormatCVar = IConsoleManager::Get().FindConsoleVariable(TEXT("r.PostProcessingColorFormat"));
-		if (ColorFormatCVar)
-		{
-			PreviousColorFormatValue = ColorFormatCVar->GetInt();
-			ColorFormatCVar->Set(1, EConsoleVariableFlags::ECVF_SetByConsole);
-		}
-	}
 
 	// Cache out the stencil layer names (from data layers or actor layers) and de-duplicate. If layers with the same name
 	// are provided, renders may fail, which is why the the names need to be de-duplicated.
@@ -357,21 +338,6 @@ void UMoviePipelineDeferredPassBase::TeardownImpl()
 			}
 		}
 	}
-	
-	if (PreviousDumpFramesValue.IsSet())
-	{
-		IConsoleVariable* DumpFramesCVar = IConsoleManager::Get().FindConsoleVariable(TEXT("r.BufferVisualizationDumpFramesAsHDR"));
-		if (DumpFramesCVar)
-		{
-			DumpFramesCVar->Set(PreviousDumpFramesValue.GetValue(),  EConsoleVariableFlags::ECVF_SetByConsole);
-		}
-		
-		IConsoleVariable* ColorFormatCVar = IConsoleManager::Get().FindConsoleVariable(TEXT("r.PostProcessingColorFormat"));
-		if (ColorFormatCVar)
-		{
-			ColorFormatCVar->Set(PreviousColorFormatValue.GetValue(), EConsoleVariableFlags::ECVF_SetByConsole);
-		}
-	}
 
 	// Preserve our view state until the rendering thread has been flushed.
 	Super::TeardownImpl();
@@ -553,6 +519,7 @@ void UMoviePipelineDeferredPassBase::RenderSample_GameThreadImpl(const FMoviePip
 				FMoviePipelinePassIdentifier LayerPassIdentifier = FMoviePipelinePassIdentifier(PassIdentifier.Name + VisMaterial->GetName(), PassIdentifierForCurrentCamera.CameraName);
 
 				auto BufferPipe = MakeShared<FImagePixelPipe, ESPMode::ThreadSafe>();
+				BufferPipe->bIsExpecting32BitPixelData = bUse32BitPostProcessMaterials;
 				BufferPipe->AddEndpoint(MakeForwardingEndpoint(LayerPassIdentifier, InOutSampleState));
 
 				View->FinalPostProcessSettings.BufferVisualizationPipes.Add(VisMaterial->GetFName(), BufferPipe);
diff --git a/Engine/Plugins/MovieScene/MovieRenderPipeline/Source/MovieRenderPipelineRenderPasses/Public/MoviePipelineDeferredPasses.h b/Engine/Plugins/MovieScene/MovieRenderPipeline/Source/MovieRenderPipelineRenderPasses/Public/MoviePipelineDeferredPasses.h
--- a/Engine/Plugins/MovieScene/MovieRenderPipeline/Source/MovieRenderPipelineRenderPasses/Public/MoviePipelineDeferredPasses.h
+++ b/Engine/Plugins/MovieScene/MovieRenderPipeline/Source/MovieRenderPipelineRenderPasses/Public/MoviePipelineDeferredPasses.h
@@ -138,9 +138,6 @@ public:
 	UPROPERTY(EditAnywhere, BlueprintReadWrite, Category = "Stencil Clip Layers")
 	TArray<FActorLayer> ActorLayers;
 	
-	UE_DEPRECATED(5.1, "Use ActorLayers property instead.")
-	TArray<FActorLayer> StencilLayers;
-	
 	/**
 	* If the map you are working with is a World Partition map, you can specify Data layers instead of Actor Layers. If any
 	* Data Layers are specified, this will take precedence over any ActorLayers in this config. Does not affect whether or
@@ -171,9 +168,9 @@ protected:
 
 	// Cache the custom stencil value. Only has meaning if they have stencil layers.
 	TOptional<int32> PreviousCustomDepthValue;
-	/** Cache the previous dump frames as HDR value. Only used if using 32-bit post processing. */
+	UE_DEPRECATED(5.5, "PreviousDumpFramesValue has been deprecated.")
 	TOptional<int32> PreviousDumpFramesValue;
-	/** Cache the previous color format value. Only used if using 32-bit post processing. */
+	UE_DEPRECATED(5.5, "PreviousColorFormatValue has been deprecated.")
 	TOptional<int32> PreviousColorFormatValue;
 
 	TSharedPtr<FAccumulatorPool, ESPMode::ThreadSafe> AccumulatorPool;
diff --git a/Engine/Source/Runtime/ImageWriteQueue/Public/ImageWriteStream.h b/Engine/Source/Runtime/ImageWriteQueue/Public/ImageWriteStream.h
--- a/Engine/Source/Runtime/ImageWriteQueue/Public/ImageWriteStream.h
+++ b/Engine/Source/Runtime/ImageWriteQueue/Public/ImageWriteStream.h
@@ -63,6 +63,9 @@ struct FImagePixelPipe
 		return EndPoints;
 	}
 
+	/** Boolean flag used to request 32-bit image pixel data, false by default. */
+	std::atomic_bool bIsExpecting32BitPixelData = false;
+
 private:
 
 	/** A lock to protect the end points array */
diff --git a/Engine/Source/Runtime/Renderer/Private/PostProcess/PostProcessVisualizeBuffer.cpp b/Engine/Source/Runtime/Renderer/Private/PostProcess/PostProcessVisualizeBuffer.cpp
--- a/Engine/Source/Runtime/Renderer/Private/PostProcess/PostProcessVisualizeBuffer.cpp
+++ b/Engine/Source/Runtime/Renderer/Private/PostProcess/PostProcessVisualizeBuffer.cpp
@@ -325,7 +325,7 @@ FScreenPassTexture AddVisualizeGBufferOverviewPass(
 	FScreenPassTexture Output;
 	
 	// Respect the r.PostProcessingColorFormat cvar just like the main rendering path
-	const EPixelFormat OutputFormat = OverridePostProcessingColorFormat(Inputs.bOutputInHDR ? PF_FloatRGBA : PF_Unknown);
+	EPixelFormat OutputFormat = OverridePostProcessingColorFormat(Inputs.bOutputInHDR ? PF_FloatRGBA : PF_Unknown);
 
 	TArray<FVisualizeBufferTile> Tiles;
 
@@ -370,14 +370,20 @@ FScreenPassTexture AddVisualizeGBufferOverviewPass(
 				EPathTracingPostProcessMaterialInput::Radiance, FScreenPassTexture(PathTracingResources.Radiance, ViewRect));
 		}
 
+		const TSharedPtr<FImagePixelPipe, ESPMode::ThreadSafe>* OutputPipe = PostProcessSettings.BufferVisualizationPipes.Find(MaterialInterface->GetFName());		
+		const bool bIsValidOutputPipe = OutputPipe && OutputPipe->IsValid();
+		
+		if (bIsValidOutputPipe && OutputPipe->Get()->bIsExpecting32BitPixelData)
+		{
+			OutputFormat = PF_A32B32G32R32F;
+		}
+
 		PostProcessMaterialInputs.SceneTextures = Inputs.SceneTextures;
 		PostProcessMaterialInputs.OutputFormat = OutputFormat;
 
 		Output = AddPostProcessMaterialPass(GraphBuilder, View, PostProcessMaterialInputs, MaterialInterface);
 
-		const TSharedPtr<FImagePixelPipe, ESPMode::ThreadSafe>* OutputPipe = PostProcessSettings.BufferVisualizationPipes.Find(MaterialInterface->GetFName());
-
-		if (OutputPipe && OutputPipe->IsValid())
+		if (bIsValidOutputPipe)
 		{
 			AddDumpToPipePass(GraphBuilder, Output, OutputPipe->Get());
 		}
