MoviePipeline: When new graph variables are created, the new job variable overrides should have defaults which match their associated graph variables.

diff --git a/Engine/Plugins/MovieScene/MovieRenderPipeline/Source/MovieRenderPipelineCore/Private/MovieJobVariableAssignmentContainer.cpp b/Engine/Plugins/MovieScene/MovieRenderPipeline/Source/MovieRenderPipelineCore/Private/MovieJobVariableAssignmentContainer.cpp
--- a/Engine/Plugins/MovieScene/MovieRenderPipeline/Source/MovieRenderPipelineCore/Private/MovieJobVariableAssignmentContainer.cpp
+++ b/Engine/Plugins/MovieScene/MovieRenderPipeline/Source/MovieRenderPipelineCore/Private/MovieJobVariableAssignmentContainer.cpp
@@ -429,8 +429,10 @@ void UMovieJobVariableAssignmentContainer::UpdateGraphVariableOverrides()
 	}
 
 	// Third, ensure that each variable has a corresponding property in the bag. If not, generate one and flag that the
-	// bag needs to be regenerated.
-	for (const UMovieGraphVariable* Variable : GraphVariables)
+	// bag needs to be regenerated. Any new job variables will get a default value assigned to them that matches the default
+	// for the associated graph variable.
+	TMap<UMovieGraphVariable*, FName> VariablesThatNeedDefaultAssigned;
+	for (UMovieGraphVariable* Variable : GraphVariables)
 	{
 		bool bFoundMatchingDesc = false;
 		
@@ -453,6 +455,9 @@ void UMovieJobVariableAssignmentContainer::UpdateGraphVariableOverrides()
 				ModifiedDescs.Add(NewPropertyDesc);
 				ModifiedDescs.Add(EditConditionPropertyDesc);
 				bNeedsToRegenerate = true;
+
+				// Mark this variable as needing a default assigned (property descs do not contain the property value)
+				VariablesThatNeedDefaultAssigned.Add({Variable, NewPropertyDesc.Name});
 			}
 		}
 	}
@@ -461,6 +466,17 @@ void UMovieJobVariableAssignmentContainer::UpdateGraphVariableOverrides()
 	{
 		const UPropertyBag* NewBagStruct = UPropertyBag::GetOrCreateFromDescs(ModifiedDescs);
 		Value.MigrateToNewBagStruct(NewBagStruct);
+
+		// Also update the default for newly-added job variables to match the associated graph variables
+		for (const TPair<UMovieGraphVariable*, FName>& VariableNeedingDefault : VariablesThatNeedDefaultAssigned)
+		{
+			UMovieGraphVariable* GraphVariable = VariableNeedingDefault.Key;
+			const FName VariableName = VariableNeedingDefault.Value;
+			
+			// Not ideal that this is being set by string, but it's the most universal method of setting all variable types and we don't need
+			// to special-case for each variable type
+			Value.SetValueSerializedString(VariableName, GraphVariable->GetValueSerializedString());
+		}
 	}
 }
 #endif
