MoviePipeline: Added ability to remove multiple items from a SMovieGraphSimpleList at the same time. Affects queries within Collection nodes that contain lists, and the Collections list within the Modifiers node.

diff --git a/Engine/Plugins/MovieScene/MovieRenderPipeline/Source/MovieRenderPipelineCore/Private/Graph/MovieGraphRenderLayerSubsystem.cpp b/Engine/Plugins/MovieScene/MovieRenderPipeline/Source/MovieRenderPipelineCore/Private/Graph/MovieGraphRenderLayerSubsystem.cpp
--- a/Engine/Plugins/MovieScene/MovieRenderPipeline/Source/MovieRenderPipelineCore/Private/Graph/MovieGraphRenderLayerSubsystem.cpp
+++ b/Engine/Plugins/MovieScene/MovieRenderPipeline/Source/MovieRenderPipelineCore/Private/Graph/MovieGraphRenderLayerSubsystem.cpp
@@ -633,13 +633,17 @@ TArray<TSharedRef<SWidget>> UMovieGraphConditionGroupQuery_Actor::GetWidgets()
 			.DataTypePlural(FText::FromString("Actors"))
 			.OnGetRowText_Static(&GetRowText)
 			.OnGetRowIcon_Static(&GetRowIcon)
-			.OnDelete_Lambda([this](const TSharedPtr<TSoftObjectPtr<AActor>> InActor)
+			.OnDelete_Lambda([this](const TArray<TSharedPtr<TSoftObjectPtr<AActor>>> InActors)
 			{
 				const FScopedTransaction Transaction(LOCTEXT("RemoveActorsFromCollection", "Remove Actors from Collection"));
 				Modify();
+
+				for (const TSharedPtr<TSoftObjectPtr<AActor>>& Actor : InActors)
+				{
+					ListDataSource.Remove(Actor);
+					ActorsToMatch.Remove(*Actor.Get());
+				}
 				
-				ListDataSource.Remove(InActor);
-				ActorsToMatch.Remove(*InActor.Get());
 				ActorsList->Refresh();
 			})
 		]
@@ -1043,12 +1047,16 @@ TArray<TSharedRef<SWidget>> UMovieGraphConditionGroupQuery_ActorType::GetWidgets
 			.DataTypePlural(FText::FromString("Actor Types"))
 			.OnGetRowText_Static(&GetRowText)
 			.OnGetRowIcon_Static(&GetRowIcon)
-			.OnDelete_Lambda([this](UClass* InActorClass)
+			.OnDelete_Lambda([this](const TArray<UClass*> InActorClasses)
 			{
 				const FScopedTransaction Transaction(LOCTEXT("RemoveActorTypesFromCollection", "Remove Actor Types from Collection"));
 				Modify();
+
+				for (UClass* ActorClass : InActorClasses)
+				{
+					ActorTypes.Remove(ActorClass);
+				}
 				
-				ActorTypes.Remove(InActorClass);
 				ActorTypesList->Refresh();
 			})
 		]
@@ -1276,12 +1284,16 @@ TArray<TSharedRef<SWidget>> UMovieGraphConditionGroupQuery_ComponentType::GetWid
 			.DataTypePlural(FText::FromString("Component Types"))
 			.OnGetRowText_Static(&GetRowText)
 			.OnGetRowIcon_Static(&GetRowIcon)
-			.OnDelete_Lambda([this](UClass* InComponentType)
+			.OnDelete_Lambda([this](const TArray<UClass*> InComponentTypes)
 			{
 				const FScopedTransaction Transaction(LOCTEXT("RemoveComponentTypesFromCollection", "Remove Component Types from Collection"));
 				Modify();
+
+				for (UClass* ComponentType : InComponentTypes)
+				{
+					ComponentTypes.Remove(ComponentType);
+				}
 				
-				ComponentTypes.Remove(InComponentType);
 				ComponentTypesList->Refresh();
 			})
 	);
@@ -1464,12 +1476,16 @@ TArray<TSharedRef<SWidget>> UMovieGraphConditionGroupQuery_EditorFolder::GetWidg
 			.DataTypePlural(FText::FromString("Folders"))
 			.OnGetRowText_Static(&GetRowText)
 			.OnGetRowIcon_Static(&GetRowIcon)
-			.OnDelete_Lambda([this](FName InFolderPath)
+			.OnDelete_Lambda([this](const TArray<FName> InFolderPaths)
 			{
 				const FScopedTransaction Transaction(LOCTEXT("RemoveEditorFoldersFromCollection", "Remove Editor Folders from Collection"));
 				Modify();
 
-				FolderPaths.Remove(InFolderPath);
+				for (const FName& FolderPath : InFolderPaths)
+				{
+					FolderPaths.Remove(FolderPath);
+				}
+				
 				FolderPathsList->Refresh();
 
 				if (FolderPickerWidget.IsValid())
@@ -1689,13 +1705,16 @@ TArray<TSharedRef<SWidget>> UMovieGraphConditionGroupQuery_Sublevel::GetWidgets(
 			.DataTypePlural(FText::FromString("Sublevels"))
 			.OnGetRowText_Static(&GetRowText)
 			.OnGetRowIcon_Static(&GetRowIcon)
-			.OnDelete_Lambda([this](const TSharedPtr<TSoftObjectPtr<UWorld>> InSublevel)
+			.OnDelete_Lambda([this](const TArray<TSharedPtr<TSoftObjectPtr<UWorld>>> InSublevels)
 			{
 				const FScopedTransaction Transaction(LOCTEXT("RemoveSublevelsFromCollection", "Remove Sublevels from Collection"));
 				Modify();
-				
-				ListDataSource.Remove(InSublevel);
-				Sublevels.Remove(*InSublevel.Get());
+
+				for (TSharedPtr<TSoftObjectPtr<UWorld>> Sublevel : InSublevels)
+				{
+					ListDataSource.Remove(Sublevel);
+					Sublevels.Remove(*Sublevel.Get());
+				}
 				
 				SublevelsList->Refresh();
 				
@@ -1870,12 +1889,16 @@ TArray<TSharedRef<SWidget>> UMovieGraphConditionGroupQuery_ActorLayer::GetWidget
 			.DataTypePlural(FText::FromString("Actor Layers"))
 			.OnGetRowText_Static(&GetRowText)
 			.OnGetRowIcon_Static(&GetRowIcon)
-			.OnDelete_Lambda([this](FName InLayerName)
+			.OnDelete_Lambda([this](const TArray<FName> InLayerNames)
 			{
 				const FScopedTransaction Transaction(LOCTEXT("RemoveActorLayerFromCollection", "Remove Actor Layers from Collection"));
 				Modify();
 
-				LayerNames.Remove(InLayerName);
+				for (const FName& LayerName : InLayerNames)
+				{
+					LayerNames.Remove(LayerName);
+				}
+				
 				LayerNamesList->Refresh();
 			})
 		]
@@ -2042,13 +2065,17 @@ TArray<TSharedRef<SWidget>> UMovieGraphConditionGroupQuery_DataLayer::GetWidgets
 			.DataTypePlural(FText::FromString("Data Layers"))
 			.OnGetRowText_Static(&GetRowText)
 			.OnGetRowIcon_Static(&GetRowIcon)
-			.OnDelete_Lambda([this](TSharedPtr<TSoftObjectPtr<UDataLayerAsset>> InLayer)
+			.OnDelete_Lambda([this](const TArray<TSharedPtr<TSoftObjectPtr<UDataLayerAsset>>> InLayers)
 			{
 				const FScopedTransaction Transaction(LOCTEXT("RemoveDataLayerFromCollection", "Remove Data Layers from Collection"));
 				Modify();
 
-				ListDataSource.Remove(InLayer);
-				DataLayers.Remove(*InLayer.Get());
+				for (TSharedPtr<TSoftObjectPtr<UDataLayerAsset>> Layer : InLayers)
+				{
+					ListDataSource.Remove(Layer);
+					DataLayers.Remove(*Layer.Get());
+				}
+				
 				DataLayersList->Refresh();
 
 				constexpr bool bUpdateSources = true;
diff --git a/Engine/Plugins/MovieScene/MovieRenderPipeline/Source/MovieRenderPipelineCore/Public/Graph/MovieGraphSharedWidgets.h b/Engine/Plugins/MovieScene/MovieRenderPipeline/Source/MovieRenderPipelineCore/Public/Graph/MovieGraphSharedWidgets.h
--- a/Engine/Plugins/MovieScene/MovieRenderPipeline/Source/MovieRenderPipelineCore/Public/Graph/MovieGraphSharedWidgets.h
+++ b/Engine/Plugins/MovieScene/MovieRenderPipeline/Source/MovieRenderPipelineCore/Public/Graph/MovieGraphSharedWidgets.h
@@ -29,13 +29,15 @@ class SMovieGraphSimpleList final : public SCompoundWidget
 public:
 	DECLARE_DELEGATE_RetVal_OneParam(const FSlateBrush*, FGetRowIcon, ListType);
 	DECLARE_DELEGATE_RetVal_OneParam(FText, FGetRowText, ListType);
-	DECLARE_DELEGATE_OneParam(FOnDelete, ListType);
+	DECLARE_DELEGATE_OneParam(FOnDelete, TArray<ListType>);
 	DECLARE_DELEGATE_RetVal_OneParam(bool, FGetRowEnableState, ListType);
 	DECLARE_DELEGATE_TwoParams(FSetRowEnableState, ListType, bool);
 	
 	SLATE_BEGIN_ARGS(SMovieGraphSimpleList<ListType>)
-		: _ShowEnableDisable(false)
+		: _SelectionMode(ESelectionMode::Multi)
+		, _ShowEnableDisable(false)
 		{}
+		
 		/** The source of data that the list will display. */
 		SLATE_ATTRIBUTE(TArray<ListType>*, DataSource)
 
@@ -45,6 +47,9 @@ public:
 		/** The plural of the DataType attribute. */
 		SLATE_ATTRIBUTE(FText, DataTypePlural)
 
+		/** The selection mode the list should be put into. Defaults to Multi. */
+		SLATE_ATTRIBUTE(ESelectionMode::Type, SelectionMode)
+
 		/** Gets the icon for a row in the list. */
 		SLATE_EVENT(FGetRowIcon, OnGetRowIcon);
 
@@ -69,6 +74,7 @@ public:
 		DataSource = InArgs._DataSource.Get();
 		DataType = InArgs._DataType.Get();
 		DataTypePlural = InArgs._DataTypePlural.Get();
+		SelectionMode = InArgs._SelectionMode.Get();
 		OnGetRowIcon = InArgs._OnGetRowIcon;
 		OnGetRowText = InArgs._OnGetRowText;
 		OnDelete = InArgs._OnDelete;
@@ -96,6 +102,7 @@ public:
 			[
 				SAssignNew(ListView, SListView<ListType>)
 				.ListItemsSource(DataSource)
+				.SelectionMode(SelectionMode)
 				.OnKeyDownHandler_Lambda([this](const FGeometry& InGeometry, const FKeyEvent& InKeyEvent)
 				{
 					return CommandList->ProcessCommandBindings(InKeyEvent) ? FReply::Handled() : FReply::Unhandled();
@@ -177,7 +184,7 @@ private:
 		{
 			if (OnDelete.IsBound())
 			{
-				OnDelete.Execute(SelectedItems[0]);
+				OnDelete.Execute(SelectedItems);
 			}
 		}
 	}
@@ -271,6 +278,7 @@ private:
 	FText DataType;
 	FText DataTypePlural;
 	TArray<ListType>* DataSource = nullptr;
+	ESelectionMode::Type SelectionMode = ESelectionMode::Multi;
 };
 
 /**
diff --git a/Engine/Plugins/MovieScene/MovieRenderPipeline/Source/MovieRenderPipelineEditor/Private/Customizations/Graph/MovieGraphModifiersCustomization.cpp b/Engine/Plugins/MovieScene/MovieRenderPipeline/Source/MovieRenderPipelineEditor/Private/Customizations/Graph/MovieGraphModifiersCustomization.cpp
--- a/Engine/Plugins/MovieScene/MovieRenderPipeline/Source/MovieRenderPipelineEditor/Private/Customizations/Graph/MovieGraphModifiersCustomization.cpp
+++ b/Engine/Plugins/MovieScene/MovieRenderPipeline/Source/MovieRenderPipelineEditor/Private/Customizations/Graph/MovieGraphModifiersCustomization.cpp
@@ -206,13 +206,17 @@ void FMovieGraphModifiersCustomization::CustomizeDetails(IDetailLayoutBuilder& I
 		.DataSource(&ListDataSource)
 		.DataType(FText::FromString("Collection"))
 		.DataTypePlural(FText::FromString("Collections"))
-		.OnDelete_Lambda([this, ModifierNode](const FName DeletedCollectionName)
+		.OnDelete_Lambda([this, ModifierNode](const TArray<FName> DeletedCollectionNames)
 		{
 			if (ModifierNode.IsValid())
 			{
-				const FScopedTransaction Transaction(LOCTEXT("RemoveCollectionFromModifier", "Remove Collection from Modifier"));
+				const FScopedTransaction Transaction(LOCTEXT("RemoveCollectionsFromModifier", "Remove Collections from Modifier"));
+
+				for (const FName& DeletedCollectionName : DeletedCollectionNames)
+				{
+					ModifierNode.Get()->RemoveCollection(DeletedCollectionName);
+				}
 				
-				ModifierNode.Get()->RemoveCollection(DeletedCollectionName);
 				ListDataSource = ModifierNode->GetCollections();
 				CollectionsList->Refresh();
 			}
