MovieGraph: Add support for 32-bit output on additional post-process materials.

diff --git a/Engine/Plugins/MovieScene/MovieRenderPipeline/Source/MovieRenderPipelineRenderPasses/Private/Graph/Renderers/MovieGraphDeferredPass.cpp b/Engine/Plugins/MovieScene/MovieRenderPipeline/Source/MovieRenderPipelineRenderPasses/Private/Graph/Renderers/MovieGraphDeferredPass.cpp
--- a/Engine/Plugins/MovieScene/MovieRenderPipeline/Source/MovieRenderPipelineRenderPasses/Private/Graph/Renderers/MovieGraphDeferredPass.cpp
+++ b/Engine/Plugins/MovieScene/MovieRenderPipeline/Source/MovieRenderPipelineRenderPasses/Private/Graph/Renderers/MovieGraphDeferredPass.cpp
@@ -277,6 +277,8 @@ void FMovieGraphDeferredPass::Render(const FMovieGraphTraversalContext& InFrameT
 
 		if (UMovieGraphImagePassBaseNode* ParentNode = GetParentNode(InFrameTraversalContext.Time.EvaluatedConfig))
 		{
+			TSet<UMaterialInterface*> HighPrecisionMaterials;
+
 			for (const FMoviePipelinePostProcessPass& PostProcessPass : ParentNode->GetAdditionalPostProcessMaterials())
 			{
 				if (PostProcessPass.bEnabled)
@@ -285,6 +287,11 @@ void FMovieGraphDeferredPass::Render(const FMovieGraphTraversalContext& InFrameT
 					if (Material)
 					{
 						NewView->FinalPostProcessSettings.BufferVisualizationOverviewMaterials.Add(Material);
+						
+						if (PostProcessPass.bHighPrecisionOutput)
+						{
+							HighPrecisionMaterials.Add(Material);
+						}
 					}
 				}
 			}
@@ -292,6 +299,11 @@ void FMovieGraphDeferredPass::Render(const FMovieGraphTraversalContext& InFrameT
 			for (UMaterialInterface* VisMaterial : NewView->FinalPostProcessSettings.BufferVisualizationOverviewMaterials)
 			{
 				auto BufferPipe = MakeShared<FImagePixelPipe, ESPMode::ThreadSafe>();
+
+				if (HighPrecisionMaterials.Contains(VisMaterial))
+				{
+					BufferPipe->bIsExpecting32BitPixelData = true;
+				}
 				
 				FMovieGraphRenderDataIdentifier Identifier = RenderDataIdentifier;
 				Identifier.SubResourceName = VisMaterial->GetName();
diff --git a/Engine/Plugins/MovieScene/MovieRenderPipeline/Source/MovieRenderPipelineRenderPasses/Public/MoviePipelineDeferredPasses.h b/Engine/Plugins/MovieScene/MovieRenderPipeline/Source/MovieRenderPipelineRenderPasses/Public/MoviePipelineDeferredPasses.h
--- a/Engine/Plugins/MovieScene/MovieRenderPipeline/Source/MovieRenderPipelineRenderPasses/Public/MoviePipelineDeferredPasses.h
+++ b/Engine/Plugins/MovieScene/MovieRenderPipeline/Source/MovieRenderPipelineRenderPasses/Public/MoviePipelineDeferredPasses.h
@@ -28,6 +28,10 @@ public:
 	*/
 	UPROPERTY(EditAnywhere, BlueprintReadWrite, Category = "Settings")
 	TSoftObjectPtr<UMaterialInterface> Material;
+
+	/** Request output to be 32-bit, usually for data exports. Note that scene color precision is still defined by r.SceneColorFormat.*/
+	UPROPERTY(EditAnywhere, BlueprintReadWrite, Category = "Settings", DisplayName="High Precision Output (32-bit)")
+	bool bHighPrecisionOutput = false;
 };
 
 UCLASS(BlueprintType)
diff --git a/Engine/Source/Runtime/Renderer/Private/PostProcess/PostProcessVisualizeBuffer.cpp b/Engine/Source/Runtime/Renderer/Private/PostProcess/PostProcessVisualizeBuffer.cpp
--- a/Engine/Source/Runtime/Renderer/Private/PostProcess/PostProcessVisualizeBuffer.cpp
+++ b/Engine/Source/Runtime/Renderer/Private/PostProcess/PostProcessVisualizeBuffer.cpp
@@ -370,16 +370,19 @@ FScreenPassTexture AddVisualizeGBufferOverviewPass(
 				EPathTracingPostProcessMaterialInput::Radiance, FScreenPassTexture(PathTracingResources.Radiance, ViewRect));
 		}
 
-		const TSharedPtr<FImagePixelPipe, ESPMode::ThreadSafe>* OutputPipe = PostProcessSettings.BufferVisualizationPipes.Find(MaterialInterface->GetFName());		
+		PostProcessMaterialInputs.SceneTextures = Inputs.SceneTextures;
+
+		const TSharedPtr<FImagePixelPipe, ESPMode::ThreadSafe>* OutputPipe = PostProcessSettings.BufferVisualizationPipes.Find(MaterialInterface->GetFName());
 		const bool bIsValidOutputPipe = OutputPipe && OutputPipe->IsValid();
-		
+
 		if (bIsValidOutputPipe && OutputPipe->Get()->bIsExpecting32BitPixelData)
 		{
-			OutputFormat = PF_A32B32G32R32F;
+			PostProcessMaterialInputs.OutputFormat = PF_A32B32G32R32F;
+		}
+		else
+		{
+			PostProcessMaterialInputs.OutputFormat = OutputFormat;
 		}
-
-		PostProcessMaterialInputs.SceneTextures = Inputs.SceneTextures;
-		PostProcessMaterialInputs.OutputFormat = OutputFormat;
 
 		Output = AddPostProcessMaterialPass(GraphBuilder, View, PostProcessMaterialInputs, MaterialInterface);
 
