MoviePipeline: Updated the picker for the Actor condition group query to include a checkbox selection mode.

diff --git a/Engine/Plugins/MovieScene/MovieRenderPipeline/Source/MovieRenderPipelineCore/Private/Graph/MovieGraphRenderLayerSubsystem.cpp b/Engine/Plugins/MovieScene/MovieRenderPipeline/Source/MovieRenderPipelineCore/Private/Graph/MovieGraphRenderLayerSubsystem.cpp
--- a/Engine/Plugins/MovieScene/MovieRenderPipeline/Source/MovieRenderPipelineCore/Private/Graph/MovieGraphRenderLayerSubsystem.cpp
+++ b/Engine/Plugins/MovieScene/MovieRenderPipeline/Source/MovieRenderPipelineCore/Private/Graph/MovieGraphRenderLayerSubsystem.cpp
@@ -16,6 +16,7 @@
 #include "Styling/SlateIconFinder.h"
 #include "UObject/Package.h"
 #include "Widgets/Input/SMultiLineEditableTextBox.h"
+#include "Widgets/Layout/SSpacer.h"
 #include "WorldPartition/DataLayer/DataLayerAsset.h"
 #include "WorldPartition/DataLayer/DataLayerManager.h"
 
@@ -40,6 +41,7 @@
 #include "Graph/MovieGraphSharedWidgets.h"
 #include "IContentBrowserSingleton.h"
 #include "ISceneOutliner.h"
+#include "ISceneOutlinerTreeItem.h"
 #include "LayersDragDropOp.h"
 #include "Layers/LayersSubsystem.h"
 #include "SceneOutlinerModule.h"
@@ -635,16 +637,16 @@ TArray<TSharedRef<SWidget>> UMovieGraphConditionGroupQuery_Actor::GetWidgets()
 			.OnGetRowIcon_Static(&GetRowIcon)
 			.OnDelete_Lambda([this](const TArray<TSharedPtr<TSoftObjectPtr<AActor>>> InActors)
 			{
-				const FScopedTransaction Transaction(LOCTEXT("RemoveActorsFromCollection", "Remove Actors from Collection"));
-				Modify();
-
-				for (const TSharedPtr<TSoftObjectPtr<AActor>>& Actor : InActors)
+				TArray<TSoftObjectPtr<AActor>> ActorsToRemove;
+				for (const TSharedPtr<TSoftObjectPtr<AActor>>& InActor : InActors)
 				{
-					ListDataSource.Remove(Actor);
-					ActorsToMatch.Remove(*Actor.Get());
+					if (InActor.IsValid())
+					{
+						ActorsToRemove.Add(*InActor.Get());
+					}
 				}
 				
-				ActorsList->Refresh();
+				RemoveActors(ActorsToRemove);
 			})
 		]
 	);
@@ -665,20 +667,24 @@ TSharedRef<SWidget> UMovieGraphConditionGroupQuery_Actor::GetAddMenuContents(con
 	SceneOutlinerInitOptions.bShowCreateNewFolder = false;
 	SceneOutlinerInitOptions.bFocusSearchBoxWhenOpened = true;
 
-	// Show the name/label column and the type column
+	// Show the custom "Add" column, as well as the built-in name/label/type columns
+	SceneOutlinerInitOptions.ColumnMap.Add(
+		FActorSelectionColumn::GetID(),
+		FSceneOutlinerColumnInfo(
+			ESceneOutlinerColumnVisibility::Visible,
+			0,
+			FCreateSceneOutlinerColumn::CreateLambda([this](ISceneOutliner& InSceneOutliner)
+			{
+				return MakeShared<FActorSelectionColumn>(MakeWeakObjectPtr(this));
+			}),
+			false,
+			TOptional<float>()));
 	SceneOutlinerInitOptions.ColumnMap.Add(
 		FSceneOutlinerBuiltInColumnTypes::Label(),
-		FSceneOutlinerColumnInfo(ESceneOutlinerColumnVisibility::Visible, 0, FCreateSceneOutlinerColumn(), false, TOptional<float>(), FSceneOutlinerBuiltInColumnTypes::Label_Localized()));
+		FSceneOutlinerColumnInfo(ESceneOutlinerColumnVisibility::Visible, 1, FCreateSceneOutlinerColumn(), false, TOptional<float>(), FSceneOutlinerBuiltInColumnTypes::Label_Localized()));
 	SceneOutlinerInitOptions.ColumnMap.Add(
 		FSceneOutlinerBuiltInColumnTypes::ActorInfo(),
 		FSceneOutlinerColumnInfo(ESceneOutlinerColumnVisibility::Visible, 10,  FCreateSceneOutlinerColumn(), false, TOptional<float>(), FSceneOutlinerBuiltInColumnTypes::ActorInfo_Localized()));
-
-	// Don't show actors which have already been picked
-	SceneOutlinerInitOptions.Filters->AddFilterPredicate<FActorTreeItem>(
-		FActorTreeItem::FFilterPredicate::CreateLambda([this](const AActor* InActor)
-		{
-			return !ActorsToMatch.Contains(InActor);
-		}));
 	
 	const FSceneOutlinerModule& SceneOutlinerModule = FModuleManager::LoadModuleChecked<FSceneOutlinerModule>("SceneOutliner");
 
@@ -761,7 +767,7 @@ FText UMovieGraphConditionGroupQuery_Actor::GetRowText(TSharedPtr<TSoftObjectPtr
 	return LOCTEXT("MovieGraphActorConditionGroupQuery_InvalidActor", "(invalid)");
 }
 
-void UMovieGraphConditionGroupQuery_Actor::AddActors(const TArray<AActor*>& InActors, const FMovieGraphConditionGroupQueryContentsChanged& InOnAddFinished)
+void UMovieGraphConditionGroupQuery_Actor::AddActors(const TArray<AActor*>& InActors, const FMovieGraphConditionGroupQueryContentsChanged& InOnAddFinished, const bool bCloseAddMenu)
 {
 	const FScopedTransaction Transaction(LOCTEXT("AddActorsToCollection", "Add Actors to Collection"));
 	Modify();
@@ -784,7 +790,115 @@ void UMovieGraphConditionGroupQuery_Actor::AddActors(const TArray<AActor*>& InAc
 	}
 	
 	ActorsList->Refresh();
-	FSlateApplication::Get().DismissAllMenus();
+
+	if (bCloseAddMenu)
+	{
+		FSlateApplication::Get().DismissAllMenus();
+	}
+}
+
+void UMovieGraphConditionGroupQuery_Actor::RemoveActors(const TArray<TSoftObjectPtr<AActor>>& InActors)
+{
+	const FScopedTransaction Transaction(LOCTEXT("RemoveActorsFromCollection", "Remove Actors from Collection"));
+	Modify();
+	
+	for (const TSoftObjectPtr<AActor>& Actor : InActors)
+	{		
+		ListDataSource.RemoveAll([&Actor](const TSharedPtr<TSoftObjectPtr<AActor>>& ListActor)
+		{
+			if (const TSoftObjectPtr<AActor>* SoftListActor = ListActor.Get())
+			{
+				return *SoftListActor == Actor;
+			}
+			return false;
+		});
+		
+		ActorsToMatch.Remove(Actor);
+	}
+	
+	ActorsList->Refresh();
+}
+		
+FName UMovieGraphConditionGroupQuery_Actor::FActorSelectionColumn::GetID()
+{
+	static const FName ColumnId = FName("ActorSelection");
+	return ColumnId;
+}
+
+FName UMovieGraphConditionGroupQuery_Actor::FActorSelectionColumn::GetColumnID()
+{
+	return GetID();
+}
+
+SHeaderRow::FColumn::FArguments UMovieGraphConditionGroupQuery_Actor::FActorSelectionColumn::ConstructHeaderRowColumn()
+{
+	return SHeaderRow::Column(GetColumnID())
+		.FixedWidth(25.f)
+		[
+			SNew(SSpacer)
+		];
+}
+
+const TSharedRef<SWidget> UMovieGraphConditionGroupQuery_Actor::FActorSelectionColumn::ConstructRowWidget(FSceneOutlinerTreeItemRef TreeItem, const STableRow<FSceneOutlinerTreeItemPtr>& Row)
+{
+	const FActorTreeItem* ActorTreeItem = TreeItem->CastTo<FActorTreeItem>();
+	if (!ActorTreeItem)
+	{
+		return SNew(SSpacer);
+	}
+	
+	return SNew(SHorizontalBox)
+		+ SHorizontalBox::Slot()
+		.FillWidth(1)
+		.VAlign(VAlign_Center)
+		.HAlign(HAlign_Center)
+		[
+			SNew(SCheckBox)
+			.IsChecked(this, &FActorSelectionColumn::IsRowChecked, ActorTreeItem)
+			.OnCheckStateChanged(this, &FActorSelectionColumn::OnCheckStateChanged, ActorTreeItem)
+		];
+}
+
+ECheckBoxState UMovieGraphConditionGroupQuery_Actor::FActorSelectionColumn::IsRowChecked(const FActorTreeItem* InActorTreeItem) const
+{
+	if (!WeakActorQuery.IsValid())
+	{
+		return ECheckBoxState::Unchecked;
+	}
+
+	const TStrongObjectPtr<UMovieGraphConditionGroupQuery_Actor> ActorQuery = WeakActorQuery.Pin();
+	for (TSoftObjectPtr<AActor>& ActorToMatch : ActorQuery->ActorsToMatch)
+	{
+		if (ActorToMatch.Get() == InActorTreeItem->Actor.Get())
+		{
+			return ECheckBoxState::Checked;
+		}
+	}
+				
+	return ECheckBoxState::Unchecked;
+}
+
+void UMovieGraphConditionGroupQuery_Actor::FActorSelectionColumn::OnCheckStateChanged(const ECheckBoxState NewState, const FActorTreeItem* InActorTreeItem) const
+{
+	if (!WeakActorQuery.IsValid())
+	{
+		return;
+	}
+
+	const TStrongObjectPtr<UMovieGraphConditionGroupQuery_Actor> ActorQuery = WeakActorQuery.Pin();
+	if (InActorTreeItem->Actor.IsValid())
+	{
+		if (NewState == ECheckBoxState::Unchecked)
+		{
+			ActorQuery->RemoveActors({InActorTreeItem->Actor.Get()});
+		}
+		else
+		{
+			const FMovieGraphConditionGroupQueryContentsChanged OnAddFinished = nullptr;
+			constexpr bool bCloseAddMenu = false;
+			ActorQuery->AddActors({InActorTreeItem->Actor.Get()}, OnAddFinished, bCloseAddMenu);
+		}
+	}
 }
 #endif	// WITH_EDITOR
 
diff --git a/Engine/Plugins/MovieScene/MovieRenderPipeline/Source/MovieRenderPipelineCore/Public/Graph/MovieGraphRenderLayerSubsystem.h b/Engine/Plugins/MovieScene/MovieRenderPipeline/Source/MovieRenderPipelineCore/Public/Graph/MovieGraphRenderLayerSubsystem.h
--- a/Engine/Plugins/MovieScene/MovieRenderPipeline/Source/MovieRenderPipelineCore/Public/Graph/MovieGraphRenderLayerSubsystem.h
+++ b/Engine/Plugins/MovieScene/MovieRenderPipeline/Source/MovieRenderPipelineCore/Public/Graph/MovieGraphRenderLayerSubsystem.h
@@ -9,6 +9,7 @@
 
 #if WITH_EDITOR
 #include "ContentBrowserDelegates.h"
+#include "ISceneOutlinerColumn.h"
 #endif	// WITH_EDITOR
 
 #include "MovieGraphRenderLayerSubsystem.generated.h"
@@ -165,8 +166,11 @@ private:
 	static const FSlateBrush* GetRowIcon(TSharedPtr<TSoftObjectPtr<AActor>> InActor);
 	static FText GetRowText(TSharedPtr<TSoftObjectPtr<AActor>> InActor);
 
-	/** Adds the provided actors to the query, updating the UI as needed. Calls InOnAddFinished when done. */
-	void AddActors(const TArray<AActor*>& InActors, const FMovieGraphConditionGroupQueryContentsChanged& InOnAddFinished);
+	/** Adds the provided actors to the query, updating the UI as needed. Calls InOnAddFinished when done. Can optionally close the Add menu. */
+	void AddActors(const TArray<AActor*>& InActors, const FMovieGraphConditionGroupQueryContentsChanged& InOnAddFinished, const bool bCloseAddMenu = true);
+
+	/** Removes the provided actors from the query, updating the UI as needed. */
+	void RemoveActors(const TArray<TSoftObjectPtr<AActor>>& InActors);
 #endif
 
 public:
@@ -176,6 +180,31 @@ public:
 
 private:
 #if WITH_EDITOR
+	/** Custom outliner column that allows adding/removing an actor from an Actor condition group query (via checkbox). */
+	class FActorSelectionColumn final : public ISceneOutlinerColumn
+	{
+	public:
+		explicit FActorSelectionColumn(const TWeakObjectPtr<UMovieGraphConditionGroupQuery_Actor> InWeakActorQuery)
+			: WeakActorQuery(InWeakActorQuery)
+		{}
+		
+		static FName GetID();
+		virtual FName GetColumnID() override;
+		virtual SHeaderRow::FColumn::FArguments ConstructHeaderRowColumn() override;
+		virtual const TSharedRef<SWidget> ConstructRowWidget(FSceneOutlinerTreeItemRef TreeItem, const STableRow<FSceneOutlinerTreeItemPtr>& Row) override;
+
+	private:
+		/** Determines if the given tree item (corresponding to one actor) is checked. */
+		ECheckBoxState IsRowChecked(const FActorTreeItem* InActorTreeItem) const;
+
+		/** Updates the associated actor query when a row is checked or unchecked. */
+		void OnCheckStateChanged(const ECheckBoxState NewState, const FActorTreeItem* InActorTreeItem) const;
+	
+	private:
+		/** The Actor condition group query that populates the data for this column. */
+		TWeakObjectPtr<UMovieGraphConditionGroupQuery_Actor> WeakActorQuery;
+	};
+	
 	TSharedPtr<class ISceneOutliner> ActorPickerWidget;
 
 	/** Displays the actors which have been chosen. */
