MovieGraph: Updated high precision output display name.

#rnx

diff --git a/Engine/Plugins/MovieScene/MovieRenderPipeline/Source/MovieRenderPipelineRenderPasses/Public/MoviePipelineDeferredPasses.h b/Engine/Plugins/MovieScene/MovieRenderPipeline/Source/MovieRenderPipelineRenderPasses/Public/MoviePipelineDeferredPasses.h
--- a/Engine/Plugins/MovieScene/MovieRenderPipeline/Source/MovieRenderPipelineRenderPasses/Public/MoviePipelineDeferredPasses.h
+++ b/Engine/Plugins/MovieScene/MovieRenderPipeline/Source/MovieRenderPipelineRenderPasses/Public/MoviePipelineDeferredPasses.h
@@ -30,7 +30,7 @@ public:
 	TSoftObjectPtr<UMaterialInterface> Material;
 
 	/** Request output to be 32-bit, usually for data exports. Note that scene color precision is still defined by r.SceneColorFormat.*/
-	UPROPERTY(EditAnywhere, BlueprintReadWrite, Category = "Settings", DisplayName="High Precision Output (32-bit)")
+	UPROPERTY(EditAnywhere, BlueprintReadWrite, Category = "Settings", DisplayName="Use High Precision (32-bit) Output")
 	bool bHighPrecisionOutput = false;
 };
 
