MoviePipeline: Migrate MRQ to new 32-bit precision setting per post-process material, since the setting is now also visible there.

We also default depth to high precision by default.

#rnx

diff --git a/Engine/Plugins/MovieScene/MovieRenderPipeline/Source/MovieRenderPipelineRenderPasses/Private/Graph/Nodes/MovieGraphDeferredPassNode.cpp b/Engine/Plugins/MovieScene/MovieRenderPipeline/Source/MovieRenderPipelineRenderPasses/Private/Graph/Nodes/MovieGraphDeferredPassNode.cpp
--- a/Engine/Plugins/MovieScene/MovieRenderPipeline/Source/MovieRenderPipelineRenderPasses/Private/Graph/Nodes/MovieGraphDeferredPassNode.cpp
+++ b/Engine/Plugins/MovieScene/MovieRenderPipeline/Source/MovieRenderPipelineRenderPasses/Private/Graph/Nodes/MovieGraphDeferredPassNode.cpp
@@ -27,6 +27,7 @@ UMovieGraphDeferredRenderPassNode::UMovieGraphDeferredRenderPassNode()
 		FMoviePipelinePostProcessPass& NewPass = AdditionalPostProcessMaterials.AddDefaulted_GetRef();
 		NewPass.Material = TSoftObjectPtr<UMaterialInterface>(FSoftObjectPath(MaterialPath));
 		NewPass.bEnabled = false;
+		NewPass.bHighPrecisionOutput = MaterialPath.Equals(DefaultDepthAsset);
 	}
 }
 
diff --git a/Engine/Plugins/MovieScene/MovieRenderPipeline/Source/MovieRenderPipelineRenderPasses/Private/Graph/Nodes/MovieGraphPathTracerPassNode.cpp b/Engine/Plugins/MovieScene/MovieRenderPipeline/Source/MovieRenderPipelineRenderPasses/Private/Graph/Nodes/MovieGraphPathTracerPassNode.cpp
--- a/Engine/Plugins/MovieScene/MovieRenderPipeline/Source/MovieRenderPipelineRenderPasses/Private/Graph/Nodes/MovieGraphPathTracerPassNode.cpp
+++ b/Engine/Plugins/MovieScene/MovieRenderPipeline/Source/MovieRenderPipelineRenderPasses/Private/Graph/Nodes/MovieGraphPathTracerPassNode.cpp
@@ -31,6 +31,7 @@ UMovieGraphPathTracerRenderPassNode::UMovieGraphPathTracerRenderPassNode()
 		FMoviePipelinePostProcessPass& NewPass = AdditionalPostProcessMaterials.AddDefaulted_GetRef();
 		NewPass.Material = TSoftObjectPtr<UMaterialInterface>(FSoftObjectPath(MaterialPath));
 		NewPass.bEnabled = false;
+		NewPass.bHighPrecisionOutput = MaterialPath.Equals(DefaultDepthAsset);
 	}
 }
 
diff --git a/Engine/Plugins/MovieScene/MovieRenderPipeline/Source/MovieRenderPipelineRenderPasses/Private/MoviePipelineDeferredPasses.cpp b/Engine/Plugins/MovieScene/MovieRenderPipeline/Source/MovieRenderPipelineRenderPasses/Private/MoviePipelineDeferredPasses.cpp
--- a/Engine/Plugins/MovieScene/MovieRenderPipeline/Source/MovieRenderPipelineRenderPasses/Private/MoviePipelineDeferredPasses.cpp
+++ b/Engine/Plugins/MovieScene/MovieRenderPipeline/Source/MovieRenderPipelineRenderPasses/Private/MoviePipelineDeferredPasses.cpp
@@ -62,9 +62,26 @@ UMoviePipelineDeferredPassBase::UMoviePipelineDeferredPassBase()
 		FMoviePipelinePostProcessPass& NewPass = AdditionalPostProcessMaterials.AddDefaulted_GetRef();
 		NewPass.Material = TSoftObjectPtr<UMaterialInterface>(FSoftObjectPath(MaterialPath));
 		NewPass.bEnabled = false;
+		NewPass.bHighPrecisionOutput = MaterialPath.Equals(DefaultDepthAsset);
 	}
 	bRenderMainPass = true;
-	bUse32BitPostProcessMaterials = false;
+}
+
+void UMoviePipelineDeferredPassBase::PostLoad()
+{
+	Super::PostLoad();
+
+#if WITH_EDITORONLY_DATA
+	PRAGMA_DISABLE_DEPRECATION_WARNINGS
+	if (bUse32BitPostProcessMaterials_DEPRECATED)
+	{
+		for (FMoviePipelinePostProcessPass& Pass : AdditionalPostProcessMaterials)
+		{
+			Pass.bHighPrecisionOutput = true;
+		}
+	}
+	PRAGMA_ENABLE_DEPRECATION_WARNINGS
+#endif // WITH_EDITORONLY_DATA
 }
 
 FIntPoint UMoviePipelineDeferredPassBase::GetEffectiveOutputResolutionForCamera(const int32 InCameraIndex) const
@@ -165,6 +182,11 @@ void UMoviePipelineDeferredPassBase::SetupImpl(const MoviePipeline::FMoviePipeli
 			if (Material)
 			{
 				ActivePostProcessMaterials.Add(Material);
+				
+				if (AdditionalPass.bHighPrecisionOutput)
+				{
+					ActiveHighPrecisionPostProcessMaterials.Add(Material);
+				}
 			}
 		}
 	}
@@ -304,6 +326,7 @@ void UMoviePipelineDeferredPassBase::SetupImpl(const MoviePipeline::FMoviePipeli
 void UMoviePipelineDeferredPassBase::TeardownImpl()
 {
 	ActivePostProcessMaterials.Reset();
+	ActiveHighPrecisionPostProcessMaterials.Reset();
 	UniqueStencilLayerNames.Reset();
 
 	for (FMultiCameraViewStateData& CameraData : CameraViewStateData)
@@ -519,7 +542,7 @@ void UMoviePipelineDeferredPassBase::RenderSample_GameThreadImpl(const FMoviePip
 				FMoviePipelinePassIdentifier LayerPassIdentifier = FMoviePipelinePassIdentifier(PassIdentifier.Name + VisMaterial->GetName(), PassIdentifierForCurrentCamera.CameraName);
 
 				auto BufferPipe = MakeShared<FImagePixelPipe, ESPMode::ThreadSafe>();
-				BufferPipe->bIsExpecting32BitPixelData = bUse32BitPostProcessMaterials;
+				BufferPipe->bIsExpecting32BitPixelData = ActiveHighPrecisionPostProcessMaterials.Contains(VisMaterial);
 				BufferPipe->AddEndpoint(MakeForwardingEndpoint(LayerPassIdentifier, InOutSampleState));
 
 				View->FinalPostProcessSettings.BufferVisualizationPipes.Add(VisMaterial->GetFName(), BufferPipe);
diff --git a/Engine/Plugins/MovieScene/MovieRenderPipeline/Source/MovieRenderPipelineRenderPasses/Public/MoviePipelineDeferredPasses.h b/Engine/Plugins/MovieScene/MovieRenderPipeline/Source/MovieRenderPipelineRenderPasses/Public/MoviePipelineDeferredPasses.h
--- a/Engine/Plugins/MovieScene/MovieRenderPipeline/Source/MovieRenderPipelineRenderPasses/Public/MoviePipelineDeferredPasses.h
+++ b/Engine/Plugins/MovieScene/MovieRenderPipeline/Source/MovieRenderPipelineRenderPasses/Public/MoviePipelineDeferredPasses.h
@@ -42,6 +42,10 @@ class MOVIERENDERPIPELINERENDERPASSES_API UMoviePipelineDeferredPassBase : publi
 public:
 	UMoviePipelineDeferredPassBase();
 	
+	// UObject Interface
+	virtual void PostLoad() override;
+	// ~UObject Interface
+
 protected:
 	// UMoviePipelineRenderPass API
 	virtual void SetupImpl(const MoviePipeline::FMoviePipelineRenderPassInitSettings& InPassInitSettings) override;
@@ -107,11 +111,14 @@ public:
 	UPROPERTY(EditAnywhere, BlueprintReadWrite, Category = "Post Processing")
 	bool bDisableMultisampleEffects;
 
+#if WITH_EDITORONLY_DATA
 	/**
 	* Should the additional post-process materials write out to a 32-bit render target instead of 16-bit?
 	*/
-	UPROPERTY(EditAnywhere, BlueprintReadWrite, Category = "Deferred Renderer Data")
-	bool bUse32BitPostProcessMaterials;
+	UE_DEPRECATED(5.5, "bUse32BitPostProcessMaterials has been deprecated, please use the setting per material.")
+	UPROPERTY()
+	bool bUse32BitPostProcessMaterials_DEPRECATED;
+#endif
 
 	/**
 	* An array of additional post-processing materials to run after the frame is rendered. Using this feature may add a notable amount of render time.
@@ -155,6 +162,9 @@ protected:
 	UPROPERTY(Transient, DuplicateTransient)
 	TArray<TObjectPtr<UMaterialInterface>> ActivePostProcessMaterials;
 
+	UPROPERTY(Transient, DuplicateTransient)
+	TSet<TObjectPtr<UMaterialInterface>> ActiveHighPrecisionPostProcessMaterials;
+
 	UPROPERTY(Transient, DuplicateTransient)
 	TObjectPtr<UMaterialInterface> StencilLayerMaterial;
 
