MRQ: Fix typo in tiling logic to allocate history per tile

The total number of tiles was being applied per dimension leading to O(n^4) allocations instead of O(n^2).

diff --git a/Engine/Plugins/MovieScene/MovieRenderPipeline/Source/MovieRenderPipelineRenderPasses/Private/MoviePipelineDeferredPasses.cpp b/Engine/Plugins/MovieScene/MovieRenderPipeline/Source/MovieRenderPipelineRenderPasses/Private/MoviePipelineDeferredPasses.cpp
--- a/Engine/Plugins/MovieScene/MovieRenderPipeline/Source/MovieRenderPipelineRenderPasses/Private/MoviePipelineDeferredPasses.cpp
+++ b/Engine/Plugins/MovieScene/MovieRenderPipeline/Source/MovieRenderPipelineRenderPasses/Private/MoviePipelineDeferredPasses.cpp
@@ -217,10 +217,10 @@ void UMoviePipelineDeferredPassBase::SetupImpl(const MoviePipeline::FMoviePipeli
 
 		// We don't always want to allocate a unique history per tile as very large resolutions can OOM the GPU in backbuffer images alone.
 		// But we do need the history for some features (like Lumen) to work, so it's optional.
-		int32 NumHighResTiles = HighResSettings->bAllocateHistoryPerTile ? (HighResSettings->TileCount * HighResSettings->TileCount) : 1;
-		for (int32 TileIndexX = 0; TileIndexX < NumHighResTiles; TileIndexX++)
+		int32 HighResTileCount = HighResSettings->bAllocateHistoryPerTile ? HighResSettings->TileCount : 1;
+		for (int32 TileIndexX = 0; TileIndexX < HighResTileCount; TileIndexX++)
 		{
-			for (int32 TileIndexY = 0; TileIndexY < NumHighResTiles; TileIndexY++)
+			for (int32 TileIndexY = 0; TileIndexY < HighResTileCount; TileIndexY++)
 			{
 				FMultiCameraViewStateData::FPerTile& PerTile = CameraData.TileData.FindOrAdd(FIntPoint(TileIndexX, TileIndexY));
 				// If they want to render the main pass (most likely) add a view state for it
