Movie Render Queue: Attempt to throttle the render queue so that we don't accumulate too much work even when the path tracer is being used.

This improves memory usage when using MRQ and high spatial sample counts at the same time by ensuring the CPU thread does not get too far ahead of the GPU (which causes a lot of temporary resources to be allocated).

Fix tiled rendering when per tile history is disabled (need to fallback the lookup to the (0,0) tile view state in that case).

Remove log message in CopyReadbackTexture_RenderThread when the texture dimensions are not aligned as this case is correctly handled.

diff --git a/Engine/Plugins/MovieScene/MovieRenderPipeline/Source/MovieRenderPipelineCore/Private/MoviePipelineRendering.cpp b/Engine/Plugins/MovieScene/MovieRenderPipeline/Source/MovieRenderPipelineCore/Private/MoviePipelineRendering.cpp
--- a/Engine/Plugins/MovieScene/MovieRenderPipeline/Source/MovieRenderPipelineCore/Private/MoviePipelineRendering.cpp
+++ b/Engine/Plugins/MovieScene/MovieRenderPipeline/Source/MovieRenderPipelineCore/Private/MoviePipelineRendering.cpp
@@ -46,6 +46,12 @@ static TAutoConsoleVariable<bool> CVarMoviePipelineDisableShaderFlushing(
 	TEXT("If true, rendered frames may be missing objects (meshes, particles, etc.) or objects may show the default checkerboard material."),
 	ECVF_Default);
 
+static TAutoConsoleVariable<int> CVarMoviePipelineThrottleFrameCount(
+	TEXT("MoviePipeline.ThrottleFrameCount"), 2,
+	TEXT("Number of rendered frames that can be submitted to the rendering thread before waiting. A value of 0 will allow the CPU to submit all work without waiting on the GPU.\n")
+	TEXT("The default value of 2 tries to balance between performance and memory usage. The maximum value is 4."),
+	ECVF_Default);
+
 #define LOCTEXT_NAMESPACE "MoviePipeline"
 
 void UMoviePipeline::SetupRenderingPipelineForShot(UMoviePipelineExecutorShot* InShot)
@@ -290,6 +296,11 @@ void UMoviePipeline::RenderFrame()
 	}
 #endif
 
+	constexpr int FenceBufferMax = 4;
+	const int FrameThrotteCount = FMath::Clamp(CVarMoviePipelineThrottleFrameCount.GetValueOnGameThread(), 0, FenceBufferMax);
+	FGPUFenceRHIRef MRQThrottleFence[FenceBufferMax];
+	int FenceIndex = 0;
+	
 	for (int32 TileY = 0; TileY < TileCount.Y; TileY++)
 	{
 		for (int32 TileX = 0; TileX < TileCount.X; TileX++)
@@ -382,6 +393,35 @@ void UMoviePipeline::RenderFrame()
 				SampleState.GlobalScreenPercentageFraction = FLegacyScreenPercentageDriver::GetCVarResolutionFraction();
 				SampleState.OverscanPercentage = FMath::Clamp(CameraSettings->OverscanPercentage, 0.0f, 1.0f);
 
+				if (FrameThrotteCount > 0)
+				{
+					// Before we render, wait for previous samples to have completed so the GPU command list doesn't get too far behind
+					if (MRQThrottleFence[FenceIndex] && !MRQThrottleFence[FenceIndex]->Poll())
+					{
+						TRACE_CPUPROFILER_EVENT_SCOPE(MRQFrameThrottle);
+						for (;;)
+						{
+							FPlatformProcess::SleepNoStats(0.001f);
+							if (MRQThrottleFence[FenceIndex]->Poll())
+							{
+								break;
+							}
+						}
+					}
+
+					// Create a fence for this frame and insert a signal to it
+					MRQThrottleFence[FenceIndex] = RHICreateGPUFence(TEXT("MRQThrottleFence"));
+					ENQUEUE_RENDER_COMMAND(MRQFrameThrottle)([Fence = MRQThrottleFence[FenceIndex]]
+						(FRHICommandListImmediate& RHICmdList)
+						{
+							RHICmdList.WriteGPUFence(Fence);
+							RHICmdList.SubmitCommandsHint();
+						});
+
+					// Switch fences for the next frame (this makes us wait on a different frame than what we just made the signal for)
+					FenceIndex = (FenceIndex + 1) % FrameThrotteCount;
+				}
+
 				// Render each output pass
 				FMoviePipelineRenderPassMetrics SampleStateForCurrentResolution = UE::MoviePipeline::GetRenderPassMetrics(GetPipelinePrimaryConfig(), ActiveShotList[CurrentShotIndex], SampleState, OutputResolution);
 				for (UMoviePipelineRenderPass* RenderPass : InputBuffers)
diff --git a/Engine/Plugins/MovieScene/MovieRenderPipeline/Source/MovieRenderPipelineCore/Private/MoviePipelineSurfaceReader.cpp b/Engine/Plugins/MovieScene/MovieRenderPipeline/Source/MovieRenderPipelineCore/Private/MoviePipelineSurfaceReader.cpp
--- a/Engine/Plugins/MovieScene/MovieRenderPipeline/Source/MovieRenderPipelineCore/Private/MoviePipelineSurfaceReader.cpp
+++ b/Engine/Plugins/MovieScene/MovieRenderPipeline/Source/MovieRenderPipelineCore/Private/MoviePipelineSurfaceReader.cpp
@@ -222,11 +222,6 @@ void FMoviePipelineSurfaceReader::CopyReadbackTexture_RenderThread(TUniqueFuncti
 		}
 		else
 		{
-			UE_LOG(LogMovieRenderPipeline, Log, TEXT("Unexpected size in FMoviePipelineSurfaceReader::CopyReadbackTexture_RenderThread."));
-			// UE_LOG(LogMovieRenderPipeline, Log, TEXT("    Tile size:     %d x %d"), InSampleState.TileSize.X, InSampleState.TileSize.Y);
-			UE_LOG(LogMovieRenderPipeline, Log, TEXT("    Expected size: %d x %d"), ExpectedSizeX, ExpectedSizeY);
-			UE_LOG(LogMovieRenderPipeline, Log, TEXT("    Actual size:   %d x %d"), ActualSizeX, ActualSizeY);
-
 			// Make sure the target is larger than expected size.
 			check(ExpectedSizeX <= ActualSizeX);
 			check(ExpectedSizeY <= ActualSizeY);
diff --git a/Engine/Plugins/MovieScene/MovieRenderPipeline/Source/MovieRenderPipelineRenderPasses/Private/MoviePipelineDeferredPasses.cpp b/Engine/Plugins/MovieScene/MovieRenderPipeline/Source/MovieRenderPipelineRenderPasses/Private/MoviePipelineDeferredPasses.cpp
--- a/Engine/Plugins/MovieScene/MovieRenderPipeline/Source/MovieRenderPipelineRenderPasses/Private/MoviePipelineDeferredPasses.cpp
+++ b/Engine/Plugins/MovieScene/MovieRenderPipeline/Source/MovieRenderPipelineRenderPasses/Private/MoviePipelineDeferredPasses.cpp
@@ -411,7 +411,8 @@ FSceneViewStateInterface* UMoviePipelineDeferredPassBase::GetSceneViewStateInter
 	int32 LocalCameraIndex = FMath::Clamp(Payload->CameraIndex, 0, Payload->CameraIndex);
 
 	FMultiCameraViewStateData& CameraData = CameraViewStateData[LocalCameraIndex];
-	if (FMultiCameraViewStateData::FPerTile* TileData = CameraData.TileData.Find(Payload->TileIndex))
+	FIntPoint TileIndex = CameraData.TileData.Num() == 1 ? FIntPoint(0, 0) : Payload->TileIndex;
+	if (FMultiCameraViewStateData::FPerTile* TileData = CameraData.TileData.Find(TileIndex))
 	{
 		return TileData->SceneViewStates[Payload->SceneViewIndex].GetReference();
 	}
