MoviePipeline: Added the ability to re-order graph variables (drag/drop) and assign them to categories. The categories themselves are also re-orderable. The variable order/categorization is reflected in the job variables in the queue UI.

diff --git a/Engine/Plugins/MovieScene/MovieRenderPipeline/Source/MovieRenderPipelineCore/Private/Graph/MovieGraphConfig.cpp b/Engine/Plugins/MovieScene/MovieRenderPipeline/Source/MovieRenderPipelineCore/Private/Graph/MovieGraphConfig.cpp
--- a/Engine/Plugins/MovieScene/MovieRenderPipeline/Source/MovieRenderPipelineCore/Private/Graph/MovieGraphConfig.cpp
+++ b/Engine/Plugins/MovieScene/MovieRenderPipeline/Source/MovieRenderPipelineCore/Private/Graph/MovieGraphConfig.cpp
@@ -69,6 +69,141 @@ bool UMovieGraphVariable::IsGlobal() const
 	return IsA<UMovieGraphGlobalVariable>();
 }
 
+const FString& UMovieGraphVariable::GetCategory() const
+{
+	return Category;
+}
+
+void UMovieGraphVariable::SetCategory(const FString& InNewCategory)
+{
+	// Sets this variable to have a specific category with Pre/Post change events sent properly
+	auto SetCategoryWithEvents = [this](const FString& NewCategory)
+	{
+#if WITH_EDITOR
+		FProperty* CategoryProperty = FindFProperty<FProperty>(GetClass(), GET_MEMBER_NAME_CHECKED(UMovieGraphVariable, Category));
+		PreEditChange(CategoryProperty);
+
+		Modify();
+#endif
+
+		// The category needs to be put through NameToDisplayString() to prevent a host of category-matching issues in the graph's action menu that
+		// cannot be controlled/changed on the MRG side. For example, when dragging/dropping a category, the provided category in DroppedOnCategory()
+		// will be the *display name* of the category, not the variable's actual category. Sometimes the display name and actual category can differ,
+		// (for example, if the variable category was set to "Test1", the display name would be "Test 1"). To mitigate all of these mismatch issues,
+		// just put all categories through NameToDisplayString(), which is how the display name is generated. The ReplaceInline() fixes a weird issue
+		// with NameToDisplayString().
+		constexpr bool bIsBool = false;
+		Category = FName::NameToDisplayString(NewCategory, bIsBool);
+		Category.ReplaceInline(TEXT("| "), TEXT("|"), ESearchCase::CaseSensitive);
+
+#if WITH_EDITOR
+		FPropertyChangedEvent PropertyUpdate(CategoryProperty);
+		PostEditChangeProperty(PropertyUpdate);
+#endif	// WITH_EDITOR
+	};
+
+	// Nothing to do if the category didn't change
+	if (Category == InNewCategory)
+	{
+		return;
+	}
+	
+	const TArray<UMovieGraphVariable*> GraphVariables = GetOwningGraph()->GetVariables();
+	
+	// Make sure the variable goes in the correct position in the array when changing the category. Variables need to remain sorted by their category,
+	// otherwise other operations (eg, moving a variable up/down in the array) will not work as expected when variables are categorized.
+
+	// First, get all the categories that are present.
+	TArray<FString> Categories;
+	Algo::Transform(GraphVariables, Categories, [](const UMovieGraphVariable* Variable) { return Variable->GetCategory(); });
+
+	// Second, if the newly-set category already exists, then just move the variable to be immediately after the last variable within the category.
+	if (Categories.Contains(InNewCategory))
+	{
+		const int32 LastIndex = GraphVariables.FindLastByPredicate([&InNewCategory](const UMovieGraphVariable* GraphVariable)
+		{
+			return GraphVariable->GetCategory() == InNewCategory;
+		});
+
+		if (LastIndex != INDEX_NONE)
+		{
+			GetOwningGraph()->MoveVariableToIndex(this, LastIndex + 1);
+		}
+	}
+
+	// Otherwise, if the newly-set category does not exist, the situation is a bit more complex.
+	else
+	{
+		// If moving the variable into category "Foo|Bar|Baz", generate all potential root categories (Foo, Foo|Bar, and Foo|Bar|Baz).
+		TArray<FString> CategoryParts;
+		InNewCategory.ParseIntoArray(CategoryParts, TEXT("|"));
+		TArray<FString> CategoryCandidates;
+		for (int32 CategoryIndex = 0; CategoryIndex < CategoryParts.Num(); ++CategoryIndex)
+		{
+			if (CategoryIndex == 0)
+			{
+				CategoryCandidates.Add(CategoryParts[0]);
+			}
+			else
+			{
+				CategoryCandidates.Add(CategoryCandidates.Last() + TEXT("|") + CategoryParts[CategoryIndex]);
+			}
+		}
+
+		// Find the variable with the closest matching category. For example, if setting the category to "Foo|Bar|Baz", if no variable was set in either
+		// "Foo|Bar|Baz" or "Foo|Bar", "Foo" would be the closest matching category. Candidates are iterated in reverse because the deepest category
+		// should be matched first if possible ("Foo|Bar|Baz" in this example).
+		int32 BestIndex = INDEX_NONE;
+		for (int32 CategoryIndex = CategoryCandidates.Num() - 1; CategoryIndex >= 0; --CategoryIndex)
+		{
+			const FString CurrentCategory = CategoryCandidates[CategoryIndex];
+
+			BestIndex = GraphVariables.IndexOfByPredicate([&CurrentCategory](const UMovieGraphVariable* GraphVariable)
+			{
+				return GraphVariable->GetCategory() == CurrentCategory;
+			});
+
+			if (BestIndex != INDEX_NONE)
+			{
+				break;
+			}
+		}
+
+		// If there are variables with a matching root set already, move the variable to be immediately before the variable with a matching category.
+		if (BestIndex != INDEX_NONE)
+		{
+			GetOwningGraph()->MoveVariableToIndex(this, BestIndex);
+		}
+		else
+		{
+			// Find the last variable with any category set, and move it to be after that. New categories go to the end.
+			BestIndex = GraphVariables.FindLastByPredicate([](const UMovieGraphVariable* GraphVariable)
+			{
+				return !GraphVariable->GetCategory().IsEmpty();
+			});
+			
+			if (BestIndex == INDEX_NONE)
+			{
+				// If ALL categories are currently empty (meaning no variables have a category set yet) then move this variable to be the first
+				// in the array.
+				GetOwningGraph()->MoveVariableToIndex(this, 0);
+			}
+			else if (BestIndex == (GraphVariables.Num() - 1))
+			{
+				// ALL variables have a category set, so move to the very end.
+				GetOwningGraph()->MoveVariableToIndex(this, GraphVariables.Num());
+			}
+			else
+			{
+				// Move to after the last variable with a category set 
+				GetOwningGraph()->MoveVariableToIndex(this, BestIndex + 1);
+			}
+		}
+	}
+	
+	SetCategoryWithEvents(InNewCategory);
+}
+
 bool UMovieGraphVariable::IsDeletable() const
 {
 	return true;
@@ -1151,6 +1286,162 @@ void UMovieGraphConfig::GetOutputDirectory(FString& OutOutputDirectory) const
 	}
 }
 
+void UMovieGraphConfig::MoveVariableBefore(UMovieGraphVariable* InTargetVariable, UMovieGraphVariable* InBeforeVariable)
+{
+	const int32 TargetVariableIndex = Variables.Find(InTargetVariable);
+	const int32 BeforeVariableIndex = Variables.Find(InBeforeVariable);
+	if ((TargetVariableIndex == INDEX_NONE) || (BeforeVariableIndex == INDEX_NONE) || (TargetVariableIndex == BeforeVariableIndex))
+	{
+		return;
+	}
+
+#if WITH_EDITOR
+	Modify();
+	InTargetVariable->Modify();
+#endif
+
+	// Moving the target before another variable means that the target should inherit the other variable's category
+	InTargetVariable->Category = InBeforeVariable->Category;
+
+	Variables.RemoveSingle(InTargetVariable);
+	if (BeforeVariableIndex < TargetVariableIndex)
+	{
+		Variables.Insert(InTargetVariable, BeforeVariableIndex);
+	}
+	else
+	{
+		Variables.Insert(InTargetVariable, BeforeVariableIndex - 1);
+	}
+
+#if WITH_EDITOR
+	if (OnGraphVariablesChangedDelegate.IsBound())
+	{
+		OnGraphVariablesChangedDelegate.Broadcast();
+	}
+#endif
+}
+
+void UMovieGraphConfig::MoveVariableToIndex(UMovieGraphVariable* InTargetVariable, int32 NewIndex)
+{
+	if (!InTargetVariable)
+	{
+		return;
+	}
+	
+	// Check if the index is valid. IsValidIndex() won't work here because we may want to insert at one-past-the-end, where IsValidIndex() will
+	// return false.
+	const bool bIsValidIndex = (NewIndex >= 0) && (NewIndex <= Variables.Num());
+	if (!bIsValidIndex)
+	{
+		return;
+	}
+	
+	const int32 CurrentIndex = Variables.Find(InTargetVariable);
+	if ((CurrentIndex == NewIndex) || (CurrentIndex == INDEX_NONE))
+	{
+		return;
+	}
+
+#if WITH_EDITOR
+	Modify();
+#endif
+
+	Variables.RemoveAt(CurrentIndex);
+	if (NewIndex < CurrentIndex)
+	{
+		Variables.Insert(InTargetVariable, NewIndex);
+	}
+	else
+	{
+		Variables.Insert(InTargetVariable, NewIndex - 1);
+	}
+}
+
+void UMovieGraphConfig::MoveCategoryBefore(const FString& InCategoryToMove, const FString& InCategoryBefore)
+{
+	// Determines if a variable is part of the given category. Matches by exact category, or if the variable is a child under the given category
+	// (eg, if "Foo" is the provided category, and the variable's category is "Foo|Bar", the variable would be a match).
+	auto VariableBelongsToCategory = [](const UMovieGraphVariable* InVariable, const FString& InCategory)
+	{
+		if (!InVariable)
+		{
+			return false;
+		}
+		
+		const FString ParentCategoryPrefix = FString::Format(TEXT("{0}|"), {InCategory});
+		const FString& VariableCategory = InVariable->GetCategory();
+		
+		return (VariableCategory == InCategory) || VariableCategory.StartsWith(ParentCategoryPrefix);
+	};
+	
+	// If dragging a category into a subcategory (eg, Foo|Bar) we have to use the root category (Foo) as the category to move before. Variable
+	// reparenting via category moves is not supported currently.
+	FString CategoryBefore = InCategoryBefore;
+	if (CategoryBefore.Contains(TEXT("|")))
+	{
+		TArray<FString> CategoryParts;
+		CategoryBefore.ParseIntoArray(CategoryParts, TEXT("|"));
+
+		if (!CategoryParts.IsEmpty())
+		{
+			CategoryBefore = CategoryParts[0];
+		}
+	}
+	
+	// Cache all variables that will be moved
+	const TArray<UMovieGraphVariable*> SourceVariables = Variables.FilterByPredicate([&InCategoryToMove, &VariableBelongsToCategory](const UMovieGraphVariable* InVariable)
+	{
+		return VariableBelongsToCategory(InVariable, InCategoryToMove);
+	});
+
+	// Ensure that a "before" variable can be found before removing the "from" variables
+	const TObjectPtr<UMovieGraphVariable>* BeforeVariable = Variables.FindByPredicate([&CategoryBefore, &VariableBelongsToCategory](const TObjectPtr<UMovieGraphVariable>& InVariable)
+	{
+		return VariableBelongsToCategory(InVariable, CategoryBefore);
+	});
+	if (!BeforeVariable)
+	{
+		return;
+	}
+
+#if WITH_EDITOR
+	Modify();
+#endif
+
+	// Remove all the variables that will be moved
+	Variables.RemoveAll([&InCategoryToMove, &VariableBelongsToCategory](const UMovieGraphVariable* InVariable)
+	{
+		return VariableBelongsToCategory(InVariable, InCategoryToMove);
+	});
+
+	// Find the first variable that marks the variable to be "before" variable
+	int32 BeforeVariableIndex = INDEX_NONE;
+	for (int32 VariableIndex = 0; VariableIndex < Variables.Num(); ++VariableIndex)
+	{
+		if (VariableBelongsToCategory(Variables[VariableIndex], CategoryBefore))
+		{
+			BeforeVariableIndex = VariableIndex;
+			break;
+		}
+	}
+
+	// If found, move all "from" variables to be immediately before the "before" variable.
+	if (BeforeVariableIndex != INDEX_NONE)
+	{
+		for (int32 SourceVariableIndex = 0; SourceVariableIndex < SourceVariables.Num(); ++ SourceVariableIndex)
+		{
+			Variables.Insert(SourceVariables[SourceVariableIndex], BeforeVariableIndex + SourceVariableIndex);
+		}
+	}
+
+#if WITH_EDITOR
+	if (OnGraphVariablesChangedDelegate.IsBound())
+	{
+		OnGraphVariablesChangedDelegate.Broadcast();
+	}
+#endif
+}
+
 void UMovieGraphConfig::InitializeFlattenedNode(UMovieGraphNode* InNode)
 {
 	// We go through each of the bOverride_ properties on this new instance and set
diff --git a/Engine/Plugins/MovieScene/MovieRenderPipeline/Source/MovieRenderPipelineCore/Private/MovieJobVariableAssignmentContainer.cpp b/Engine/Plugins/MovieScene/MovieRenderPipeline/Source/MovieRenderPipelineCore/Private/MovieJobVariableAssignmentContainer.cpp
--- a/Engine/Plugins/MovieScene/MovieRenderPipeline/Source/MovieRenderPipelineCore/Private/MovieJobVariableAssignmentContainer.cpp
+++ b/Engine/Plugins/MovieScene/MovieRenderPipeline/Source/MovieRenderPipelineCore/Private/MovieJobVariableAssignmentContainer.cpp
@@ -284,6 +284,8 @@ bool UMovieJobVariableAssignmentContainer::GenerateVariableOverride(const UMovie
 #if WITH_EDITOR
 	NewProperty.MetaData.Add(FPropertyBagPropertyDescMetaData(VariableGuidMetaDataKey, InGraphVariable->GetGuid().ToString()));
 	NewProperty.MetaData.Add(FPropertyBagPropertyDescMetaData(ToolTipMetaDataKey, InGraphVariable->Description));
+	NewProperty.MetaData.Add(FPropertyBagPropertyDescMetaData(CategoryMetaDataKey, InGraphVariable->GetCategory()));
+	NewProperty.MetaData.Add(FPropertyBagPropertyDescMetaData(EnableCategoriesMetaDataKey, FString()));
 #endif
 
 	// Track a separate EditCondition property that can enable/disable the above property. Since the variable can be
@@ -415,6 +417,7 @@ void UMovieJobVariableAssignmentContainer::UpdateGraphVariableOverrides()
 				Desc.ContainerTypes = { static_cast<EPropertyBagContainerType>(Variable->GetValueContainerType()) };
 			}
 
+			// Check the ToolTip metadata
 			FPropertyBagPropertyDescMetaData* ToolTipMetaDataEntry =
 				Desc.MetaData.FindByPredicate([](const FPropertyBagPropertyDescMetaData& InMetaDataEntry)
 				{
@@ -425,6 +428,18 @@ void UMovieJobVariableAssignmentContainer::UpdateGraphVariableOverrides()
 				bNeedsToRegenerate = true;
 				ToolTipMetaDataEntry->Value = Variable->Description;
 			}
+
+			// Check the Category metadata
+			FPropertyBagPropertyDescMetaData* CategoryMetaDataEntry =
+				Desc.MetaData.FindByPredicate([](const FPropertyBagPropertyDescMetaData& InMetaDataEntry)
+				{
+					return InMetaDataEntry.Key == CategoryMetaDataKey;
+				});
+			if (CategoryMetaDataEntry && (CategoryMetaDataEntry->Value != FName(Variable->GetCategory())))
+			{
+				bNeedsToRegenerate = true;
+				CategoryMetaDataEntry->Value = Variable->GetCategory();
+			}
 		}
 	}
 
@@ -462,6 +477,25 @@ void UMovieJobVariableAssignmentContainer::UpdateGraphVariableOverrides()
 		}
 	}
 
+	// Fourth, make sure that the variables are in the correct order. They should reflect the exact order that they're defined within the graph.
+	TArray<FGuid> OriginalOrder;
+	Algo::Transform(ModifiedDescs, OriginalOrder, [&GetVariableGuidFromDesc](const FPropertyBagPropertyDesc& Desc) { return GetVariableGuidFromDesc(Desc); });
+	Algo::SortBy(ModifiedDescs, [&GraphVariables, &GetVariableGuidFromDesc](const FPropertyBagPropertyDesc& Desc)
+	{
+		// Sort the descs by the order of variable GUIDs within the graph
+		return GraphVariables.IndexOfByPredicate([&Desc, &GetVariableGuidFromDesc](const UMovieGraphVariable* Variable)
+		{
+			return Variable->GetGuid() == GetVariableGuidFromDesc(Desc);
+		});
+	});
+	TArray<FGuid> ChangedOrder;
+	Algo::Transform(ModifiedDescs, ChangedOrder, [&GetVariableGuidFromDesc](const FPropertyBagPropertyDesc& Desc) { return GetVariableGuidFromDesc(Desc); });
+
+	if (OriginalOrder != ChangedOrder)
+	{
+		bNeedsToRegenerate = true;
+	}
+
 	if (bNeedsToRegenerate)
 	{
 		const UPropertyBag* NewBagStruct = UPropertyBag::GetOrCreateFromDescs(ModifiedDescs);
diff --git a/Engine/Plugins/MovieScene/MovieRenderPipeline/Source/MovieRenderPipelineCore/Public/Graph/MovieGraphConfig.h b/Engine/Plugins/MovieScene/MovieRenderPipeline/Source/MovieRenderPipelineCore/Public/Graph/MovieGraphConfig.h
--- a/Engine/Plugins/MovieScene/MovieRenderPipeline/Source/MovieRenderPipelineCore/Public/Graph/MovieGraphConfig.h
+++ b/Engine/Plugins/MovieScene/MovieRenderPipeline/Source/MovieRenderPipelineCore/Public/Graph/MovieGraphConfig.h
@@ -119,13 +119,21 @@ public:
 	/** Returns true if this variable is a global variable. */
 	bool IsGlobal() const;
 
+	/** Gets the category (if any) assigned to this variable. */
+	const FString& GetCategory() const;
+
+	/**
+	 * Sets the variable to the provided category. Be aware that the category provided here may not be the final category set on the
+	 * variable (InNewCategory will be put through FName::NameToDisplayString().
+	 */
+	void SetCategory(const FString& InNewCategory);
+
 	//~ Begin UMovieGraphMember interface
 	virtual bool IsDeletable() const override;
 	virtual bool CanRename(const FText& InNewName, FText& OutError) const override;
 	virtual bool SetMemberName(const FString& InNewName) override;
 	//~ End UMovieGraphMember interface
 
-public:
 #if WITH_EDITOR
 	FOnMovieGraphVariableChanged OnMovieGraphVariableChangedDelegate;
 
@@ -133,6 +141,11 @@ public:
 	virtual void PostEditChangeProperty(FPropertyChangedEvent& PropertyChangedEvent) override;
 	//~ End UObject overrides
 #endif // WITH_EDITOR
+
+private:
+	/** The category assigned to the variable. Defaults to empty, which means no category. */
+	UPROPERTY()
+	FString Category;
 };
 
 /**
@@ -655,7 +668,6 @@ public:
 	UFUNCTION(BlueprintPure, Category = "Movie Graph")
 	UMovieGraphNode* GetOutputNode() const { return OutputNode; }
 
-
 	const TArray<TObjectPtr<UMovieGraphNode>>& GetNodes() const { return AllNodes; }
 
 	/**
@@ -750,6 +762,23 @@ public:
 	 */
 	void GetOutputDirectory(FString& OutOutputDirectory) const;
 
+	/**
+	 * Moves one variable (InTargetVariable) before another variable (InBeforeVariable). Takes care of ensuring the variable's category is
+	 * set properly after the move.
+	 */
+	void MoveVariableBefore(UMovieGraphVariable* InTargetVariable, UMovieGraphVariable* InBeforeVariable);
+
+	/**
+	 * Moves one variable (InTargetVariable) to the specified index among all user graph variables.
+	 *
+	 * Note that MoveVariableBefore() should be used in almost all cases unless there is very specific use case. This method will not take care of setting
+	 * the category for you after the move unlike MoveVariableBefore().
+	 */
+	void MoveVariableToIndex(UMovieGraphVariable* InTargetVariable, int32 NewIndex);
+
+	/** Moves one category (InCategoryToMove) and its variables before another category (InCategoryBefore). */
+	void MoveCategoryBefore(const FString& InCategoryToMove, const FString& InCategoryBefore);
+
 protected:
 	/** Look for the output directory in the UMovieGraphOutputSettings nodes found upstream of InNode. */
 	void RecurseUpGlobalsBranchToFindOutputDirectory(const UMovieGraphNode* InNode, FString& OutOutputDirectory, TArray<const UMovieGraphConfig*>& VisitedGraphStack) const;
diff --git a/Engine/Plugins/MovieScene/MovieRenderPipeline/Source/MovieRenderPipelineCore/Public/MovieJobVariableAssignmentContainer.h b/Engine/Plugins/MovieScene/MovieRenderPipeline/Source/MovieRenderPipelineCore/Public/MovieJobVariableAssignmentContainer.h
--- a/Engine/Plugins/MovieScene/MovieRenderPipeline/Source/MovieRenderPipelineCore/Public/MovieJobVariableAssignmentContainer.h
+++ b/Engine/Plugins/MovieScene/MovieRenderPipeline/Source/MovieRenderPipelineCore/Public/MovieJobVariableAssignmentContainer.h
@@ -262,6 +262,12 @@ private:
 	
 	/* The metadata key on properties that is associated with a tooltip value. */
 	static inline FName ToolTipMetaDataKey = FName(TEXT("ToolTip"));
+
+	/* The metadata key on properties that is associated with the variable's category. */
+	static inline FName CategoryMetaDataKey = FName(TEXT("Category"));
+	
+	/* The metadata key on property bag properties that allows them to show up grouped in the details panel. */
+	static inline FName EnableCategoriesMetaDataKey = FName(TEXT("EnableCategories"));
 	
 	/** The properties managed by this object. */
 	UPROPERTY(EditAnywhere, meta=(ShowOnlyInnerProperties, FixedLayout), Category = "Value")
diff --git a/Engine/Plugins/MovieScene/MovieRenderPipeline/Source/MovieRenderPipelineEditor/Private/Customizations/Graph/MovieGraphMemberCustomization.h b/Engine/Plugins/MovieScene/MovieRenderPipeline/Source/MovieRenderPipelineEditor/Private/Customizations/Graph/MovieGraphMemberCustomization.h
--- a/Engine/Plugins/MovieScene/MovieRenderPipeline/Source/MovieRenderPipelineEditor/Private/Customizations/Graph/MovieGraphMemberCustomization.h
+++ b/Engine/Plugins/MovieScene/MovieRenderPipeline/Source/MovieRenderPipelineEditor/Private/Customizations/Graph/MovieGraphMemberCustomization.h
@@ -11,6 +11,7 @@
 #include "Graph/MovieGraphConfig.h"
 #include "PropertyBagDetails.h"
 #include "PropertyHandle.h"
+#include "ScopedTransaction.h"
 #include "SPinTypeSelector.h"
 #include "Widgets/Input/SEditableTextBox.h"
 #include "Widgets/Text/STextBlock.h"
@@ -155,12 +156,12 @@ protected:
 				continue;
 			}
 
-			const bool bIsNameRowEnabled = MemberObject->IsEditable();
+			const bool bIsEditable = MemberObject->IsEditable();
 			
 			// Add a custom row for the Name property (to allow for proper validation)
 			IDetailCategoryBuilder& GeneralCategory = DetailBuilder.EditCategory("General");
 			GeneralCategory.AddCustomRow(FText::GetEmpty())
-			.IsEnabled(bIsNameRowEnabled)
+			.IsEnabled(bIsEditable)
 			.NameContent()
 			[
 				SNew(STextBlock)
@@ -177,6 +178,35 @@ protected:
 				.Font(IDetailLayoutBuilder::GetDetailFont())
 			];
 
+			// Add a custom row for the Category property (its property is not edited directly in order to give the graph a chance to sort the
+			// variables correctly). This is only applicable to variables.
+			if (UMovieGraphVariable* VariableMember = Cast<UMovieGraphVariable>(MemberObject))
+			{
+				GeneralCategory.AddCustomRow(FText::GetEmpty())
+				.IsEnabled(bIsEditable)
+				.NameContent()
+				[
+					SNew(STextBlock)
+					.Text(LOCTEXT("MemberPropertyLabel_Category", "Category"))
+					.ToolTipText(LOCTEXT("MemberPropertyTooltip_Category", "The category assigned to the variable. Use a '|' to separate category names to create a category hierarchy (eg, Settings|Resolution)."))
+					.Font(DetailBuilder.GetDetailFont())
+				]
+				.ValueContent()
+				[
+					SNew(SEditableTextBox)
+					.Text_Lambda([VariableMember]()
+					{
+						return FText::FromString(VariableMember->GetCategory());
+					})
+					.OnTextCommitted_Lambda([VariableMember](const FText& InNewCategory, ETextCommit::Type TextCommitType)
+					{
+						const FScopedTransaction Transaction(LOCTEXT("SetVariableCategory", "Set Variable Category"));
+						VariableMember->SetCategory(InNewCategory.ToString());
+					})
+					.Font(IDetailLayoutBuilder::GetDetailFont())
+				];
+			}
+
 			// If this is an interface (eg, input/output) only enable the type selector if it's not a branch and editable.
 			// Otherwise, editability is the only factor in enable state.
 			const bool bIsTypeRowEnabled = InterfaceBase ? (!InterfaceBase->bIsBranch && MemberObject->IsEditable()) : MemberObject->IsEditable();
diff --git a/Engine/Plugins/MovieScene/MovieRenderPipeline/Source/MovieRenderPipelineEditor/Private/Graph/MovieGraphSchema.cpp b/Engine/Plugins/MovieScene/MovieRenderPipeline/Source/MovieRenderPipelineEditor/Private/Graph/MovieGraphSchema.cpp
--- a/Engine/Plugins/MovieScene/MovieRenderPipeline/Source/MovieRenderPipelineEditor/Private/Graph/MovieGraphSchema.cpp
+++ b/Engine/Plugins/MovieScene/MovieRenderPipeline/Source/MovieRenderPipelineEditor/Private/Graph/MovieGraphSchema.cpp
@@ -41,6 +41,9 @@ const FName UMovieGraphSchema::PC_SoftObject(UEdGraphSchema_K2::PC_SoftObject);
 const FName UMovieGraphSchema::PC_Class(UEdGraphSchema_K2::PC_Class);
 const FName UMovieGraphSchema::PC_SoftClass(UEdGraphSchema_K2::PC_SoftClass);
 
+const FText FMovieGraphSchemaAction::UserVariablesCategory = LOCTEXT("UserVariablesCategory", "User Variables");
+const FText FMovieGraphSchemaAction::GlobalVariablesCategory = LOCTEXT("GlobalVariablesCategory", "Global Variables");
+
 namespace UE::MovieGraph::Private
 {
 	UMovieGraphNode* GetGraphNodeFromEdPin(const UEdGraphPin* InPin)
@@ -451,6 +454,50 @@ FConnectionDrawingPolicy* UMovieGraphSchema::CreateConnectionDrawingPolicy(int32
 	return new FMovieEdGraphConnectionDrawingPolicy(InBackLayerID, InFrontLayerID, InZoomFactor, InClippingRect, InDrawElements, InGraphObj);
 }
 
+void FMovieGraphSchemaAction::MovePersistentItemToCategory(const FText& NewCategoryName)
+{
+	if (const TObjectPtr<UMovieGraphVariable> TargetVariable = Cast<UMovieGraphVariable>(ActionTarget))
+	{
+		FString NewCategory = NewCategoryName.ToString();
+		
+		// If moving to the root, the category will be User Variables
+		if (NewCategory == UserVariablesCategory.ToString())
+		{
+			NewCategory = FString();
+		}
+
+		const FScopedTransaction Transaction(LOCTEXT("GraphEditor_SetVariableCategory", "Set Variable Category"));
+		
+		// Remove the "User Variables" prefix. Variables themselves do not store that part of the category.
+		const FString UserVariablesRootPrefix = FString::Format(TEXT("{0}|"), {UserVariablesCategory.ToString()});
+		NewCategory = NewCategory.StartsWith(UserVariablesRootPrefix) ? NewCategory.RightChop(UserVariablesRootPrefix.Len()) : NewCategory;
+		TargetVariable->SetCategory(NewCategory);
+	}
+}
+
+bool FMovieGraphSchemaAction::ReorderToBeforeAction(TSharedRef<FEdGraphSchemaAction> OtherAction)
+{
+	const TSharedRef<FMovieGraphSchemaAction> GraphAction = StaticCastSharedRef<FMovieGraphSchemaAction>(OtherAction);
+
+	const TObjectPtr<UMovieGraphVariable> BeforeVariable = Cast<UMovieGraphVariable>(GraphAction->ActionTarget);
+	if (!BeforeVariable)
+	{
+		return false;
+	}
+
+	const TObjectPtr<UMovieGraphVariable> TargetVariable = Cast<UMovieGraphVariable>(ActionTarget);
+	if (!TargetVariable)
+	{
+		return false;
+	}
+
+	const FScopedTransaction Transaction(LOCTEXT("GraphEditor_MoveVariable", "Move Variable"));
+
+	BeforeVariable->GetOwningGraph()->MoveVariableBefore(TargetVariable, BeforeVariable);
+
+	return true;
+}
+
 FMovieGraphSchemaAction_NewNode::FMovieGraphSchemaAction_NewNode(FText InNodeCategory, FText InDisplayName, FText InToolTip, int32 InGrouping, FText InKeywords)
 	: FMovieGraphSchemaAction(MoveTemp(InNodeCategory), MoveTemp(InDisplayName), MoveTemp(InToolTip), InGrouping, MoveTemp(InKeywords))
 {
diff --git a/Engine/Plugins/MovieScene/MovieRenderPipeline/Source/MovieRenderPipelineEditor/Private/Graph/MovieGraphSchema.h b/Engine/Plugins/MovieScene/MovieRenderPipeline/Source/MovieRenderPipelineEditor/Private/Graph/MovieGraphSchema.h
--- a/Engine/Plugins/MovieScene/MovieRenderPipeline/Source/MovieRenderPipelineEditor/Private/Graph/MovieGraphSchema.h
+++ b/Engine/Plugins/MovieScene/MovieRenderPipeline/Source/MovieRenderPipelineEditor/Private/Graph/MovieGraphSchema.h
@@ -97,6 +97,12 @@ struct FMovieGraphSchemaAction : public FEdGraphSchemaAction
 		return Type;
 	}
 
+	static const FText UserVariablesCategory;
+	static const FText GlobalVariablesCategory;
+
+	virtual void MovePersistentItemToCategory(const FText& NewCategoryName) override;
+	virtual bool ReorderToBeforeAction(TSharedRef<FEdGraphSchemaAction> OtherAction) override;
+
 	/** The object the action relates to. */
 	UPROPERTY()
 	TObjectPtr<UObject> ActionTarget = nullptr;
diff --git a/Engine/Plugins/MovieScene/MovieRenderPipeline/Source/MovieRenderPipelineEditor/Private/Graph/SMovieGraphMembersTabContent.cpp b/Engine/Plugins/MovieScene/MovieRenderPipeline/Source/MovieRenderPipelineEditor/Private/Graph/SMovieGraphMembersTabContent.cpp
--- a/Engine/Plugins/MovieScene/MovieRenderPipeline/Source/MovieRenderPipelineEditor/Private/Graph/SMovieGraphMembersTabContent.cpp
+++ b/Engine/Plugins/MovieScene/MovieRenderPipeline/Source/MovieRenderPipelineEditor/Private/Graph/SMovieGraphMembersTabContent.cpp
@@ -10,6 +10,7 @@
 #include "MovieEdGraphNode.h"
 #include "ScopedTransaction.h"
 #include "SGraphActionMenu.h"
+#include "SGraphPalette.h"
 #include "Toolkits/AssetEditorToolkit.h"
 
 #define LOCTEXT_NAMESPACE "SMoviePipelineMembersTabContent"
@@ -78,6 +79,7 @@ void SMovieGraphMembersTabContent::Construct(const FArguments& InArgs)
 		.AutoExpandActionMenu(true)
 		.AlphaSortItems(false)
 		.OnActionDragged(this, &SMovieGraphMembersTabContent::OnActionDragged)
+		.OnCategoryDragged(this, &SMovieGraphMembersTabContent::OnCategoryDragged)
 		.OnCreateWidgetForAction(this, &SMovieGraphMembersTabContent::CreateActionWidget)
 		.OnCollectStaticSections(this, &SMovieGraphMembersTabContent::CollectStaticSections)
 		.OnContextMenuOpening(this, &SMovieGraphMembersTabContent::OnContextMenuOpening)
@@ -108,6 +110,12 @@ FReply SMovieGraphMembersTabContent::OnActionDragged(const TArray<TSharedPtr<FEd
 	return FReply::Unhandled();
 }
 
+FReply SMovieGraphMembersTabContent::OnCategoryDragged(const FText& InCategory, const FPointerEvent& MouseEvent)
+{
+	const TSharedRef<FMovieGraphDragAction_Category> DragOperation = FMovieGraphDragAction_Category::New(InCategory, CurrentGraph);
+	return FReply::Handled().BeginDragDrop(DragOperation);
+}
+
 TSharedRef<SWidget> SMovieGraphMembersTabContent::CreateActionWidget(FCreateWidgetForActionData* InCreateData) const
 {
 	// For variables, show an icon w/ the color representing the variable's type (in addition to the variable name)
@@ -121,24 +129,23 @@ TSharedRef<SWidget> SMovieGraphMembersTabContent::CreateActionWidget(FCreateWidg
 			: FEdGraphPinType();
 		const FLinearColor PinColor = CurrentGraph->PipelineEdGraph->GetSchema()->GetPinTypeColor(PinType);
 		
-		return
-			SNew(SHorizontalBox)
-			+ SHorizontalBox::Slot()
+		const TSharedPtr<SHorizontalBox> RowBox = SNew(SHorizontalBox);
+		RowBox->AddSlot()
 			.AutoWidth()
 			.Padding(0, 0, 5, 0)
 			[
 				SNew(SImage)
 				.Image(FAppStyle::GetBrush("Kismet.AllClasses.VariableIcon"))
 				.ColorAndOpacity(PinColor)
-			]
+			];
 			
-			+ SHorizontalBox::Slot()
+		RowBox->AddSlot()
 			.FillWidth(1.f)
 			[
-				SNew(STextBlock)
-				.Text(InCreateData->Action->GetMenuDescription())
-				.ToolTipText(InCreateData->Action->GetTooltipDescription())
+				SNew(SGraphPaletteItem, InCreateData)
 			];
+
+		return RowBox->AsShared();
 	}
 	
 	return
@@ -221,8 +228,6 @@ void SMovieGraphMembersTabContent::PostRedo(bool bSuccess)
 
 void SMovieGraphMembersTabContent::CollectAllActions(FGraphActionListBuilderBase& OutAllActions)
 {
-	static const FText UserVariablesCategory = LOCTEXT("UserVariablesCategory", "User Variables");
-	static const FText GlobalVariablesCategory = LOCTEXT("GlobalVariablesCategory", "Global Variables");
 	static const FText EmptyCategory = FText::GetEmpty();
 	
 	if (!CurrentGraph)
@@ -293,7 +298,15 @@ void SMovieGraphMembersTabContent::CollectAllActions(FGraphActionListBuilderBase
 	{
 		if (Variable && !Variable->IsGlobal())
 		{
-			AddToActionMenu(Variable, EActionSection::Variables, UserVariablesCategory);
+			// If the user didn't specify a custom category, just use the default. If one was specified, concat the default category and the user's category
+			// so they display in a hierarchy.
+			FText VariableCategory = FMovieGraphSchemaAction::UserVariablesCategory;
+			if (!Variable->GetCategory().IsEmpty())
+			{
+				VariableCategory = FText::Format(INVTEXT("{0}|{1}"), FMovieGraphSchemaAction::UserVariablesCategory, FText::FromString(Variable->GetCategory()));
+			}
+			
+			AddToActionMenu(Variable, EActionSection::Variables, VariableCategory);
 		}
 	}
 
@@ -302,7 +315,7 @@ void SMovieGraphMembersTabContent::CollectAllActions(FGraphActionListBuilderBase
 	{
 		if (Variable && Variable->IsGlobal())
 		{
-			AddToActionMenu(Variable, EActionSection::Variables, GlobalVariablesCategory);
+			AddToActionMenu(Variable, EActionSection::Variables, FMovieGraphSchemaAction::GlobalVariablesCategory);
 		}
 	}
 	
@@ -465,9 +478,83 @@ TSharedRef<FMovieGraphDragAction_Variable> FMovieGraphDragAction_Variable::New(
 
 void FMovieGraphDragAction_Variable::HoverTargetChanged()
 {
+	const FSlateBrush* ErrorBrush = FAppStyle::GetBrush(TEXT("Graph.ConnectorFeedback.Error"));
+	const FSlateBrush* OkBrush = FAppStyle::GetBrush(TEXT("Graph.ConnectorFeedback.OK"));
+	
+	if (SourceAction.IsValid())
+	{
+		// Moving the action to a category
+		if (!HoveredCategoryName.IsEmpty())
+		{
+			FFormatNamedArguments Args;
+			Args.Add(TEXT("DisplayName"), SourceAction->GetMenuDescription());
+			Args.Add(TEXT("HoveredCategoryName"), HoveredCategoryName);
+			
+			if (HoveredCategoryName.EqualTo(SourceAction->GetCategory()))
+			{
+				const FText ErrorMessage = LOCTEXT("MoveVariableToCategory_Error", "Cannot move variable to '{HoveredCategoryName}' because it is already in that category.");
+				SetSimpleFeedbackMessage(ErrorBrush, FLinearColor::White, FText::Format(ErrorMessage, Args));
+			}
+			else
+			{
+				const FText Message = LOCTEXT("MoveVariableToCategory_OK", "Move '{DisplayName}' to category '{HoveredCategoryName}'.");
+				SetSimpleFeedbackMessage(OkBrush, FLinearColor::White, FText::Format(Message, Args));
+			}
+
+			return;
+		}
+
+		// Moving the action before another action
+		if (HoveredAction.IsValid())
+		{
+			const TSharedPtr<FEdGraphSchemaAction> HoveredActionPtr = HoveredAction.Pin();
+			FFormatNamedArguments Args;
+			Args.Add(TEXT("DraggedDisplayName"), SourceAction->GetMenuDescription());
+			Args.Add(TEXT("DropTargetDisplayName"), HoveredActionPtr->GetMenuDescription());
+
+			if (HoveredAction == SourceAction)
+			{
+				const FText ErrorMessage = LOCTEXT("MoveVariable_SameVariableError", "Cannot move variable '{DraggedDisplayName}' before itself.");
+				SetSimpleFeedbackMessage(ErrorBrush, FLinearColor::White, FText::Format(ErrorMessage, Args));
+			}
+			else
+			{
+				const FText Message = LOCTEXT("MoveVariable_OK", "Move '{DraggedDisplayName}' before '{DropTargetDisplayName}'.");
+				SetSimpleFeedbackMessage(OkBrush, FLinearColor::White, FText::Format(Message, Args));
+			}
+
+			return;
+		}
+	}
+	
 	FGraphSchemaActionDragDropAction::HoverTargetChanged();
 }
 
+FReply FMovieGraphDragAction_Variable::DroppedOnAction(TSharedRef<FEdGraphSchemaAction> Action)
+{
+	// The drop can only target another variable action
+	const TSharedRef<FMovieGraphSchemaAction> GraphAction = StaticCastSharedRef<FMovieGraphSchemaAction>(Action);
+	if (!GraphAction->ActionTarget->IsA<UMovieGraphVariable>())
+	{
+		return FReply::Unhandled();
+	}
+
+	FScopedTransaction Transaction(LOCTEXT("MoveGraphVariable", "Move Graph Variable"));
+
+	SourceAction->ReorderToBeforeAction(GraphAction);
+
+	return FReply::Handled();
+}
+
+FReply FMovieGraphDragAction_Variable::DroppedOnCategory(FText Category)
+{
+	FScopedTransaction Transaction(LOCTEXT("MoveGraphVariableToCategory", "Move Graph Variable to Category"));
+	
+	SourceAction->MovePersistentItemToCategory(Category);
+
+	return FReply::Handled();
+}
+
 FReply FMovieGraphDragAction_Variable::DroppedOnPanel(
 	const TSharedRef<SWidget>& InPanel, FVector2D InScreenPosition, FVector2D InGraphPosition, UEdGraph& InGraph)
 {
@@ -500,4 +587,97 @@ void FMovieGraphDragAction_Variable::GetDefaultStatusSymbol(
 	}
 }
 
+TSharedRef<FMovieGraphDragAction_Category> FMovieGraphDragAction_Category::New(const FText& InCategory, UMovieGraphConfig* InGraph)
+{
+	TSharedRef<FMovieGraphDragAction_Category> Operation = MakeShared<FMovieGraphDragAction_Category>();
+	Operation->Construct();
+	Operation->DraggedCategory = InCategory;
+	Operation->GraphConfig = InGraph;
+	
+	return Operation;
+}
+
+void FMovieGraphDragAction_Category::HoverTargetChanged()
+{
+	const FSlateBrush* ErrorBrush = FAppStyle::GetBrush(TEXT("Graph.ConnectorFeedback.Error"));
+	const FSlateBrush* OkBrush = FAppStyle::GetBrush(TEXT("Graph.ConnectorFeedback.OK"));
+
+	// Get the name of the category without all of the "|" separators
+	TArray<FString> CategoryHierarchy;
+	DraggedCategory.ToString().ParseIntoArray(CategoryHierarchy, TEXT("|"));
+
+	if (CategoryHierarchy.IsEmpty())
+	{
+		FGraphSchemaActionDragDropAction::HoverTargetChanged();
+		return;
+	}
+	
+	const FString DraggedCategoryName = CategoryHierarchy.Last();
+	
+	FFormatNamedArguments Args;
+	Args.Add(TEXT("CategoryName"), FText::FromString(DraggedCategoryName));
+	Args.Add(TEXT("HoveredCategoryName"), HoveredCategoryName);
+
+	// Moving the category to another category
+	if (!HoveredCategoryName.IsEmpty())
+	{
+		if (HoveredCategoryName.EqualTo(DraggedCategory))
+		{
+			const FText ErrorMessage = LOCTEXT("MoveCategoryToCategory_SelfError", "Cannot move category '{CategoryName}' before itself.");
+			SetSimpleFeedbackMessage(ErrorBrush, FLinearColor::White, FText::Format(ErrorMessage, Args));
+		}
+		else
+		{
+			const FText Message = LOCTEXT("MoveCategoryToCategory_OK", "Move category '{CategoryName}' before category '{HoveredCategoryName}'.");
+			SetSimpleFeedbackMessage(OkBrush, FLinearColor::White, FText::Format(Message, Args));
+		}
+
+		// Unfortunately the DraggedCategory is set to the display name and there's no easy way to change this. Because of that, we can't warn about
+		// dragging a child category into its immediate parent, or a parent category into a child category. That will have to be handled when the
+		// drop actually happens.
+
+		return;
+	}
+
+	// Invalid drop of a category on a non-category item
+	if (HoveredAction.IsValid())
+	{
+		const FText ErrorMessage = LOCTEXT("MoveCategoryToCategory_InvalidActionError", "Can only insert category '{CategoryName}' before another category.");
+		SetSimpleFeedbackMessage(ErrorBrush, FLinearColor::White, FText::Format(ErrorMessage, Args));
+
+		return;
+	}
+	
+	FGraphSchemaActionDragDropAction::HoverTargetChanged();
+}
+
+FReply FMovieGraphDragAction_Category::DroppedOnCategory(FText Category)
+{
+	static const FString UserVariablesPrefix = FString::Format(TEXT("{0}|"), {FMovieGraphSchemaAction::UserVariablesCategory.ToString()});;
+	
+	if (GraphConfig)
+	{
+		const FString DraggedCategoryString = DraggedCategory.ToString();
+		const FString CategoryString = Category.ToString();
+
+		// Remove the "User Variables|" prefix. "User Variables" is the top-level category the variables are displayed under; variables should not
+		// have this stored as part of their category.
+		const FString SourceCategory = DraggedCategoryString.StartsWith(UserVariablesPrefix)
+			? DraggedCategoryString.RightChop(UserVariablesPrefix.Len())
+			: DraggedCategoryString;
+		const FString DestinationCategory = CategoryString.StartsWith(UserVariablesPrefix)
+			? CategoryString.RightChop(UserVariablesPrefix.Len())
+			: CategoryString;
+
+		FScopedTransaction Transaction(LOCTEXT("MoveGraphVariableCategory", "Move Graph Variable Category"));
+		
+		// SourceAction will be invalid for a category drag/drop, so go directly to the graph
+		GraphConfig->MoveCategoryBefore(SourceCategory, DestinationCategory);
+		
+		return FReply::Handled();
+	}
+
+	return FReply::Unhandled();
+}
+
 #undef LOCTEXT_NAMESPACE
\ No newline at end of file
diff --git a/Engine/Plugins/MovieScene/MovieRenderPipeline/Source/MovieRenderPipelineEditor/Private/Graph/SMovieGraphMembersTabContent.h b/Engine/Plugins/MovieScene/MovieRenderPipeline/Source/MovieRenderPipelineEditor/Private/Graph/SMovieGraphMembersTabContent.h
--- a/Engine/Plugins/MovieScene/MovieRenderPipeline/Source/MovieRenderPipelineEditor/Private/Graph/SMovieGraphMembersTabContent.h
+++ b/Engine/Plugins/MovieScene/MovieRenderPipeline/Source/MovieRenderPipelineEditor/Private/Graph/SMovieGraphMembersTabContent.h
@@ -44,6 +44,9 @@ public:
 	/** Handler for when an action in the action menu is dragged. */
 	FReply OnActionDragged(const TArray<TSharedPtr<FEdGraphSchemaAction>>& InActions, const FPointerEvent& MouseEvent);
 
+	/** Handler for when a category in the action menu is dragged. */
+	FReply OnCategoryDragged(const FText& InCategory, const FPointerEvent& MouseEvent);
+
 	/** Creates the widget for an action within the action menu. */
 	TSharedRef<SWidget> CreateActionWidget(FCreateWidgetForActionData* CreateWidgetForActionData) const;
 
@@ -121,6 +124,8 @@ public:
 		TSharedPtr<FEdGraphSchemaAction> InAction, UMovieGraphVariable* InVariable);
 
 	virtual void HoverTargetChanged() override;
+	virtual FReply DroppedOnAction(TSharedRef<FEdGraphSchemaAction> Action) override;
+	virtual FReply DroppedOnCategory(FText Category) override;
 	virtual FReply DroppedOnPanel(
 		const TSharedRef<SWidget>& InPanel, FVector2D InScreenPosition, FVector2D InGraphPosition, UEdGraph& InGraph) override;
 
@@ -131,4 +136,23 @@ protected:
 private:
 	/** The variable member that the drag-and-drop is associated with. */
 	TWeakObjectPtr<UMovieGraphVariable> WeakVariable;
+};
+
+/* Drag-and-drop action which handles variable categories. */
+class FMovieGraphDragAction_Category : public FGraphSchemaActionDragDropAction
+{
+public:
+	DRAG_DROP_OPERATOR_TYPE(FMovieGraphDragAction_Category, FGraphSchemaActionDragDropAction)
+
+	static TSharedRef<FMovieGraphDragAction_Category> New(const FText& InCategory, UMovieGraphConfig* InGraph);
+
+	virtual void HoverTargetChanged() override;
+	virtual FReply DroppedOnCategory(FText Category) override;
+
+public:
+	/** The category that is being dragged. */
+	FText DraggedCategory;
+
+	/** The graph associated with the drag/drop action. */
+	UMovieGraphConfig* GraphConfig = nullptr;
 };
\ No newline at end of file
