Removing additional comma in the comment section for a python code  (no code change)

diff --git a/Engine/Plugins/MovieScene/MovieRenderPipeline/Source/MovieRenderPipelineCore/Public/Graph/Nodes/MovieGraphExecuteScriptNode.h b/Engine/Plugins/MovieScene/MovieRenderPipeline/Source/MovieRenderPipelineCore/Public/Graph/Nodes/MovieGraphExecuteScriptNode.h
--- a/Engine/Plugins/MovieScene/MovieRenderPipeline/Source/MovieRenderPipelineCore/Public/Graph/Nodes/MovieGraphExecuteScriptNode.h
+++ b/Engine/Plugins/MovieScene/MovieRenderPipeline/Source/MovieRenderPipelineCore/Public/Graph/Nodes/MovieGraphExecuteScriptNode.h
@@ -165,7 +165,7 @@ public:
 	* For Python implementations override
 	*	@unreal.ufunction(override=True)
 	*	def is_per_shot_callback_needed():
-	*		return False;
+	*		return False
 	*/
 	UFUNCTION(BlueprintNativeEvent, Category = "Movie Graph")
 	bool IsPerShotCallbackNeeded() const;
