MovieRenderQueue: Only use frame throttle for path traced renders.

The fence mechanism was causing deadlocks on DX11 since they don't have real fences. Since the frame throttling is mainly intended to prevent memory growth with the path tracer, limit it to that case only for now.

#rnx

diff --git a/Engine/Plugins/MovieScene/MovieRenderPipeline/Source/MovieRenderPipelineCore/Private/MoviePipelineRendering.cpp b/Engine/Plugins/MovieScene/MovieRenderPipeline/Source/MovieRenderPipelineCore/Private/MoviePipelineRendering.cpp
--- a/Engine/Plugins/MovieScene/MovieRenderPipeline/Source/MovieRenderPipelineCore/Private/MoviePipelineRendering.cpp
+++ b/Engine/Plugins/MovieScene/MovieRenderPipeline/Source/MovieRenderPipelineCore/Private/MoviePipelineRendering.cpp
@@ -49,7 +49,8 @@ static TAutoConsoleVariable<bool> CVarMoviePipelineDisableShaderFlushing(
 static TAutoConsoleVariable<int> CVarMoviePipelineThrottleFrameCount(
 	TEXT("MoviePipeline.ThrottleFrameCount"), 2,
 	TEXT("Number of rendered frames that can be submitted to the rendering thread before waiting. A value of 0 will allow the CPU to submit all work without waiting on the GPU.\n")
-	TEXT("The default value of 2 tries to balance between performance and memory usage. The maximum value is 4."),
+	TEXT("The default value of 2 tries to balance between performance and memory usage. The maximum value is 4.\n")
+	TEXT("This option only applies to path traced renders, as deferred rendering is synchronized through pixel readbacks.\n"),
 	ECVF_Default);
 
 #define LOCTEXT_NAMESPACE "MoviePipeline"
@@ -297,7 +298,15 @@ void UMoviePipeline::RenderFrame()
 #endif
 
 	constexpr int FenceBufferMax = 4;
-	const int FrameThrotteCount = FMath::Clamp(CVarMoviePipelineThrottleFrameCount.GetValueOnGameThread(), 0, FenceBufferMax);
+	int FrameThrotteCount = 0;
+	for (const UMoviePipelineRenderPass* RenderPass : InputBuffers)
+	{
+		if (RenderPass->NeedsFrameThrottle())
+		{
+			FrameThrotteCount = FMath::Clamp(CVarMoviePipelineThrottleFrameCount.GetValueOnGameThread(), 0, FenceBufferMax);
+			break;
+		}
+	}
 	FGPUFenceRHIRef MRQThrottleFence[FenceBufferMax];
 	int FenceIndex = 0;
 	
diff --git a/Engine/Plugins/MovieScene/MovieRenderPipeline/Source/MovieRenderPipelineCore/Public/MoviePipelineRenderPass.h b/Engine/Plugins/MovieScene/MovieRenderPipeline/Source/MovieRenderPipelineCore/Public/MoviePipelineRenderPass.h
--- a/Engine/Plugins/MovieScene/MovieRenderPipeline/Source/MovieRenderPipelineCore/Public/MoviePipelineRenderPass.h
+++ b/Engine/Plugins/MovieScene/MovieRenderPipeline/Source/MovieRenderPipelineCore/Public/MoviePipelineRenderPass.h
@@ -44,6 +44,8 @@ public:
 		return IsAlphaInTonemapperRequiredImpl();
 	}
 
+	virtual bool NeedsFrameThrottle() const { return false; }
+
 
 protected:
 	virtual bool IsValidOnShots() const override { return true; }
diff --git a/Engine/Plugins/MovieScene/MovieRenderPipeline/Source/MovieRenderPipelineRenderPasses/Public/MoviePipelineDeferredPasses.h b/Engine/Plugins/MovieScene/MovieRenderPipeline/Source/MovieRenderPipelineRenderPasses/Public/MoviePipelineDeferredPasses.h
--- a/Engine/Plugins/MovieScene/MovieRenderPipeline/Source/MovieRenderPipelineRenderPasses/Public/MoviePipelineDeferredPasses.h
+++ b/Engine/Plugins/MovieScene/MovieRenderPipeline/Source/MovieRenderPipelineRenderPasses/Public/MoviePipelineDeferredPasses.h
@@ -320,6 +320,8 @@ public:
 	virtual void ValidateStateImpl() override;
 	virtual void SetupImpl(const MoviePipeline::FMoviePipelineRenderPassInitSettings& InPassInitSettings) override;
 
+	virtual bool NeedsFrameThrottle() const override { return true; }
+
 	/** When enabled, the path tracer will blend all spatial and temporal samples prior to the denoising and will disable post-processed motion blur.
 	 *  In this mode it is possible to use higher temporal sample counts to improve the motion blur quality.
 	 *  When this option is disabled, the path tracer will accumulate spatial samples, but denoise them prior to accumulation of temporal samples.
