MovieGraph: Fixed {camera_name} token always saying "Unsupported" instead of the correct camera name from the sequence.

diff --git a/Engine/Plugins/MovieScene/MovieRenderPipeline/Source/MovieRenderPipelineCore/Private/Graph/MovieGraphDefaultRenderer.cpp b/Engine/Plugins/MovieScene/MovieRenderPipeline/Source/MovieRenderPipelineCore/Private/Graph/MovieGraphDefaultRenderer.cpp
--- a/Engine/Plugins/MovieScene/MovieRenderPipeline/Source/MovieRenderPipelineCore/Private/Graph/MovieGraphDefaultRenderer.cpp
+++ b/Engine/Plugins/MovieScene/MovieRenderPipeline/Source/MovieRenderPipelineCore/Private/Graph/MovieGraphDefaultRenderer.cpp
@@ -14,6 +14,7 @@
 #include "UObject/Package.h"
 #include "MoviePipelineSurfaceReader.h"
 #include "RenderCaptureInterface.h"
+#include "MoviePipelineQueue.h"
 
 // For flushing async systems
 #include "EngineModule.h"
@@ -393,7 +394,13 @@ UE::MovieGraph::DefaultRenderer::FCameraInfo UMovieGraphDefaultRenderer::GetCame
 	{
 		CameraInfo.ViewInfo = LocalPlayerController->PlayerCameraManager->GetCameraCacheView();
 		CameraInfo.ViewActor = LocalPlayerController->GetViewTarget();
-		CameraInfo.CameraName = TEXT("Unsupported"); // ToDo: This eventually needs to come from Level Sequences
+		
+		UMoviePipelineExecutorShot* CurrentShot = GetOwningGraph()->GetActiveShotList()[GetOwningGraph()->GetCurrentShotIndex()];
+
+		// INDEX_NONE returns the primary camera from the rendered shot until we have
+		// multi-camera support.
+		const int32 CameraIndex = INDEX_NONE;
+		CameraInfo.CameraName = CurrentShot->GetCameraName(CameraIndex);
 	}
 	else
 	{
