Add reference motion blur parameter to Pathtracer node.
#tests Behaves as expected with/without reference motion blur parameter when using path tracer node in MRG.
[RN] Add reference motion blur parameter to Pathtracer node.

diff --git a/Engine/Plugins/MovieScene/MovieRenderPipeline/Source/MovieRenderPipelineRenderPasses/Private/Graph/Nodes/MovieGraphPathTracerPassNode.cpp b/Engine/Plugins/MovieScene/MovieRenderPipeline/Source/MovieRenderPipelineRenderPasses/Private/Graph/Nodes/MovieGraphPathTracerPassNode.cpp
--- a/Engine/Plugins/MovieScene/MovieRenderPipeline/Source/MovieRenderPipelineRenderPasses/Private/Graph/Nodes/MovieGraphPathTracerPassNode.cpp
+++ b/Engine/Plugins/MovieScene/MovieRenderPipeline/Source/MovieRenderPipelineRenderPasses/Private/Graph/Nodes/MovieGraphPathTracerPassNode.cpp
@@ -13,6 +13,7 @@ TUniquePtr<UE::MovieGraph::Rendering::FMovieGraphImagePassBase> UMovieGraphPathT
 
 UMovieGraphPathTracerRenderPassNode::UMovieGraphPathTracerRenderPassNode()
 	: SpatialSampleCount(1)
+	, bEnableReferenceMotionBlur(false)
 	, bDenoiser(true)
 	, bWriteAllSamples(false)
 	, bDisableToneCurve(false)
@@ -140,3 +141,11 @@ bool UMovieGraphPathTracerRenderPassNode::GetAllowDenoiser() const
 	return bDenoiser;
 }
 
+FEngineShowFlags UMovieGraphPathTracerRenderPassNode::GetShowFlags() const
+{
+	FEngineShowFlags OutShowFlag = FEngineShowFlags(EShowFlagInitMode::ESFIM_Game);
+	OutShowFlag.SetPathTracing(true);
+	OutShowFlag.SetMotionBlur(!bEnableReferenceMotionBlur);
+
+	return OutShowFlag;
+}
\ No newline at end of file
diff --git a/Engine/Plugins/MovieScene/MovieRenderPipeline/Source/MovieRenderPipelineRenderPasses/Private/Graph/Renderers/MovieGraphPathTracerPass.cpp b/Engine/Plugins/MovieScene/MovieRenderPipeline/Source/MovieRenderPipelineRenderPasses/Private/Graph/Renderers/MovieGraphPathTracerPass.cpp
--- a/Engine/Plugins/MovieScene/MovieRenderPipeline/Source/MovieRenderPipelineRenderPasses/Private/Graph/Renderers/MovieGraphPathTracerPass.cpp
+++ b/Engine/Plugins/MovieScene/MovieRenderPipeline/Source/MovieRenderPipelineRenderPasses/Private/Graph/Renderers/MovieGraphPathTracerPass.cpp
@@ -33,7 +33,7 @@ namespace Private
 		//    can help reduce strobing
 
 		// NOTE: Tiling is not compatible with the reference motion blur mode because it changes the order of the loops over the image.
-		const bool bHasTiles = InCameraInfo.TilingParams.TileCount.X > 1 || InCameraInfo.TilingParams.TileCount.Y;
+		const bool bHasTiles = InCameraInfo.TilingParams.TileCount.X * InCameraInfo.TilingParams.TileCount.Y > 1;
 		const bool bAccumulateSpatialSamplesOnly = InOutFamily->EngineShowFlags.MotionBlur || bHasTiles;
 
 		OutSampleCount = bAccumulateSpatialSamplesOnly ? InCameraInfo.SamplingParams.SpatialSampleCount : InCameraInfo.SamplingParams.TemporalSampleCount * InCameraInfo.SamplingParams.SpatialSampleCount;
diff --git a/Engine/Plugins/MovieScene/MovieRenderPipeline/Source/MovieRenderPipelineRenderPasses/Public/Graph/Nodes/MovieGraphPathTracerPassNode.h b/Engine/Plugins/MovieScene/MovieRenderPipeline/Source/MovieRenderPipelineRenderPasses/Public/Graph/Nodes/MovieGraphPathTracerPassNode.h
--- a/Engine/Plugins/MovieScene/MovieRenderPipeline/Source/MovieRenderPipelineRenderPasses/Public/Graph/Nodes/MovieGraphPathTracerPassNode.h
+++ b/Engine/Plugins/MovieScene/MovieRenderPipeline/Source/MovieRenderPipelineRenderPasses/Public/Graph/Nodes/MovieGraphPathTracerPassNode.h
@@ -34,6 +34,7 @@ public:
 	virtual bool GetAllowOCIO() const override;
 	virtual bool GetAllowDenoiser() const override;
 	virtual TUniquePtr<UE::MovieGraph::Rendering::FMovieGraphImagePassBase> CreateInstance() const;
+	virtual FEngineShowFlags GetShowFlags() const override;
 	// ~UMovieGraphImagePassBaseNode Interface
 
 protected:
@@ -51,6 +52,9 @@ public:
 	UPROPERTY(EditAnywhere, BlueprintReadWrite, Category = Overrides, meta = (InlineEditConditionToggle))
 	uint8 bOverride_SpatialSampleCount : 1;
 
+	UPROPERTY(EditAnywhere, BlueprintReadWrite, Category = Overrides, meta = (InlineEditConditionToggle))
+	uint8 bOverride_bEnableReferenceMotionBlur : 1;
+
 	UPROPERTY(EditAnywhere, BlueprintReadWrite, Category = Overrides, meta = (InlineEditConditionToggle))
 	uint8 bOverride_bDenoiser : 1;
 
@@ -75,6 +79,14 @@ public:
 	UPROPERTY(EditAnywhere, BlueprintReadWrite, meta = (UIMin = 1, ClampMin = 1), Category = "Sampling", meta = (EditCondition = "bOverride_SpatialSampleCount"))
 	int32 SpatialSampleCount;
 
+	/** 
+	 *  When enabled, the path tracer will blend all spatial and temporal samples prior to the denoising and will disable post-processed motion blur.
+	 *  In this mode it is possible to use higher temporal sample counts to improve the motion blur quality.
+	 *  When this option is disabled, the path tracer will accumulate spatial samples, but denoise them prior to accumulation of temporal samples.
+	 */
+	UPROPERTY(EditAnywhere, BlueprintReadWrite, meta = (UIMin = 1, ClampMin = 1), Category = "Reference Motion Blur", meta = (EditCondition = "bOverride_bEnableReferenceMotionBlur"))
+	bool bEnableReferenceMotionBlur;
+
 	/** If true the resulting image will be denoised at the end of each set of Spatial Samples. */
 	UPROPERTY(EditAnywhere, BlueprintReadWrite, meta = (UIMin = 1, ClampMin = 1), Category = "Sampling", meta = (EditCondition = "bOverride_bDenoiser"))
 	bool bDenoiser;
