MoviePipeline: Fixed CIS error.

#rnx

diff --git a/Engine/Plugins/MovieScene/MovieRenderPipeline/Source/MovieRenderPipelineCore/Private/Graph/Nodes/MovieGraphVideoOutputNode.cpp b/Engine/Plugins/MovieScene/MovieRenderPipeline/Source/MovieRenderPipelineCore/Private/Graph/Nodes/MovieGraphVideoOutputNode.cpp
--- a/Engine/Plugins/MovieScene/MovieRenderPipeline/Source/MovieRenderPipelineCore/Private/Graph/Nodes/MovieGraphVideoOutputNode.cpp
+++ b/Engine/Plugins/MovieScene/MovieRenderPipeline/Source/MovieRenderPipelineCore/Private/Graph/Nodes/MovieGraphVideoOutputNode.cpp
@@ -148,6 +148,13 @@ void UMovieGraphVideoOutputNode::OnAllFramesFinalizedImpl(UMovieGraphPipeline* I
 	bHasError = false;
 }
 
+UMovieGraphVideoOutputNode::FMovieGraphCodecWriterWithPromise::FMovieGraphCodecWriterWithPromise(TUniquePtr<MovieRenderGraph::IVideoCodecWriter>&& InWriter, TPromise<bool>&& InPromise)
+	: CodecWriter(MoveTemp(InWriter))
+	, Promise(MoveTemp(InPromise))
+{
+	
+}
+
 void UMovieGraphVideoOutputNode::GetOutputFilePaths(const UMovieGraphPipeline* InPipeline, const UE::MovieGraph::FMovieGraphOutputMergerFrame* InRawFrameData, FMovieGraphPassData& InRenderPassData, const TArray<FMovieGraphPassData>& InCompositedPasses, FString& OutFinalFilePath, FString& OutStableFilePath)
 {
 	const TObjectPtr<UMoviePipelineExecutorShot> Shot = InRawFrameData->TraversalContext.Shot;
diff --git a/Engine/Plugins/MovieScene/MovieRenderPipeline/Source/MovieRenderPipelineCore/Public/Graph/Nodes/MovieGraphVideoOutputNode.h b/Engine/Plugins/MovieScene/MovieRenderPipeline/Source/MovieRenderPipelineCore/Public/Graph/Nodes/MovieGraphVideoOutputNode.h
--- a/Engine/Plugins/MovieScene/MovieRenderPipeline/Source/MovieRenderPipelineCore/Public/Graph/Nodes/MovieGraphVideoOutputNode.h
+++ b/Engine/Plugins/MovieScene/MovieRenderPipeline/Source/MovieRenderPipelineCore/Public/Graph/Nodes/MovieGraphVideoOutputNode.h
@@ -60,6 +60,8 @@ protected:
 private:
 	struct FMovieGraphCodecWriterWithPromise
 	{
+		FMovieGraphCodecWriterWithPromise(TUniquePtr<MovieRenderGraph::IVideoCodecWriter>&& InWriter, TPromise<bool>&& InPromise);
+		
 		/** The codec writer. */
 		TUniquePtr<MovieRenderGraph::IVideoCodecWriter> CodecWriter;
 
