MovieGraph: Fixed an issue where Filename Tokens and Metadata were only being gathered from the first instance of a given settings class. This is correct in most cases (since the flattened graph only contains one copy of the setting) but in the case of "Named" nodes, the flattened graph can contain multiple copies which have unique data.

diff --git a/Engine/Plugins/MovieScene/MovieRenderPipeline/Source/MovieRenderPipelineCore/Private/Graph/MovieGraphBlueprintLibrary.cpp b/Engine/Plugins/MovieScene/MovieRenderPipeline/Source/MovieRenderPipelineCore/Private/Graph/MovieGraphBlueprintLibrary.cpp
--- a/Engine/Plugins/MovieScene/MovieRenderPipeline/Source/MovieRenderPipelineCore/Private/Graph/MovieGraphBlueprintLibrary.cpp
+++ b/Engine/Plugins/MovieScene/MovieRenderPipeline/Source/MovieRenderPipelineCore/Private/Graph/MovieGraphBlueprintLibrary.cpp
@@ -115,29 +115,31 @@ FString UMovieGraphBlueprintLibrary::ResolveFilenameFormatArguments(const FStrin
 	
 
 	//  Now get the settings from our config. We need to gather KVP data from all possible nodes, even if not expressed in your configuration. This is because you might want to
-	// always use the {ts_count} token even if you don't have a Temporal Sample Count node to add it. So we loop through all the possible class types, and call a function on the 
-	// CDO, but then we pass that class type from the evaluated config (if it exists), and we pass the CDO as an argument if it doesn't.
+	// always use the {ts_count} token even if you don't have a Temporal Sample Count node to add it. We look at all possible class types and then get all instances of them in
+	// the evaluated graph. We look at all instances because "Named" nodes can end up having multiple copies in the final graph which may want to provide different data.
 	TArray<UClass*> AllSettingsNodeClasses = UE::MovieRenderPipeline::FindMoviePipelineSettingClasses(UMovieGraphSettingNode::StaticClass(), false);
 
 	// ToDo: This loops through class iterators every frame, we should probably initialize a copy of everything into the flattened config, since we could cache the classes
 	// once per run there. We don't cache the returned results here because you could potentially add/remove classes (via Blueprints) which would invalidate our cache.
 	for (UClass* InClass : AllSettingsNodeClasses)
 	{
-		const UMovieGraphSettingNode* SettingInstance = nullptr;
+		TArray<UMovieGraphSettingNode*> SettingInstances;
 		if (InParams.EvaluatedConfig)
 		{
 			const bool bIncludeCDOs = true;
 			const bool bExactMatch = true;
-			SettingInstance = InParams.EvaluatedConfig->GetSettingForBranch(InClass, InParams.RenderDataIdentifier.RootBranchName, bIncludeCDOs, bExactMatch);
+			SettingInstances = InParams.EvaluatedConfig->GetSettingsForBranch(InClass, InParams.RenderDataIdentifier.RootBranchName, bIncludeCDOs, bExactMatch);
 		}
 		else
 		{
-			SettingInstance = GetDefault<UMovieGraphSettingNode>(InClass);
+			// GetFormatResolveArgs is const, but GetSettingsForBranch doesn't return a const pointer array,
+			// so we just get a non-const CDO pointer here instead, knowing that the function being called is const.
+			SettingInstances.Add(GetMutableDefault<UMovieGraphSettingNode>(InClass));
 		}
 
-		if (SettingInstance)
+		for(UMovieGraphSettingNode* Instance : SettingInstances)
 		{
-			SettingInstance->GetFormatResolveArgs(OutMergedFormatArgs, InParams.RenderDataIdentifier);
+			Instance->GetFormatResolveArgs(OutMergedFormatArgs, InParams.RenderDataIdentifier);
 		}
 	}
 
