Integrate NFOR denoiser into Movie Render Graph. Inside Path Traced Render node,
1. Denoiser type 'Spatial' and 'Temporal' setting is added. Default to Temporal. If 'Temporal' denosier is selected but no temporal denosier plugin is enabled. The denoising will fallback to the first available spatial denoiser.
2.  Frame Count is added and is editable only when `Temporal` denoiser is set. It enables the user to use n future and n past frames for denoising. Default to 2 with a range of 0 to 3.

#tests Run with different frame count and spatial temporal config correctly with and without plugin enabled.

diff --git a/Engine/Plugins/MovieScene/MovieRenderPipeline/Source/MovieRenderPipelineCore/Private/Graph/MovieGraphCoreTimeStep.cpp b/Engine/Plugins/MovieScene/MovieRenderPipeline/Source/MovieRenderPipelineCore/Private/Graph/MovieGraphCoreTimeStep.cpp
--- a/Engine/Plugins/MovieScene/MovieRenderPipeline/Source/MovieRenderPipelineCore/Private/Graph/MovieGraphCoreTimeStep.cpp
+++ b/Engine/Plugins/MovieScene/MovieRenderPipeline/Source/MovieRenderPipelineCore/Private/Graph/MovieGraphCoreTimeStep.cpp
@@ -171,7 +171,7 @@ void UMovieGraphCoreTimeStep::TickProducingFrames()
 
 	// We only automatically advance the CurrentOutputFrameRange during Rendering, or when warm-up frames
 	// are counting down (if we're not going to emulate motion blur).
-	const bool bIncrementBecauseRenderingState = CurrentCameraCut->ShotInfo.State == EMovieRenderShotState::Rendering;
+	const bool bIncrementBecauseRenderingState = CurrentCameraCut->ShotInfo.State == EMovieRenderShotState::Rendering || CurrentCameraCut->ShotInfo.State == EMovieRenderShotState::CoolingDown;
 	const bool bIncrementBecauseWarmUpState = CurrentCameraCut->ShotInfo.State == EMovieRenderShotState::WarmingUp && !CurrentCameraCut->ShotInfo.bEmulateFirstFrameMotionBlur;
 	bool bIncrementInternalCounters = bIncrementBecauseRenderingState || bIncrementBecauseWarmUpState;
 	
@@ -312,6 +312,21 @@ void UMovieGraphCoreTimeStep::TickProducingFrames()
 			// Now that we've calculated the total range of time we're trying to represent, we can check to see
 			// if this would put us beyond our range of time this shot is supposed to represent.
 			if (CurrentFrameData.CurrentOutputFrameRange.GetUpperBoundValue() > CurrentCameraCut->ShotInfo.TotalOutputRangeRoot.GetUpperBoundValue())
+			{
+				// We've reached the end of the sequence. We're going 
+				// to transition into the Cooling Down state (which we may 
+				// immediately leave if there are no cooling down frames 
+				// to actually render).
+				CurrentCameraCut->ShotInfo.State = EMovieRenderShotState::CoolingDown;
+			}
+		}
+
+		// If we've run past the end of the sequence, we'll be in the CoolingDown state now.
+		if (CurrentCameraCut->ShotInfo.State == EMovieRenderShotState::CoolingDown)
+		{
+			// We may or may not actually have any CoolingDown frames to process. If there's no frames left,
+			// then we'll move to the Finished state and tear-down the render process for this frame.
+			if (CurrentCameraCut->ShotInfo.NumEngineCoolDownFramesRemaining <= 0)
 			{
 				// We're going to spend this frame tearing down the shot (not rendering anything), next frame
 				// we'll re-enter this loop and pick up the start of the next shot.
@@ -325,6 +340,36 @@ void UMovieGraphCoreTimeStep::TickProducingFrames()
 				GetOwningGraph()->TeardownShot(CurrentCameraCut);
 				return;
 			}
+			else
+			{
+				// Only jump backwards in time if there's actually a cooldown frames to run.
+				// Cooling down repeats the last frame of the render entirely, spatial and temporal samples included. This is because 
+				// we want to generate good looking images to send to the denoiser (to denoise the normal final frames), so we need to
+				// generate them like they would have been in the normal render.
+				//CurrentFrameData.CurrentOutputFrameRange = CurrentFrameData.LastOutputFrameRange;
+				//CurrentFrameData.LastSampleRange = CurrentFrameData.LastOutputFrameRange;
+
+				// Go back FrameCount output frame for the current output frame range, and FrameCount+1 output frame for the last output frame.
+				FFrameTime UpperBound = EndOfPreviousFrame.GetValue() - CurrentFrameMetrics.FrameTimePerOutputFrame;
+				FFrameTime NewStartTime = UpperBound;
+				CurrentFrameData.LastOutputFrameRange = TRange<FFrameTime>(UpperBound - CurrentFrameMetrics.FrameTimePerOutputFrame, UpperBound);
+				CurrentFrameData.LastSampleRange = CurrentFrameData.LastOutputFrameRange;
+				CurrentFrameData.CurrentOutputFrameRange = TRange<FFrameTime>(NewStartTime, NewStartTime + CurrentFrameMetrics.FrameTimePerOutputFrame);
+
+				// Update CurrentFrameData.RangeShutterOpen and CurrentFrameData.RangeShutterClosed
+				UpdateShutterRanges();
+				// Update CurrentFrameData.TemporalRanges
+				UpdateTemporalRanges();
+
+				// We're going backwards so we need to flag the jump
+				bShouldJump = true;
+
+				// Cooldown frames 
+				// ToDo: We need to separate the idea of "should we send this to the output merger" (false during cooldown)
+				// and "should we render and schedule a readback" (true during cooldown)
+				CurrentTimeStepData.bDiscardOutput = false;
+				CurrentCameraCut->ShotInfo.NumEngineCoolDownFramesRemaining--;
+			}
 		}
 	}
 
diff --git a/Engine/Plugins/MovieScene/MovieRenderPipeline/Source/MovieRenderPipelineCore/Private/Graph/MovieGraphPipeline.cpp b/Engine/Plugins/MovieScene/MovieRenderPipeline/Source/MovieRenderPipelineCore/Private/Graph/MovieGraphPipeline.cpp
--- a/Engine/Plugins/MovieScene/MovieRenderPipeline/Source/MovieRenderPipelineCore/Private/Graph/MovieGraphPipeline.cpp
+++ b/Engine/Plugins/MovieScene/MovieRenderPipeline/Source/MovieRenderPipelineCore/Private/Graph/MovieGraphPipeline.cpp
@@ -23,6 +23,7 @@
 #include "Graph/Nodes/MovieGraphSamplingMethodNode.h"
 #include "Graph/Nodes/MovieGraphSubgraphNode.h"
 #include "Graph/Nodes/MovieGraphWarmUpSettingNode.h"
+#include "Graph/Nodes/MovieGraphRenderPassNode.h"
 
 #include "HAL/PlatformFileManager.h"
 #include "ImageWriteQueue.h"
@@ -567,7 +568,23 @@ void UMovieGraphPipeline::BuildShotListFromDataSource()
 		ExpandShot(Shot, OutputNode->HandleFrameCount, bExpandForTemporalSubSample, bPrePass, FinalFrameRate, TickResolution, WarmUpNode->NumWarmUpFrames);
 
 		Shot->ShotInfo.CurrentTimeInRoot = Shot->ShotInfo.TotalOutputRangeRoot.GetLowerBoundValue();
-		Shot->ShotInfo.NumEngineWarmUpFramesRemaining = WarmUpNode->NumWarmUpFrames;
+
+		// Query the max frame count.
+		int32 MaxCoolingDownFrameCount = 0;
+		for (const FName& BranchName : EvaluatedConfig->GetBranchNames())
+		{
+			const bool bIncludeCDOs = false;
+			const bool bExactMatch = false;
+			TArray<UMovieGraphRenderPassNode*> Renderers = EvaluatedConfig->GetSettingsForBranch<UMovieGraphRenderPassNode>(BranchName, bIncludeCDOs, bExactMatch);
+			for (const UMovieGraphRenderPassNode* Render : Renderers)
+			{
+				MaxCoolingDownFrameCount = FMath::Max(Render->GetCoolingDownFrameCount(), MaxCoolingDownFrameCount);
+			}
+		}
+		
+		Shot->ShotInfo.NumEngineCoolDownFramesRemaining = MaxCoolingDownFrameCount;
+		// When using cooldown, we need at least that many warm-up frames even if they have otherwise chose not to do warm-ups.
+		Shot->ShotInfo.NumEngineWarmUpFramesRemaining = FMath::Max(WarmUpNode->NumWarmUpFrames, MaxCoolingDownFrameCount);
 		Shot->ShotInfo.bEmulateFirstFrameMotionBlur = WarmUpNode->bEmulateMotionBlur;
 		Shot->ShotInfo.CalculateWorkMetrics();
 		Shot->ShotInfo.VersionNumber = ResolveVersionForShot(Shot, EvaluatedConfig);
diff --git a/Engine/Plugins/MovieScene/MovieRenderPipeline/Source/MovieRenderPipelineCore/Public/Graph/Nodes/MovieGraphRenderPassNode.h b/Engine/Plugins/MovieScene/MovieRenderPipeline/Source/MovieRenderPipelineCore/Public/Graph/Nodes/MovieGraphRenderPassNode.h
--- a/Engine/Plugins/MovieScene/MovieRenderPipeline/Source/MovieRenderPipelineCore/Public/Graph/Nodes/MovieGraphRenderPassNode.h
+++ b/Engine/Plugins/MovieScene/MovieRenderPipeline/Source/MovieRenderPipelineCore/Public/Graph/Nodes/MovieGraphRenderPassNode.h
@@ -52,6 +52,11 @@ public:
 	* or more than one for things like panoramic or tiling. */
 	int32 GetNumSceneViewsRendered() const { return GetNumSceneViewsRenderedImpl(); }
 
+	/**
+	* Get the cooling down frame count when denoising based on temporal frames are used.
+	**/
+	virtual int32 GetCoolingDownFrameCount() const { return 0; }
+
 #if WITH_EDITOR
 	virtual FText GetMenuCategory() const override
 	{
diff --git a/Engine/Plugins/MovieScene/MovieRenderPipeline/Source/MovieRenderPipelineCore/Public/MovieRenderPipelineDataTypes.h b/Engine/Plugins/MovieScene/MovieRenderPipeline/Source/MovieRenderPipelineCore/Public/MovieRenderPipelineDataTypes.h
--- a/Engine/Plugins/MovieScene/MovieRenderPipeline/Source/MovieRenderPipelineCore/Public/MovieRenderPipelineDataTypes.h
+++ b/Engine/Plugins/MovieScene/MovieRenderPipeline/Source/MovieRenderPipelineCore/Public/MovieRenderPipelineDataTypes.h
@@ -83,11 +83,17 @@ enum class EMovieRenderShotState : uint8
 	* a whole frame.
 	*/
 	Rendering = 3,
+	/** 
+	* The shot is cooling down. Engine ticks are passing and frames are being rendered, but
+	* not saved to disk. This is needed because temporal-based denoisers need to look at the
+	* future (which would be after a shot) to finish the current frames of the shot.
+	*/
+	CoolingDown = 4,
 	/*
 	* The shot has produced all frames it will produce. No more evaluation should be
 	* done for this shot once it reaches this state.
 	*/
-	Finished = 4
+	Finished = 5
 };
 
 USTRUCT(BlueprintType)
@@ -485,6 +491,7 @@ public:
 		, State(EMovieRenderShotState::Uninitialized)
 		, bHasEvaluatedMotionBlurFrame(false)
 		, NumEngineWarmUpFramesRemaining(0)
+		, NumEngineCoolDownFramesRemaining(0)
 		, VersionNumber(0)
 	{
 	}
@@ -552,6 +559,10 @@ public:
 	/** How many engine warm up frames are left to process for this shot. May be zero. */
 	int32 NumEngineWarmUpFramesRemaining;
 
+	/** How many cool down frames for this shot. May be zero.*/
+	int32 NumEngineCoolDownFramesRemaining;
+
+
 	/** What version number should this shot use when resolving format arguments. This is the highest version number found across all branches. */
 	int32 VersionNumber;
 };
diff --git a/Engine/Plugins/MovieScene/MovieRenderPipeline/Source/MovieRenderPipelineRenderPasses/Private/Graph/Nodes/MovieGraphPathTracerPassNode.cpp b/Engine/Plugins/MovieScene/MovieRenderPipeline/Source/MovieRenderPipelineRenderPasses/Private/Graph/Nodes/MovieGraphPathTracerPassNode.cpp
--- a/Engine/Plugins/MovieScene/MovieRenderPipeline/Source/MovieRenderPipelineRenderPasses/Private/Graph/Nodes/MovieGraphPathTracerPassNode.cpp
+++ b/Engine/Plugins/MovieScene/MovieRenderPipeline/Source/MovieRenderPipelineRenderPasses/Private/Graph/Nodes/MovieGraphPathTracerPassNode.cpp
@@ -5,6 +5,7 @@
 #include "Engine/EngineBaseTypes.h"
 #include "RenderUtils.h"
 #include "ShowFlags.h"
+#include "PathTracingDenoiser.h"
 
 TUniquePtr<UE::MovieGraph::Rendering::FMovieGraphImagePassBase> UMovieGraphPathTracerRenderPassNode::CreateInstance() const
 {
@@ -14,7 +15,9 @@ TUniquePtr<UE::MovieGraph::Rendering::FMovieGraphImagePassBase> UMovieGraphPathT
 UMovieGraphPathTracerRenderPassNode::UMovieGraphPathTracerRenderPassNode()
 	: SpatialSampleCount(1)
 	, bEnableReferenceMotionBlur(false)
-	, bDenoiser(true)
+	, bEnableDenoiser(true)
+	, DenoiserType(EMovieGraphPathTracerDenoiserType::Temporal)
+	, FrameCount(2)
 	, bWriteAllSamples(false)
 	, bDisableToneCurve(false)
 	, bAllowOCIO(true)
@@ -62,6 +65,63 @@ void UMovieGraphPathTracerRenderPassNode::SetupImpl(const FMovieGraphRenderPassS
 		ProgressDisplayCvar->Set(false);
 	}
 
+	// Different Path Traced renders can't currently use different NFOR settings, so we take the largest setting from all branches.
+	int32 MaxFrameCount = 0;
+	int32 MaxDenoiserType = 0;
+	FString WarningLayerName = TEXT("None");
+
+	for (const FMovieGraphRenderPassLayerData& LayerData : InSetupData.Layers)
+	{
+		UMovieGraphRenderPassNode* RenderPassNode = LayerData.RenderPassNode.Get();
+		UMovieGraphPathTracerRenderPassNode* PathTracerNode = CastChecked<UMovieGraphPathTracerRenderPassNode>(RenderPassNode);
+		if (PathTracerNode)
+		{
+			MaxFrameCount = FMath::Max(MaxFrameCount, PathTracerNode->GetCoolingDownFrameCount());
+			MaxDenoiserType = FMath::Max(MaxDenoiserType, static_cast<int32>(PathTracerNode->DenoiserType));
+
+			// Record the last temporal layer
+			if (PathTracerNode->DenoiserType == EMovieGraphPathTracerDenoiserType::Temporal)
+			{
+				WarningLayerName = LayerData.LayerName;
+			}
+		}
+	}
+
+	// Reset the max frame count to 0 if spatial denoiser is in use.
+	// We need this reset because here we use the max of all nodes.
+	const EMovieGraphPathTracerDenoiserType EffectiveDenosierType = static_cast<EMovieGraphPathTracerDenoiserType>(MaxDenoiserType);
+
+	if (EffectiveDenosierType == EMovieGraphPathTracerDenoiserType::Spatial)
+	{
+		MaxFrameCount = 0;
+	}
+
+	if (IConsoleVariable* FrameCountCvar = IConsoleManager::Get().FindConsoleVariable(TEXT("r.NFOR.FrameCount")))
+	{
+		OriginalFrameCountCvarValue = FrameCountCvar->GetInt();
+		FrameCountCvar->Set(MaxFrameCount);
+	}
+
+	if (IConsoleVariable* DenoiserTypeCvar = IConsoleManager::Get().FindConsoleVariable(TEXT("r.PathTracing.SpatialDenoiser.Type")))
+	{
+
+		OriginalDenoiserType = DenoiserTypeCvar->GetInt();
+		DenoiserTypeCvar->Set(static_cast<int32>(MaxDenoiserType));
+	}
+
+	//If no temporal denoiser is enabled, provide warning if we use temporal denoiser.
+	if (!HasTemporalDenosier())
+	{
+		bool bShouldPerformTemporalDenoising = EMovieGraphPathTracerDenoiserType::Temporal == EffectiveDenosierType;
+
+		if (bShouldPerformTemporalDenoising)
+		{
+			UE_LOG(LogMovieRenderPipeline, Warning, TEXT("The Path Traced Renderer node of Layer `%s` enables temporal denoising but no temporal denoiser plugin is enabled." 
+														 "Fallback to the first available spatial denosier. Please enable at least one denoiser plugin "
+														 "with temporal denoising capability to fully apply the node setting."), *WarningLayerName);
+		}
+	}
+
 	bool bSupportsPathTracing = false;
 	if (IsRayTracingEnabled())
 	{
@@ -90,6 +150,16 @@ void UMovieGraphPathTracerRenderPassNode::TeardownImpl()
 	{
 		ProgressDisplayCvar->Set(bOriginalProgressDisplayCvarValue);
 	}
+
+	if (IConsoleVariable* FrameCountCvar = IConsoleManager::Get().FindConsoleVariable(TEXT("r.NFOR.FrameCount")))
+	{
+		FrameCountCvar->Set(OriginalFrameCountCvarValue);
+	}
+
+	if (IConsoleVariable* DenoiserTypeCvar = IConsoleManager::Get().FindConsoleVariable(TEXT("r.PathTracing.SpatialDenoiser.Type")))
+	{
+		DenoiserTypeCvar->Set(OriginalDenoiserType);
+	}
 }
 
 FString UMovieGraphPathTracerRenderPassNode::GetRendererNameImpl() const
@@ -98,6 +168,13 @@ FString UMovieGraphPathTracerRenderPassNode::GetRendererNameImpl() const
 	return RendererNameImpl;
 }
 
+int32 UMovieGraphPathTracerRenderPassNode::GetCoolingDownFrameCount() const
+{
+	return ((DenoiserType == EMovieGraphPathTracerDenoiserType::Temporal)
+		&& GetAllowDenoiser()
+		&& HasTemporalDenosier()) ? FrameCount : 0;
+}
+
 EViewModeIndex UMovieGraphPathTracerRenderPassNode::GetViewModeIndex() const
 {
 	return VMI_PathTracing;
@@ -138,7 +215,7 @@ bool UMovieGraphPathTracerRenderPassNode::GetAllowOCIO() const
 
 bool UMovieGraphPathTracerRenderPassNode::GetAllowDenoiser() const
 {
-	return bDenoiser;
+	return bEnableDenoiser;
 }
 
 FEngineShowFlags UMovieGraphPathTracerRenderPassNode::GetShowFlags() const
diff --git a/Engine/Plugins/MovieScene/MovieRenderPipeline/Source/MovieRenderPipelineRenderPasses/Private/Graph/Renderers/MovieGraphDeferredPass.cpp b/Engine/Plugins/MovieScene/MovieRenderPipeline/Source/MovieRenderPipelineRenderPasses/Private/Graph/Renderers/MovieGraphDeferredPass.cpp
--- a/Engine/Plugins/MovieScene/MovieRenderPipeline/Source/MovieRenderPipelineRenderPasses/Private/Graph/Renderers/MovieGraphDeferredPass.cpp
+++ b/Engine/Plugins/MovieScene/MovieRenderPipeline/Source/MovieRenderPipelineRenderPasses/Private/Graph/Renderers/MovieGraphDeferredPass.cpp
@@ -47,6 +47,11 @@ void FMovieGraphDeferredPass::Setup(TWeakObjectPtr<UMovieGraphDefaultRenderer> I
 	RenderDataIdentifier.CameraName =  CameraInfo.CameraName;
 
 	SceneViewState.Allocate(InRenderer->GetWorld()->GetFeatureLevel());
+
+	// The InRenderPassNode is not initialized with user's config. Use InLayer to 
+	// initialize the frames to delay for post submission.
+	UMovieGraphRenderPassNode* RenderPassNode = InLayer.RenderPassNode.Get();
+	FramesToDelayPostSubmission = RenderPassNode ? RenderPassNode->GetCoolingDownFrameCount() : 0;
 }
 
 void FMovieGraphDeferredPass::Teardown()
@@ -125,7 +130,8 @@ void FMovieGraphDeferredPass::Render(const FMovieGraphTraversalContext& InFrameT
 
 	UMovieGraphImagePassBaseNode* ParentNodeThisFrame = GetParentNode(InTimeData.EvaluatedConfig);
 	const bool bWriteAllSamples = ParentNodeThisFrame->GetWriteAllSamples();
-	const bool bIsRenderingState = InFrameTraversalContext.Shot->ShotInfo.State == EMovieRenderShotState::Rendering;
+	const bool bIsRenderingState = InFrameTraversalContext.Shot->ShotInfo.State == EMovieRenderShotState::Rendering ||
+									InFrameTraversalContext.Shot->ShotInfo.State == EMovieRenderShotState::CoolingDown;
 	int32 NumSpatialSamples = FMath::Max(1, bIsRenderingState ? ParentNodeThisFrame->GetNumSpatialSamples() : ParentNodeThisFrame->GetNumSpatialSamplesDuringWarmUp());
 
 	const ESceneCaptureSource SceneCaptureSource = ParentNodeThisFrame->GetDisableToneCurve() ? ESceneCaptureSource::SCS_FinalColorHDR : ESceneCaptureSource::SCS_FinalToneCurveHDR;
@@ -336,14 +342,59 @@ void FMovieGraphDeferredPass::Render(const FMovieGraphTraversalContext& InFrameT
 			continue;
 		}
 		
-		// Readback + Accumulate.
-		PostRendererSubmission(SampleState, RenderTargetInitParams, Canvas, CameraInfo);
+		// Post-submission is a little bit complicated to allow supporting temporal denoisers in the Path Tracer.
+		// When using the denoiser with a sample frame count of 2, for a frame to be produced it must look 2 frames
+		// backwards, and 2 frames forwards, ie: to denoise Frame 5, we need to have rendered 3, 4, 5, 6, and 7.
+		// The complication for this is that when we request a render and then immediately schedule a readback, when the
+		// readback is completed the image will be the denoised result from a previous frame. ie: If on frame 5 you schedule
+		// the readback, the result that will be copied to the CPU is the data from Frame 3.
+		//
+		// To resolve these issues, we capture the PostRendererSubmission and place it in a FIFO queue, and then when we schedule
+		// a readback, we provide the old captured state, ie: on Frame 5 we provide Frame 3's data, and that will line up with the
+		// image data actually generated on Frame 3 (which is what is returned by the GPU). A slight complication to this is that
+		// PostRendererSubmission can no longer depend on any member state (since that would be using old data in combination with new),
+		// but the one exception to this is that the current FCanvas is provided since it's only a wrapper for drawing letterboxing anyways.
+		//
+		// Under non-path-tracer temporal denoiser cases, this should effectively work out to be a no-op, the queue will dequeue immediately.
+
+		FMovieGraphPostRendererSubmissionParams Params;
+		Params.SampleState = SampleState;
+		Params.RenderTargetInitParams = RenderTargetInitParams;
+		Params.CameraInfo = CameraInfo;
+
+		// Push the current frame into our FIFO queue
+		SubmissionQueue.Enqueue(Params);
+
+		// When we first start rendering we don't want to schedule a readback (as there isn't actually finished data to read back)
+		// so we skip the first few frames. When we get to the end of a shot, we'll be in a cool-down period where we render extra
+		// frames to allow finishing the denoising on the previous "real" frames. Those frames can't have discard output set on them,
+		// otherwise we won't actually read back the end of the "real" frames. This means the queue will be left with some extra data
+		// in it (for the cool-down frames which were calculated and submitted but never themselves get read back) but that's okay.
+		if (FramesToDelayPostSubmission == 0)
+		{
+			// Now we schedule a readback using the oldest data.
+			FMovieGraphPostRendererSubmissionParams PostParamsToUse;
+			if (SubmissionQueue.Dequeue(PostParamsToUse))
+			{
+				// It's okay that we use the current FCanvas here as it's just a vessel to draw letterboxing based on state captured by the params.
+				PostRendererSubmission(PostParamsToUse.SampleState, PostParamsToUse.RenderTargetInitParams, Canvas, PostParamsToUse.CameraInfo);
+			}
+			else
+			{
+				UE_LOG(LogMovieRenderPipeline, Error, TEXT("De-queue post-submission parameters failed. Attempted to send a frame to post-render submission, but no frames were available in the FIFO queue."));
+			}
+		}
+		else
+		{
+			FramesToDelayPostSubmission--;
+		}
+
 	}
 }
 
 void FMovieGraphDeferredPass::PostRendererSubmission(
 	const UE::MovieGraph::FMovieGraphSampleState& InSampleState,
-	const UE::MovieGraph::DefaultRenderer::FRenderTargetInitParams& InRenderTargetInitParams, FCanvas& InCanvas, const UE::MovieGraph::DefaultRenderer::FCameraInfo& InCameraInfo)
+	const UE::MovieGraph::DefaultRenderer::FRenderTargetInitParams& InRenderTargetInitParams, FCanvas& InCanvas, const UE::MovieGraph::DefaultRenderer::FCameraInfo& InCameraInfo) const
 {
 	TObjectPtr<UMovieGraphDefaultRenderer> GraphRenderer = GetRenderer().Get();
 	if (!GraphRenderer)
@@ -353,10 +404,9 @@ void FMovieGraphDeferredPass::PostRendererSubmission(
 
 	// Draw letterboxing
 	// ToDo: Multi-camera support
-	APlayerCameraManager* PlayerCameraManager = GraphRenderer->GetWorld()->GetFirstPlayerController()->PlayerCameraManager;
-	if(PlayerCameraManager && PlayerCameraManager->GetCameraCacheView().bConstrainAspectRatio)
+	if(InCameraInfo.ViewInfo.bConstrainAspectRatio)
 	{
-		const FMinimalViewInfo CameraCache = PlayerCameraManager->GetCameraCacheView();
+		const FMinimalViewInfo CameraCache = InCameraInfo.ViewInfo;
 		
 		// Taking overscan into account.
 		const FIntPoint FullOutputSize = InSampleState.AccumulatorResolution;
diff --git a/Engine/Plugins/MovieScene/MovieRenderPipeline/Source/MovieRenderPipelineRenderPasses/Private/Graph/Renderers/MovieGraphImagePassBase.cpp b/Engine/Plugins/MovieScene/MovieRenderPipeline/Source/MovieRenderPipelineRenderPasses/Private/Graph/Renderers/MovieGraphImagePassBase.cpp
--- a/Engine/Plugins/MovieScene/MovieRenderPipeline/Source/MovieRenderPipelineRenderPasses/Private/Graph/Renderers/MovieGraphImagePassBase.cpp
+++ b/Engine/Plugins/MovieScene/MovieRenderPipeline/Source/MovieRenderPipelineRenderPasses/Private/Graph/Renderers/MovieGraphImagePassBase.cpp
@@ -459,7 +459,7 @@ void FMovieGraphImagePassBase::ModifyProjectionMatrixForTiling(const UE::MovieGr
 
 void FMovieGraphImagePassBase::PostRendererSubmission(
 	const UE::MovieGraph::FMovieGraphSampleState& InSampleState,
-	const UE::MovieGraph::DefaultRenderer::FRenderTargetInitParams& InRenderTargetInitParams, FCanvas& InCanvas, const UE::MovieGraph::DefaultRenderer::FCameraInfo& InCameraInfo)
+	const UE::MovieGraph::DefaultRenderer::FRenderTargetInitParams& InRenderTargetInitParams, FCanvas& InCanvas, const UE::MovieGraph::DefaultRenderer::FCameraInfo& InCameraInfo) const
 {
 	// We have a pool of accumulators - we multi-thread the accumulation on the task graph, and for each frame,
 	// the task has the previous samples as pre-reqs to keep the accumulation in order. However, each accumulator
diff --git a/Engine/Plugins/MovieScene/MovieRenderPipeline/Source/MovieRenderPipelineRenderPasses/Public/Graph/Nodes/MovieGraphPathTracerPassNode.h b/Engine/Plugins/MovieScene/MovieRenderPipeline/Source/MovieRenderPipelineRenderPasses/Public/Graph/Nodes/MovieGraphPathTracerPassNode.h
--- a/Engine/Plugins/MovieScene/MovieRenderPipeline/Source/MovieRenderPipelineRenderPasses/Public/Graph/Nodes/MovieGraphPathTracerPassNode.h
+++ b/Engine/Plugins/MovieScene/MovieRenderPipeline/Source/MovieRenderPipelineRenderPasses/Public/Graph/Nodes/MovieGraphPathTracerPassNode.h
@@ -6,6 +6,26 @@
 
 #include "MovieGraphPathTracerPassNode.generated.h"
 
+
+UENUM(BlueprintType)
+enum class EMovieGraphPathTracerDenoiserType : uint8
+{
+	/** 
+	* The active spatial denoiser plugin will be used for denoising. If the denoiser is not loaded, a warning will show in the log.
+	* If multiple spatial denoiser plugins are enabled, the last one to get loaded will be the one used.
+	*/
+	Spatial = 0,
+
+	/** 
+	* The active spatial-temporal denoiser plugin will be used for denoising. It provides more temporal stability than spatial denoiser 
+	* if the Frame Count of past/future frames are used (Frame Count > 0) in the plugin. The user needs to config `Frame Count` to
+	* match the requirements of the chosen denoiser plugin. If the denoiser is not loaded, a warning will show in the log. If multiple
+	* spatial-temporal denoiser plugins are enabled, the last one to get loaded will be the one used.
+	*/
+	Temporal = 1
+};
+
+
 /** A render node which uses the path tracer. */
 UCLASS()
 class MOVIERENDERPIPELINERENDERPASSES_API UMovieGraphPathTracerRenderPassNode : public UMovieGraphImagePassBaseNode
@@ -40,7 +60,7 @@ public:
 protected:
 	// UMovieGraphRenderPassNode Interface
 	virtual FString GetRendererNameImpl() const override;
-
+	virtual int32 GetCoolingDownFrameCount() const override;
 	// ~UMovieGraphRenderPassNode Interface
 
 	// UMovieGraphCoreRenderPassNode Interface
@@ -56,7 +76,13 @@ public:
 	uint8 bOverride_bEnableReferenceMotionBlur : 1;
 
 	UPROPERTY(EditAnywhere, BlueprintReadWrite, Category = Overrides, meta = (InlineEditConditionToggle))
-	uint8 bOverride_bDenoiser : 1;
+	uint8 bOverride_bEnableDenoiser : 1;
+
+	UPROPERTY(EditAnywhere, BlueprintReadWrite, Category = Overrides, meta = (InlineEditConditionToggle))
+	uint8 bOverride_DenoiserType : 1;
+
+	UPROPERTY(EditAnywhere, BlueprintReadWrite, Category = Overrides, meta = (InlineEditConditionToggle))
+	uint8 bOverride_FrameCount : 1;
 
 	UPROPERTY(EditAnywhere, BlueprintReadWrite, Category = Overrides, meta = (InlineEditConditionToggle))
 	uint8 bOverride_bDisableToneCurve : 1;
@@ -88,8 +114,25 @@ public:
 	bool bEnableReferenceMotionBlur;
 
 	/** If true the resulting image will be denoised at the end of each set of Spatial Samples. */
-	UPROPERTY(EditAnywhere, BlueprintReadWrite, meta = (UIMin = 1, ClampMin = 1), Category = "Sampling", meta = (EditCondition = "bOverride_bDenoiser"))
-	bool bDenoiser;
+	UPROPERTY(EditAnywhere, BlueprintReadWrite, Category = "Denoiser", meta = (EditCondition = "bOverride_bEnableDenoiser"))
+	bool bEnableDenoiser;
+
+	/**
+	* Select which type of denosier to use when the denoiser is enabled. Temporal denoisers will provide better results when
+	* denoising animated sequences (the denoising results will look more stable), especially when combined with an appropriate 
+	* Frame Count (non-zero). Denoisers are implemented as plugins so you may need to enable a plugin as well for this to work.
+	*/
+	UPROPERTY(EditAnywhere, BlueprintReadWrite, Category = "Denoiser", meta = (EditCondition = "bOverride_DenoiserType"))
+	EMovieGraphPathTracerDenoiserType DenoiserType;
+
+	/** 
+	* The number of frames to consider when using temporal-based denoisers. Generally higher numbers will result in longer
+	* denoising times and higher memory requirements. For NFOR this number refers to how many frames to consider on both sides
+	* of the current frame (ie: 2 means consider 2 before, and 2 after the currently denoised frame), but other denoiser 
+	* implementations may interpret this value differently.
+	*/
+	UPROPERTY(EditAnywhere, BlueprintReadWrite, meta = (UIMin = 0, ClampMin = 0, UIMax = 3), Category = "Denoiser", meta = (EditCondition = "DenoiserType==EMovieGraphPathTracerDenoiserType::Temporal"))
+	int32 FrameCount;
 
 	/**
 	* Debug Feature. Not currently marked BlueprintReadWrite/EditAnywhere as it's not totally implemented on the Path Tracer right now.
@@ -125,4 +168,16 @@ private:
 	 * will be hidden during the render.
 	 */
 	bool bOriginalProgressDisplayCvarValue = false;
+
+	/**
+	 * The original value of the "r.NFOR.FrameCount" cvar before the render starts. Will use the new value set in
+	 * this node during the render.
+	 */
+	int32 OriginalFrameCountCvarValue = 2;
+
+	/**
+	 * The original value of the "r.PathTracing.SpatialDenoiser.Type" cvar before the render starts. Will use the new value set in
+	 * this node during the render.
+	 */
+	int32 OriginalDenoiserType = 0;
 };
\ No newline at end of file
diff --git a/Engine/Plugins/MovieScene/MovieRenderPipeline/Source/MovieRenderPipelineRenderPasses/Public/Graph/Renderers/MovieGraphDeferredPass.h b/Engine/Plugins/MovieScene/MovieRenderPipeline/Source/MovieRenderPipelineRenderPasses/Public/Graph/Renderers/MovieGraphDeferredPass.h
--- a/Engine/Plugins/MovieScene/MovieRenderPipeline/Source/MovieRenderPipelineRenderPasses/Public/Graph/Renderers/MovieGraphDeferredPass.h
+++ b/Engine/Plugins/MovieScene/MovieRenderPipeline/Source/MovieRenderPipelineRenderPasses/Public/Graph/Renderers/MovieGraphDeferredPass.h
@@ -5,6 +5,13 @@
 
 namespace UE::MovieGraph::Rendering
 {
+	struct FMovieGraphPostRendererSubmissionParams
+	{
+		UE::MovieGraph::FMovieGraphSampleState SampleState;
+		UE::MovieGraph::DefaultRenderer::FRenderTargetInitParams RenderTargetInitParams;
+		UE::MovieGraph::DefaultRenderer::FCameraInfo CameraInfo;
+	};
+
 	struct MOVIERENDERPIPELINERENDERPASSES_API FMovieGraphDeferredPass : public FMovieGraphImagePassBase
 	{
 		// FMovieGraphImagePassBase Interface
@@ -18,7 +25,7 @@ namespace UE::MovieGraph::Rendering
 		// End FMovieGraphImagePassBase
 			
 	protected:
-		virtual void PostRendererSubmission(const UE::MovieGraph::FMovieGraphSampleState& InSampleState, const UE::MovieGraph::DefaultRenderer::FRenderTargetInitParams& InRenderTargetInitParams, FCanvas& InCanvas, const UE::MovieGraph::DefaultRenderer::FCameraInfo& InCameraInfo) override;
+		virtual void PostRendererSubmission(const UE::MovieGraph::FMovieGraphSampleState& InSampleState, const UE::MovieGraph::DefaultRenderer::FRenderTargetInitParams& InRenderTargetInitParams, FCanvas& InCanvas, const UE::MovieGraph::DefaultRenderer::FCameraInfo& InCameraInfo) const override;
 	protected:
 		FMovieGraphRenderPassLayerData LayerData;
 
@@ -27,5 +34,11 @@ namespace UE::MovieGraph::Rendering
 
 		// Scene View history used by the renderer 
 		FSceneViewStateReference SceneViewState;
+
+		// The number of frames to delay to send frames from SubmissionQueue to post-render submission.
+		int32 FramesToDelayPostSubmission;
+
+		// FIFO queue of rendered frames. It allows frames to be sent to post-render submission with a delay if needed (e.g., when temporal denoising is used with path tracers).  
+		TQueue<FMovieGraphPostRendererSubmissionParams> SubmissionQueue;
 	}; 
 }
\ No newline at end of file
diff --git a/Engine/Plugins/MovieScene/MovieRenderPipeline/Source/MovieRenderPipelineRenderPasses/Public/Graph/Renderers/MovieGraphImagePassBase.h b/Engine/Plugins/MovieScene/MovieRenderPipeline/Source/MovieRenderPipelineRenderPasses/Public/Graph/Renderers/MovieGraphImagePassBase.h
--- a/Engine/Plugins/MovieScene/MovieRenderPipeline/Source/MovieRenderPipelineRenderPasses/Public/Graph/Renderers/MovieGraphImagePassBase.h
+++ b/Engine/Plugins/MovieScene/MovieRenderPipeline/Source/MovieRenderPipelineRenderPasses/Public/Graph/Renderers/MovieGraphImagePassBase.h
@@ -81,7 +81,7 @@ namespace UE::MovieGraph::Rendering
 		virtual TSharedRef<FSceneViewFamilyContext> CreateSceneViewFamily(const FViewFamilyInitData& InInitData, const UE::MovieGraph::DefaultRenderer::FCameraInfo& InCameraInfo) const;
 		virtual void ApplyMovieGraphOverridesToSceneView(TSharedRef<FSceneViewFamilyContext> InOutFamily, const FViewFamilyInitData& InInitData, const UE::MovieGraph::DefaultRenderer::FCameraInfo& InCameraInfo) const;
 		virtual void ApplyMovieGraphOverridesToViewFamily(TSharedRef<FSceneViewFamilyContext> InOutFamily, const FViewFamilyInitData& InInitData) const;
-		virtual void PostRendererSubmission(const UE::MovieGraph::FMovieGraphSampleState& InSampleState, const UE::MovieGraph::DefaultRenderer::FRenderTargetInitParams& InRenderTargetInitParams, FCanvas& InCanvas, const UE::MovieGraph::DefaultRenderer::FCameraInfo& InCameraInfo);
+		virtual void PostRendererSubmission(const UE::MovieGraph::FMovieGraphSampleState& InSampleState, const UE::MovieGraph::DefaultRenderer::FRenderTargetInitParams& InRenderTargetInitParams, FCanvas& InCanvas, const UE::MovieGraph::DefaultRenderer::FCameraInfo& InCameraInfo) const;
 		virtual TFunction<void(TUniquePtr<FImagePixelData>&&)> MakeForwardingEndpoint(const FMovieGraphSampleState& InSampleState, const FMovieGraphTimeStepData& InTimeData);
 		virtual bool ShouldDiscardOutput(const TSharedRef<FSceneViewFamilyContext>& InFamily, const UE::MovieGraph::DefaultRenderer::FCameraInfo& InCameraInfo) const { return false; }
 		/** For this image pass, look up the associated node by type for the given config. */
diff --git a/Engine/Source/Runtime/Renderer/Private/PathTracingSpatialTemporalDenoising.cpp b/Engine/Source/Runtime/Renderer/Private/PathTracingSpatialTemporalDenoising.cpp
--- a/Engine/Source/Runtime/Renderer/Private/PathTracingSpatialTemporalDenoising.cpp
+++ b/Engine/Source/Runtime/Renderer/Private/PathTracingSpatialTemporalDenoising.cpp
@@ -105,6 +105,11 @@ void UnregisterDenoiser(FString Name)
 	FDenoiserManager::Get().UnregisterDenoiser(Name);
 }
 
+bool HasTemporalDenosier()
+{
+	return FDenoiserManager::Get().HasSpatialTemporalDenoiser();
+}
+
 #if RHI_RAYTRACING
 
 #include "DeferredShadingRenderer.h"
diff --git a/Engine/Source/Runtime/Renderer/Public/PathTracingDenoiser.h b/Engine/Source/Runtime/Renderer/Public/PathTracingDenoiser.h
--- a/Engine/Source/Runtime/Renderer/Public/PathTracingDenoiser.h
+++ b/Engine/Source/Runtime/Renderer/Public/PathTracingDenoiser.h
@@ -120,3 +120,6 @@ extern RENDERER_API void RegisterSpatialTemporalDenoiser(TUniquePtr<UE::Renderer
 
 /** Unregister all denoisers with the name Name*/
 extern RENDERER_API void UnregisterDenoiser(FString Name);
+
+/** Return true if there is any temporal denoiser plugin registered*/
+extern RENDERER_API bool HasTemporalDenosier();
