Set the default denoiser to spatial denoiser in the path tracer node of movie render graph.

diff --git a/Engine/Plugins/MovieScene/MovieRenderPipeline/Source/MovieRenderPipelineRenderPasses/Private/Graph/Nodes/MovieGraphPathTracerPassNode.cpp b/Engine/Plugins/MovieScene/MovieRenderPipeline/Source/MovieRenderPipelineRenderPasses/Private/Graph/Nodes/MovieGraphPathTracerPassNode.cpp
--- a/Engine/Plugins/MovieScene/MovieRenderPipeline/Source/MovieRenderPipelineRenderPasses/Private/Graph/Nodes/MovieGraphPathTracerPassNode.cpp
+++ b/Engine/Plugins/MovieScene/MovieRenderPipeline/Source/MovieRenderPipelineRenderPasses/Private/Graph/Nodes/MovieGraphPathTracerPassNode.cpp
@@ -16,7 +16,7 @@ UMovieGraphPathTracerRenderPassNode::UMovieGraphPathTracerRenderPassNode()
 	: SpatialSampleCount(1)
 	, bEnableReferenceMotionBlur(false)
 	, bEnableDenoiser(true)
-	, DenoiserType(EMovieGraphPathTracerDenoiserType::Temporal)
+	, DenoiserType(EMovieGraphPathTracerDenoiserType::Spatial)
 	, FrameCount(2)
 	, bWriteAllSamples(false)
 	, bDisableToneCurve(false)
