Removed LevelSequence.h from MovieRenderGraph headers

diff --git a/Engine/Plugins/MovieScene/MovieRenderPipeline/Source/MovieRenderPipelineCore/Private/Graph/MovieGraphBlueprintLibrary.cpp b/Engine/Plugins/MovieScene/MovieRenderPipeline/Source/MovieRenderPipelineCore/Private/Graph/MovieGraphBlueprintLibrary.cpp
--- a/Engine/Plugins/MovieScene/MovieRenderPipeline/Source/MovieRenderPipelineCore/Private/Graph/MovieGraphBlueprintLibrary.cpp
+++ b/Engine/Plugins/MovieScene/MovieRenderPipeline/Source/MovieRenderPipelineCore/Private/Graph/MovieGraphBlueprintLibrary.cpp
@@ -12,6 +12,7 @@
 #include "MoviePipelineBlueprintLibrary.h"
 #include "MoviePipelineUtils.h"
 
+#include "Misc/Paths.h"
 #include "HAL/FileManager.h"
 #include "Internationalization/Regex.h"
 
diff --git a/Engine/Plugins/MovieScene/MovieRenderPipeline/Source/MovieRenderPipelineCore/Private/MoviePipelineQueue.cpp b/Engine/Plugins/MovieScene/MovieRenderPipeline/Source/MovieRenderPipelineCore/Private/MoviePipelineQueue.cpp
--- a/Engine/Plugins/MovieScene/MovieRenderPipeline/Source/MovieRenderPipelineCore/Private/MoviePipelineQueue.cpp
+++ b/Engine/Plugins/MovieScene/MovieRenderPipeline/Source/MovieRenderPipelineCore/Private/MoviePipelineQueue.cpp
@@ -3,6 +3,7 @@
 #include "MoviePipelineQueue.h"
 #include "MoviePipelineOutputSetting.h"
 #include "MoviePipelineSetting.h"
+#include "LevelSequence.h"
 #include "MoviePipelineBlueprintLibrary.h"
 #include "MovieRenderPipelineCoreModule.h"
 
diff --git a/Engine/Plugins/MovieScene/MovieRenderPipeline/Source/MovieRenderPipelineCore/Public/MoviePipelineQueue.h b/Engine/Plugins/MovieScene/MovieRenderPipeline/Source/MovieRenderPipelineCore/Public/MoviePipelineQueue.h
--- a/Engine/Plugins/MovieScene/MovieRenderPipeline/Source/MovieRenderPipelineCore/Public/MoviePipelineQueue.h
+++ b/Engine/Plugins/MovieScene/MovieRenderPipeline/Source/MovieRenderPipelineCore/Public/MoviePipelineQueue.h
@@ -3,7 +3,6 @@
 
 #include "UObject/Object.h"
 #include "MovieRenderPipelineDataTypes.h"
-#include "LevelSequence.h"
 #include "MovieJobVariableAssignmentContainer.h"
 #include "MoviePipelinePrimaryConfig.h"
 #include "MoviePipelineShotConfig.h"
diff --git a/Engine/Plugins/MovieScene/MovieRenderPipeline/Source/MovieRenderPipelineEditor/Private/MoviePipelineEditorBlueprintLibrary.cpp b/Engine/Plugins/MovieScene/MovieRenderPipeline/Source/MovieRenderPipelineEditor/Private/MoviePipelineEditorBlueprintLibrary.cpp
--- a/Engine/Plugins/MovieScene/MovieRenderPipeline/Source/MovieRenderPipelineEditor/Private/MoviePipelineEditorBlueprintLibrary.cpp
+++ b/Engine/Plugins/MovieScene/MovieRenderPipeline/Source/MovieRenderPipelineEditor/Private/MoviePipelineEditorBlueprintLibrary.cpp
@@ -7,6 +7,7 @@
 #include "FileHelpers.h"
 #include "ObjectTools.h"
 #include "PackageTools.h"
+#include "LevelSequence.h"
 #include "MoviePipelineQueue.h"
 #include "MoviePipelineQueueSubsystem.h"
 #include "MovieRenderPipelineSettings.h"
diff --git a/Engine/Plugins/MovieScene/MovieRenderPipeline/Source/MovieRenderPipelineSettings/Private/MoviePipelineConsoleVariableSetting.cpp b/Engine/Plugins/MovieScene/MovieRenderPipeline/Source/MovieRenderPipelineSettings/Private/MoviePipelineConsoleVariableSetting.cpp
--- a/Engine/Plugins/MovieScene/MovieRenderPipeline/Source/MovieRenderPipelineSettings/Private/MoviePipelineConsoleVariableSetting.cpp
+++ b/Engine/Plugins/MovieScene/MovieRenderPipeline/Source/MovieRenderPipelineSettings/Private/MoviePipelineConsoleVariableSetting.cpp
@@ -5,6 +5,7 @@
 #include "MoviePipelineQueue.h"
 #include "MovieRenderPipelineCoreModule.h"
 #include "MovieScene.h"
+#include "LevelSequence.h"
 #include "MovieSceneCommonHelpers.h"
 #include "HAL/IConsoleManager.h"
 #include "Engine/World.h"
