MoviePipeline: In the graph, sort collection query types alphabetically in the picker drop-down.

diff --git a/Engine/Plugins/MovieScene/MovieRenderPipeline/Source/MovieRenderPipelineEditor/Private/Customizations/Graph/MovieGraphCollectionsCustomization.cpp b/Engine/Plugins/MovieScene/MovieRenderPipeline/Source/MovieRenderPipelineEditor/Private/Customizations/Graph/MovieGraphCollectionsCustomization.cpp
--- a/Engine/Plugins/MovieScene/MovieRenderPipeline/Source/MovieRenderPipelineEditor/Private/Customizations/Graph/MovieGraphCollectionsCustomization.cpp
+++ b/Engine/Plugins/MovieScene/MovieRenderPipeline/Source/MovieRenderPipelineEditor/Private/Customizations/Graph/MovieGraphCollectionsCustomization.cpp
@@ -865,6 +865,11 @@ TArray<UClass*>* SMovieGraphCollectionTreeQueryTypeSelectorWidget::GetAvailableQ
 				QueryTypes.Add(Class);
 			}
 		}
+
+		QueryTypes.Sort([](const UClass& InClass1, const UClass& InClass2)
+		{
+			return InClass1.GetName().Compare(InClass2.GetName(), ESearchCase::IgnoreCase) < 0;
+		});
 	}
 
 	return &QueryTypes;
