CIS: Fixed missing LevelSequence headers.

diff --git a/Engine/Plugins/MovieScene/MovieRenderPipeline/Source/MovieRenderPipelineCore/Private/MoviePipelineInProcessExecutor.cpp b/Engine/Plugins/MovieScene/MovieRenderPipeline/Source/MovieRenderPipelineCore/Private/MoviePipelineInProcessExecutor.cpp
--- a/Engine/Plugins/MovieScene/MovieRenderPipeline/Source/MovieRenderPipelineCore/Private/MoviePipelineInProcessExecutor.cpp
+++ b/Engine/Plugins/MovieScene/MovieRenderPipeline/Source/MovieRenderPipelineCore/Private/MoviePipelineInProcessExecutor.cpp
@@ -4,6 +4,7 @@
 #include "Engine/GameInstance.h"
 #include "MoviePipeline.h"
 #include "Engine/World.h"
+#include "LevelSequence.h"
 #include "MoviePipelineQueue.h"
 #include "MoviePipelineInProcessExecutorSettings.h"
 #include "MoviePipelineGameOverrideSetting.h"
diff --git a/Engine/Plugins/MovieScene/MovieRenderPipeline/Source/MovieRenderPipelineCore/Private/MoviePipelineQueueEngineSubsystem.cpp b/Engine/Plugins/MovieScene/MovieRenderPipeline/Source/MovieRenderPipelineCore/Private/MoviePipelineQueueEngineSubsystem.cpp
--- a/Engine/Plugins/MovieScene/MovieRenderPipeline/Source/MovieRenderPipelineCore/Private/MoviePipelineQueueEngineSubsystem.cpp
+++ b/Engine/Plugins/MovieScene/MovieRenderPipeline/Source/MovieRenderPipelineCore/Private/MoviePipelineQueueEngineSubsystem.cpp
@@ -4,6 +4,7 @@
 #include "Modules/ModuleManager.h"
 #include "MoviePipeline.h"
 #include "MoviePipelineInProcessExecutor.h"
+#include "LevelSequence.h"
 
 #include UE_INLINE_GENERATED_CPP_BY_NAME(MoviePipelineQueueEngineSubsystem)
 
