MRG: Reference Motion blur should enable Reference DOF

Post-processed Depth-of-Field does not work well if the input is already motion blurred. When using the reference motion blur mode, automatically turn on reference DOF.

diff --git a/Engine/Plugins/MovieScene/MovieRenderPipeline/Source/MovieRenderPipelineRenderPasses/Private/Graph/Renderers/MovieGraphPathTracerPass.cpp b/Engine/Plugins/MovieScene/MovieRenderPipeline/Source/MovieRenderPipelineRenderPasses/Private/Graph/Renderers/MovieGraphPathTracerPass.cpp
--- a/Engine/Plugins/MovieScene/MovieRenderPipeline/Source/MovieRenderPipelineRenderPasses/Private/Graph/Renderers/MovieGraphPathTracerPass.cpp
+++ b/Engine/Plugins/MovieScene/MovieRenderPipeline/Source/MovieRenderPipelineRenderPasses/Private/Graph/Renderers/MovieGraphPathTracerPass.cpp
@@ -64,6 +64,16 @@ void FMovieGraphPathTracerPass::ApplyMovieGraphOverridesToSceneView(TSharedRef<F
 	FSceneView* View = const_cast<FSceneView*>(InOutFamily->Views[0]);
 	View->FinalPostProcessSettings.bOverride_PathTracingSamplesPerPixel = true;
 	View->FinalPostProcessSettings.PathTracingSamplesPerPixel = SampleCount;
+
+	// If we are using reference motion blur, also force the use of reference DOF (as the post-processed DOF cannot behave well with motion blurred input)
+	// TODO: Is there a way to directly access bEnableReferenceMotionBlur from here?
+	const bool bHasTiles = InCameraInfo.TilingParams.TileCount.X * InCameraInfo.TilingParams.TileCount.Y > 1;
+	const bool bAccumulateSpatialSamplesOnly = InOutFamily->EngineShowFlags.MotionBlur || bHasTiles;
+	if (!bAccumulateSpatialSamplesOnly)
+	{
+		View->FinalPostProcessSettings.bOverride_PathTracingEnableReferenceDOF = true;
+		View->FinalPostProcessSettings.PathTracingEnableReferenceDOF = true;
+	}
 		
 	// reset path tracer's accumulation at the start of each sample
 	View->bForcePathTracerReset = SampleIndex == 0;
diff --git a/Engine/Plugins/MovieScene/MovieRenderPipeline/Source/MovieRenderPipelineRenderPasses/Public/Graph/Nodes/MovieGraphPathTracerPassNode.h b/Engine/Plugins/MovieScene/MovieRenderPipeline/Source/MovieRenderPipelineRenderPasses/Public/Graph/Nodes/MovieGraphPathTracerPassNode.h
--- a/Engine/Plugins/MovieScene/MovieRenderPipeline/Source/MovieRenderPipelineRenderPasses/Public/Graph/Nodes/MovieGraphPathTracerPassNode.h
+++ b/Engine/Plugins/MovieScene/MovieRenderPipeline/Source/MovieRenderPipelineRenderPasses/Public/Graph/Nodes/MovieGraphPathTracerPassNode.h
@@ -107,7 +107,7 @@ public:
 
 	/** 
 	 *  When enabled, the path tracer will blend all spatial and temporal samples prior to the denoising and will disable post-processed motion blur.
-	 *  In this mode it is possible to use higher temporal sample counts to improve the motion blur quality.
+	 *  In this mode it is possible to use higher temporal sample counts to improve the motion blur quality. This mode also automatically enabled reference DOF.
 	 *  When this option is disabled, the path tracer will accumulate spatial samples, but denoise them prior to accumulation of temporal samples.
 	 */
 	UPROPERTY(EditAnywhere, BlueprintReadWrite, meta = (UIMin = 1, ClampMin = 1), Category = "Reference Motion Blur", meta = (EditCondition = "bOverride_bEnableReferenceMotionBlur"))
