MoviePipeline: Updated the Set Console Variable node to make promoting the Value property more streamlined.
* The "Name" property is no longer promotable.
* The "Value" property includes the cvar name in the context menu -- eg, "Value (r.SomeCVar)".
* When promoting the "Value" property, the new variable is given the name of the cvar.

diff --git a/Engine/Plugins/MovieScene/MovieRenderPipeline/Source/MovieRenderPipelineCore/Private/Graph/Nodes/MovieGraphSetCVarValueNode.cpp b/Engine/Plugins/MovieScene/MovieRenderPipeline/Source/MovieRenderPipelineCore/Private/Graph/Nodes/MovieGraphSetCVarValueNode.cpp
--- a/Engine/Plugins/MovieScene/MovieRenderPipeline/Source/MovieRenderPipelineCore/Private/Graph/Nodes/MovieGraphSetCVarValueNode.cpp
+++ b/Engine/Plugins/MovieScene/MovieRenderPipeline/Source/MovieRenderPipelineCore/Private/Graph/Nodes/MovieGraphSetCVarValueNode.cpp
@@ -17,6 +17,33 @@ EMovieGraphBranchRestriction UMovieGraphSetCVarValueNode::GetBranchRestriction()
 	return EMovieGraphBranchRestriction::Globals;
 }
 
+TArray<FMovieGraphPropertyInfo> UMovieGraphSetCVarValueNode::GetOverrideablePropertyInfo() const
+{
+	TArray<FMovieGraphPropertyInfo> PropertyInfo = Super::GetOverrideablePropertyInfo();
+
+	for (FMovieGraphPropertyInfo& Info : PropertyInfo)
+	{
+		if (Info.Name == GET_MEMBER_NAME_CHECKED(UMovieGraphSetCVarValueNode, Value))
+		{
+			const FString CVarName = Name.IsEmpty() ? TEXT("NoCVarName") : Name;
+			
+			// This custom display name will show up in the property promotion context menu
+			Info.ContextMenuName = FText::Format(NSLOCTEXT("MovieGraphNodes", "SetCvarValueNode_ValuePromotionContextMenuName", "Value ({0})"), FText::FromString(CVarName));
+
+			// If the property is promoted, the created variable should have the name of the cvar
+			Info.PromotionName = FName(CVarName);
+		}
+	}
+
+	// Also remove the "Name" property -- it should not show up as a property that can be promoted
+	PropertyInfo.RemoveAll([](const FMovieGraphPropertyInfo& InPropertyInfo)
+	{
+		return InPropertyInfo.Name == GET_MEMBER_NAME_CHECKED(UMovieGraphSetCVarValueNode, Name);
+	});
+
+	return PropertyInfo;
+}
+
 #if WITH_EDITOR
 FText UMovieGraphSetCVarValueNode::GetNodeTitle(const bool bGetDescriptive) const
 {
diff --git a/Engine/Plugins/MovieScene/MovieRenderPipeline/Source/MovieRenderPipelineCore/Private/MovieJobVariableAssignmentContainer.cpp b/Engine/Plugins/MovieScene/MovieRenderPipeline/Source/MovieRenderPipelineCore/Private/MovieJobVariableAssignmentContainer.cpp
--- a/Engine/Plugins/MovieScene/MovieRenderPipeline/Source/MovieRenderPipelineCore/Private/MovieJobVariableAssignmentContainer.cpp
+++ b/Engine/Plugins/MovieScene/MovieRenderPipeline/Source/MovieRenderPipelineCore/Private/MovieJobVariableAssignmentContainer.cpp
@@ -286,6 +286,7 @@ bool UMovieJobVariableAssignmentContainer::GenerateVariableOverride(const UMovie
 	NewProperty.MetaData.Add(FPropertyBagPropertyDescMetaData(ToolTipMetaDataKey, InGraphVariable->Description));
 	NewProperty.MetaData.Add(FPropertyBagPropertyDescMetaData(CategoryMetaDataKey, InGraphVariable->GetCategory()));
 	NewProperty.MetaData.Add(FPropertyBagPropertyDescMetaData(EnableCategoriesMetaDataKey, FString()));
+	NewProperty.MetaData.Add(FPropertyBagPropertyDescMetaData(DisplayNameMetaDataKey, InGraphVariable->GetMemberName()));	// Display the raw variable name, not the "friendly" name
 #endif
 
 	// Track a separate EditCondition property that can enable/disable the above property. Since the variable can be
diff --git a/Engine/Plugins/MovieScene/MovieRenderPipeline/Source/MovieRenderPipelineCore/Public/Graph/MovieGraphNode.h b/Engine/Plugins/MovieScene/MovieRenderPipeline/Source/MovieRenderPipelineCore/Public/Graph/MovieGraphNode.h
--- a/Engine/Plugins/MovieScene/MovieRenderPipeline/Source/MovieRenderPipelineCore/Public/Graph/MovieGraphNode.h
+++ b/Engine/Plugins/MovieScene/MovieRenderPipeline/Source/MovieRenderPipelineCore/Public/Graph/MovieGraphNode.h
@@ -45,6 +45,14 @@ struct FMovieGraphPropertyInfo
 	UPROPERTY(BlueprintReadOnly, Category = "Movie Graph")
 	FName Name;
 
+	/** The display name of the property which will be shown in the context menu. If empty, the value from 'Name' will be used. */
+	UPROPERTY(BlueprintReadOnly, Category = "Movie Graph")
+	FText ContextMenuName;
+
+	/** If this property is promoted, this is the name of the variable that is created. If empty, the value from 'Name' will be used. */
+	UPROPERTY(BlueprintReadOnly, Category = "Movie Graph")
+	FName PromotionName;
+
 	/** Whether this property is dynamic (ie, it does not correspond to a native UPROPERTY on the node). */
 	UPROPERTY(BlueprintReadOnly, Category = "Movie Graph")
 	bool bIsDynamicProperty = false;
@@ -79,9 +87,12 @@ struct FMovieGraphPropertyInfo
 	bool operator==(const FMovieGraphPropertyInfo& Other) const
 	{
 		return (Name == Other.Name)
+			&& (ContextMenuName.EqualTo(Other.ContextMenuName))
+			&& (PromotionName == Other.PromotionName)
 			&& (bIsDynamicProperty == Other.bIsDynamicProperty)
 			&& (ValueType == Other.ValueType)
-			&& (ValueTypeObject == Other.ValueTypeObject);
+			&& (ValueTypeObject == Other.ValueTypeObject)
+			&& (bIsPermanentlyExposed == Other.bIsPermanentlyExposed);
 	}
 };
 
diff --git a/Engine/Plugins/MovieScene/MovieRenderPipeline/Source/MovieRenderPipelineCore/Public/Graph/Nodes/MovieGraphSetCVarValueNode.h b/Engine/Plugins/MovieScene/MovieRenderPipeline/Source/MovieRenderPipelineCore/Public/Graph/Nodes/MovieGraphSetCVarValueNode.h
--- a/Engine/Plugins/MovieScene/MovieRenderPipeline/Source/MovieRenderPipelineCore/Public/Graph/Nodes/MovieGraphSetCVarValueNode.h
+++ b/Engine/Plugins/MovieScene/MovieRenderPipeline/Source/MovieRenderPipelineCore/Public/Graph/Nodes/MovieGraphSetCVarValueNode.h
@@ -16,6 +16,7 @@ public:
 
 	virtual FString GetNodeInstanceName() const override;
 	virtual EMovieGraphBranchRestriction GetBranchRestriction() const override;
+	virtual TArray<FMovieGraphPropertyInfo> GetOverrideablePropertyInfo() const override;
 
 #if WITH_EDITOR
 	virtual FText GetNodeTitle(const bool bGetDescriptive = false) const override;
diff --git a/Engine/Plugins/MovieScene/MovieRenderPipeline/Source/MovieRenderPipelineCore/Public/MovieJobVariableAssignmentContainer.h b/Engine/Plugins/MovieScene/MovieRenderPipeline/Source/MovieRenderPipelineCore/Public/MovieJobVariableAssignmentContainer.h
--- a/Engine/Plugins/MovieScene/MovieRenderPipeline/Source/MovieRenderPipelineCore/Public/MovieJobVariableAssignmentContainer.h
+++ b/Engine/Plugins/MovieScene/MovieRenderPipeline/Source/MovieRenderPipelineCore/Public/MovieJobVariableAssignmentContainer.h
@@ -265,6 +265,9 @@ private:
 
 	/* The metadata key on properties that is associated with the variable's category. */
 	static inline FName CategoryMetaDataKey = FName(TEXT("Category"));
+
+	/* The metadata key on properties that is associated with the variable's display name. */
+	static inline FName DisplayNameMetaDataKey = FName(TEXT("DisplayName"));
 	
 	/* The metadata key on property bag properties that allows them to show up grouped in the details panel. */
 	static inline FName EnableCategoriesMetaDataKey = FName(TEXT("EnableCategories"));
diff --git a/Engine/Plugins/MovieScene/MovieRenderPipeline/Source/MovieRenderPipelineEditor/Private/Graph/MovieEdGraphNode.cpp b/Engine/Plugins/MovieScene/MovieRenderPipeline/Source/MovieRenderPipelineEditor/Private/Graph/MovieEdGraphNode.cpp
--- a/Engine/Plugins/MovieScene/MovieRenderPipeline/Source/MovieRenderPipelineEditor/Private/Graph/MovieEdGraphNode.cpp
+++ b/Engine/Plugins/MovieScene/MovieRenderPipeline/Source/MovieRenderPipelineEditor/Private/Graph/MovieEdGraphNode.cpp
@@ -317,7 +317,7 @@ void UMoviePipelineEdGraphNode::GetPropertyPromotionContextMenuActions(UToolMenu
 		
 		ExposeAsPinSection.AddMenuEntry(
 			PropertyInfo.Name,
-			FText::FromName(PropertyInfo.Name),
+			PropertyInfo.ContextMenuName.IsEmpty() ? FText::FromName(PropertyInfo.Name) : PropertyInfo.ContextMenuName,
 			LOCTEXT("PromotePropertyToPin", "Promote this property to a pin on this node."),
 			FSlateIcon(),
 			FUIAction(
@@ -348,8 +348,10 @@ void UMoviePipelineEdGraphNode::GetPropertyPromotionContextMenuActions(UToolMenu
 
 void UMoviePipelineEdGraphNode::PromotePropertyToVariable(const FMovieGraphPropertyInfo& TargetProperty) const
 {
+	const FName PromotedVariableName = TargetProperty.PromotionName.IsNone() ? TargetProperty.Name : TargetProperty.PromotionName;
+	
 	// Note: AddVariable() will take care of determining a unique name if there is already a variable with the property's name
-	if (UMovieGraphVariable* NewGraphVariable = RuntimeNode->GetGraph()->AddVariable(TargetProperty.Name))
+	if (UMovieGraphVariable* NewGraphVariable = RuntimeNode->GetGraph()->AddVariable(PromotedVariableName))
 	{
 		// Set the new variable's type to match the property that is being promoted
 		UObject* ValueTypeObject = const_cast<UObject*>(TargetProperty.ValueTypeObject.Get());
