MoviePipeline: Fixed a number of issues with how Movie Render Queue (both original and MRG) calculated the projection matrix and cropping for cameras. These should now be much more inline with what you see in the viewport. Additionally relaxed the warnings around Orthographic projection with TSR/TAA as they are now supported.

Specific things fixed:
Orthographic renders being stretched to fill view
Aspect Ratio Axis Lock not being respected
Squeeze Factor not being respected under all setting combinations

Non-exhaustive list of things tested:
Aspect Ratio constrained and unconstrained, with output aspect ratios which are taller, wider, and equal, both vertical and horizontal aspect ratios (ie: widescreen vs. portrait), with and without high resolution tiling.
Squeeze Factor with and without Aspect Ratio Constraint enabled, 16x9, 4x3 and no crop (with varying output aspect ratios), with and without High Resolution tiling enabled.
Overscan with and without aspect ratio constraint, orthographic and perspective projections, high resolution on and off. Note that an aspect-ratio constrained image which has overscan enabled will move the letterboxing (otherwise overscan isn't actually giving the user any additional data)
Aspect Ratio Axis Lock with orthographic and perspective projections, using X, Y, and Major axis constraints, with taller, wider, or equal output aspect ratios. Note that orthographic rendering doesn't fully support this feature at this time so it stretches the image, but is consistent with viewport behavior.

diff --git a/Engine/Plugins/MovieScene/MovieRenderPipeline/Source/MovieRenderPipelineCore/Public/Graph/MovieGraphDefaultRenderer.h b/Engine/Plugins/MovieScene/MovieRenderPipeline/Source/MovieRenderPipelineCore/Public/Graph/MovieGraphDefaultRenderer.h
--- a/Engine/Plugins/MovieScene/MovieRenderPipeline/Source/MovieRenderPipelineCore/Public/Graph/MovieGraphDefaultRenderer.h
+++ b/Engine/Plugins/MovieScene/MovieRenderPipeline/Source/MovieRenderPipelineCore/Public/Graph/MovieGraphDefaultRenderer.h
@@ -67,9 +67,6 @@ namespace UE::MovieGraph::DefaultRenderer
 		/** should we respect the camera's aspect ratio settings. */
 		bool bAllowCameraAspectRatio;
 
-		/** Projection Matrix this camera should use. Used instead of the one calculated by FMinimalViewInfo to handle special cases. */
-		FMatrix ProjectionMatrix;
-
 		// questionable if these are fcamerainfo
 		/** When using tiling, we scale the sensor to counteract the view changes. This value comes from modifying the ProjectionMatrix. */
 		float DoFSensorScale;
diff --git a/Engine/Plugins/MovieScene/MovieRenderPipeline/Source/MovieRenderPipelineRenderPasses/Private/Graph/Renderers/MovieGraphDeferredPass.cpp b/Engine/Plugins/MovieScene/MovieRenderPipeline/Source/MovieRenderPipelineRenderPasses/Private/Graph/Renderers/MovieGraphDeferredPass.cpp
--- a/Engine/Plugins/MovieScene/MovieRenderPipeline/Source/MovieRenderPipelineRenderPasses/Private/Graph/Renderers/MovieGraphDeferredPass.cpp
+++ b/Engine/Plugins/MovieScene/MovieRenderPipeline/Source/MovieRenderPipelineRenderPasses/Private/Graph/Renderers/MovieGraphDeferredPass.cpp
@@ -218,16 +218,6 @@ void FMovieGraphDeferredPass::Render(const FMovieGraphTraversalContext& InFrameT
 		// Move the final render by this much in the accumulator to counteract the offset put into the view matrix.
 		// Note that when bAllowSpatialJitter is false, SpatialShiftX/Y will always be zero.
 		FVector2D OverlappedSubpixelShift = FVector2D(0.5f - SpatialShiftAmount.X, 0.5f - SpatialShiftAmount.Y);
-		
-		FMatrix ProjectionMatrix = CalculateProjectionMatrix(CameraInfo);
-		float DoFSensorScale = 1.f;
-
-		// Modify the perspective matrix to do an off center projection, with overlap for high-res tiling
-		const bool bOrthographic = CameraInfo.ViewInfo.ProjectionMode == ECameraProjectionMode::Type::Orthographic;
-		ModifyProjectionMatrixForTiling(CameraInfo.TilingParams, bOrthographic, ProjectionMatrix, DoFSensorScale);
-		
-		CameraInfo.ProjectionMatrix = ProjectionMatrix;
-		CameraInfo.DoFSensorScale = DoFSensorScale;
 
 		// The Scene View Family must be constructed first as the FSceneView needs it to be constructed
 		UE::MovieGraph::Rendering::FViewFamilyInitData ViewFamilyInitData;
@@ -240,11 +230,23 @@ void FMovieGraphDeferredPass::Render(const FMovieGraphTraversalContext& InFrameT
 		ViewFamilyInitData.AntiAliasingMethod = AntiAliasingMethod;
 		ViewFamilyInitData.ShowFlags = ParentNodeThisFrame->GetShowFlags();
 		ViewFamilyInitData.ViewModeIndex = ParentNodeThisFrame->GetViewModeIndex();
+		ViewFamilyInitData.ProjectionMode = CameraInfo.ViewInfo.ProjectionMode;
+
+		TSharedRef<FSceneViewFamilyContext> ViewFamily = CreateSceneViewFamily(ViewFamilyInitData);
+		FSceneViewInitOptions SceneViewInitOptions = CreateViewInitOptions(CameraInfo, ViewFamily.ToSharedPtr().Get(), SceneViewState);
 		
-		TSharedRef<FSceneViewFamilyContext> ViewFamily = CreateSceneViewFamily(ViewFamilyInitData, CameraInfo);
+		
+		CalculateProjectionMatrix(CameraInfo, SceneViewInitOptions, BackbufferResolution, AccumulatorResolution);
+	
+		// Modify the perspective matrix to do an off center projection, with overlap for high-res tiling
+		const bool bOrthographic = CameraInfo.ViewInfo.ProjectionMode == ECameraProjectionMode::Type::Orthographic;
+		ModifyProjectionMatrixForTiling(CameraInfo.TilingParams, bOrthographic, SceneViewInitOptions.ProjectionMatrix, CameraInfo.DoFSensorScale);
+		
+		// Scale the DoF sensor scale to counteract overscan, otherwise the size of Bokeh changes when you have Overscan enabled.
+		// This needs to come after we modify it for Tiling.
+		CameraInfo.DoFSensorScale *= 1.0 + CameraInfo.OverscanFraction;
 		 
-		// Now we can construct a View to go within this family.
-		FSceneViewInitOptions SceneViewInitOptions = CreateViewInitOptions(CameraInfo, ViewFamily.ToSharedPtr().Get(), SceneViewState);
+		// Construct a View to go within this family.
 		FSceneView* NewView = CreateSceneView(SceneViewInitOptions, ViewFamily, CameraInfo);
 		
 		// Then apply Movie Render Queue specific overrides to the ViewFamily, and then to the SceneView.
diff --git a/Engine/Plugins/MovieScene/MovieRenderPipeline/Source/MovieRenderPipelineRenderPasses/Private/Graph/Renderers/MovieGraphImagePassBase.cpp b/Engine/Plugins/MovieScene/MovieRenderPipeline/Source/MovieRenderPipelineRenderPasses/Private/Graph/Renderers/MovieGraphImagePassBase.cpp
--- a/Engine/Plugins/MovieScene/MovieRenderPipeline/Source/MovieRenderPipelineRenderPasses/Private/Graph/Renderers/MovieGraphImagePassBase.cpp
+++ b/Engine/Plugins/MovieScene/MovieRenderPipeline/Source/MovieRenderPipelineRenderPasses/Private/Graph/Renderers/MovieGraphImagePassBase.cpp
@@ -38,8 +38,7 @@ FSceneViewInitOptions FMovieGraphImagePassBase::CreateViewInitOptions(const UE::
 	ViewInitOptions.SetViewRectangle(FIntRect(FIntPoint(0, 0), RenderResolution));
 	ViewInitOptions.ViewRotationMatrix = FInverseRotationMatrix(InCameraInfo.ViewInfo.Rotation);
 	ViewInitOptions.ViewActor = InCameraInfo.ViewActor;
-	ViewInitOptions.ProjectionMatrix = InCameraInfo.ProjectionMatrix;
-
+	
 	// Rotate the view 90 degrees to match the rest of the engine.
 	ViewInitOptions.ViewRotationMatrix = ViewInitOptions.ViewRotationMatrix * FMatrix(
 		FPlane(0, 0, 1, 0),
@@ -121,13 +120,13 @@ void FMovieGraphImagePassBase::ApplyCameraManagerPostProcessBlends(FSceneView* I
 	}
 }
 
-TSharedRef<FSceneViewFamilyContext> FMovieGraphImagePassBase::CreateSceneViewFamily(const FViewFamilyInitData& InInitData, const UE::MovieGraph::DefaultRenderer::FCameraInfo& InCameraInfo) const
+TSharedRef<FSceneViewFamilyContext> FMovieGraphImagePassBase::CreateSceneViewFamily(const FViewFamilyInitData& InInitData) const
 {
 
 	EViewModeIndex ViewModeIndex = InInitData.ViewModeIndex;
 	FEngineShowFlags ShowFlags = InInitData.ShowFlags;
 
-	const bool bIsPerspective = InCameraInfo.ViewInfo.ProjectionMode == ECameraProjectionMode::Type::Perspective;
+	const bool bIsPerspective = InInitData.ProjectionMode == ECameraProjectionMode::Type::Perspective;
 
 	// Allow the Engine Showflag system to override our engine showflags, based on our view mode index.
 	// This is required for certain debug view modes (to have matching show flags set for rendering).
@@ -149,6 +148,8 @@ TSharedRef<FSceneViewFamilyContext> FMovieGraphImagePassBase::CreateSceneViewFam
 	return OutViewFamily;
 }
 
+
+
 void FMovieGraphImagePassBase::ApplyMovieGraphOverridesToViewFamily(TSharedRef<FSceneViewFamilyContext> InOutFamily, const FViewFamilyInitData& InInitData) const
 {
 	// Used to specify if the Tone Curve is being applied or not to our Linear Output data
@@ -239,19 +240,6 @@ void FMovieGraphImagePassBase::ApplyMovieGraphOverridesToSceneView(TSharedRef<FS
 		}
 	}
 
-	// Orthographic cameras don't support anti-aliasing outside the path tracer (other than FXAA)
-	const bool bIsOrthographicCamera = !View->IsPerspectiveProjection();
-	if (bIsOrthographicCamera)
-	{
-		bool bIsSupportedAAMethod = View->AntiAliasingMethod == EAntiAliasingMethod::AAM_FXAA;
-		bool bIsPathTracer = InOutFamily->EngineShowFlags.PathTracing;
-		bool bWarnJitters = InCameraInfo.ProjectionMatrixJitterAmount.SquaredLength() > SMALL_NUMBER;
-		if ((!bIsPathTracer && !bIsSupportedAAMethod) || bWarnJitters)
-		{
-			UE_LOG(LogMovieRenderPipeline, Warning, TEXT("Orthographic Cameras are only supported with PathTracer or Deferred with FXAA Anti-Aliasing"));
-		}
-	}
-
 	{
 		bool bMethodWasUnsupported = false;
 		if (View->AntiAliasingMethod == AAM_TemporalAA && !SupportsGen4TAA(View->GetShaderPlatform()))
@@ -275,8 +263,8 @@ void FMovieGraphImagePassBase::ApplyMovieGraphOverridesToSceneView(TSharedRef<FS
 	{
 		// If we're not using TAA, TSR, or Path Tracing we will apply the View Matrix projection jitter. Normally TAA sets this
 		// inside FSceneRenderer::PreVisibilityFrameSetup. Path Tracing does its own anti-aliasing internally.
-		bool bApplyProjectionJitter = !bIsOrthographicCamera
-			&& !InOutFamily->EngineShowFlags.PathTracing
+		bool bApplyProjectionJitter = 
+			   !InOutFamily->EngineShowFlags.PathTracing
 			&& !IsTemporalAccumulationBasedMethod(View->AntiAliasingMethod);
 		if (bApplyProjectionJitter)
 		{
@@ -285,116 +273,54 @@ void FMovieGraphImagePassBase::ApplyMovieGraphOverridesToSceneView(TSharedRef<FS
 	}
 }
 
-FMatrix FMovieGraphImagePassBase::CalculateProjectionMatrix(const UE::MovieGraph::DefaultRenderer::FCameraInfo& InCameraInfo) const
+void FMovieGraphImagePassBase::CalculateProjectionMatrix(UE::MovieGraph::DefaultRenderer::FCameraInfo& InOutCameraInfo, FSceneViewProjectionData& InOutProjectionData, const FIntPoint InBackbufferResolution, const FIntPoint InAccumulatorResolution) const
 {
-	// Calculate a Projection Matrix. This code unfortunately ends up similar to, but not quite the same as FMinimalViewInfo::CalculateProjectionMatrixGivenView
-	FMatrix BaseProjMatrix;
-	
 	// TileSize should respect the actual backbuffer size being used by the render.
-	float ViewRectWidth = InCameraInfo.TilingParams.TileSize.X;
-	float ViewRectHeight = InCameraInfo.TilingParams.TileSize.Y;
+	float ViewRectWidth = InBackbufferResolution.X;
+	float ViewRectHeight = InBackbufferResolution.Y;
+	FIntRect ViewRect = FIntRect(FIntPoint(0, 0), InBackbufferResolution);
 
 	const float DestAspectRatio = ViewRectWidth / ViewRectHeight;
-	const float CameraAspectRatio = InCameraInfo.bAllowCameraAspectRatio ? InCameraInfo.ViewInfo.AspectRatio : DestAspectRatio;
-	
-	float ViewFOV = InCameraInfo.ViewInfo.FOV;
+	const float CameraAspectRatio = InOutCameraInfo.bAllowCameraAspectRatio ? InOutCameraInfo.ViewInfo.AspectRatio : DestAspectRatio;
 
-	// Inflate our FOV to support the overscan 
-	ViewFOV = 2.0f * FMath::RadiansToDegrees(FMath::Atan((1.0f + InCameraInfo.OverscanFraction) * FMath::Tan(FMath::DegreesToRadians(ViewFOV * 0.5f))));
-	
-	if (InCameraInfo.ViewInfo.ProjectionMode == ECameraProjectionMode::Orthographic)
-	{
-		const float YScale = 1.0f / InCameraInfo.ViewInfo.AspectRatio;
-		const float OverscanScale = 1.0f + (InCameraInfo.OverscanFraction);
+	float ViewFOV = InOutCameraInfo.ViewInfo.FOV;
 
-		const float HalfOrthoWidth = (InCameraInfo.ViewInfo.OrthoWidth / 2.0f) * OverscanScale;
-		const float ScaledOrthoHeight = (InCameraInfo.ViewInfo.OrthoWidth / 2.0f) * OverscanScale * YScale;
+	// Inflate our FOV to support the overscan 
+	ViewFOV = 2.0f * FMath::RadiansToDegrees(FMath::Atan((1.0f + InOutCameraInfo.OverscanFraction) * FMath::Tan(FMath::DegreesToRadians(ViewFOV * 0.5f))));
+	InOutCameraInfo.ViewInfo.FOV = ViewFOV;
+	InOutCameraInfo.ViewInfo.DesiredFOV = ViewFOV;
 
-		const float NearPlane = InCameraInfo.ViewInfo.OrthoNearClipPlane;
-		const float FarPlane = InCameraInfo.ViewInfo.OrthoFarClipPlane;
+	// Overscan the Orthographic pass too.
+	InOutCameraInfo.ViewInfo.OrthoWidth *= 1.0f + InOutCameraInfo.OverscanFraction;
 
-		const float ZScale = 1.0f / (FarPlane - NearPlane);
-		const float ZOffset = -NearPlane;
+	const int TotalTileCount = InOutCameraInfo.TilingParams.TileCount.X * InOutCameraInfo.TilingParams.TileCount.Y;
 
-		BaseProjMatrix = FReversedZOrthoMatrix(
-			HalfOrthoWidth,
-			ScaledOrthoHeight,
-			ZScale,
-			ZOffset
-		);
-	}
-	else
+	// If they're using high-resolution tiling we can't support letterboxing (as the blended areas we would render with
+	// would have been cropped via letterboxing), so to handle this scenario we disable aspect ratio constraints and then
+	// manually rescale the view (if needed) to mimick the effect of letterboxing.
+	TEnumAsByte<EAspectRatioAxisConstraint> AspectRatioAxisConstraint = InOutCameraInfo.ViewInfo.AspectRatioAxisConstraint.Get(EAspectRatioAxisConstraint::AspectRatio_MaintainXFOV);
+	if (TotalTileCount > 1 && InOutCameraInfo.ViewInfo.bConstrainAspectRatio)
 	{
-		float XAxisMultiplier;
-		float YAxisMultiplier;
-
-		if (InCameraInfo.ViewInfo.bConstrainAspectRatio)
+		if (CameraAspectRatio < DestAspectRatio)
 		{
-			// If the camera's aspect ratio has a thinner width, then stretch the horizontal fov more than usual to 
-			// account for the extra with of (before constraining - after constraining)
-			if (InCameraInfo.ViewInfo.AspectRatio < DestAspectRatio)
-			{
-				const float ConstrainedWidth = ViewRectHeight * InCameraInfo.ViewInfo.AspectRatio;
-				XAxisMultiplier = ConstrainedWidth / (float)ViewRectWidth;
-				YAxisMultiplier = InCameraInfo.ViewInfo.AspectRatio;
-			}
-			// Simplified some math here but effectively functions similarly to the above, the unsimplified code would look like:
-			// const float ConstrainedHeight = ViewRectWidth / CameraCache.AspectRatio;
-			// YAxisMultiplier = (ConstrainedHeight / ViewInitOptions.GetViewRect.Height()) * CameraCache.AspectRatio;
-			else
-			{
-				XAxisMultiplier = 1.0f;
-				YAxisMultiplier = ViewRectWidth / ViewRectHeight;
-			}
+			AspectRatioAxisConstraint = EAspectRatioAxisConstraint::AspectRatio_MaintainYFOV;
+			InOutCameraInfo.ViewInfo.OrthoWidth *= (DestAspectRatio / CameraAspectRatio);
 		}
-		else
+		else if (CameraAspectRatio > DestAspectRatio)
 		{
-			const EAspectRatioAxisConstraint AspectRatioAxisConstraint = GetDefault<ULocalPlayer>()->AspectRatioAxisConstraint;
-			if (((ViewRectWidth > ViewRectHeight) && (AspectRatioAxisConstraint == AspectRatio_MajorAxisFOV)) || (AspectRatioAxisConstraint == AspectRatio_MaintainXFOV))
-			{
-				//if the viewport is wider than it is tall
-				XAxisMultiplier = 1.0f;
-				YAxisMultiplier = ViewRectWidth / ViewRectHeight;
-			}
-			else
-			{
-				//if the viewport is taller than it is wide
-				XAxisMultiplier = ViewRectHeight / ViewRectWidth;
-				YAxisMultiplier = 1.0f;
-			}
+			// Don't rescale the width and keep it X-constrained.
+			AspectRatioAxisConstraint = EAspectRatioAxisConstraint::AspectRatio_MaintainXFOV;
 		}
+		InOutCameraInfo.ViewInfo.bConstrainAspectRatio = false;
+	}
 
-		const float MinZ = InCameraInfo.ViewInfo.GetFinalPerspectiveNearClipPlane();
-		const float MaxZ = MinZ;
-		// Avoid zero ViewFOV's which cause divide by zero's in projection matrix
-		const float MatrixFOV = FMath::Max(0.001f, ViewFOV) * (float)PI / 360.0f;
+	FIntRect ViewExtents = FViewport::CalculateViewExtents(InOutCameraInfo.ViewInfo.AspectRatio, DestAspectRatio, ViewRect, InAccumulatorResolution);
 
 
-		if ((bool)ERHIZBuffer::IsInverted)
-		{
-			BaseProjMatrix = FReversedZPerspectiveMatrix(
-				MatrixFOV,
-				MatrixFOV,
-				XAxisMultiplier,
-				YAxisMultiplier,
-				MinZ,
-				MaxZ
-			);
-		}
-		else
-		{
-			BaseProjMatrix = FPerspectiveMatrix(
-				MatrixFOV,
-				MatrixFOV,
-				XAxisMultiplier,
-				YAxisMultiplier,
-				MinZ,
-				MaxZ
-			);
-		}
-	}
 	
-	return BaseProjMatrix;
+
+	// This function updates data in both the FMinimalViewInfo and the ProjectionData
+	FMinimalViewInfo::CalculateProjectionMatrixGivenViewRectangle(InOutCameraInfo.ViewInfo, AspectRatioAxisConstraint, ViewExtents, InOutProjectionData);
 }
 
 FVector4f FMovieGraphImagePassBase::CalculatePrinciplePointOffsetForTiling(const UE::MovieGraph::DefaultRenderer::FMovieGraphTilingParams& InTilingParams) const 
diff --git a/Engine/Plugins/MovieScene/MovieRenderPipeline/Source/MovieRenderPipelineRenderPasses/Private/MoviePipelineImagePassBase.cpp b/Engine/Plugins/MovieScene/MovieRenderPipeline/Source/MovieRenderPipelineRenderPasses/Private/MoviePipelineImagePassBase.cpp
--- a/Engine/Plugins/MovieScene/MovieRenderPipeline/Source/MovieRenderPipelineRenderPasses/Private/MoviePipelineImagePassBase.cpp
+++ b/Engine/Plugins/MovieScene/MovieRenderPipeline/Source/MovieRenderPipelineRenderPasses/Private/MoviePipelineImagePassBase.cpp
@@ -292,19 +292,6 @@ TSharedPtr<FSceneViewFamilyContext> UMoviePipelineImagePassBase::CalculateViewFa
 		}
 	}
 
-	// Orthographic cameras don't support anti-aliasing outside the path tracer (other than FXAA)
-	const bool bIsOrthographicCamera = !View->IsPerspectiveProjection();
-	if (bIsOrthographicCamera)
-	{
-		bool bIsSupportedAAMethod = View->AntiAliasingMethod == EAntiAliasingMethod::AAM_FXAA;
-		bool bIsPathTracer = OutViewFamily->EngineShowFlags.PathTracing;
-		bool bWarnJitters = InOutSampleState.ProjectionMatrixJitterAmount.SquaredLength() > SMALL_NUMBER;
-		if ((!bIsPathTracer && !bIsSupportedAAMethod) || bWarnJitters)
-		{
-			UE_LOG(LogMovieRenderPipeline, Warning, TEXT("Orthographic Cameras are only supported with PathTracer or Deferred with FXAA Anti-Aliasing"));
-		}
-	}
-
 	OutViewFamily->ViewExtensions.Append(GEngine->ViewExtensions->GatherActiveExtensions(FSceneViewExtensionContext(GetWorld()->Scene)));
 
 	AddViewExtensions(*OutViewFamily, InOutSampleState);
@@ -345,8 +332,8 @@ TSharedPtr<FSceneViewFamilyContext> UMoviePipelineImagePassBase::CalculateViewFa
 	{
 		// If we're not using Temporal Anti-Aliasing or Path Tracing we will apply the View Matrix projection jitter. Normally TAA sets this
 		// inside FSceneRenderer::PreVisibilityFrameSetup. Path Tracing does its own anti-aliasing internally.
-		bool bApplyProjectionJitter = !bIsOrthographicCamera
-									&& !OutViewFamily->EngineShowFlags.PathTracing 
+		bool bApplyProjectionJitter = 
+									!OutViewFamily->EngineShowFlags.PathTracing 
 									&& !IsTemporalAccumulationBasedMethod(View->AntiAliasingMethod);
 		if (bApplyProjectionJitter)
 		{
@@ -470,7 +457,9 @@ FSceneView* UMoviePipelineImagePassBase::GetSceneViewForSampleState(FSceneViewFa
 	FSceneViewInitOptions ViewInitOptions;
 	ViewInitOptions.ViewFamily = ViewFamily;
 	ViewInitOptions.ViewOrigin = CameraInfo.ViewInfo.Location;
-	ViewInitOptions.SetViewRectangle(FIntRect(FIntPoint(0, 0), FIntPoint(TileSizeX, TileSizeY)));
+	FIntRect ViewRect = FIntRect(FIntPoint(0, 0), FIntPoint(TileSizeX, TileSizeY));
+	ViewInitOptions.SetViewRectangle(ViewRect);
+
 	ViewInitOptions.ViewRotationMatrix = FInverseRotationMatrix(CameraInfo.ViewInfo.Rotation);
 	ViewInitOptions.ViewActor = CameraInfo.ViewActor;
 
@@ -480,135 +469,60 @@ FSceneView* UMoviePipelineImagePassBase::GetSceneViewForSampleState(FSceneViewFa
 		FPlane(1, 0, 0, 0),
 		FPlane(0, 1, 0, 0),
 		FPlane(0, 0, 0, 1));
-	float ViewFOV = CameraInfo.ViewInfo.FOV;
 
-	// Inflate our FOV to support the overscan 
-	ViewFOV = 2.0f * FMath::RadiansToDegrees(FMath::Atan((1.0f + InOutSampleState.OverscanPercentage) * FMath::Tan(FMath::DegreesToRadians( ViewFOV * 0.5f ))));
+	// Inflate our FOV to support the overscan, we do this before calculating the Projection Matrix so that any downstream consumers of the FOV/DesiredFOV are taken into account (ie: culling)
+	float ViewFOV = 2.0f * FMath::RadiansToDegrees(FMath::Atan((1.0f + InOutSampleState.OverscanPercentage) * FMath::Tan(FMath::DegreesToRadians(CameraInfo.ViewInfo.FOV * 0.5f ))));
+	CameraInfo.ViewInfo.FOV = ViewFOV;
+	CameraInfo.ViewInfo.DesiredFOV = ViewFOV;
+	ViewInitOptions.FOV = CameraInfo.ViewInfo.FOV;
+	ViewInitOptions.DesiredFOV = CameraInfo.ViewInfo.DesiredFOV;
 
-	float DofSensorScale = 1.0f;
+	// Overscan the Orthographic pass too.
+	CameraInfo.ViewInfo.OrthoWidth *= 1.0f + InOutSampleState.OverscanPercentage;
 
-	// Calculate a Projection Matrix. This code unfortunately ends up similar to, but not quite the same as FMinimalViewInfo::CalculateProjectionMatrixGivenView
-	FMatrix BaseProjMatrix;
+	float DofSensorScale = 1.0f;
 
 	if (CameraInfo.bUseCustomProjectionMatrix)
 	{
-		BaseProjMatrix = CameraInfo.CustomProjectionMatrix;
+		ViewInitOptions.ProjectionMatrix = CameraInfo.CustomProjectionMatrix;
 
 		// Modify the custom matrix to do an off center projection, with overlap for high-res tiling
 		const bool bOrthographic = false;
-		ModifyProjectionMatrixForTiling(InOutSampleState, bOrthographic, /*InOut*/ BaseProjMatrix, DofSensorScale);
+		ModifyProjectionMatrixForTiling(InOutSampleState, bOrthographic, /*InOut*/ ViewInitOptions.ProjectionMatrix, /*Out*/ DofSensorScale);
 	}
 	else
 	{
-		if (CameraInfo.ViewInfo.ProjectionMode == ECameraProjectionMode::Orthographic)
+		// If they're using high-resolution tiling we can't support letterboxing (as the blended areas we would render with
+		// would have been cropped via letterboxing), so to handle this scenario we disable aspect ratio constraints and then
+		// manually rescale the view (if needed) to mimick the effect of letterboxing.
+		TEnumAsByte<EAspectRatioAxisConstraint> AspectRatioAxisConstraint = CameraInfo.ViewInfo.AspectRatioAxisConstraint.Get(EAspectRatioAxisConstraint::AspectRatio_MaintainXFOV);
+		if (InOutSampleState.GetTileCount() > 1 && CameraInfo.ViewInfo.bConstrainAspectRatio)
 		{
-			const float YScale = 1.0f / CameraAspectRatio;
-			const float OverscanScale = 1.0f + InOutSampleState.OverscanPercentage;
-
-			const float HalfOrthoWidth = (CameraInfo.ViewInfo.OrthoWidth / 2.0f) * OverscanScale;
-			const float ScaledOrthoHeight = (CameraInfo.ViewInfo.OrthoWidth / 2.0f) * OverscanScale * YScale;
-
-			const float NearPlane = CameraInfo.ViewInfo.OrthoNearClipPlane;
-			const float FarPlane = CameraInfo.ViewInfo.OrthoFarClipPlane;
-
-			const float ZScale = 1.0f / (FarPlane - NearPlane);
-			const float ZOffset = -NearPlane;
-
-			BaseProjMatrix = FReversedZOrthoMatrix(
-				HalfOrthoWidth,
-				ScaledOrthoHeight,
-				ZScale,
-				ZOffset
-			);
-			
-			// Modify the projection matrix to do an off center projection, with overlap for high-res tiling
-			const bool bOrthographic = true;
-			ModifyProjectionMatrixForTiling(InOutSampleState, bOrthographic, /*InOut*/ BaseProjMatrix, DofSensorScale);
-		}
-		else
-		{
-			float XAxisMultiplier;
-			float YAxisMultiplier;
-
-			if (CameraInfo.ViewInfo.bConstrainAspectRatio)
+			if (CameraAspectRatio < DestAspectRatio)
 			{
-				// If the camera's aspect ratio has a thinner width, then stretch the horizontal fov more than usual to 
-				// account for the extra with of (before constraining - after constraining)
-				if (CameraAspectRatio < DestAspectRatio)
-				{
-					const float ConstrainedWidth = ViewInitOptions.GetViewRect().Height() * CameraAspectRatio;
-					XAxisMultiplier = ConstrainedWidth / (float)ViewInitOptions.GetViewRect().Width();
-					YAxisMultiplier = CameraAspectRatio;
-				}
-				// Simplified some math here but effectively functions similarly to the above, the unsimplified code would look like:
-				// const float ConstrainedHeight = ViewInitOptions.GetViewRect().Width() / CameraCache.AspectRatio;
-				// YAxisMultiplier = (ConstrainedHeight / ViewInitOptions.GetViewRect.Height()) * CameraCache.AspectRatio;
-				else
-				{
-					XAxisMultiplier = 1.0f;
-					YAxisMultiplier = ViewInitOptions.GetViewRect().Width() / (float)ViewInitOptions.GetViewRect().Height();
-				}
+				AspectRatioAxisConstraint = EAspectRatioAxisConstraint::AspectRatio_MaintainYFOV;
+				CameraInfo.ViewInfo.OrthoWidth *= (DestAspectRatio / CameraAspectRatio);
 			}
-			else
+			else if (CameraAspectRatio > DestAspectRatio)
 			{
-				const int32 DestSizeX = ViewInitOptions.GetViewRect().Width();
-				const int32 DestSizeY = ViewInitOptions.GetViewRect().Height();
-				const EAspectRatioAxisConstraint AspectRatioAxisConstraint = GetDefault<ULocalPlayer>()->AspectRatioAxisConstraint;
-				if (((DestSizeX > DestSizeY) && (AspectRatioAxisConstraint == AspectRatio_MajorAxisFOV)) || (AspectRatioAxisConstraint == AspectRatio_MaintainXFOV))
-				{
-					//if the viewport is wider than it is tall
-					XAxisMultiplier = 1.0f;
-					YAxisMultiplier = ViewInitOptions.GetViewRect().Width() / (float)ViewInitOptions.GetViewRect().Height();
-				}
-				else
-				{
-					//if the viewport is taller than it is wide
-					XAxisMultiplier = ViewInitOptions.GetViewRect().Height() / (float)ViewInitOptions.GetViewRect().Width();
-					YAxisMultiplier = 1.0f;
-				}
+				// Don't rescale the width and keep it X-constrained.
+				AspectRatioAxisConstraint = EAspectRatioAxisConstraint::AspectRatio_MaintainXFOV;
 			}
+			CameraInfo.ViewInfo.bConstrainAspectRatio = false;
+		}
 
-			const float MinZ = CameraInfo.ViewInfo.GetFinalPerspectiveNearClipPlane();
-			const float MaxZ = MinZ;
-			// Avoid zero ViewFOV's which cause divide by zero's in projection matrix
-			const float MatrixFOV = FMath::Max(0.001f, ViewFOV) * (float)PI / 360.0f;
 
+		FIntRect ViewExtents = FViewport::CalculateViewExtents(CameraInfo.ViewInfo.AspectRatio, DestAspectRatio, ViewRect, InOutSampleState.BackbufferSize);
+		FMinimalViewInfo::CalculateProjectionMatrixGivenViewRectangle(CameraInfo.ViewInfo, AspectRatioAxisConstraint, ViewExtents, ViewInitOptions);
 
-			if ((bool)ERHIZBuffer::IsInverted)
-			{
-				BaseProjMatrix = FReversedZPerspectiveMatrix(
-					MatrixFOV,
-					MatrixFOV,
-					XAxisMultiplier,
-					YAxisMultiplier,
-					MinZ,
-					MaxZ
-				);
-			}
-			else
-			{
-				BaseProjMatrix = FPerspectiveMatrix(
-					MatrixFOV,
-					MatrixFOV,
-					XAxisMultiplier,
-					YAxisMultiplier,
-					MinZ,
-					MaxZ
-				);
-			}
-
-			// Modify the perspective matrix to do an off center projection, with overlap for high-res tiling
-			const bool bOrthographic = false;
-			ModifyProjectionMatrixForTiling(InOutSampleState, bOrthographic, /*InOut*/ BaseProjMatrix, DofSensorScale);
-			// ToDo: Does orthographic support tiling in the same way or do I need to modify the values before creating the ortho view.
-		}
+		ModifyProjectionMatrixForTiling(InOutSampleState, CameraInfo.ViewInfo.ProjectionMode == ECameraProjectionMode::Orthographic,  /*InOut*/ ViewInitOptions.ProjectionMatrix, /*Out*/ DofSensorScale);
 	}
-		// BaseProjMatrix may be perspective or orthographic.
-		ViewInitOptions.ProjectionMatrix = BaseProjMatrix;
+
+
+	// Scale the DoF sensor scale to counteract overscan, otherwise the size of Bokeh changes when you have Overscan enabled.
+	DofSensorScale *= 1.0 + InOutSampleState.OverscanPercentage;
 
 	ViewInitOptions.SceneViewStateInterface = GetSceneViewStateInterface(OptPayload);
-	ViewInitOptions.FOV = ViewFOV;
-	ViewInitOptions.DesiredFOV = ViewFOV;
 
 	FSceneView* View = new FSceneView(ViewInitOptions);
 	ViewFamily->Views.Add(View);
@@ -625,6 +539,7 @@ FSceneView* UMoviePipelineImagePassBase::GetSceneViewForSampleState(FSceneViewFa
 	// Scaling sensor size inversely with the the projection matrix [0][0] should physically
 	// cause the circle of confusion to be unchanged.
 	View->FinalPostProcessSettings.DepthOfFieldSensorWidth *= DofSensorScale;
+
 	// Modify the 'center' of the lens to be offset for high-res tiling, helps some effects (vignette) etc. still work.
 	View->LensPrincipalPointOffsetScale = (FVector4f)CalculatePrinciplePointOffsetForTiling(InOutSampleState); // LWC_TODO: precision loss. CalculatePrinciplePointOffsetForTiling() could return float, it's normalized?
 	View->EndFinalPostprocessSettings(ViewInitOptions);
diff --git a/Engine/Plugins/MovieScene/MovieRenderPipeline/Source/MovieRenderPipelineRenderPasses/Public/Graph/Renderers/MovieGraphImagePassBase.h b/Engine/Plugins/MovieScene/MovieRenderPipeline/Source/MovieRenderPipelineRenderPasses/Public/Graph/Renderers/MovieGraphImagePassBase.h
--- a/Engine/Plugins/MovieScene/MovieRenderPipeline/Source/MovieRenderPipelineRenderPasses/Public/Graph/Renderers/MovieGraphImagePassBase.h
+++ b/Engine/Plugins/MovieScene/MovieRenderPipeline/Source/MovieRenderPipelineRenderPasses/Public/Graph/Renderers/MovieGraphImagePassBase.h
@@ -48,6 +48,7 @@ namespace UE::MovieGraph::Rendering
 		EAntiAliasingMethod AntiAliasingMethod;
 		FEngineShowFlags ShowFlags;
 		EViewModeIndex ViewModeIndex;
+		TEnumAsByte<ECameraProjectionMode::Type> ProjectionMode;
 	};
 	
 
@@ -67,7 +68,7 @@ namespace UE::MovieGraph::Rendering
 
 	protected:
 		/** Utility function for calculating a Projection Matrix (Orthographic or Perspective), and modify it based on the overscan percentage, aspect ratio, etc. */
-		virtual FMatrix CalculateProjectionMatrix(const UE::MovieGraph::DefaultRenderer::FCameraInfo& InCameraInfo) const;
+		virtual void CalculateProjectionMatrix(UE::MovieGraph::DefaultRenderer::FCameraInfo& InOutCameraInfo, FSceneViewProjectionData& InOutProjectionData, const FIntPoint InBackbufferResolution, const FIntPoint InAccumulatorResolution) const;
 		/** Utility function for modifying the projection matrix to match the given tiling parameters. */
 		virtual void ModifyProjectionMatrixForTiling(const UE::MovieGraph::DefaultRenderer::FMovieGraphTilingParams& InTilingParams, const bool bInOrthographic, FMatrix& InOutProjectionMatrix, float& OutDoFSensorScale) const;
 		/** Utility function for calculating the PrinciplePointOffset with the given tiling parameters. Used to make some effects (like vignette) work with tiling. */
@@ -78,7 +79,7 @@ namespace UE::MovieGraph::Rendering
 		virtual FSceneView* CreateSceneView(const FSceneViewInitOptions& InInitOptions, TSharedRef<FSceneViewFamilyContext> InViewFamily, const UE::MovieGraph::DefaultRenderer::FCameraInfo& InCameraInfo) const;
 
 		virtual void ApplyCameraManagerPostProcessBlends(FSceneView* InView) const;
-		virtual TSharedRef<FSceneViewFamilyContext> CreateSceneViewFamily(const FViewFamilyInitData& InInitData, const UE::MovieGraph::DefaultRenderer::FCameraInfo& InCameraInfo) const;
+		virtual TSharedRef<FSceneViewFamilyContext> CreateSceneViewFamily(const FViewFamilyInitData& InInitData) const;
 		virtual void ApplyMovieGraphOverridesToSceneView(TSharedRef<FSceneViewFamilyContext> InOutFamily, const FViewFamilyInitData& InInitData, const UE::MovieGraph::DefaultRenderer::FCameraInfo& InCameraInfo) const;
 		virtual void ApplyMovieGraphOverridesToViewFamily(TSharedRef<FSceneViewFamilyContext> InOutFamily, const FViewFamilyInitData& InInitData) const;
 		virtual void PostRendererSubmission(const UE::MovieGraph::FMovieGraphSampleState& InSampleState, const UE::MovieGraph::DefaultRenderer::FRenderTargetInitParams& InRenderTargetInitParams, FCanvas& InCanvas, const UE::MovieGraph::DefaultRenderer::FCameraInfo& InCameraInfo) const;
diff --git a/Engine/Source/Runtime/Engine/Private/UnrealClient.cpp b/Engine/Source/Runtime/Engine/Private/UnrealClient.cpp
--- a/Engine/Source/Runtime/Engine/Private/UnrealClient.cpp
+++ b/Engine/Source/Runtime/Engine/Private/UnrealClient.cpp
@@ -2181,28 +2181,28 @@ void FViewport::UpdateViewportRHI(bool bDestroyed, uint32 NewSizeX, uint32 NewSi
 	}
 }
 
-FIntRect FViewport::CalculateViewExtents(float AspectRatio, const FIntRect& ViewRect)
+FIntRect FViewport::CalculateViewExtents(float AspectRatio, float DesiredAspectRatio, const FIntRect& ViewRect, const FIntPoint& DestSize)
 {
 	FIntRect Result = ViewRect;
 
 	const float CurrentSizeX = ViewRect.Width();
 	const float CurrentSizeY = ViewRect.Height();
 
-	// the viewport's SizeX/SizeY may not always match the GetDesiredAspectRatio(), so adjust the requested AspectRatio to compensate
-	const float AdjustedAspectRatio = AspectRatio / (GetDesiredAspectRatio() / ((float)GetSizeXY().X / (float)GetSizeXY().Y));
+	// the viewport's SizeX/SizeY may not always match the DesiredAspectRatio, so adjust the requested AspectRatio to compensate
+	const float AdjustedAspectRatio = AspectRatio / (DesiredAspectRatio / ((float)DestSize.X / (float)DestSize.Y));
 
 	// If desired, enforce a particular aspect ratio for the render of the scene. 
 	// Results in black bars at top/bottom etc.
 	const float AspectRatioDifference = AdjustedAspectRatio - (CurrentSizeX / CurrentSizeY);
 
-	if( FMath::Abs( AspectRatioDifference ) > 0.01f )
+	if (FMath::Abs(AspectRatioDifference) > 0.01f)
 	{
 		// If desired aspect ratio is bigger than current - we need black bars on top and bottom.
-		if( AspectRatioDifference > 0.0f )
+		if (AspectRatioDifference > 0.0f)
 		{
 			// Calculate desired Y size.
-			const int32 NewSizeY = FMath::Max(1, FMath::RoundToInt( CurrentSizeX / AdjustedAspectRatio ) );
-			Result.Min.Y = FMath::RoundToInt( 0.5f * (CurrentSizeY - NewSizeY) );
+			const int32 NewSizeY = FMath::Max(1, FMath::RoundToInt(CurrentSizeX / AdjustedAspectRatio));
+			Result.Min.Y = FMath::RoundToInt(0.5f * (CurrentSizeY - NewSizeY));
 			Result.Max.Y = Result.Min.Y + NewSizeY;
 			Result.Min.Y += ViewRect.Min.Y;
 			Result.Max.Y += ViewRect.Min.Y;
@@ -2210,8 +2210,8 @@ FIntRect FViewport::CalculateViewExtents(float AspectRatio, const FIntRect& View
 		// Otherwise - will place bars on the sides.
 		else
 		{
-			const int32 NewSizeX = FMath::Max(1, FMath::RoundToInt( CurrentSizeY * AdjustedAspectRatio ) );
-			Result.Min.X = FMath::RoundToInt( 0.5f * (CurrentSizeX - NewSizeX) );
+			const int32 NewSizeX = FMath::Max(1, FMath::RoundToInt(CurrentSizeY * AdjustedAspectRatio));
+			Result.Min.X = FMath::RoundToInt(0.5f * (CurrentSizeX - NewSizeX));
 			Result.Max.X = Result.Min.X + NewSizeX;
 			Result.Min.X += ViewRect.Min.X;
 			Result.Max.X += ViewRect.Min.X;
@@ -2221,6 +2221,11 @@ FIntRect FViewport::CalculateViewExtents(float AspectRatio, const FIntRect& View
 	return Result;
 }
 
+FIntRect FViewport::CalculateViewExtents(float AspectRatio, const FIntRect& ViewRect)
+{
+	return CalculateViewExtents(AspectRatio, GetDesiredAspectRatio(), ViewRect, GetSizeXY());
+}
+
 /**
  *	Sets a viewport client if one wasn't provided at construction time.
  *	@param InViewportClient	- The viewport client to set.
diff --git a/Engine/Source/Runtime/Engine/Public/UnrealClient.h b/Engine/Source/Runtime/Engine/Public/UnrealClient.h
--- a/Engine/Source/Runtime/Engine/Public/UnrealClient.h
+++ b/Engine/Source/Runtime/Engine/Public/UnrealClient.h
@@ -604,6 +604,17 @@ public:
 	 * @return	constrained view rectangle
 	 */
 	ENGINE_API FIntRect CalculateViewExtents(float AspectRatio, const FIntRect& ViewRect);
+	
+	/**
+	* Calculates the view inside the viewport when the aspect ratio is locked.
+	* Used for creating cinematic bars.
+	* @param AspectRatio [in] ratio to lock to
+	* @param DesiredAspectRatio [in] the aspect ratio of the viewport
+	* @param ViewRect [in] unconstrained view rectangle
+	* @param DestSize [in] the size of the the viewport
+	* @return constrained view rectangle
+	*/
+	ENGINE_API static FIntRect CalculateViewExtents(float AspectRatio, float DesiredAspectRatio, const FIntRect& ViewRect, const FIntPoint& DestSize);
 
 	/**
 	 *	Sets a viewport client if one wasn't provided at construction time.
