MoviePipeline: Added a new "Set Metadata Attributes" node which adds metadata to graph output files (currently, only EXRs will pick up the metadata).

diff --git a/Engine/Plugins/MovieScene/MovieRenderPipeline/Source/MovieRenderPipelineCore/Private/Graph/Nodes/MovieGraphSetMetadataAttributesNode.cpp b/Engine/Plugins/MovieScene/MovieRenderPipeline/Source/MovieRenderPipelineCore/Private/Graph/Nodes/MovieGraphSetMetadataAttributesNode.cpp
new file mode 100644
--- /dev/null
+++ b/Engine/Plugins/MovieScene/MovieRenderPipeline/Source/MovieRenderPipelineCore/Private/Graph/Nodes/MovieGraphSetMetadataAttributesNode.cpp
@@ -0,0 +1,138 @@
+// Copyright Epic Games, Inc. All Rights Reserved.
+
+#include "Graph/Nodes/MovieGraphSetMetadataAttributesNode.h"
+
+#include "Graph/MovieGraphConfig.h"
+#include "Styling/AppStyle.h"
+
+void UMovieGraphMetadataAttributeCollection::Merge(const IMovieGraphTraversableObject* InSourceObject)
+{
+	const UMovieGraphMetadataAttributeCollection* MovieGraphMetadataAttributeCollection = Cast<UMovieGraphMetadataAttributeCollection>(InSourceObject);
+	checkf(MovieGraphMetadataAttributeCollection, TEXT("UMovieGraphMetadataAttributeCollection cannot merge with null or an object of another type."));
+
+	// The Movie Graph is evaluated right to left, i.e., from the output to the input, thus we need to keep track of the metadata attributes that
+	// have already been merged so they are not added again upon merging a new source object.
+	TSet<FString> ExistingMetadataAttributes;
+	for (const FMovieGraphMetadataAttribute& MetadataAttribute : MetadataAttributes)
+	{
+		if (MetadataAttribute.bIsEnabled && !MetadataAttribute.Name.IsEmpty())
+		{
+			ExistingMetadataAttributes.Add(MetadataAttribute.Name);
+		}
+	}
+
+	for (const FMovieGraphMetadataAttribute& SourceMetadataAttribute : MovieGraphMetadataAttributeCollection->MetadataAttributes)
+	{
+		if (SourceMetadataAttribute.bIsEnabled && !SourceMetadataAttribute.Name.IsEmpty() && !ExistingMetadataAttributes.Contains(SourceMetadataAttribute.Name))
+		{
+			MetadataAttributes.Add(FMovieGraphMetadataAttribute(SourceMetadataAttribute.Name, SourceMetadataAttribute.Value, SourceMetadataAttribute.bIsEnabled));
+		}
+	}
+}
+
+TArray<TPair<FString, FString>> UMovieGraphMetadataAttributeCollection::GetMergedProperties() const
+{
+	TArray<TPair<FString, FString>> MergedProperties;
+
+	for (const FMovieGraphMetadataAttribute& Metadata : MetadataAttributes)
+	{
+		if (Metadata.bIsEnabled && !Metadata.Name.IsEmpty())
+		{
+			MergedProperties.Add({Metadata.Name, Metadata.Value});
+		}
+	}
+	
+	return MergedProperties;
+}
+
+UMovieGraphSetMetadataAttributesNode::UMovieGraphSetMetadataAttributesNode()
+{
+	MetadataAttributeCollection = CreateDefaultSubobject<UMovieGraphMetadataAttributeCollection>(TEXT("Metadata Attributes"));
+}
+
+void UMovieGraphSetMetadataAttributesNode::GetFormatResolveArgs(FMovieGraphResolveArgs& OutMergedFormatArgs, const FMovieGraphRenderDataIdentifier& InRenderDataIdentifier) const
+{
+	for (const FMovieGraphMetadataAttribute& MetadataAttribute: MetadataAttributeCollection->MetadataAttributes)
+	{
+		if (MetadataAttribute.bIsEnabled && !MetadataAttribute.Name.IsEmpty())
+		{
+			OutMergedFormatArgs.FileMetadata.Add(MetadataAttribute.Name, MetadataAttribute.Value);
+		}
+	}
+}
+
+#if WITH_EDITOR
+FText UMovieGraphSetMetadataAttributesNode::GetNodeTitle(const bool bGetDescriptive) const
+{
+	static const FText SetMetadataAttributesNodeName = NSLOCTEXT("MovieGraphNodes", "NodeName_SetMetadataAttributes", "Set Metadata Attributes");
+	static const FText SetMetadataAttributeNodesDescription = NSLOCTEXT("MovieGraphNodes", "NodeDescription_SetMetadataAttributes", "Set Metadata Attributes\n{0}");
+
+	if (bGetDescriptive && IsValid(MetadataAttributeCollection))
+	{
+		TArray<FString> MetadataAttributeNames;
+		for (const FMovieGraphMetadataAttribute& MetadataAttribute: MetadataAttributeCollection->MetadataAttributes)
+		{
+			if (MetadataAttribute.bIsEnabled && !MetadataAttribute.Name.IsEmpty())
+			{
+				MetadataAttributeNames.Add(MetadataAttribute.Name);
+			}
+		}
+
+		if (!MetadataAttributeNames.IsEmpty())
+		{
+			// Shorten the list of attribute names if needed so it doesn't make the node too wide
+			FString AttributeNames = FString::Join(MetadataAttributeNames, TEXT(", "));
+			if (AttributeNames.Len() > 50)
+			{
+				AttributeNames = AttributeNames.Left(50) + TEXT("...");
+			}
+			
+			return FText::Format(SetMetadataAttributeNodesDescription, FText::FromString(AttributeNames));
+		}
+	}
+	
+	return SetMetadataAttributesNodeName;
+}
+
+FText UMovieGraphSetMetadataAttributesNode::GetMenuCategory() const
+{
+	return NSLOCTEXT("MovieGraphNodes", "SetMetadataAttributesGraphNode_Category", "Utility");
+}
+
+FText UMovieGraphSetMetadataAttributesNode::GetKeywords() const
+{
+	static const FText Keywords = NSLOCTEXT("MovieGraphNodes", "SetMetadataAttributesGraphNode_Keywords", "metadata attributes exr");
+	return Keywords;
+}
+
+FLinearColor UMovieGraphSetMetadataAttributesNode::GetNodeTitleColor() const
+{
+	static const FLinearColor SetMetadataAttributesNodeColor = FLinearColor(0.22f, 0.04f,  0.36f);
+	return SetMetadataAttributesNodeColor;
+}
+
+FSlateIcon UMovieGraphSetMetadataAttributesNode::GetIconAndTint(FLinearColor& OutColor) const
+{
+	static const FSlateIcon SetMetadataAttributesIcon = FSlateIcon(FAppStyle::GetAppStyleSetName(), "DerivedData.ResourceUsage");
+
+	OutColor = FLinearColor::White;
+	return SetMetadataAttributesIcon;
+}
+
+void UMovieGraphSetMetadataAttributesNode::PostEditChangeProperty(FPropertyChangedEvent& PropertyChangedEvent)
+{
+	Super::PostEditChangeProperty(PropertyChangedEvent);
+	
+	// Skip rapid updates from properties; only refresh on commit
+	if (PropertyChangedEvent.ChangeType == EPropertyChangeType::Interactive)
+	{
+		return;
+	}
+	
+	if (PropertyChangedEvent.GetPropertyName() == GET_MEMBER_NAME_CHECKED(UMovieGraphSetMetadataAttributesNode, MetadataAttributeCollection))
+	{
+		OnNodeChangedDelegate.Broadcast(this);
+	}
+}
+
+#endif // WITH_EDITOR
diff --git a/Engine/Plugins/MovieScene/MovieRenderPipeline/Source/MovieRenderPipelineCore/Public/Graph/Nodes/MovieGraphSetMetadataAttributesNode.h b/Engine/Plugins/MovieScene/MovieRenderPipeline/Source/MovieRenderPipelineCore/Public/Graph/Nodes/MovieGraphSetMetadataAttributesNode.h
new file mode 100644
--- /dev/null
+++ b/Engine/Plugins/MovieScene/MovieRenderPipeline/Source/MovieRenderPipelineCore/Public/Graph/Nodes/MovieGraphSetMetadataAttributesNode.h
@@ -0,0 +1,93 @@
+// Copyright Epic Games, Inc. All Rights Reserved.
+
+#pragma once
+
+#include "Graph/MovieGraphNode.h"
+#include "Graph/MovieGraphConfig.h"
+
+#include "MovieGraphSetMetadataAttributesNode.generated.h"
+
+/**
+ * Represents a metadata attribute that can be included in a file.
+ */
+USTRUCT(BlueprintType)
+struct FMovieGraphMetadataAttribute
+{
+	GENERATED_BODY()
+	
+	FMovieGraphMetadataAttribute(const FString& InName, const FString& InValue, const bool bInIsEnabled = true)
+		: Name(InName)
+		, Value(InValue)
+		, bIsEnabled(bInIsEnabled)
+	{
+	}
+
+	FMovieGraphMetadataAttribute()
+		: bIsEnabled(true)
+	{
+	}
+
+	/* The name of the metadata attribute. */
+	UPROPERTY(EditAnywhere, BlueprintReadWrite, Category = "Settings", meta=(DisplayName="Metadata Name"))
+	FString Name;
+
+	/* The value of the metadata attribute. */
+	UPROPERTY(EditAnywhere, BlueprintReadWrite, Category = "Settings", meta=(DisplayName="Metadata Value"))
+	FString Value;
+
+	/* Enable state. If disabled, this metadata attribute won't be added to the file. */
+	UPROPERTY(EditAnywhere, BlueprintReadWrite, Category = "Settings", meta=(DisplayName="Metadata Enable State"))
+	bool bIsEnabled;
+};
+
+UCLASS(BlueprintType)
+class MOVIERENDERPIPELINECORE_API UMovieGraphMetadataAttributeCollection : public UObject, public IMovieGraphTraversableObject
+{
+	GENERATED_BODY()
+
+public:
+	/**
+	* An array of metadata attributes to be added to output files.
+	*/
+	UPROPERTY(EditAnywhere, BlueprintReadWrite, Category="Metadata", meta=(DisplayName="Metadata Attributes"))
+	TArray<FMovieGraphMetadataAttribute> MetadataAttributes;
+
+	// IMovieGraphTraversableObject interface
+	virtual void Merge(const IMovieGraphTraversableObject* InSourceObject) override;
+	virtual TArray<TPair<FString, FString>> GetMergedProperties() const override;
+	// ~IMovieGraphTraversableObject interface
+};
+
+/** A node which can set a specific metadata attributes. */
+UCLASS()
+class MOVIERENDERPIPELINECORE_API UMovieGraphSetMetadataAttributesNode : public UMovieGraphSettingNode
+{
+	GENERATED_BODY()
+
+public:
+	UMovieGraphSetMetadataAttributesNode();
+
+	virtual void GetFormatResolveArgs(FMovieGraphResolveArgs& OutMergedFormatArgs, const FMovieGraphRenderDataIdentifier& InRenderDataIdentifier) const override;
+
+#if WITH_EDITOR
+	virtual FText GetNodeTitle(const bool bGetDescriptive = false) const override;
+	virtual FText GetMenuCategory() const override;
+	virtual FText GetKeywords() const override;
+	virtual FLinearColor GetNodeTitleColor() const override;
+	virtual FSlateIcon GetIconAndTint(FLinearColor& OutColor) const override;
+	
+	//~ Begin UObject interface
+	virtual void PostEditChangeProperty(FPropertyChangedEvent& PropertyChangedEvent) override;
+	//~ End UObject interface
+#endif
+
+public:
+	UPROPERTY(EditAnywhere, BlueprintReadWrite, Category = Overrides, meta = (InlineEditConditionToggle))
+	uint8 bOverride_MetadataAttributeCollection : 1 = 1;
+	
+	/**
+	* A container of metadata attributes to be added to the file.
+	*/
+	UPROPERTY(VisibleAnywhere, BlueprintReadWrite, Instanced, Category="Metadata", meta=(DisplayName="Metadata Attribute Collection", EditCondition="bOverride_MetadataAttributeCollection"))
+	TObjectPtr<UMovieGraphMetadataAttributeCollection> MetadataAttributeCollection;
+};
\ No newline at end of file
diff --git a/Engine/Plugins/MovieScene/MovieRenderPipeline/Source/MovieRenderPipelineEditor/Private/Customizations/Graph/MovieGraphMetadataAttributeCustomization.h b/Engine/Plugins/MovieScene/MovieRenderPipeline/Source/MovieRenderPipelineEditor/Private/Customizations/Graph/MovieGraphMetadataAttributeCustomization.h
new file mode 100644
--- /dev/null
+++ b/Engine/Plugins/MovieScene/MovieRenderPipeline/Source/MovieRenderPipelineEditor/Private/Customizations/Graph/MovieGraphMetadataAttributeCustomization.h
@@ -0,0 +1,158 @@
+// Copyright Epic Games, Inc. All Rights Reserved.
+
+#pragma once
+
+#include "Graph/Nodes/MovieGraphSetMetadataAttributesNode.h"
+
+#include "DetailLayoutBuilder.h"
+#include "Internationalization/Internationalization.h"
+#include "PropertyCustomizationHelpers.h"
+#include "Widgets/Input/SCheckBox.h"
+#include "Widgets/Input/SEditableTextBox.h"
+
+#define LOCTEXT_NAMESPACE "FMovieGraphMetadataAttributeCustomization"
+
+/** Customize how a metadata attribute (FMovieGraphMetadataAttribute) appears in the details panel. */
+class MOVIERENDERPIPELINEEDITOR_API FMovieGraphMetadataAttributeCustomization : public IPropertyTypeCustomization
+{
+
+public:
+	static TSharedRef<IPropertyTypeCustomization> MakeInstance()
+	{
+		return MakeShared<FMovieGraphMetadataAttributeCustomization>();
+	}
+
+protected:
+	TSharedPtr<IPropertyHandle> StructPropertyHandle;
+
+	virtual void CustomizeHeader(TSharedRef<IPropertyHandle> InStructPropertyHandle, FDetailWidgetRow& HeaderRow, IPropertyTypeCustomizationUtils& CustomizationUtils) override
+	{
+		StructPropertyHandle = InStructPropertyHandle;
+		
+		const TSharedRef<IPropertyHandle> NamePropertyHandle = StructPropertyHandle->GetChildHandle(
+			GET_MEMBER_NAME_CHECKED(FMovieGraphMetadataAttribute, Name)).ToSharedRef();
+		const TSharedRef<IPropertyHandle> ValuePropertyHandle = StructPropertyHandle->GetChildHandle(
+			GET_MEMBER_NAME_CHECKED(FMovieGraphMetadataAttribute, Value)).ToSharedRef();
+		const TSharedRef<IPropertyHandle> IsEnabledPropertyHandle = StructPropertyHandle->GetChildHandle(
+			GET_MEMBER_NAME_CHECKED(FMovieGraphMetadataAttribute, bIsEnabled)).ToSharedRef();
+
+		HeaderRow
+		.WholeRowContent()
+		[
+			SNew(SHorizontalBox)
+			+ SHorizontalBox::Slot()
+			.Padding(0, 2)
+			.AutoWidth()
+			.HAlign(HAlign_Left)
+			[
+				SNew(SCheckBox)
+				.IsChecked(this, &FMovieGraphMetadataAttributeCustomization::IsCheckBoxChecked, IsEnabledPropertyHandle)
+				.OnCheckStateChanged(this, &FMovieGraphMetadataAttributeCustomization::OnCheckBoxCheckStateChanged, IsEnabledPropertyHandle)
+			]
+
+			+ SHorizontalBox::Slot()
+			.Padding(5, 2)
+			.FillWidth(0.5f)
+			[
+				SNew(SEditableTextBox)
+				.IsEnabled(this, &FMovieGraphMetadataAttributeCustomization::IsEnabled, IsEnabledPropertyHandle)
+				.Text(this, &FMovieGraphMetadataAttributeCustomization::GetPropertyHandleValue, NamePropertyHandle)
+				.Font(IDetailLayoutBuilder::GetDetailFont())
+				.OnTextCommitted_Lambda([this](const FText& InText, ETextCommit::Type InCommitType)
+				{
+					SetPropertyValueWithText(GET_MEMBER_NAME_CHECKED(FMovieGraphMetadataAttribute, Name), InText);
+				})
+				.HintText(LOCTEXT("MetadataAttributeNameInputHintText", "Attribute Name"))
+			]
+
+			+ SHorizontalBox::Slot()
+			.Padding(0, 2)
+			.FillWidth(0.5f)
+			[
+				SNew(SEditableTextBox)
+				.IsEnabled(this, &FMovieGraphMetadataAttributeCustomization::IsEnabled, IsEnabledPropertyHandle)
+				.Text(this, &FMovieGraphMetadataAttributeCustomization::GetPropertyHandleValue, ValuePropertyHandle)
+				.Font(IDetailLayoutBuilder::GetDetailFont())
+				.OnTextCommitted_Lambda([this](const FText& InText, ETextCommit::Type InCommitType)
+				{
+					SetPropertyValueWithText(GET_MEMBER_NAME_CHECKED(FMovieGraphMetadataAttribute, Value), InText);
+				})
+				.HintText(LOCTEXT("MetadataAttributeValueInputHintText", "Attribute Value"))
+			]
+
+			+ SHorizontalBox::Slot()
+			.Padding(5, 2)
+			.AutoWidth()
+			[
+				PropertyCustomizationHelpers::MakeInsertDeleteDuplicateButton(
+					{},
+					FExecuteAction::CreateLambda([InStructPropertyHandle]()
+					{
+						const TSharedPtr<IPropertyHandleArray> ParentPropertyHandleArray = InStructPropertyHandle->GetParentHandle()->AsArray();
+						const int32 ArrayIndex = InStructPropertyHandle->IsValidHandle() ? InStructPropertyHandle->GetIndexInArray() : INDEX_NONE;
+						if (ParentPropertyHandleArray.IsValid() && ArrayIndex >= 0)
+						{
+							ParentPropertyHandleArray->DeleteItem(ArrayIndex);
+						}
+					}),
+					{})
+			]
+		];
+	}
+	
+	virtual void CustomizeChildren(TSharedRef<IPropertyHandle> InStructPropertyHandle, IDetailChildrenBuilder& StructBuilder, IPropertyTypeCustomizationUtils& StructCustomizationUtils) override
+	{
+	}
+	
+	ECheckBoxState IsCheckBoxChecked(TSharedRef<IPropertyHandle> PropertyHandle) const
+	{
+		bool bValue = false;
+		if (PropertyHandle->IsValidHandle() && (PropertyHandle->GetValue(bValue) != FPropertyAccess::Fail))
+		{
+			return bValue ? ECheckBoxState::Checked : ECheckBoxState::Unchecked;
+		}
+
+		return ECheckBoxState::Unchecked;
+	}
+
+	void OnCheckBoxCheckStateChanged(ECheckBoxState NewState, TSharedRef<IPropertyHandle> PropertyHandle)
+	{
+		if (PropertyHandle->IsValidHandle())
+		{
+			const bool bValue = (NewState == ECheckBoxState::Checked) ? true : false;
+			PropertyHandle->SetValue(bValue);
+		}
+	}
+
+	bool IsEnabled(TSharedRef<IPropertyHandle> InIsEnabledPropertyHandle) const
+	{
+		bool bValue = false;
+		if (InIsEnabledPropertyHandle->IsValidHandle() && (InIsEnabledPropertyHandle->GetValue(bValue) != FPropertyAccess::Fail))
+		{
+			return bValue;
+		}
+
+		return false;
+	}
+
+	FText GetPropertyHandleValue(TSharedRef<IPropertyHandle> PropertyHandle) const
+	{
+		FString Value;
+		if (PropertyHandle->IsValidHandle() && (PropertyHandle->GetValue(Value) != FPropertyAccess::Fail))
+		{
+			return FText::FromString(Value);
+		}
+
+		return FText();
+	}
+
+	void SetPropertyValueWithText(const FName& Member, const FText& InText) const
+	{
+		const TSharedPtr<IPropertyHandle> PropertyHandle = StructPropertyHandle->GetChildHandle(Member);
+		check(PropertyHandle);
+
+		PropertyHandle->SetValue(InText.ToString());
+	}
+};
+
+#undef LOCTEXT_NAMESPACE
diff --git a/Engine/Plugins/MovieScene/MovieRenderPipeline/Source/MovieRenderPipelineEditor/Private/Graph/MovieGraphAssetToolkit.cpp b/Engine/Plugins/MovieScene/MovieRenderPipeline/Source/MovieRenderPipelineEditor/Private/Graph/MovieGraphAssetToolkit.cpp
--- a/Engine/Plugins/MovieScene/MovieRenderPipeline/Source/MovieRenderPipelineEditor/Private/Graph/MovieGraphAssetToolkit.cpp
+++ b/Engine/Plugins/MovieScene/MovieRenderPipeline/Source/MovieRenderPipelineEditor/Private/Graph/MovieGraphAssetToolkit.cpp
@@ -6,6 +6,7 @@
 #include "Customizations/Graph/MovieGraphCollectionsCustomization.h"
 #include "Customizations/Graph/MovieGraphFormatTokenCustomization.h"
 #include "Customizations/Graph/MovieGraphMemberCustomization.h"
+#include "Customizations/Graph/MovieGraphMetadataAttributeCustomization.h"
 #include "Customizations/Graph/MovieGraphModifiersCustomization.h"
 #include "Customizations/Graph/MovieGraphNamedResolutionCustomization.h"
 #include "Customizations/Graph/MovieGraphNodeCustomization.h"
@@ -23,6 +24,7 @@
 #include "Graph/Nodes/MovieGraphFileOutputNode.h"
 #include "Graph/Nodes/MovieGraphModifierNode.h"
 #include "Graph/Nodes/MovieGraphSelectNode.h"
+#include "Graph/Nodes/MovieGraphSetMetadataAttributesNode.h"
 #include "Graph/Nodes/MovieGraphSetCVarValueNode.h"
 #include "MovieEdGraphNode.h"
 #include "MovieGraphSchema.h"
@@ -397,6 +399,10 @@ TSharedRef<SDockTab> FMovieGraphAssetToolkit::SpawnTab_RenderGraphDetails(const
 		FMovieGraphNamedResolution::StaticStruct()->GetFName(),
 		FOnGetPropertyTypeCustomizationInstance::CreateStatic(&FMovieGraphNamedResolutionCustomization::MakeInstance));
 
+	SelectedGraphObjectsDetailsWidget->RegisterInstancedCustomPropertyTypeLayout(
+		FMovieGraphMetadataAttribute::StaticStruct()->GetFName(),
+		FOnGetPropertyTypeCustomizationInstance::CreateStatic(&FMovieGraphMetadataAttributeCustomization::MakeInstance));
+	
 	SelectedGraphObjectsDetailsWidget->RegisterInstancedCustomPropertyTypeLayout(
 		FMovieGraphVersioningSettings::StaticStruct()->GetFName(),
 		FOnGetPropertyTypeCustomizationInstance::CreateStatic(&FMovieGraphVersioningSettingsCustomization::MakeInstance));
