MoviePipeline: Updated the Apply Console Variable Preset node to be able to allow per-cvar pin promotions (so individual cvars within the preset can be overridden easily).

diff --git a/Engine/Plugins/MovieScene/MovieRenderPipeline/Source/MovieRenderPipelineCore/Private/Graph/MovieGraphCVarManager.cpp b/Engine/Plugins/MovieScene/MovieRenderPipeline/Source/MovieRenderPipelineCore/Private/Graph/MovieGraphCVarManager.cpp
--- a/Engine/Plugins/MovieScene/MovieRenderPipeline/Source/MovieRenderPipelineCore/Private/Graph/MovieGraphCVarManager.cpp
+++ b/Engine/Plugins/MovieScene/MovieRenderPipeline/Source/MovieRenderPipelineCore/Private/Graph/MovieGraphCVarManager.cpp
@@ -81,6 +81,12 @@ void FMovieGraphCVarManager::AddEvaluatedGraph(const UMovieGraphEvaluatedConfig*
 		else if (const UMovieGraphApplyCVarPresetNode* ApplyPresetNode = Cast<UMovieGraphApplyCVarPresetNode>(Node))
 		{
 			AddPreset(ApplyPresetNode->ConsoleVariablePreset);
+
+			// Preset nodes can optionally have overrides on top of the preset via promoted pins (and variables connected to them)
+			for (const TPair<FString, float>& CvarOverride : ApplyPresetNode->GetConsoleVariableOverrides())
+			{
+				AddCVar(CvarOverride.Key, CvarOverride.Value);
+			}
 		}
 	}
 }
diff --git a/Engine/Plugins/MovieScene/MovieRenderPipeline/Source/MovieRenderPipelineCore/Private/Graph/MovieGraphConfig.cpp b/Engine/Plugins/MovieScene/MovieRenderPipeline/Source/MovieRenderPipelineCore/Private/Graph/MovieGraphConfig.cpp
--- a/Engine/Plugins/MovieScene/MovieRenderPipeline/Source/MovieRenderPipelineCore/Private/Graph/MovieGraphConfig.cpp
+++ b/Engine/Plugins/MovieScene/MovieRenderPipeline/Source/MovieRenderPipelineCore/Private/Graph/MovieGraphConfig.cpp
@@ -1724,13 +1724,17 @@ bool UMovieGraphConfig::CreateFlattenedGraph_Recursive(UMovieGraphEvaluatedConfi
 		const UMovieGraphSettingNode* NodeAsSetting = CastChecked<UMovieGraphSettingNode>(Node);
 		const FString& NodeInstanceName = NodeAsSetting->GetNodeInstanceName();
 		
-		UMovieGraphNode* ExistingNode = OutBranchConfig.GetNodeByClassExactMatch(Node->GetClass(), NodeInstanceName);
+		UMovieGraphSettingNode* ExistingNode = Cast<UMovieGraphSettingNode>(OutBranchConfig.GetNodeByClassExactMatch(Node->GetClass(), NodeInstanceName));
 		if (!ExistingNode)
 		{
 			// Create a new instance of this node inside our flattened eval graph
-			ExistingNode = NewObject<UMovieGraphNode>(InOwningConfig, Node->GetClass());
+			ExistingNode = NewObject<UMovieGraphSettingNode>(InOwningConfig, Node->GetClass());
 			OutBranchConfig.NamedNodes.FindOrAdd(NodeInstanceName).NodeInstances.Add(ExistingNode);
 
+			// This shouldn't be abused by nodes, but in some very rare cases the flattened node needs to be "primed" before being put through the
+			// evaluation process.
+			ExistingNode->PrepareForFlattening(NodeAsSetting);
+
 			// Set all of the boolean edit condition values to false, so we can use "true" to indicate
 			// that the value was overridden already during traversal.
 			InitializeFlattenedNode(ExistingNode);
diff --git a/Engine/Plugins/MovieScene/MovieRenderPipeline/Source/MovieRenderPipelineCore/Private/Graph/Nodes/MovieGraphApplyCVarPresetNode.cpp b/Engine/Plugins/MovieScene/MovieRenderPipeline/Source/MovieRenderPipelineCore/Private/Graph/Nodes/MovieGraphApplyCVarPresetNode.cpp
--- a/Engine/Plugins/MovieScene/MovieRenderPipeline/Source/MovieRenderPipelineCore/Private/Graph/Nodes/MovieGraphApplyCVarPresetNode.cpp
+++ b/Engine/Plugins/MovieScene/MovieRenderPipeline/Source/MovieRenderPipelineCore/Private/Graph/Nodes/MovieGraphApplyCVarPresetNode.cpp
@@ -1,4 +1,4 @@
-﻿// Copyright Epic Games, Inc. All Rights Reserved.
+// Copyright Epic Games, Inc. All Rights Reserved.
 
 #include "Graph/Nodes/MovieGraphApplyCVarPresetNode.h"
 
@@ -21,6 +21,96 @@ EMovieGraphBranchRestriction UMovieGraphApplyCVarPresetNode::GetBranchRestrictio
 	return EMovieGraphBranchRestriction::Globals;
 }
 
+TArray<FMovieGraphPropertyInfo> UMovieGraphApplyCVarPresetNode::GetOverrideablePropertyInfo() const
+{
+	TArray<FMovieGraphPropertyInfo> PropertyInfo = Super::GetOverrideablePropertyInfo();
+
+	for (FMovieGraphPropertyInfo& Info : PropertyInfo)
+	{
+		if (!Info.bIsDynamicProperty)
+		{
+			continue;
+		}
+
+		Info.ContextMenuName = FText::Format(NSLOCTEXT("MovieGraphNodes", "ApplyCVarPresetNode_ValuePromotionContextMenuName", "Value ({0})"), FText::FromName(Info.Name));
+		Info.PromotionName = Info.Name;
+	}
+
+	return PropertyInfo;
+}
+
+TArray<FPropertyBagPropertyDesc> UMovieGraphApplyCVarPresetNode::GetDynamicPropertyDescriptions() const
+{
+	static const FName HideInActiveRenderSettingsMetaDataKey(TEXT("HideInActiveRenderSettings"));
+	
+	TArray<FPropertyBagPropertyDesc> CvarDescs;
+
+	// No dynamic properties if the console variable preset has not yet been chosen or is invalid
+	if (!ConsoleVariablePreset)
+	{
+		return CvarDescs;
+	}
+
+	constexpr bool bOnlyIncludeChecked = true;
+	TArray<TTuple<FString, FString>> ConsoleVariables;
+	ConsoleVariablePreset->GetConsoleVariablesForTrack(bOnlyIncludeChecked, ConsoleVariables);
+
+	// Add in properties which correspond to the cvars within the preset
+	for (const TTuple<FString, FString>& CVarPair : ConsoleVariables)
+	{
+		const FName CVarName = FName(CVarPair.Key);
+
+		// Hide these dynamic properties in the Active Render Settings. They will not show correct values if they're not connected to a variable.
+		// Cvars still need to be fully resolved, and these properties are not the place for that.
+		FPropertyBagPropertyDesc NewDesc(CVarName, EPropertyBagPropertyType::Float);
+#if WITH_EDITOR
+		NewDesc.MetaData.Add(FPropertyBagPropertyDescMetaData(HideInActiveRenderSettingsMetaDataKey, FString()));
+#endif
+		CvarDescs.Add(MoveTemp(NewDesc));
+
+		// Need a bOverride_* property in order for this to be picked up properly by graph evaluation
+		const FString OverrideName = FString::Printf(TEXT("bOverride_%s"), *CVarName.ToString());
+		FPropertyBagPropertyDesc OverrideDesc(FName(OverrideName), EPropertyBagPropertyType::Bool);
+		CvarDescs.Add(MoveTemp(OverrideDesc));
+	}
+
+	return CvarDescs;
+}
+
+void UMovieGraphApplyCVarPresetNode::TogglePromotePropertyToPin(const FName& PropertyName)
+{
+	Super::TogglePromotePropertyToPin(PropertyName);
+	
+	UpdateDynamicProperties();
+}
+
+void UMovieGraphApplyCVarPresetNode::PrepareForFlattening(const UMovieGraphSettingNode* InSourceNode)
+{
+	Super::PrepareForFlattening(InSourceNode);
+
+	// This node's dynamic properties rely on the cvar preset being valid and correct in order to be properly initialized.
+	if (const UMovieGraphApplyCVarPresetNode* SourcePresetNode = Cast<UMovieGraphApplyCVarPresetNode>(InSourceNode))
+	{
+		ConsoleVariablePreset = SourcePresetNode->ConsoleVariablePreset;
+	}
+}
+
+TArray<TPair<FString, float>> UMovieGraphApplyCVarPresetNode::GetConsoleVariableOverrides() const
+{
+	TArray<TPair<FString, float>> OverrideValues;
+	
+	for (const FPropertyBagPropertyDesc& Desc : GetDynamicPropertyDescriptions())
+	{
+		if (IsDynamicPropertyOverridden(Desc.Name))
+		{
+			TPair<FString, float> NewOverride(Desc.Name.ToString(), DynamicProperties.GetValueFloat(Desc.Name).GetValue());
+			OverrideValues.Add(MoveTemp(NewOverride));
+		}
+	}
+
+	return OverrideValues;
+}
+
 #if WITH_EDITOR
 FText UMovieGraphApplyCVarPresetNode::GetNodeTitle(const bool bGetDescriptive) const
 {
@@ -71,6 +161,9 @@ void UMovieGraphApplyCVarPresetNode::PostEditChangeProperty(FPropertyChangedEven
 	
 	if (PropertyChangedEvent.GetPropertyName() == GET_MEMBER_NAME_CHECKED(UMovieGraphApplyCVarPresetNode, ConsoleVariablePreset))
 	{
+		// Update the available cvars to override based on the preset that is chosen
+		UpdateDynamicProperties();
+		
 		OnNodeChangedDelegate.Broadcast(this);
 	}
 }
diff --git a/Engine/Plugins/MovieScene/MovieRenderPipeline/Source/MovieRenderPipelineCore/Public/Graph/MovieGraphNode.h b/Engine/Plugins/MovieScene/MovieRenderPipeline/Source/MovieRenderPipelineCore/Public/Graph/MovieGraphNode.h
--- a/Engine/Plugins/MovieScene/MovieRenderPipeline/Source/MovieRenderPipelineCore/Public/Graph/MovieGraphNode.h
+++ b/Engine/Plugins/MovieScene/MovieRenderPipeline/Source/MovieRenderPipelineCore/Public/Graph/MovieGraphNode.h
@@ -429,6 +429,14 @@ public:
 	 * identifier.
 	 */
 	virtual FString GetNodeInstanceName() const { return FString(); }
+
+	/**
+	 * In some very rare cases, a node needs to be "primed" from the node that the flattening is starting from before it's actually put through the
+	 * flattening process.
+	 *
+	 * The need for this should be exceedingly uncommon, so only use if absolutely necessary.
+	 */
+	virtual void PrepareForFlattening(const UMovieGraphSettingNode* InSourceNode) { }
 	
 	/*
 	* This is called either on the CDO, or on a "flattened" instance of the node every frame when
diff --git a/Engine/Plugins/MovieScene/MovieRenderPipeline/Source/MovieRenderPipelineCore/Public/Graph/Nodes/MovieGraphApplyCVarPresetNode.h b/Engine/Plugins/MovieScene/MovieRenderPipeline/Source/MovieRenderPipelineCore/Public/Graph/Nodes/MovieGraphApplyCVarPresetNode.h
--- a/Engine/Plugins/MovieScene/MovieRenderPipeline/Source/MovieRenderPipelineCore/Public/Graph/Nodes/MovieGraphApplyCVarPresetNode.h
+++ b/Engine/Plugins/MovieScene/MovieRenderPipeline/Source/MovieRenderPipelineCore/Public/Graph/Nodes/MovieGraphApplyCVarPresetNode.h
@@ -19,6 +19,16 @@ public:
 
 	virtual FString GetNodeInstanceName() const override;
 	virtual EMovieGraphBranchRestriction GetBranchRestriction() const override;
+	virtual TArray<FMovieGraphPropertyInfo> GetOverrideablePropertyInfo() const override;
+	virtual TArray<FPropertyBagPropertyDesc> GetDynamicPropertyDescriptions() const override;
+	virtual void TogglePromotePropertyToPin(const FName& PropertyName) override;
+	virtual void PrepareForFlattening(const UMovieGraphSettingNode* InSourceNode) override;
+
+	/**
+	 * Gets the name and resolved value of any cvars that were overridden via promoted pins. This method only makes sense to call on a
+	 * flattened (resolved) node.
+	 */
+	TArray<TPair<FString, float>> GetConsoleVariableOverrides() const;
 
 #if WITH_EDITOR
 	virtual FText GetNodeTitle(const bool bGetDescriptive = false) const override;
diff --git a/Engine/Plugins/MovieScene/MovieRenderPipeline/Source/MovieRenderPipelineEditor/Private/Customizations/Graph/MovieGraphApplyCVarPresetNodeCustomization.h b/Engine/Plugins/MovieScene/MovieRenderPipeline/Source/MovieRenderPipelineEditor/Private/Customizations/Graph/MovieGraphApplyCVarPresetNodeCustomization.h
--- a/Engine/Plugins/MovieScene/MovieRenderPipeline/Source/MovieRenderPipelineEditor/Private/Customizations/Graph/MovieGraphApplyCVarPresetNodeCustomization.h
+++ b/Engine/Plugins/MovieScene/MovieRenderPipeline/Source/MovieRenderPipelineEditor/Private/Customizations/Graph/MovieGraphApplyCVarPresetNodeCustomization.h
@@ -120,6 +120,10 @@ protected:
 				DetailBuilder.ForceRefreshDetails();
 			}));
 
+			// Don't show the dynamic properties (which hold values for cvars which are exposed via pins). For some reason hiding the individual
+			// property doesn't work, so the entire category is hidden instead.
+			DetailBuilder.HideCategory("Properties");
+
 			// Hide the default preset property since we'll be making a group that contains it instead (and the group's rows will be the individual
 			// console variables within the preset)
 			ConsoleVariablePresetHandle->MarkHiddenByCustomization();
diff --git a/Engine/Plugins/MovieScene/MovieRenderPipeline/Source/MovieRenderPipelineEditor/Private/Graph/SMovieGraphActiveRenderSettingsTabContent.cpp b/Engine/Plugins/MovieScene/MovieRenderPipeline/Source/MovieRenderPipelineEditor/Private/Graph/SMovieGraphActiveRenderSettingsTabContent.cpp
--- a/Engine/Plugins/MovieScene/MovieRenderPipeline/Source/MovieRenderPipelineEditor/Private/Graph/SMovieGraphActiveRenderSettingsTabContent.cpp
+++ b/Engine/Plugins/MovieScene/MovieRenderPipeline/Source/MovieRenderPipelineEditor/Private/Graph/SMovieGraphActiveRenderSettingsTabContent.cpp
@@ -90,6 +90,8 @@ bool FActiveRenderSettingsTreeElement::IsBranchRenderable() const
 
 const TArray<TSharedPtr<FActiveRenderSettingsTreeElement>>& FActiveRenderSettingsTreeElement::GetChildren() const
 {
+	static const FName HideInActiveRenderSettingsMetaDataKey(TEXT("HideInActiveRenderSettings"));
+	
 	auto MakeElementFromNode = [this](const TObjectPtr<UMovieGraphNode>& Node) -> TSharedPtr<FActiveRenderSettingsTreeElement>
 	{
 		// The element name should include the node's instance name (if the instance name isn't empty)
@@ -179,6 +181,12 @@ const TArray<TSharedPtr<FActiveRenderSettingsTreeElement>>& FActiveRenderSetting
 	{
 		for (const FProperty* Property : SettingsNode->GetAllOverrideableProperties())
 		{
+			// Some properties opt-out of being shown in the Active Render Settings editor
+			if (Property->HasMetaData(HideInActiveRenderSettingsMetaDataKey))
+			{
+				continue;
+			}
+			
 			TSharedPtr<FActiveRenderSettingsTreeElement> Element =
 				MakeShared<FActiveRenderSettingsTreeElement>(FName(Property->GetDisplayNameText().ToString()), EElementType::Property);
 			Element->SettingsNode = SettingsNode;
