MoviePipeline: Made several improvements to EXR metadata.
* EXR metadata is now aggregated together when using multiple layers and the multilayer EXR node.
* Updated the Set Metadata Attributes node to resolve filename format {tokens} in both the key and value data.
* Also updated the node to include {token} auto-complete.

Refactored FMovieGraphFormatTokenCustomization and SMoviePipelineFormatTokenAutoCompleteBox a bit so the autocomplete box could be more cleanly used within the Set Metadata Attributes customization.

diff --git a/Engine/Plugins/MovieScene/MovieRenderPipeline/Source/MovieRenderPipelineEditor/Private/Customizations/Graph/MovieGraphFormatTokenCustomization.cpp b/Engine/Plugins/MovieScene/MovieRenderPipeline/Source/MovieRenderPipelineEditor/Private/Customizations/Graph/MovieGraphFormatTokenCustomization.cpp
--- a/Engine/Plugins/MovieScene/MovieRenderPipeline/Source/MovieRenderPipelineEditor/Private/Customizations/Graph/MovieGraphFormatTokenCustomization.cpp
+++ b/Engine/Plugins/MovieScene/MovieRenderPipeline/Source/MovieRenderPipelineEditor/Private/Customizations/Graph/MovieGraphFormatTokenCustomization.cpp
@@ -2,7 +2,6 @@
 
 #include "MovieGraphFormatTokenCustomization.h"
 
-#include "Graph/MovieGraphBlueprintLibrary.h"
 #include "Graph/Nodes/MovieGraphDebugNode.h"
 #include "Graph/Nodes/MovieGraphFileOutputNode.h"
 
@@ -13,9 +12,6 @@
 
 void FMovieGraphFormatTokenCustomization::CustomizeDetails(IDetailLayoutBuilder& DetailBuilder)
 {
-	CustomizedObject = DetailBuilder.GetSelectedObjects()[0];
-	check(CustomizedObject.IsValid() && !CustomizedObject.IsStale());
-
 	// This should work just fine for UMovieGraphFileOutputNode and UMovieGraphCommandLineEncoderNode as long as the property names stay in sync
 	OutputFormatPropertyHandle = DetailBuilder.GetProperty(GET_MEMBER_NAME_CHECKED(UMovieGraphFileOutputNode, FileNameFormat));
 
@@ -31,13 +27,6 @@ void FMovieGraphFormatTokenCustomization::CustomizeDetails(IDetailLayoutBuilder&
 		return;
 	}
 
-	FText StartingDisplayText;
-	OutputFormatPropertyHandle->GetValueAsDisplayText(StartingDisplayText);
-
-	// Update the text box when the handle value changes
-	OutputFormatPropertyHandle->SetOnPropertyValueChanged(
-		FSimpleDelegate::CreateSP(this, &FMovieGraphFormatTokenCustomization::OnPropertyChange));
-
 	DetailBuilder.EditDefaultProperty(OutputFormatPropertyHandle)->CustomWidget()
 		.NameContent()
 		[
@@ -45,55 +34,11 @@ void FMovieGraphFormatTokenCustomization::CustomizeDetails(IDetailLayoutBuilder&
 		]
 		.ValueContent()
 		.MinDesiredWidth(200.0f)
-	[
+		[
 			// We choose not to bind the text box text here because simultaneously setting and getting 
 			// the handle value can cause the binding to stop working, so we manually update it when necessary 
 			SAssignNew(AutoCompleteBox, SMoviePipelineFormatTokenAutoCompleteBox)
-			.InitialText(StartingDisplayText)
-			.Suggestions(this, &FMovieGraphFormatTokenCustomization::GetSuggestions)
-			.OnTextChanged(this, &FMovieGraphFormatTokenCustomization::OnTextChanged)
+			.TextHandle(OutputFormatPropertyHandle)
+			.Suggestions(SMoviePipelineFormatTokenAutoCompleteBox::GetFileNameFormatSuggestions())
 		];
-}
-
-void FMovieGraphFormatTokenCustomization::OnPropertyChange()
-{
-	// Sync the text box back to the handle value
-	FText DisplayText;
-	OutputFormatPropertyHandle->GetValueAsDisplayText(DisplayText);
-
-	AutoCompleteBox->SetText(DisplayText);
-	AutoCompleteBox->CloseMenuAndReset();
-}
-
-TArray<FString> FMovieGraphFormatTokenCustomization::GetSuggestions() const
-{
-	TArray<FString> Suggestions;
-	FMovieGraphResolveArgs FormatArgs;
-	GetFormatArguments(FormatArgs);
-
-	// Display the token names alphabetically
-	TArray<FString> FilenameFormatTokens;
-	FormatArgs.FilenameArguments.GetKeys(FilenameFormatTokens);
-	FilenameFormatTokens.Sort();
-
-	for (const FString& FilenameFormatToken : FilenameFormatTokens)
-	{
-		Suggestions.Add(FilenameFormatToken);
-	}
-
-	return Suggestions;
-}
-
-void FMovieGraphFormatTokenCustomization::OnTextChanged(const FText& InValue)
-{
-	OutputFormatPropertyHandle->SetValue(InValue.ToString());
-}
-
-void FMovieGraphFormatTokenCustomization::GetFormatArguments(FMovieGraphResolveArgs& InOutFormatArgs)
-{
-	// Just fetch the format arguments (by keeping the format string empty). The tokens themselves will not be resolved correctly here (no context is
-	// provided in the resolve params), but all we care about here is the token list, not the resolved token values.
-	const FString FormatString;
-	const FMovieGraphFilenameResolveParams ResolveParams;
-	UMovieGraphBlueprintLibrary::ResolveFilenameFormatArguments(FormatString, ResolveParams, InOutFormatArgs);
-}
+}
\ No newline at end of file
diff --git a/Engine/Plugins/MovieScene/MovieRenderPipeline/Source/MovieRenderPipelineEditor/Private/Customizations/Graph/MovieGraphFormatTokenCustomization.h b/Engine/Plugins/MovieScene/MovieRenderPipeline/Source/MovieRenderPipelineEditor/Private/Customizations/Graph/MovieGraphFormatTokenCustomization.h
--- a/Engine/Plugins/MovieScene/MovieRenderPipeline/Source/MovieRenderPipelineEditor/Private/Customizations/Graph/MovieGraphFormatTokenCustomization.h
+++ b/Engine/Plugins/MovieScene/MovieRenderPipeline/Source/MovieRenderPipelineEditor/Private/Customizations/Graph/MovieGraphFormatTokenCustomization.h
@@ -20,16 +20,7 @@ protected:
 	//~ Begin IDetailCustomization interface
 	virtual void CustomizeDetails(IDetailLayoutBuilder& DetailBuilder) override;
 	//~ End IDetailCustomization interface
-
-	void OnPropertyChange();
-
-	TArray<FString> GetSuggestions() const;
-
-	void OnTextChanged(const FText& InValue);
-
-	static void GetFormatArguments(FMovieGraphResolveArgs& InOutFormatArgs);
 	
-	TWeakObjectPtr<UObject> CustomizedObject;
 	TSharedPtr<IPropertyHandle> OutputFormatPropertyHandle;
 	TSharedPtr<class SMoviePipelineFormatTokenAutoCompleteBox> AutoCompleteBox;
 };
diff --git a/Engine/Plugins/MovieScene/MovieRenderPipeline/Source/MovieRenderPipelineEditor/Private/Customizations/Graph/MovieGraphMetadataAttributeCustomization.h b/Engine/Plugins/MovieScene/MovieRenderPipeline/Source/MovieRenderPipelineEditor/Private/Customizations/Graph/MovieGraphMetadataAttributeCustomization.h
--- a/Engine/Plugins/MovieScene/MovieRenderPipeline/Source/MovieRenderPipelineEditor/Private/Customizations/Graph/MovieGraphMetadataAttributeCustomization.h
+++ b/Engine/Plugins/MovieScene/MovieRenderPipeline/Source/MovieRenderPipelineEditor/Private/Customizations/Graph/MovieGraphMetadataAttributeCustomization.h
@@ -9,6 +9,7 @@
 #include "PropertyCustomizationHelpers.h"
 #include "Widgets/Input/SCheckBox.h"
 #include "Widgets/Input/SEditableTextBox.h"
+#include "Widgets/SMoviePipelineFormatTokenAutoCompleteBox.h"
 
 #define LOCTEXT_NAMESPACE "FMovieGraphMetadataAttributeCustomization"
 
@@ -37,7 +38,9 @@ protected:
 			GET_MEMBER_NAME_CHECKED(FMovieGraphMetadataAttribute, bIsEnabled)).ToSharedRef();
 
 		HeaderRow
-		.WholeRowContent()
+		.NameContent()
+		.HAlign(HAlign_Fill)
+		.VAlign(VAlign_Fill)
 		[
 			SNew(SHorizontalBox)
 			+ SHorizontalBox::Slot()
@@ -52,51 +55,31 @@ protected:
 
 			+ SHorizontalBox::Slot()
 			.Padding(5, 2)
-			.FillWidth(0.5f)
+			.FillWidth(1.f)
 			[
-				SNew(SEditableTextBox)
+				SNew(SMoviePipelineFormatTokenAutoCompleteBox)
+				.TextHandle(NamePropertyHandle)
+				.Suggestions(SMoviePipelineFormatTokenAutoCompleteBox::GetFileNameFormatSuggestions())
 				.IsEnabled(this, &FMovieGraphMetadataAttributeCustomization::IsEnabled, IsEnabledPropertyHandle)
-				.Text(this, &FMovieGraphMetadataAttributeCustomization::GetPropertyHandleValue, NamePropertyHandle)
-				.Font(IDetailLayoutBuilder::GetDetailFont())
-				.OnTextCommitted_Lambda([this](const FText& InText, ETextCommit::Type InCommitType)
-				{
-					SetPropertyValueWithText(GET_MEMBER_NAME_CHECKED(FMovieGraphMetadataAttribute, Name), InText);
-				})
 				.HintText(LOCTEXT("MetadataAttributeNameInputHintText", "Attribute Name"))
 			]
+		]
 
+		.ValueContent()
+		.HAlign(HAlign_Fill)
+		.VAlign(VAlign_Fill)
+		[
+			SNew(SHorizontalBox)
 			+ SHorizontalBox::Slot()
 			.Padding(0, 2)
-			.FillWidth(0.5f)
+			.FillWidth(1.f)
 			[
-				SNew(SEditableTextBox)
+				SNew(SMoviePipelineFormatTokenAutoCompleteBox)
+				.TextHandle(ValuePropertyHandle)
+				.Suggestions(SMoviePipelineFormatTokenAutoCompleteBox::GetFileNameFormatSuggestions())
 				.IsEnabled(this, &FMovieGraphMetadataAttributeCustomization::IsEnabled, IsEnabledPropertyHandle)
-				.Text(this, &FMovieGraphMetadataAttributeCustomization::GetPropertyHandleValue, ValuePropertyHandle)
-				.Font(IDetailLayoutBuilder::GetDetailFont())
-				.OnTextCommitted_Lambda([this](const FText& InText, ETextCommit::Type InCommitType)
-				{
-					SetPropertyValueWithText(GET_MEMBER_NAME_CHECKED(FMovieGraphMetadataAttribute, Value), InText);
-				})
 				.HintText(LOCTEXT("MetadataAttributeValueInputHintText", "Attribute Value"))
 			]
-
-			+ SHorizontalBox::Slot()
-			.Padding(5, 2)
-			.AutoWidth()
-			[
-				PropertyCustomizationHelpers::MakeInsertDeleteDuplicateButton(
-					{},
-					FExecuteAction::CreateLambda([InStructPropertyHandle]()
-					{
-						const TSharedPtr<IPropertyHandleArray> ParentPropertyHandleArray = InStructPropertyHandle->GetParentHandle()->AsArray();
-						const int32 ArrayIndex = InStructPropertyHandle->IsValidHandle() ? InStructPropertyHandle->GetIndexInArray() : INDEX_NONE;
-						if (ParentPropertyHandleArray.IsValid() && ArrayIndex >= 0)
-						{
-							ParentPropertyHandleArray->DeleteItem(ArrayIndex);
-						}
-					}),
-					{})
-			]
 		];
 	}
 	
@@ -134,25 +117,6 @@ protected:
 
 		return false;
 	}
-
-	FText GetPropertyHandleValue(TSharedRef<IPropertyHandle> PropertyHandle) const
-	{
-		FString Value;
-		if (PropertyHandle->IsValidHandle() && (PropertyHandle->GetValue(Value) != FPropertyAccess::Fail))
-		{
-			return FText::FromString(Value);
-		}
-
-		return FText();
-	}
-
-	void SetPropertyValueWithText(const FName& Member, const FText& InText) const
-	{
-		const TSharedPtr<IPropertyHandle> PropertyHandle = StructPropertyHandle->GetChildHandle(Member);
-		check(PropertyHandle);
-
-		PropertyHandle->SetValue(InText.ToString());
-	}
 };
 
 #undef LOCTEXT_NAMESPACE
diff --git a/Engine/Plugins/MovieScene/MovieRenderPipeline/Source/MovieRenderPipelineEditor/Private/Widgets/SMoviePipelineConfigEditor.cpp b/Engine/Plugins/MovieScene/MovieRenderPipeline/Source/MovieRenderPipelineEditor/Private/Widgets/SMoviePipelineConfigEditor.cpp
--- a/Engine/Plugins/MovieScene/MovieRenderPipeline/Source/MovieRenderPipelineEditor/Private/Widgets/SMoviePipelineConfigEditor.cpp
+++ b/Engine/Plugins/MovieScene/MovieRenderPipeline/Source/MovieRenderPipelineEditor/Private/Widgets/SMoviePipelineConfigEditor.cpp
@@ -71,10 +71,6 @@ protected:
 		FText InitialText;
 		OutputFormatPropertyHandle->GetValueAsDisplayText(InitialText);
 
-		// Update the text box when the handle value changes
-		OutputFormatPropertyHandle->SetOnPropertyValueChanged(
-			FSimpleDelegate::CreateSP(this, &FOutputFormatDetailsCustomization::OnPropertyChange));
-
 		DetailBuilder.EditDefaultProperty(OutputFormatPropertyHandle)->CustomWidget(bShowChildren)
 			.NameContent()
 			[
@@ -85,24 +81,12 @@ protected:
 			[
 				SAssignNew(AutoCompleteBox, SMoviePipelineFormatTokenAutoCompleteBox)
 				.InitialText(InitialText)
+				.TextHandle(OutputFormatPropertyHandle)
 				.Suggestions(this, &FOutputFormatDetailsCustomization::GetSuggestions)
-				.OnTextChanged(this, &FOutputFormatDetailsCustomization::OnTextChanged)
 			];
-
-
 	}
 	//~ End IDetailCustomization interface
 
-	void OnPropertyChange()
-	{
-		// Sync the text box back to the handle value
-		FText DisplayText;
-		OutputFormatPropertyHandle->GetValueAsDisplayText(DisplayText);
-
-		AutoCompleteBox->SetText(DisplayText);
-		AutoCompleteBox->CloseMenuAndReset();
-	}
-
 	TArray<FString> GetSuggestions() const
 	{
 		TSharedPtr<SMoviePipelineConfigEditor> ConfigEditor = OwningEditor.Pin();
@@ -134,11 +118,6 @@ protected:
 		return TArray<FString>();
 	}
 
-	void OnTextChanged(const FText& InValue)
-	{
-		OutputFormatPropertyHandle->SetValue(InValue.ToString());
-	}
-
 	TSharedPtr<IPropertyHandle> OutputFormatPropertyHandle;
 	TWeakPtr<SMoviePipelineConfigEditor> OwningEditor;
 	TSharedPtr<SMoviePipelineFormatTokenAutoCompleteBox> AutoCompleteBox;
diff --git a/Engine/Plugins/MovieScene/MovieRenderPipeline/Source/MovieRenderPipelineEditor/Private/Widgets/SMoviePipelineFormatTokenAutoCompleteBox.cpp b/Engine/Plugins/MovieScene/MovieRenderPipeline/Source/MovieRenderPipelineEditor/Private/Widgets/SMoviePipelineFormatTokenAutoCompleteBox.cpp
--- a/Engine/Plugins/MovieScene/MovieRenderPipeline/Source/MovieRenderPipelineEditor/Private/Widgets/SMoviePipelineFormatTokenAutoCompleteBox.cpp
+++ b/Engine/Plugins/MovieScene/MovieRenderPipeline/Source/MovieRenderPipelineEditor/Private/Widgets/SMoviePipelineFormatTokenAutoCompleteBox.cpp
@@ -3,22 +3,35 @@
 #include "SMoviePipelineFormatTokenAutoCompleteBox.h"
 
 #include "DetailLayoutBuilder.h"
+#include "Graph/MovieGraphBlueprintLibrary.h"
+#include "Graph/MovieGraphFilenameResolveParams.h"
 #include "Layout/WidgetPath.h"
 
 void SMoviePipelineFormatTokenAutoCompleteBox::Construct(const FArguments& InArgs)
 {
+	TextHandle = InArgs._TextHandle;
+	
 	ChildSlot
 	[
 		SAssignNew(MenuAnchor, SMenuAnchor)
 		.Placement(MenuPlacement_ComboBox)
 		[
 			SAssignNew(TextBox, SMultiLineEditableTextBox)
-			.Text(InArgs._InitialText)
+			.Text_Lambda([this]()
+			{
+				FString TextValue;
+				TextHandle->GetValue(TextValue);
+				
+				return FText::FromString(TextValue);
+			})
+			.HintText(InArgs._HintText)
 			.Font(IDetailLayoutBuilder::GetDetailFont())
 			.OnKeyDownHandler(this, &SMoviePipelineFormatTokenAutoCompleteBox::OnKeyDown)
 			.OnTextChanged(this, &SMoviePipelineFormatTokenAutoCompleteBox::HandleTextBoxTextChanged)
+			.OnTextCommitted(this, &SMoviePipelineFormatTokenAutoCompleteBox::HandleTextBoxTextCommitted)
 			.SelectWordOnMouseDoubleClick(true)
 			.AllowMultiLine(false)
+			.IsEnabled(InArgs._IsEnabled)
 		]
 		.MenuContent
 		(
@@ -42,7 +55,6 @@ void SMoviePipelineFormatTokenAutoCompleteBox::Construct(const FArguments& InArg
 	// We just call it once and cache it for now as the selection code isn't tested against
 	// the amount of suggestions changing.
 	AllSuggestions.Append(InArgs._Suggestions.Get());
-	OnTextChanged = InArgs._OnTextChanged;
 }
 
 void SMoviePipelineFormatTokenAutoCompleteBox::OnFocusChanging(const FWeakWidgetPath& PreviousFocusPath, const FWidgetPath& NewWidgetPath, const FFocusEvent& InFocusEvent)
@@ -195,7 +207,7 @@ void SMoviePipelineFormatTokenAutoCompleteBox::ReplaceRelevantTextWithSuggestion
 
 void SMoviePipelineFormatTokenAutoCompleteBox::HandleTextBoxTextChanged(const FText& InText)
 {
-	OnTextChanged.ExecuteIfBound(InText);
+	TextHandle->SetValue(InText.ToString(), EPropertyValueSetFlags::InteractiveChange);
 
 	const FString TextAsStr = InText.ToString();
 	if (TextAsStr.Len() > 0)
@@ -220,6 +232,11 @@ void SMoviePipelineFormatTokenAutoCompleteBox::HandleTextBoxTextChanged(const FT
 	}
 }
 
+void SMoviePipelineFormatTokenAutoCompleteBox::HandleTextBoxTextCommitted(const FText& InText, ETextCommit::Type CommitInfo)
+{
+	TextHandle->SetValue(InText.ToString(), EPropertyValueSetFlags::DefaultFlags);
+}
+
 void SMoviePipelineFormatTokenAutoCompleteBox::FilterVisibleSuggestions(const FString& StrToMatch, const bool bForceShowAll)
 {
 	Suggestions.Reset();
@@ -284,4 +301,22 @@ TSharedRef<ITableRow> SMoviePipelineFormatTokenAutoCompleteBox::HandleSuggestion
 				.Text(FText::FromString(SuggestionText))
 			]
 		];
+}
+
+TArray<FString> SMoviePipelineFormatTokenAutoCompleteBox::GetFileNameFormatSuggestions()
+{
+	TArray<FString> FileNameFormatSuggestions;
+	
+	// Just fetch the format arguments (by keeping the format string empty). The tokens themselves will not be resolved correctly here (no context is
+	// provided in the resolve params), but all we care about here is the token list, not the resolved token values.
+	const FString FormatString;
+	const FMovieGraphFilenameResolveParams ResolveParams;
+	FMovieGraphResolveArgs FormatArgs;
+	UMovieGraphBlueprintLibrary::ResolveFilenameFormatArguments(FormatString, ResolveParams, FormatArgs);
+
+	// Display the token names alphabetically
+	FormatArgs.FilenameArguments.GetKeys(FileNameFormatSuggestions);
+	FileNameFormatSuggestions.Sort();
+
+	return FileNameFormatSuggestions;
 }
\ No newline at end of file
diff --git a/Engine/Plugins/MovieScene/MovieRenderPipeline/Source/MovieRenderPipelineEditor/Private/Widgets/SMoviePipelineFormatTokenAutoCompleteBox.h b/Engine/Plugins/MovieScene/MovieRenderPipeline/Source/MovieRenderPipelineEditor/Private/Widgets/SMoviePipelineFormatTokenAutoCompleteBox.h
--- a/Engine/Plugins/MovieScene/MovieRenderPipeline/Source/MovieRenderPipelineEditor/Private/Widgets/SMoviePipelineFormatTokenAutoCompleteBox.h
+++ b/Engine/Plugins/MovieScene/MovieRenderPipeline/Source/MovieRenderPipelineEditor/Private/Widgets/SMoviePipelineFormatTokenAutoCompleteBox.h
@@ -18,9 +18,9 @@ public:
 	SLATE_BEGIN_ARGS(SMoviePipelineFormatTokenAutoCompleteBox){}
 
 	SLATE_ARGUMENT(FText, InitialText)
+	SLATE_ARGUMENT(FText, HintText)
+	SLATE_ARGUMENT(TSharedPtr<IPropertyHandle>, TextHandle)
 	SLATE_ATTRIBUTE(TArray<FString>, Suggestions)
-	/** Called whenever the text is changed programmatically or interactively by the user. */
-	SLATE_EVENT(FOnTextChanged, OnTextChanged)
 	SLATE_END_ARGS()
 
 	void Construct(const FArguments& InArgs);
@@ -40,6 +40,8 @@ public:
 
 	void HandleTextBoxTextChanged(const FText& InText);
 
+	void HandleTextBoxTextCommitted(const FText& InText, ETextCommit::Type CommitInfo);
+
 	void FilterVisibleSuggestions(const FString& StrToMatch, const bool bForceShowAll);
 
 	void CloseMenuAndReset();
@@ -48,13 +50,17 @@ public:
 
 	TSharedRef<ITableRow> HandleSuggestionListViewGenerateRow(TSharedPtr<FString> Text, const TSharedRef<STableViewBase>& OwnerTable) const;
 
+	/**
+	 * A helper to get file name format suggestions. Can be passed to the "Suggestions" argument.
+	 */
+	static TArray<FString> GetFileNameFormatSuggestions();
+
 private:
 	TSharedPtr<SListView<TSharedPtr<FString>>> SuggestionListView;
 	TSharedPtr<SMultiLineEditableTextBox> TextBox;
 	TSharedPtr<SMenuAnchor> MenuAnchor;
 	TSharedPtr<SVerticalBox> VerticalBox;
-	// Holds a delegate that is executed when the text has changed.
-	FOnTextChanged OnTextChanged;
+	TSharedPtr<IPropertyHandle> TextHandle;
 
 	// The pool of suggestions to show
 	TArray<FString> AllSuggestions;
diff --git a/Engine/Plugins/MovieScene/MovieRenderPipeline/Source/MovieRenderPipelineRenderPasses/Private/MovieGraphImageSequenceOutputNode.cpp b/Engine/Plugins/MovieScene/MovieRenderPipeline/Source/MovieRenderPipelineRenderPasses/Private/MovieGraphImageSequenceOutputNode.cpp
--- a/Engine/Plugins/MovieScene/MovieRenderPipeline/Source/MovieRenderPipelineRenderPasses/Private/MovieGraphImageSequenceOutputNode.cpp
+++ b/Engine/Plugins/MovieScene/MovieRenderPipeline/Source/MovieRenderPipelineRenderPasses/Private/MovieGraphImageSequenceOutputNode.cpp
@@ -549,6 +549,32 @@ void UMovieGraphImageSequenceOutputNode_MultiLayerEXR::GetFilenameToRenderIDMapp
 	TMap<FString, TArray<FMovieGraphRenderDataIdentifier>>& OutFilenameToRenderIDs,
 	TMap<FString, FMovieGraphResolveArgs>& OutFilenameToResolveArgs) const
 {
+	// Merge one layer's resolve args (InNewResolveArgs) into an existing set of resolve args (InExistingResolveArgs).
+	auto MergeResolveArgs = [](FMovieGraphResolveArgs& InNewResolveArgs, FMovieGraphResolveArgs& InExistingResolveArgs)
+	{
+		// Covert the filename arguments to FormatNamedArguments once; this is needed by FString::Format() in the loop
+		FStringFormatNamedArguments NamedArguments;
+		for (const TPair<FString, FString>& FilenameArgument : InNewResolveArgs.FilenameArguments)
+		{
+			NamedArguments.Add(FilenameArgument.Key, FilenameArgument.Value);
+		}
+
+		for (TPair<FString, FString>& MetadataPair : InNewResolveArgs.FileMetadata)
+		{
+			// The metadata key and/or value may contain filename format {tokens}; resolve any of them BEFORE merging in with existing metadata. This
+			// is important because the metadata may contain a {token} that, once resolved, prevents a collision with an existing key.
+			MetadataPair.Key = FString::Format(*MetadataPair.Key, NamedArguments);
+			MetadataPair.Value = FString::Format(*MetadataPair.Value, NamedArguments);
+
+			// Merge in the resolved metadata into the existing metadata
+			InExistingResolveArgs.FileMetadata.Add(MetadataPair.Key, MetadataPair.Value);
+		}
+
+		// The filename arguments are not needed after merging + resolving; however, the last set of arguments is passed along anyway if they are needed.
+		// They aren't merged though, because they differ too much between layers to make merging of any practical usefulness (eg, {layer_name}).
+		InExistingResolveArgs.FilenameArguments = InNewResolveArgs.FilenameArguments;
+	};
+	
 	TMap<FString, TArray<FIntPoint>> FilenameToResolutions;
 
 	// First, generate filename -> renderID mapping, and filename -> resolution mapping.
@@ -567,7 +593,7 @@ void UMovieGraphImageSequenceOutputNode_MultiLayerEXR::GetFilenameToRenderIDMapp
 	{
 		constexpr int32 ResolutionIndex = 0;
 		FMovieGraphResolveArgs ResolveArgs;
-		const FString PreliminaryFileName = ResolveOutputFilename(InParentNode, InPipeline, ResolutionIndex, InRawFrameData, RenderPassData.Key.RootBranchName, ResolveArgs);
+		const FString PreliminaryFileName = ResolveOutputFilename(InParentNode, InPipeline, ResolutionIndex, InRawFrameData, RenderPassData.Key, ResolveArgs);
 		
 		TArray<FMovieGraphRenderDataIdentifier>& RenderIDs = OutFilenameToRenderIDs.FindOrAdd(PreliminaryFileName);
 		RenderIDs.Add(RenderPassData.Key);
@@ -575,7 +601,7 @@ void UMovieGraphImageSequenceOutputNode_MultiLayerEXR::GetFilenameToRenderIDMapp
 		TArray<FIntPoint>& Resolutions = FilenameToResolutions.FindOrAdd(PreliminaryFileName);
 		Resolutions.AddUnique(RenderPassData.Value->GetSize());
 
-		OutFilenameToResolveArgs.Add(PreliminaryFileName, ResolveArgs);
+		MergeResolveArgs(ResolveArgs, OutFilenameToResolveArgs.FindOrAdd(PreliminaryFileName));
 	}
 
 	// Second, re-generate filenames if any render passes of differing resolutions map to the same file.
@@ -601,12 +627,12 @@ void UMovieGraphImageSequenceOutputNode_MultiLayerEXR::GetFilenameToRenderIDMapp
 			// Re-resolve the filename, this time using the resolution index to generate a filename that will only contain
 			// passes with this particular resolution
 			FMovieGraphResolveArgs ResolveArgs;
-			const FString FinalFilename = ResolveOutputFilename(InParentNode, InPipeline, ResolutionIndex, InRawFrameData, RenderID.RootBranchName, ResolveArgs);
+			const FString FinalFilename = ResolveOutputFilename(InParentNode, InPipeline, ResolutionIndex, InRawFrameData, RenderID, ResolveArgs);
 
 			TArray<FMovieGraphRenderDataIdentifier>& RenderIDs = OutFilenameToRenderIDs.FindOrAdd(FinalFilename);
 			RenderIDs.Add(RenderID);
 
-			OutFilenameToResolveArgs.Add(FinalFilename, ResolveArgs);
+			MergeResolveArgs(ResolveArgs, OutFilenameToResolveArgs.FindOrAdd(FinalFilename));
 		}
 	}
 }
@@ -615,12 +641,12 @@ FString UMovieGraphImageSequenceOutputNode_MultiLayerEXR::ResolveOutputFilename(
 	const UMovieGraphImageSequenceOutputNode_MultiLayerEXR* InParentNode,
 	const UMovieGraphPipeline* InPipeline,
 	const int32 ResolutionIndex, const UE::MovieGraph::FMovieGraphOutputMergerFrame* InRawFrameData,
-	const FName& InBranchName, FMovieGraphResolveArgs& OutResolveArgs) const
+	const FMovieGraphRenderDataIdentifier& InRenderDataIdentifier, FMovieGraphResolveArgs& OutResolveArgs) const
 {
 	const TCHAR* Extension = TEXT("exr");
 
 	constexpr bool bIncludeCDOs = true;
-	const UMovieGraphGlobalOutputSettingNode* OutputSettings = InRawFrameData->EvaluatedConfig->GetSettingForBranch<UMovieGraphGlobalOutputSettingNode>(InBranchName, bIncludeCDOs);
+	const UMovieGraphGlobalOutputSettingNode* OutputSettings = InRawFrameData->EvaluatedConfig->GetSettingForBranch<UMovieGraphGlobalOutputSettingNode>(InRenderDataIdentifier.RootBranchName, bIncludeCDOs);
 	if (!ensure(OutputSettings))
 	{
 		return FString();
@@ -657,13 +683,13 @@ FString UMovieGraphImageSequenceOutputNode_MultiLayerEXR::ResolveOutputFilename(
 		FormatOverrides.Add(TEXT("ExtraTag"), FString::Printf(TEXT("_Add(%d)"), ResolutionIndex));
 	}
 
-	// Since a multi-layer EXR can store renders from multiple cameras, renderers, etc, the render data identifier isn't
-	// very useful. However, we still need to provide the branch name -- this is important for resolving the correct output path.
-	FMovieGraphRenderDataIdentifier TempRenderDataIdentifier;
-	TempRenderDataIdentifier.RootBranchName = InBranchName;
-
-	FMovieGraphFilenameResolveParams Params = FMovieGraphFilenameResolveParams::MakeResolveParams(
-		TempRenderDataIdentifier, InPipeline, InRawFrameData->EvaluatedConfig.Get(), InRawFrameData->TraversalContext, FormatOverrides);
+	// The layer's render data identifier is used here in the resolve. Usually this is not a problem. However, the user may include some tokens, like
+	// {layer_name}, that come from the identifier, which will prevent all layers from being placed in the same multi-layer EXR (because now the path
+	// isn't resolving to the path that other layers are resolving to). We have to assume that the user is doing this intentionally, even though it's
+	// a bit strange. Including the full identifier here is important so all custom metadata is resolved correctly (see
+	// UMovieGraphSetMetadataAttributesNode) when ResolveFilenameFormatArguments() is called.
+	const FMovieGraphFilenameResolveParams Params = FMovieGraphFilenameResolveParams::MakeResolveParams(
+		InRenderDataIdentifier, InPipeline, InRawFrameData->EvaluatedConfig.Get(), InRawFrameData->TraversalContext, FormatOverrides);
 	
 	const FString FilePathFormatString = OutputSettings->OutputDirectory.Path / FileNameFormatString;
 
diff --git a/Engine/Plugins/MovieScene/MovieRenderPipeline/Source/MovieRenderPipelineRenderPasses/Public/MovieGraphImageSequenceOutputNode.h b/Engine/Plugins/MovieScene/MovieRenderPipeline/Source/MovieRenderPipelineRenderPasses/Public/MovieGraphImageSequenceOutputNode.h
--- a/Engine/Plugins/MovieScene/MovieRenderPipeline/Source/MovieRenderPipelineRenderPasses/Public/MovieGraphImageSequenceOutputNode.h
+++ b/Engine/Plugins/MovieScene/MovieRenderPipeline/Source/MovieRenderPipelineRenderPasses/Public/MovieGraphImageSequenceOutputNode.h
@@ -239,7 +239,7 @@ private:
 	FString ResolveOutputFilename(
 		const UMovieGraphImageSequenceOutputNode_MultiLayerEXR* InParentNode,
 		const UMovieGraphPipeline* InPipeline, const int32 ResolutionIndex, const UE::MovieGraph::FMovieGraphOutputMergerFrame* InRawFrameData,
-		const FName& InBranchName, FMovieGraphResolveArgs& OutResolveArgs) const;
+		const FMovieGraphRenderDataIdentifier& InRenderDataIdentifier, FMovieGraphResolveArgs& OutResolveArgs) const;
 };
 
 /**
