MovieGraph: Fixed crash when TemporalSampleCount was zero (which could be accomplished by promoting something to a variable and then leaving the default value of the promoted variable at zero).

diff --git a/Engine/Plugins/MovieScene/MovieRenderPipeline/Source/MovieRenderPipelineCore/Private/Graph/MovieGraphLinearTimeStep.cpp b/Engine/Plugins/MovieScene/MovieRenderPipeline/Source/MovieRenderPipelineCore/Private/Graph/MovieGraphLinearTimeStep.cpp
--- a/Engine/Plugins/MovieScene/MovieRenderPipeline/Source/MovieRenderPipelineCore/Private/Graph/MovieGraphLinearTimeStep.cpp
+++ b/Engine/Plugins/MovieScene/MovieRenderPipeline/Source/MovieRenderPipelineCore/Private/Graph/MovieGraphLinearTimeStep.cpp
@@ -1,6 +1,7 @@
-﻿// Copyright Epic Games, Inc. All Rights Reserved.
+// Copyright Epic Games, Inc. All Rights Reserved.
 
 #include "Graph/MovieGraphLinearTimeStep.h"
+#include "MovieRenderPipelineCoreModule.h"
 
 #include "Graph/Nodes/MovieGraphSamplingMethodNode.h"
 
@@ -15,6 +16,10 @@ int32 UMovieGraphLinearTimeStep::GetTemporalSampleCount() const
 	constexpr bool bIncludeCDOs = true;
 	const UMovieGraphSamplingMethodNode* SamplingMethod =
 		CurrentFrameData.EvaluatedConfig->GetSettingForBranch<UMovieGraphSamplingMethodNode>(UMovieGraphNode::GlobalsPinName, bIncludeCDOs);
+	if (SamplingMethod->TemporalSampleCount <= 0)
+	{
+		UE_LOG(LogMovieRenderPipeline, Error, TEXT("Sampling Method > Temporal Sample Count was zero, this is not allowed. Forcing value to 1!"));
+	}
 
-	return SamplingMethod->TemporalSampleCount;
+	return FMath::Max(SamplingMethod->TemporalSampleCount, 1);
 }
diff --git a/Engine/Plugins/MovieScene/MovieRenderPipeline/Source/MovieRenderPipelineCore/Private/Graph/MovieGraphPipeline.cpp b/Engine/Plugins/MovieScene/MovieRenderPipeline/Source/MovieRenderPipelineCore/Private/Graph/MovieGraphPipeline.cpp
--- a/Engine/Plugins/MovieScene/MovieRenderPipeline/Source/MovieRenderPipelineCore/Private/Graph/MovieGraphPipeline.cpp
+++ b/Engine/Plugins/MovieScene/MovieRenderPipeline/Source/MovieRenderPipelineCore/Private/Graph/MovieGraphPipeline.cpp
@@ -553,7 +553,7 @@ void UMovieGraphPipeline::BuildShotListFromDataSource()
 
 		UMovieGraphWarmUpSettingNode* WarmUpNode = EvaluatedConfig->GetSettingForBranch<UMovieGraphWarmUpSettingNode>(UMovieGraphSettingNode::GlobalsPinName);
 
-		Shot->ShotInfo.NumTemporalSamples = SamplingMethodNode->TemporalSampleCount;
+		Shot->ShotInfo.NumTemporalSamples = FMath::Max(SamplingMethodNode->TemporalSampleCount, 1);
 		Shot->ShotInfo.NumSpatialSamples = 1;
 		Shot->ShotInfo.NumTiles = FIntPoint(1,1);
 		Shot->ShotInfo.CachedFrameRate = FinalFrameRate;
