TObjectPtr conversion for MovieScene plugin.

diff --git a/Engine/Plugins/MovieScene/MovieRenderPipeline/Source/MovieRenderPipelineCore/Private/Graph/MovieGraphRenderLayerSubsystem.cpp b/Engine/Plugins/MovieScene/MovieRenderPipeline/Source/MovieRenderPipelineCore/Private/Graph/MovieGraphRenderLayerSubsystem.cpp
--- a/Engine/Plugins/MovieScene/MovieRenderPipeline/Source/MovieRenderPipelineCore/Private/Graph/MovieGraphRenderLayerSubsystem.cpp
+++ b/Engine/Plugins/MovieScene/MovieRenderPipeline/Source/MovieRenderPipelineCore/Private/Graph/MovieGraphRenderLayerSubsystem.cpp
@@ -64,7 +64,7 @@ namespace UE::MovieGraph::Private
 	class FClassViewerTypeFilter final : public IClassViewerFilter
 	{
 	public:
-		explicit FClassViewerTypeFilter(TArray<UClass*>* InClassesToDisallow, UClass* InRequiredBaseClass = nullptr)
+		explicit FClassViewerTypeFilter(TArray<TObjectPtr<UClass>>* InClassesToDisallow, UClass* InRequiredBaseClass = nullptr)
 			: ClassesToDisallow(InClassesToDisallow)
 			, RequiredBaseClass(InRequiredBaseClass)
 		{
@@ -86,7 +86,7 @@ namespace UE::MovieGraph::Private
 
 	private:
 		/** Classes which should be prevented from showing up in the class viewer. */
-		TArray<UClass*>* ClassesToDisallow;
+		TArray<TObjectPtr<UClass>>* ClassesToDisallow;
 
 		/** Classes must have this base class to pass the filter. */
 		UClass* RequiredBaseClass = nullptr;
@@ -1155,7 +1155,7 @@ TArray<TSharedRef<SWidget>> UMovieGraphConditionGroupQuery_ActorType::GetWidgets
 			return FReply::Handled();
 		})
 		[
-			SAssignNew(ActorTypesList, SMovieGraphSimpleList<UClass*>)
+			SAssignNew(ActorTypesList, SMovieGraphSimpleList<TObjectPtr<UClass>>)
 			.DataSource(&ActorTypes)
 			.DataType(FText::FromString("Actor Type"))
 			.DataTypePlural(FText::FromString("Actor Types"))
@@ -1215,12 +1215,12 @@ TSharedRef<SWidget> UMovieGraphConditionGroupQuery_ActorType::GetAddMenuContents
 		];
 }
 
-const FSlateBrush* UMovieGraphConditionGroupQuery_ActorType::GetRowIcon(UClass* InActorType)
+const FSlateBrush* UMovieGraphConditionGroupQuery_ActorType::GetRowIcon(TObjectPtr<UClass> InActorType)
 {
 	return FSlateIconFinder::FindIconForClass(InActorType).GetIcon();
 }
 
-FText UMovieGraphConditionGroupQuery_ActorType::GetRowText(UClass* InActorType)
+FText UMovieGraphConditionGroupQuery_ActorType::GetRowText(TObjectPtr<UClass> InActorType)
 {
 	if (InActorType)
 	{
@@ -1392,7 +1392,7 @@ TArray<TSharedRef<SWidget>> UMovieGraphConditionGroupQuery_ComponentType::GetWid
 	TArray<TSharedRef<SWidget>> Widgets;
 
 	Widgets.Add(
-		SAssignNew(ComponentTypesList, SMovieGraphSimpleList<UClass*>)
+		SAssignNew(ComponentTypesList, SMovieGraphSimpleList<TObjectPtr<UClass>>)
 			.DataSource(&ComponentTypes)
 			.DataType(FText::FromString("Component Type"))
 			.DataTypePlural(FText::FromString("Component Types"))
@@ -1465,12 +1465,12 @@ TSharedRef<SWidget> UMovieGraphConditionGroupQuery_ComponentType::GetAddMenuCont
 		];
 }
 
-const FSlateBrush* UMovieGraphConditionGroupQuery_ComponentType::GetRowIcon(UClass* InComponentType)
+const FSlateBrush* UMovieGraphConditionGroupQuery_ComponentType::GetRowIcon(TObjectPtr<UClass> InComponentType)
 {
 	return FSlateIconFinder::FindIconForClass(InComponentType).GetIcon();
 }
 
-FText UMovieGraphConditionGroupQuery_ComponentType::GetRowText(UClass* InComponentType)
+FText UMovieGraphConditionGroupQuery_ComponentType::GetRowText(TObjectPtr<UClass> InComponentType)
 {
 	if (InComponentType)
 	{
@@ -2371,7 +2371,7 @@ TSet<AActor*> UMovieGraphConditionGroup::Evaluate(const UWorld* InWorld) const
 		// Similar to EvaluationResult, QueryResult is persisted+reset to prevent constantly re-allocating it
 		QueryResult.Reset();
 
-		Query->Evaluate(AllActors, InWorld, QueryResult);
+		Query->Evaluate(AllActors, InWorld, MutableView(QueryResult));
 		
 		switch (Query->GetOperationType())
 		{
@@ -2390,7 +2390,7 @@ TSet<AActor*> UMovieGraphConditionGroup::Evaluate(const UWorld* InWorld) const
 		}
 	}
 	
-	return EvaluationResult;
+	return ObjectPtrDecay(EvaluationResult);
 }
 
 UMovieGraphConditionGroupQueryBase* UMovieGraphConditionGroup::AddQuery(const TSubclassOf<UMovieGraphConditionGroupQueryBase>& InQueryType, const int32 InsertIndex)
diff --git a/Engine/Plugins/MovieScene/MovieRenderPipeline/Source/MovieRenderPipelineCore/Public/Graph/MovieGraphRenderLayerSubsystem.h b/Engine/Plugins/MovieScene/MovieRenderPipeline/Source/MovieRenderPipelineCore/Public/Graph/MovieGraphRenderLayerSubsystem.h
--- a/Engine/Plugins/MovieScene/MovieRenderPipeline/Source/MovieRenderPipelineCore/Public/Graph/MovieGraphRenderLayerSubsystem.h
+++ b/Engine/Plugins/MovieScene/MovieRenderPipeline/Source/MovieRenderPipelineCore/Public/Graph/MovieGraphRenderLayerSubsystem.h
@@ -304,18 +304,18 @@ public:
 public:
 	/** The type (class) that the actor needs to have in order to be a match. */
 	UPROPERTY(EditAnywhere, Category="General")
-	TArray<UClass*> ActorTypes;
+	TArray<TObjectPtr<UClass>> ActorTypes;
 
 private:
 #if WITH_EDITOR
-	static const FSlateBrush* GetRowIcon(UClass* InActorType);
-	static FText GetRowText(UClass* InActorType);
+	static const FSlateBrush* GetRowIcon(TObjectPtr<UClass> InActorType);
+	static FText GetRowText(TObjectPtr<UClass> InActorType);
 	
 	/** Adds the provided actor types to the query, updating the UI as needed. Calls InOnAddFinished when done. */
 	void AddActorTypes(const TArray<UClass*>& InActorTypes, const FMovieGraphConditionGroupQueryContentsChanged& InOnAddFinished);
 
 	/** Displays the actor types which have been chosen. */
-	TSharedPtr<SMovieGraphSimpleList<UClass*>> ActorTypesList;
+	TSharedPtr<SMovieGraphSimpleList<TObjectPtr<UClass>>> ActorTypesList;
 #endif
 	
 	/** The class viewer widget to show in the Add menu. */
@@ -374,11 +374,11 @@ public:
 
 private:
 #if WITH_EDITOR
-	static const FSlateBrush* GetRowIcon(UClass* InComponentType);
-	static FText GetRowText(UClass* InComponentType);
+	static const FSlateBrush* GetRowIcon(TObjectPtr<UClass> InComponentType);
+	static FText GetRowText(TObjectPtr<UClass> InComponentType);
 
 	/** Displays the component types which have been chosen. */
-	TSharedPtr<SMovieGraphSimpleList<UClass*>> ComponentTypesList;
+	TSharedPtr<SMovieGraphSimpleList<TObjectPtr<UClass>>> ComponentTypesList;
 #endif
 	
 	/** The class viewer widget to show in the Add menu. */
@@ -387,7 +387,7 @@ private:
 public:
 	/** The actor must have one or more of the component type(s) in order to be a match. */
 	UPROPERTY(EditAnywhere, Category="General")
-	TArray<UClass*> ComponentTypes;
+	TArray<TObjectPtr<UClass>> ComponentTypes;
 };
 
 /** Query type which filters actors via the editor folder that they're contained in. */
@@ -628,11 +628,11 @@ private:
 
 	/** Persisted actor set which can be re-used for query evaluations across frames to prevent constantly re-allocating it. */
 	UPROPERTY(Transient, DuplicateTransient)
-	mutable TSet<AActor*> QueryResult;
+	mutable TSet<TObjectPtr<AActor>> QueryResult;
 
 	/** Persisted actor set which can be re-used for condition group evaluations across frames to prevent constantly re-allocating it. */
 	UPROPERTY(Transient, DuplicateTransient)
-	mutable TSet<AActor*> EvaluationResult;
+	mutable TSet<TObjectPtr<AActor>> EvaluationResult;
 };
 
 /** A group of actors generated by actor queries. */
