MoviePipeline: Fixed a bug where the C++ add-node API didn't work properly when called from Python.

Normally when nodes are created, it is done interactively through the editor, causing the associated editor node to be created. However, when done via Python, the runtime node was created, but not the editor node, making it appear like the node was never created.

To fix the issue, every time the graph editor is started, the runtime graph is put through a process to make sure there are editor equivalents for all nodes and connections.

diff --git a/Engine/Plugins/MovieScene/MovieRenderPipeline/Source/MovieRenderPipelineEditor/Private/Graph/MovieEdGraph.cpp b/Engine/Plugins/MovieScene/MovieRenderPipeline/Source/MovieRenderPipelineEditor/Private/Graph/MovieEdGraph.cpp
--- a/Engine/Plugins/MovieScene/MovieRenderPipeline/Source/MovieRenderPipelineEditor/Private/Graph/MovieEdGraph.cpp
+++ b/Engine/Plugins/MovieScene/MovieRenderPipeline/Source/MovieRenderPipelineEditor/Private/Graph/MovieEdGraph.cpp
@@ -27,26 +27,44 @@ UMovieGraphConfig* UMoviePipelineEdGraph::GetPipelineGraph() const
 
 void UMoviePipelineEdGraph::InitFromRuntimeGraph(UMovieGraphConfig* InGraph)
 {
-	// Don't allow reinitialization of an existing graph
-	check(InGraph && !bInitialized);
+	check(InGraph);
 
 	TMap<UMovieGraphNode*, UMoviePipelineEdGraphNodeBase*> NodeLookup;
 
 	// Input
 	{
 		UMovieGraphNode* InputNode = InGraph->GetInputNode();
-		NodeLookup.Add(InputNode, CreateNodeFromRuntimeNode<UMoviePipelineEdGraphNodeInput>(InputNode));
+		UMoviePipelineEdGraphNodeBase* EdInputNode = Cast<UMoviePipelineEdGraphNodeBase>(InputNode->GraphNode);
+		if (!EdInputNode)
+		{
+			EdInputNode = CreateNodeFromRuntimeNode<UMoviePipelineEdGraphNodeInput>(InputNode);
+		}
+		
+		NodeLookup.Add(InputNode, EdInputNode);
 	}
 
 	// Output
 	{
 		UMovieGraphNode* OutputNode = InGraph->GetOutputNode();
-		NodeLookup.Add(OutputNode, CreateNodeFromRuntimeNode<UMoviePipelineEdGraphNodeOutput>(OutputNode));
+		UMoviePipelineEdGraphNodeBase* EdOutputNode = Cast<UMoviePipelineEdGraphNodeBase>(OutputNode->GraphNode);
+		if (!EdOutputNode)
+		{
+			EdOutputNode = CreateNodeFromRuntimeNode<UMoviePipelineEdGraphNodeOutput>(OutputNode);
+		}
+		
+		NodeLookup.Add(OutputNode, EdOutputNode);
 	}
 
 	// Create the rest of the nodes in the graph
 	for (const TObjectPtr<UMovieGraphNode>& RuntimeNode : InGraph->GetNodes())
 	{
+		// Don't re-create the editor node if it already exists
+		if (RuntimeNode->GraphNode)
+		{
+			NodeLookup.Add(RuntimeNode, Cast<UMoviePipelineEdGraphNodeBase>(RuntimeNode->GraphNode));
+			continue;
+		}
+		
 		if (RuntimeNode->IsA<UMovieGraphVariableNode>())
 		{
 			NodeLookup.Add(RuntimeNode, CreateNodeFromRuntimeNode<UMoviePipelineEdGraphVariableNode>(RuntimeNode));
@@ -63,24 +81,27 @@ void UMoviePipelineEdGraph::InitFromRuntimeGraph(UMovieGraphConfig* InGraph)
 
 	// Now that we've added an Editor Graph representation for every runtime node in the graph, link
 	// the editor nodes together to match the Runtime Layout.
-	for (const TPair< UMovieGraphNode*, UMoviePipelineEdGraphNodeBase*> Pair : NodeLookup)
+	for (const TPair<UMovieGraphNode*, UMoviePipelineEdGraphNodeBase*>& Pair : NodeLookup)
 	{
-		const bool bCreateInboundLinks = false;
-		const bool bCreateOutboundLinks = true;
+		constexpr bool bCreateInboundLinks = false;
+		constexpr bool bCreateOutboundLinks = true;
 		CreateLinks(Pair.Value, bCreateInboundLinks, bCreateOutboundLinks, NodeLookup);
 	}
 
-	// Restore editor-only nodes, which have no runtime node equivalent
-	for (const TObjectPtr<UObject>& EditorOnlyNodeObject : InGraph->GetEditorOnlyNodes())
-	{
-		if (const UEdGraphNode* EdGraphNode = Cast<UEdGraphNode>(EditorOnlyNodeObject))
-		{
-			UEdGraphNode* NewEdGraphNode = DuplicateObject(EdGraphNode, /*Outer=*/this);
-			const bool bIsUserAction = false;
-			const bool bSelectNewNode = false;
-			AddNode(NewEdGraphNode, bIsUserAction, bSelectNewNode);
-		}
-	}
+	// NOTE: Commented-out because the runtime graph has been updated to persist an editor graph, so editor-only nodes do not need to be
+	// restored anymore. If the runtime graph is updated to not persist the editor graph in the future, the below code will need to be run again.
+	//
+	// Restore editor-only nodes, which have no runtime node equivalent.
+	// for (const TObjectPtr<UObject>& EditorOnlyNodeObject : InGraph->GetEditorOnlyNodes())
+	// {
+	// 	if (const UEdGraphNode* EdGraphNode = Cast<UEdGraphNode>(EditorOnlyNodeObject))
+	// 	{
+	// 		UEdGraphNode* NewEdGraphNode = DuplicateObject(EdGraphNode, /*Outer=*/this);
+	// 		const bool bIsUserAction = false;
+	// 		const bool bSelectNewNode = false;
+	// 		AddNode(NewEdGraphNode, bIsUserAction, bSelectNewNode);
+	// 	}
+	// }
 
 	RegisterDelegates(InGraph);
 
@@ -89,8 +110,15 @@ void UMoviePipelineEdGraph::InitFromRuntimeGraph(UMovieGraphConfig* InGraph)
 
 void UMoviePipelineEdGraph::RegisterDelegates(UMovieGraphConfig* InGraph)
 {
-	InGraph->OnGraphChangedDelegate.AddUObject(this, &UMoviePipelineEdGraph::OnGraphConfigChanged);
-	InGraph->OnGraphNodesDeletedDelegate.AddUObject(this, &UMoviePipelineEdGraph::OnGraphNodesDeleted);
+	if (!InGraph->OnGraphChangedDelegate.IsBoundToObject(this))
+	{
+		InGraph->OnGraphChangedDelegate.AddUObject(this, &UMoviePipelineEdGraph::OnGraphConfigChanged);
+	}
+
+	if (!InGraph->OnGraphNodesDeletedDelegate.IsBoundToObject(this))
+	{
+		InGraph->OnGraphNodesDeletedDelegate.AddUObject(this, &UMoviePipelineEdGraph::OnGraphNodesDeleted);
+	}
 }
 
 void UMoviePipelineEdGraph::CreateLinks(UMoviePipelineEdGraphNodeBase* InGraphNode, bool bCreateInboundLinks, bool bCreateOutboundLinks)
diff --git a/Engine/Plugins/MovieScene/MovieRenderPipeline/Source/MovieRenderPipelineEditor/Private/Graph/MovieEdGraph.h b/Engine/Plugins/MovieScene/MovieRenderPipeline/Source/MovieRenderPipelineEditor/Private/Graph/MovieEdGraph.h
--- a/Engine/Plugins/MovieScene/MovieRenderPipeline/Source/MovieRenderPipelineEditor/Private/Graph/MovieEdGraph.h
+++ b/Engine/Plugins/MovieScene/MovieRenderPipeline/Source/MovieRenderPipelineEditor/Private/Graph/MovieEdGraph.h
@@ -20,7 +20,14 @@ class UMoviePipelineEdGraph : public UEdGraph
 	GENERATED_BODY()
 
 public:
-	/** Initialize this Editor Graph from a Runtime Graph */
+	/**
+	 * Initialize this Editor Graph from a Runtime Graph.
+	 *
+	 * This is typically called when the editor representation of the runtime graph is initially created (in order to ensure there are editor
+	 * equivalents of all runtime nodes, connections, etc). However, it can also be called to ensure that all runtime objects have editor
+	 * equivalents even after initial editor graph creation. This can be useful when, for example, nodes are created via the Python API (which
+	 * creates runtime nodes), and editor nodes need to be created afterwards.
+	 **/
 	void InitFromRuntimeGraph(UMovieGraphConfig* InGraph);
 
 	/** Register delegates that relate to the runtime graph. */
diff --git a/Engine/Plugins/MovieScene/MovieRenderPipeline/Source/MovieRenderPipelineEditor/Private/Widgets/Graph/SMovieGraphConfigPanel.cpp b/Engine/Plugins/MovieScene/MovieRenderPipeline/Source/MovieRenderPipelineEditor/Private/Widgets/Graph/SMovieGraphConfigPanel.cpp
--- a/Engine/Plugins/MovieScene/MovieRenderPipeline/Source/MovieRenderPipelineEditor/Private/Widgets/Graph/SMovieGraphConfigPanel.cpp
+++ b/Engine/Plugins/MovieScene/MovieRenderPipeline/Source/MovieRenderPipelineEditor/Private/Widgets/Graph/SMovieGraphConfigPanel.cpp
@@ -108,6 +108,10 @@ void SMoviePipelineGraphPanel::Construct(const FArguments& InArgs)
 
 	// Create the EdGraph if it has not yet been created. It is saved as part of the runtime graph to prevent it from being re-created every time the
 	// graph is opened (and therefore dirtying the package).
+	//
+	// Note: Now that the editor graph is saved with the runtime graph, the !EdGraph condition will rarely be true (because when the "template"
+	// runtime graph is duplicated into a new runtime graph, the editor graph comes along as well). There may be other scenarios where a new runtime
+	// graph is created without a duplication though, so this code path remains.
 	if (!EdGraph)
 	{
 		EdGraph = Cast<UMoviePipelineEdGraph>(FBlueprintEditorUtils::CreateNewGraph(CurrentGraph, TEXT("MoviePipelineEdGraph"), UMoviePipelineEdGraph::StaticClass(), UMovieGraphSchema::StaticClass()));
@@ -122,7 +126,7 @@ void SMoviePipelineGraphPanel::Construct(const FArguments& InArgs)
 	}
 	else
 	{
-		EdGraph->RegisterDelegates(CurrentGraph);
+		EdGraph->InitFromRuntimeGraph(CurrentGraph);
 	}
 	
 	MakeEditorCommands();
