MoviePipeline: Added "script tags" to nodes in the graph to allow for easy identification within scripts.

diff --git a/Engine/Plugins/MovieScene/MovieRenderPipeline/Source/MovieRenderPipelineCore/Public/Graph/MovieGraphNode.h b/Engine/Plugins/MovieScene/MovieRenderPipeline/Source/MovieRenderPipelineCore/Public/Graph/MovieGraphNode.h
--- a/Engine/Plugins/MovieScene/MovieRenderPipeline/Source/MovieRenderPipelineCore/Public/Graph/MovieGraphNode.h
+++ b/Engine/Plugins/MovieScene/MovieRenderPipeline/Source/MovieRenderPipelineCore/Public/Graph/MovieGraphNode.h
@@ -336,6 +336,13 @@ public:
 public:
 	FOnMovieGraphNodeChanged OnNodeChangedDelegate;
 
+	/**
+	 * Tags that can be used to identify this node within a pre/post render script. Tags can be unique in order to identify this specific node,
+	 * or the same tag can be applied to multiple nodes in order to identify a grouping of nodes.
+	 */
+	UPROPERTY(EditAnywhere, BlueprintReadWrite, Category = "Settings", AdvancedDisplay)
+	TArray<FString> ScriptTags;
+
 #if WITH_EDITORONLY_DATA
 	/** Editor Node Graph representation. Not strongly typed to avoid circular dependency between editor/runtime modules. */
 	UPROPERTY()
