MoviePipeline: Updated Post Process Materials within renderer nodes to include an additional "Name" field (optional) that will be used as the sub-resource name (the {renderer_sub_name} token) in file name formatting. Without this, the entire PPM's name name will be used, which can be lengthy.

diff --git a/Engine/Plugins/MovieScene/MovieRenderPipeline/Source/MovieRenderPipelineRenderPasses/Private/Graph/Renderers/MovieGraphDeferredPass.cpp b/Engine/Plugins/MovieScene/MovieRenderPipeline/Source/MovieRenderPipelineRenderPasses/Private/Graph/Renderers/MovieGraphDeferredPass.cpp
--- a/Engine/Plugins/MovieScene/MovieRenderPipeline/Source/MovieRenderPipelineRenderPasses/Private/Graph/Renderers/MovieGraphDeferredPass.cpp
+++ b/Engine/Plugins/MovieScene/MovieRenderPipeline/Source/MovieRenderPipelineRenderPasses/Private/Graph/Renderers/MovieGraphDeferredPass.cpp
@@ -77,11 +77,10 @@ void FMovieGraphDeferredPass::GatherOutputPasses(UMovieGraphEvaluatedConfig* InC
 		{
 			if (AdditionalPass.bEnabled)
 			{
-				UMaterialInterface* Material = AdditionalPass.Material.LoadSynchronous();
-				if (Material)
+				if (const UMaterialInterface* Material = AdditionalPass.Material.LoadSynchronous())
 				{
 					FMovieGraphRenderDataIdentifier Identifier = RenderDataIdentifier;
-					Identifier.SubResourceName = Material->GetName();
+					Identifier.SubResourceName = AdditionalPass.Name.IsEmpty() ? Material->GetName() : AdditionalPass.Name;
 					OutExpectedPasses.Add(Identifier);
 				}
 			}
@@ -279,41 +278,34 @@ void FMovieGraphDeferredPass::Render(const FMovieGraphTraversalContext& InFrameT
 
 			for (const FMoviePipelinePostProcessPass& PostProcessPass : ParentNode->GetAdditionalPostProcessMaterials())
 			{
-				if (PostProcessPass.bEnabled)
+				UMaterialInterface* Material = PostProcessPass.Material.LoadSynchronous();
+				
+				if (!PostProcessPass.bEnabled || !Material)
 				{
-					UMaterialInterface* Material = PostProcessPass.Material.LoadSynchronous();
-					if (Material)
-					{
-						NewView->FinalPostProcessSettings.BufferVisualizationOverviewMaterials.Add(Material);
-						
-						if (PostProcessPass.bHighPrecisionOutput)
-						{
-							HighPrecisionMaterials.Add(Material);
-						}
-					}
+					continue;
 				}
-			}
-
-			for (UMaterialInterface* VisMaterial : NewView->FinalPostProcessSettings.BufferVisualizationOverviewMaterials)
-			{
+				
 				auto BufferPipe = MakeShared<FImagePixelPipe, ESPMode::ThreadSafe>();
-
-				if (HighPrecisionMaterials.Contains(VisMaterial))
+				
+				NewView->FinalPostProcessSettings.BufferVisualizationOverviewMaterials.Add(Material);
+				
+				if (PostProcessPass.bHighPrecisionOutput)
 				{
+					HighPrecisionMaterials.Add(Material);
 					BufferPipe->bIsExpecting32BitPixelData = true;
 				}
 				
 				FMovieGraphRenderDataIdentifier Identifier = RenderDataIdentifier;
-				Identifier.SubResourceName = VisMaterial->GetName();
+				Identifier.SubResourceName = PostProcessPass.Name.IsEmpty() ? Material->GetName() : PostProcessPass.Name;
 				
-				UE::MovieGraph::FMovieGraphSampleState PassSampleState = SampleState;
+				FMovieGraphSampleState PassSampleState = SampleState;
 				PassSampleState.TraversalContext.RenderDataIdentifier = Identifier;
 				
 				// Give a lower priority to materials so they show up after the main pass in multi-layer exrs.
 				PassSampleState.CompositingSortOrder = SampleState.CompositingSortOrder + 1;
 				BufferPipe->AddEndpoint(MakeForwardingEndpoint(PassSampleState, InTimeData));
 
-				NewView->FinalPostProcessSettings.BufferVisualizationPipes.Add(VisMaterial->GetFName(), BufferPipe);
+				NewView->FinalPostProcessSettings.BufferVisualizationPipes.Add(Material->GetFName(), BufferPipe);
 			}
 		}
 
diff --git a/Engine/Plugins/MovieScene/MovieRenderPipeline/Source/MovieRenderPipelineRenderPasses/Public/MoviePipelineDeferredPasses.h b/Engine/Plugins/MovieScene/MovieRenderPipeline/Source/MovieRenderPipelineRenderPasses/Public/MoviePipelineDeferredPasses.h
--- a/Engine/Plugins/MovieScene/MovieRenderPipeline/Source/MovieRenderPipelineRenderPasses/Public/MoviePipelineDeferredPasses.h
+++ b/Engine/Plugins/MovieScene/MovieRenderPipeline/Source/MovieRenderPipelineRenderPasses/Public/MoviePipelineDeferredPasses.h
@@ -22,6 +22,13 @@ public:
 	UPROPERTY(EditAnywhere, BlueprintReadWrite, Category = "Settings")
 	bool bEnabled = false;
 
+	/**
+	 * An optional name which which will be used as the {renderer_sub_name} file name format token for this material. If a name is not specified here,
+	 * the name of the material will be used for {renderer_sub_name}.
+	 */
+	UPROPERTY(EditAnywhere, BlueprintReadWrite, Category = "Settings")
+	FString Name;
+
 	/** 
 	* Material should be set to Post Process domain, and Blendable Location = After Tonemapping. 
 	* This will need bDisableMultisampleEffects enabled for pixels to line up(ie : no DoF, MotionBlur, TAA)
