MoviePipeline: Promoting a pinned property to a variable now sets the new variable's default value to the current value of the property.

diff --git a/Engine/Plugins/MovieScene/MovieRenderPipeline/Source/MovieRenderPipelineEditor/Private/Graph/MovieEdGraphNode.cpp b/Engine/Plugins/MovieScene/MovieRenderPipeline/Source/MovieRenderPipelineEditor/Private/Graph/MovieEdGraphNode.cpp
--- a/Engine/Plugins/MovieScene/MovieRenderPipeline/Source/MovieRenderPipelineEditor/Private/Graph/MovieEdGraphNode.cpp
+++ b/Engine/Plugins/MovieScene/MovieRenderPipeline/Source/MovieRenderPipelineEditor/Private/Graph/MovieEdGraphNode.cpp
@@ -357,6 +357,22 @@ void UMoviePipelineEdGraphNode::PromotePropertyToVariable(const FMovieGraphPrope
 		UObject* ValueTypeObject = const_cast<UObject*>(TargetProperty.ValueTypeObject.Get());
 		NewGraphVariable->SetValueType(TargetProperty.ValueType, ValueTypeObject);
 
+		// When promoting, set the variable's default value to the connected property's current value. That will ensure that there's no
+		// behavior change in the graph after the promotion.
+		{
+			FString TargetPropertyValue;
+			if (TargetProperty.bIsDynamicProperty)
+			{
+				RuntimeNode->GetDynamicPropertyValue(TargetProperty.Name, TargetPropertyValue); 
+			}
+			else if (const FProperty* TargetFProperty = FindFProperty<FProperty>(RuntimeNode->GetClass(), TargetProperty.Name))
+			{
+				TargetFProperty->ExportTextItem_InContainer(TargetPropertyValue, RuntimeNode, nullptr, RuntimeNode, PPF_None);
+			}
+			
+			NewGraphVariable->SetValueSerializedString(TargetPropertyValue);
+		}
+
 		// When creating the new action, since it's only being used to create a node, the category, display name, and tooltip can just be empty
 		const TSharedPtr<FMovieGraphSchemaAction_NewVariableNode> NewAction = MakeShared<FMovieGraphSchemaAction_NewVariableNode>(
 			FText::GetEmpty(), FText::GetEmpty(), NewGraphVariable->GetGuid(), FText::GetEmpty());
