MoviePipeline: Updated the lighting components on the path tracer node to work properly. Previously they were exposed under Show Flags (like they are for the deferred renderer); however, these show flags do not work for the path tracer, so they were removed.

A dedicated Lighting Components section was added for the path tracer node, which instead impact the post processing volume's path tracing Lighting Components properties.

Also adjusted the settings categories for a couple properties in the Deferred/Path Tracer renderer nodes to "Post Processing".

diff --git a/Engine/Plugins/MovieScene/MovieRenderPipeline/Source/MovieRenderPipelineEditor/Private/Customizations/Graph/MovieGraphShowFlagsCustomization.h b/Engine/Plugins/MovieScene/MovieRenderPipeline/Source/MovieRenderPipelineEditor/Private/Customizations/Graph/MovieGraphShowFlagsCustomization.h
--- a/Engine/Plugins/MovieScene/MovieRenderPipeline/Source/MovieRenderPipelineEditor/Private/Customizations/Graph/MovieGraphShowFlagsCustomization.h
+++ b/Engine/Plugins/MovieScene/MovieRenderPipeline/Source/MovieRenderPipelineEditor/Private/Customizations/Graph/MovieGraphShowFlagsCustomization.h
@@ -5,6 +5,7 @@
 #include "DetailWidgetRow.h"
 #include "EditorShowFlags.h"
 #include "Graph/Nodes/MovieGraphImagePassBaseNode.h"
+#include "Graph/Nodes/MovieGraphPathTracerPassNode.h"
 #include "Graph/Renderers/MovieGraphShowFlags.h"
 #include "IDetailChildrenBuilder.h"
 #include "IDetailGroup.h"
@@ -62,10 +63,19 @@ protected:
 		UMovieGraphShowFlags* ShowFlagObject = Cast<UMovieGraphShowFlags>(
 			ObjectProperty->GetObjectPropertyValue(ObjectProperty->ContainerPtrToValuePtr<void>(OuterObjects[0])));
 
+		const UClass* NodeClass = OuterObjects[0]->GetClass();
+
 		// Group together the show flags with the groups they show up under in the UI
 		TMap<EShowFlagGroup, TArray<FShowFlagData>> GroupedShowFlags;
 		for (FShowFlagData& ShowFlag : GetShowFlagMenuItems())
 		{
+			// Path Traced Renderer nodes don't show the Lighting Components group because these show flags have no effect on PT renders. Instead,
+			// these are controlled by the dedicated Lighting Components settings for the path tracer.
+			if ((NodeClass == UMovieGraphPathTracerRenderPassNode::StaticClass()) && (ShowFlag.Group == SFG_LightingComponents))
+			{
+				continue;
+			}
+			
 			TArray<FShowFlagData>& GroupShowFlags = GroupedShowFlags.FindOrAdd(ShowFlag.Group);
 			GroupShowFlags.Add(MoveTemp(ShowFlag));
 		}
diff --git a/Engine/Plugins/MovieScene/MovieRenderPipeline/Source/MovieRenderPipelineRenderPasses/Private/Graph/Renderers/MovieGraphPathTracerPass.cpp b/Engine/Plugins/MovieScene/MovieRenderPipeline/Source/MovieRenderPipelineRenderPasses/Private/Graph/Renderers/MovieGraphPathTracerPass.cpp
--- a/Engine/Plugins/MovieScene/MovieRenderPipeline/Source/MovieRenderPipelineRenderPasses/Private/Graph/Renderers/MovieGraphPathTracerPass.cpp
+++ b/Engine/Plugins/MovieScene/MovieRenderPipeline/Source/MovieRenderPipelineRenderPasses/Private/Graph/Renderers/MovieGraphPathTracerPass.cpp
@@ -74,7 +74,28 @@ void FMovieGraphPathTracerPass::ApplyMovieGraphOverridesToSceneView(TSharedRef<F
 		View->FinalPostProcessSettings.bOverride_PathTracingEnableReferenceDOF = true;
 		View->FinalPostProcessSettings.PathTracingEnableReferenceDOF = true;
 	}
+
+	// Update the post processing settings if the node has overridden any of them
+	if (const UMovieGraphPathTracerRenderPassNode* PathTracerNode = Cast<UMovieGraphPathTracerRenderPassNode>(GetParentNode(InInitData.TimeData.EvaluatedConfig)))
+	{
+#define OVERRIDE_COMPONENT(ComponentName) \
+		if (PathTracerNode->bOverride_bLightingComponents_Include##ComponentName) \
+		{ \
+			View->FinalPostProcessSettings.bOverride_PathTracingInclude##ComponentName = true; \
+			View->FinalPostProcessSettings.PathTracingInclude##ComponentName = PathTracerNode->bLightingComponents_Include##ComponentName; \
+		} \
 		
+		OVERRIDE_COMPONENT(Emissive);
+		OVERRIDE_COMPONENT(Diffuse);
+		OVERRIDE_COMPONENT(IndirectDiffuse);
+		OVERRIDE_COMPONENT(Specular);
+		OVERRIDE_COMPONENT(IndirectSpecular);
+		OVERRIDE_COMPONENT(Volume);
+		OVERRIDE_COMPONENT(IndirectVolume);
+		
+#undef OVERRIDE_COMPONENT
+	}
+	
 	// reset path tracer's accumulation at the start of each sample
 	View->bForcePathTracerReset = SampleIndex == 0;
 }
diff --git a/Engine/Plugins/MovieScene/MovieRenderPipeline/Source/MovieRenderPipelineRenderPasses/Public/Graph/Nodes/MovieGraphDeferredPassNode.h b/Engine/Plugins/MovieScene/MovieRenderPipeline/Source/MovieRenderPipelineRenderPasses/Public/Graph/Nodes/MovieGraphDeferredPassNode.h
--- a/Engine/Plugins/MovieScene/MovieRenderPipeline/Source/MovieRenderPipelineRenderPasses/Public/Graph/Nodes/MovieGraphDeferredPassNode.h
+++ b/Engine/Plugins/MovieScene/MovieRenderPipeline/Source/MovieRenderPipelineRenderPasses/Public/Graph/Nodes/MovieGraphDeferredPassNode.h
@@ -93,13 +93,13 @@ public:
 	* and can optionally be combined with OCIO profiles on the file output nodes to convert from Linear Values in Working Color Space
 	* (which is sRGB  (Rec. 709) by default, unless changed in the project settings).
 	*/
-	UPROPERTY(EditAnywhere, BlueprintReadWrite, Category = "Settings", meta = (EditCondition = "bOverride_bDisableToneCurve"))
+	UPROPERTY(EditAnywhere, BlueprintReadWrite, Category = "Post Processing", meta = (EditCondition = "bOverride_bDisableToneCurve"))
 	bool bDisableToneCurve;
 
 	/**
 	* Allow the output file OpenColorIO transform to be used on this render.
 	*/
-	UPROPERTY(EditAnywhere, BlueprintReadWrite, Category = "Settings", meta = (EditCondition = "bOverride_bAllowOCIO"))
+	UPROPERTY(EditAnywhere, BlueprintReadWrite, Category = "Post Processing", meta = (EditCondition = "bOverride_bAllowOCIO"))
 	bool bAllowOCIO;
 	
 	/** 
diff --git a/Engine/Plugins/MovieScene/MovieRenderPipeline/Source/MovieRenderPipelineRenderPasses/Public/Graph/Nodes/MovieGraphPathTracerPassNode.h b/Engine/Plugins/MovieScene/MovieRenderPipeline/Source/MovieRenderPipelineRenderPasses/Public/Graph/Nodes/MovieGraphPathTracerPassNode.h
--- a/Engine/Plugins/MovieScene/MovieRenderPipeline/Source/MovieRenderPipelineRenderPasses/Public/Graph/Nodes/MovieGraphPathTracerPassNode.h
+++ b/Engine/Plugins/MovieScene/MovieRenderPipeline/Source/MovieRenderPipelineRenderPasses/Public/Graph/Nodes/MovieGraphPathTracerPassNode.h
@@ -90,6 +90,27 @@ public:
 	UPROPERTY(EditAnywhere, BlueprintReadWrite, Category = Overrides, meta = (InlineEditConditionToggle))
 	uint8 bOverride_bAllowOCIO : 1;
 
+	UPROPERTY(EditAnywhere, BlueprintReadWrite, Category = Overrides, meta = (InlineEditConditionToggle))
+	uint8 bOverride_bLightingComponents_IncludeEmissive : 1;
+
+	UPROPERTY(EditAnywhere, BlueprintReadWrite, Category = Overrides, meta = (InlineEditConditionToggle))
+	uint8 bOverride_bLightingComponents_IncludeDiffuse : 1;
+	
+	UPROPERTY(EditAnywhere, BlueprintReadWrite, Category = Overrides, meta = (InlineEditConditionToggle))
+	uint8 bOverride_bLightingComponents_IncludeIndirectDiffuse : 1;
+
+	UPROPERTY(EditAnywhere, BlueprintReadWrite, Category = Overrides, meta = (InlineEditConditionToggle))
+	uint8 bOverride_bLightingComponents_IncludeSpecular : 1;
+
+	UPROPERTY(EditAnywhere, BlueprintReadWrite, Category = Overrides, meta = (InlineEditConditionToggle))
+	uint8 bOverride_bLightingComponents_IncludeIndirectSpecular : 1;
+
+	UPROPERTY(EditAnywhere, BlueprintReadWrite, Category = Overrides, meta = (InlineEditConditionToggle))
+	uint8 bOverride_bLightingComponents_IncludeVolume : 1;
+
+	UPROPERTY(EditAnywhere, BlueprintReadWrite, Category = Overrides, meta = (InlineEditConditionToggle))
+	uint8 bOverride_bLightingComponents_IncludeIndirectVolume : 1;
+
 	UPROPERTY(EditAnywhere, BlueprintReadWrite, Category = Overrides, meta = (InlineEditConditionToggle))
 	uint8 bOverride_bWriteAllSamples : 1;
 
@@ -145,16 +166,42 @@ public:
 	* and can optionally be combined with OCIO profiles on the file output nodes to convert from Linear Values in Working Color Space
 	* (which is sRGB  (Rec. 709) by default, unless changed in the project settings).
 	*/
-	UPROPERTY(EditAnywhere, BlueprintReadWrite, Category = "Settings", meta = (EditCondition = "bOverride_bDisableToneCurve"))
+	UPROPERTY(EditAnywhere, BlueprintReadWrite, Category = "Post Processing", meta = (EditCondition = "bOverride_bDisableToneCurve"))
 	bool bDisableToneCurve;
 
 	/**
 	* Allow the output file OpenColorIO transform to be used on this render.
 	*/
-	UPROPERTY(EditAnywhere, BlueprintReadWrite, Category = "Settings", meta = (EditCondition = "bOverride_bAllowOCIO"))
+	UPROPERTY(EditAnywhere, BlueprintReadWrite, Category = "Post Processing", meta = (EditCondition = "bOverride_bAllowOCIO"))
 	bool bAllowOCIO;
 
+	/** Whether the render should include directly visible emissive components. */
+	UPROPERTY(EditAnywhere, BlueprintReadWrite, Category = "Lighting Components", DisplayName = "Emissive", meta = (EditCondition = "bOverride_bLightingComponents_IncludeEmissive"))
+	bool bLightingComponents_IncludeEmissive = true;
+
+	/** Whether the render should include diffuse lighting contributions. */
+	UPROPERTY(EditAnywhere, BlueprintReadWrite, Category = "Lighting Components", DisplayName = "Diffuse", meta = (EditCondition = "bOverride_bLightingComponents_IncludeDiffuse"))
+	bool bLightingComponents_IncludeDiffuse = true;
+
+	/** Whether the render should include indirect diffuse lighting contributions. */
+	UPROPERTY(EditAnywhere, BlueprintReadWrite, Category = "Lighting Components", DisplayName = "Indirect Diffuse", meta = (EditCondition = "bOverride_bLightingComponents_IncludeIndirectDiffuse"))
+	bool bLightingComponents_IncludeIndirectDiffuse = true;
+
+	/** Whether the render should include specular lighting contributions. */
+	UPROPERTY(EditAnywhere, BlueprintReadWrite, Category = "Lighting Components", DisplayName = "Specular", meta = (EditCondition = "bOverride_bLightingComponents_IncludeSpecular"))
+	bool bLightingComponents_IncludeSpecular = true;
+
+	/** Whether the render should include indirect specular lighting contributions. */
+	UPROPERTY(EditAnywhere, BlueprintReadWrite, Category = "Lighting Components", DisplayName = "Indirect Specular", meta = (EditCondition = "bOverride_bLightingComponents_IncludeIndirectSpecular"))
+	bool bLightingComponents_IncludeIndirectSpecular = true;
+
+	/** Whether the render should include volume lighting contributions. */
+	UPROPERTY(EditAnywhere, BlueprintReadWrite, Category = "Lighting Components", DisplayName = "Volume", meta = (EditCondition = "bOverride_bLightingComponents_IncludeVolume"))
+	bool bLightingComponents_IncludeVolume = true;
 
+	/** Whether the render should include indirect volume lighting contributions. */
+	UPROPERTY(EditAnywhere, BlueprintReadWrite, Category = "Lighting Components", DisplayName = "Indirect Volume", meta = (EditCondition = "bOverride_bLightingComponents_IncludeIndirectVolume"))
+	bool bLightingComponents_IncludeIndirectVolume = true;
 
 	/**
 	* An array of additional post-processing materials to run after the frame is rendered. Using this feature may add a notable amount of render time.
