MoviePipeline: Fixed several actions in the Queue UI (such as converting to jobs, or assigning configurations) not supporting multi-select. Previously only toggling shots/jobs supported multi-select, but now converting job types, and assigning job/shot configurations should also respect multi-select.

If you have multiple items selected and you use the drop-down on a non-selected item, only the item for which the drop-down was clicked will be changed (as it was not part of the multi-select, so we're considering it not a multi-select action). If the item who's drop-down you clicked was part of the multi-select then the change will be propagated to same-type jobs/shots (ie: if you have both graph and non-graph jobs selected and you click the dropdown for a graph job and choose a new asset, the multi-selected graph jobs will use that new configuration but no changes will be made to non-graph jobs). When converting between graph and non-graph configurations, jobs that are already of the target type will be unchanged.

Specific scenarios tested:
1) Having a Config-based job selected, a child shot selected, and a second config-based job selected, and A) assigning a Job Configuration via the dropdown, B) clearing a Job Configuration, C) applying and clearing a shot-configuration D) converting to graph jobs
2) The above, but with a Graph-based job selected as part of the mix.
3) Having a Graph based job selected, a child shot selected, and a second graph-based job selected, and A) assigning a Graph Configuration via the dropdown, B) creating a new Graph (for both the job and the sub-shot), C) replacing with preset (converting back to non-graph job)
4) The above, but with a non-graph job selected as part of the mix

diff --git a/Engine/Plugins/MovieScene/MovieRenderPipeline/Source/MovieRenderPipelineEditor/Private/Widgets/SMoviePipelineQueueEditor.cpp b/Engine/Plugins/MovieScene/MovieRenderPipeline/Source/MovieRenderPipelineEditor/Private/Widgets/SMoviePipelineQueueEditor.cpp
--- a/Engine/Plugins/MovieScene/MovieRenderPipeline/Source/MovieRenderPipelineEditor/Private/Widgets/SMoviePipelineQueueEditor.cpp
+++ b/Engine/Plugins/MovieScene/MovieRenderPipeline/Source/MovieRenderPipelineEditor/Private/Widgets/SMoviePipelineQueueEditor.cpp
@@ -88,13 +88,52 @@ struct IMoviePipelineQueueTreeItem : TSharedFromThis<IMoviePipelineQueueTreeItem
 	virtual ~IMoviePipelineQueueTreeItem() {}
 
 	virtual TSharedPtr<FMoviePipelineQueueJobTreeItem> AsJob() { return nullptr; }
+	virtual TSharedPtr<FMoviePipelineShotItem> AsShot() { return nullptr; }
 	virtual UMoviePipelineExecutorJob* GetOwningJob() { return nullptr; }
 	virtual UMoviePipelineExecutorShot* GetOwningShot() { return nullptr; }
 	virtual void Delete(UMoviePipelineQueue* InOwningQueue) {}
 	virtual void ResetStatus() {}
 	virtual UMoviePipelineExecutorJob* Duplicate(UMoviePipelineQueue* InOwningQueue) { return nullptr; }
 
-	virtual TSharedRef<ITableRow> ConstructWidget(TWeakPtr<SMoviePipelineQueueEditor> InQueueWidget, const TSharedRef<STableViewBase>& OwnerTable) = 0;
+	virtual TSharedRef<ITableRow> ConstructWidget(TWeakPtr<SMoviePipelineQueueEditor> InQueueWidget, const TSharedRef<STableViewBase>& OwnerTable) = 0;	
+
+protected:
+	TSet<TSharedPtr<IMoviePipelineQueueTreeItem>> GetMultiSelectAffectedItems(TWeakPtr<SMoviePipelineQueueEditor> WeakQueueEditor, bool bIsShot) const
+	{
+		TSet<TSharedPtr<IMoviePipelineQueueTreeItem>> AffectedItems;
+
+		TSharedPtr<SMoviePipelineQueueEditor> QueueEditor = WeakQueueEditor.Pin();
+		if (!QueueEditor)
+		{
+			return AffectedItems;
+		}
+		
+		// We always want to apply to ourselves
+		TSharedPtr<IMoviePipelineQueueTreeItem> NonConstSelf = ConstCastSharedRef<IMoviePipelineQueueTreeItem>(AsShared());
+		AffectedItems.Add(NonConstSelf);
+
+		// If the item you clicked on (the one running this code) isn't part of the multi-selected items then we only apply the requested change
+		// to the current item itself.
+		if (!QueueEditor->GetSelectedItems().Contains(AsShared()))
+		{
+			return AffectedItems;
+		}
+
+		// Otherwise we add all the other items from the selection (of the appropriate type)
+		for (TSharedPtr<IMoviePipelineQueueTreeItem> Item : QueueEditor->GetSelectedItems())
+		{
+			if (Item->AsJob() && !bIsShot)
+			{
+				AffectedItems.Add(Item);
+			}
+			else if (Item->AsShot() && bIsShot)
+			{
+				AffectedItems.Add(Item);
+			}
+		}
+
+		return AffectedItems;
+	}
 };
 
 class SQueueJobListRow : public SMultiColumnTableRow<TSharedPtr<IMoviePipelineQueueTreeItem>>
@@ -192,26 +231,22 @@ public:
 		return ECheckBoxState::Unchecked;
 	}
 
-	void SetCheckState(const ECheckBoxState NewState) const
+	void SetCheckState(const ECheckBoxState NewState)
 	{
-		if (WeakQueueEditor.IsValid() && WeakQueueEditor.Pin()->GetSelectedItems().Contains(SharedThis(this)))
+		const bool bIsShot = false;
+		TSet<TSharedPtr<IMoviePipelineQueueTreeItem>> AffectedItems = GetMultiSelectAffectedItems(WeakQueueEditor, bIsShot);
+		
+		for (TSharedPtr<IMoviePipelineQueueTreeItem> Item : AffectedItems)
 		{
-			for (TSharedPtr<IMoviePipelineQueueTreeItem> Item : WeakQueueEditor.Pin()->GetSelectedItems())
+			TSharedPtr<FMoviePipelineQueueJobTreeItem> JobTreeItem = Item->AsJob();
+			if (JobTreeItem.IsValid())
 			{
-				TSharedPtr<FMoviePipelineQueueJobTreeItem> JobTreeItem = StaticCastSharedPtr<FMoviePipelineQueueJobTreeItem>(Item);
-				if (JobTreeItem.IsValid())
+				if (UMoviePipelineExecutorJob* Job = JobTreeItem->WeakJob.Get())
 				{
-					if (UMoviePipelineExecutorJob* Job = JobTreeItem->WeakJob.Get())
-					{
-						Job->SetIsEnabled(NewState == ECheckBoxState::Checked);
-					}
+					Job->SetIsEnabled(NewState == ECheckBoxState::Checked);
 				}
 			}
 		}
-		else if (UMoviePipelineExecutorJob* Job = WeakJob.Get())
-		{
-			Job->SetIsEnabled(NewState == ECheckBoxState::Checked);
-		}
 	}
 
 	FText GetJobName() const
@@ -266,82 +301,160 @@ public:
 	}
 
 	void OnPickPresetFromAsset(const FAssetData& AssetData)
-	{
+	{	
 		// Close the dropdown menu that showed them the assets to pick from.
 		FSlateApplication::Get().DismissAllMenus();
 
-		UMoviePipelineExecutorJob* Job = WeakJob.Get();
-		if (Job)
+		UMoviePipelineExecutorJob* CurrentJob = WeakJob.Get();
+		if (!CurrentJob)
 		{
-			FScopedTransaction Transaction(LOCTEXT("PickJobPresetAsset_Transaction", "Set Job Configuration Asset"));
-			Job->Modify();
+			return;
+		}
 
-			if (Job->IsUsingGraphConfiguration())
-			{
-				Job->SetGraphPreset(CastChecked<UMovieGraphConfig>(AssetData.GetAsset()));
-			}
-			else
+		FScopedTransaction Transaction(LOCTEXT("PickJobPresetAsset_Transaction", "Set Job Configuration Asset"));
+					
+		const bool bIsShot = false;
+		TSet<TSharedPtr<IMoviePipelineQueueTreeItem>> AffectedItems = GetMultiSelectAffectedItems(WeakQueueEditor, bIsShot);
+		
+		for (TSharedPtr<IMoviePipelineQueueTreeItem> Item : AffectedItems)
+		{
+			TSharedPtr<FMoviePipelineQueueJobTreeItem> JobTreeItem = Item->AsJob();
+			if (JobTreeItem.IsValid())
 			{
-				Job->SetPresetOrigin(CastChecked<UMoviePipelinePrimaryConfig>(AssetData.GetAsset()));
+				UMoviePipelineExecutorJob* Job = JobTreeItem->WeakJob.Get();
+				if (Job)
+				{
+					// Only apply changes to jobs of the same configuration type as our selected item.
+					const bool bSelectedJobIsGraphJob = CurrentJob->IsUsingGraphConfiguration();
+					const bool bCurrentJobIsGraphJob = Job->IsUsingGraphConfiguration();
+
+					if (bSelectedJobIsGraphJob && bCurrentJobIsGraphJob)
+					{
+						Job->Modify();
+						Job->SetGraphPreset(CastChecked<UMovieGraphConfig>(AssetData.GetAsset()));
+					}
+					else if (!bSelectedJobIsGraphJob && !bCurrentJobIsGraphJob)
+					{
+						Job->Modify();
+						Job->SetPresetOrigin(CastChecked<UMoviePipelinePrimaryConfig>(AssetData.GetAsset()));
+					}
+				}
 			}
+
 		}
 
 		OnChosePresetCallback.ExecuteIfBound(WeakJob, nullptr);
 	}
 
-	void OnPickNewPreset()
-	{
+	void OnClearNonGraphPreset()
+	{	
 		// Close the dropdown menu that showed them the assets to pick from.
 		FSlateApplication::Get().DismissAllMenus();
 
-		UMoviePipelineExecutorJob* Job = WeakJob.Get();
-		if (Job)
+		UMoviePipelineExecutorJob* CurrentJob = WeakJob.Get();
+		if (!CurrentJob)
 		{
-			if (Job->IsUsingGraphConfiguration())
-			{
-				Job->SetGraphPreset(GetMutableDefault<UMoviePipelineExecutorJob>()->GetGraphPreset());
-			}
-			else
+			return;
+		}
+
+		FScopedTransaction Transaction(LOCTEXT("ClearJobPreset_Transaction", "Clear Job Configuration Preset"));
+
+
+		const bool bIsShot = false;
+		TSet<TSharedPtr<IMoviePipelineQueueTreeItem>> AffectedItems = GetMultiSelectAffectedItems(WeakQueueEditor, bIsShot);
+		
+		for (TSharedPtr<IMoviePipelineQueueTreeItem> Item : AffectedItems)
+		{
+			TSharedPtr<FMoviePipelineQueueJobTreeItem> JobTreeItem = Item->AsJob();
+			if (JobTreeItem.IsValid())
 			{
-				// Copy from the CDO's version of the job to pick up the right name.
-				Job->SetConfiguration(GetMutableDefault<UMoviePipelineExecutorJob>()->GetConfiguration());
-				UMoviePipelineEditorBlueprintLibrary::EnsureJobHasDefaultSettings(Job);
+				UMoviePipelineExecutorJob* Job = JobTreeItem->WeakJob.Get();
+				if (Job)
+				{
+					// Only apply changes to jobs that aren't using graph configurations
+					if (!Job->IsUsingGraphConfiguration())
+					{
+						Job->SetConfiguration(GetMutableDefault<UMoviePipelineExecutorJob>()->GetConfiguration());
+						UMoviePipelineEditorBlueprintLibrary::EnsureJobHasDefaultSettings(Job);
+					}
+				}
 			}
+
 		}
 
-		OnChosePresetCallback.ExecuteIfBound(WeakJob, nullptr);
 	}
 
 	void OnReplaceWithRenderGraph()
-	{
-		UMoviePipelineExecutorJob* Job = WeakJob.Get();
-		if (!Job)
+	{	
+		const bool bIsShot = false;
+		TSet<TSharedPtr<IMoviePipelineQueueTreeItem>> AffectedItems = GetMultiSelectAffectedItems(WeakQueueEditor, bIsShot);
+		
+		for (TSharedPtr<IMoviePipelineQueueTreeItem> Item : AffectedItems)
 		{
-			return;
+			TSharedPtr<FMoviePipelineQueueJobTreeItem> JobTreeItem = Item->AsJob();
+			if (JobTreeItem.IsValid())
+			{
+				UMoviePipelineExecutorJob* Job = JobTreeItem->WeakJob.Get();
+				if (Job)
+				{
+					// Only replace jobs that don't already use the render-graph configuration
+					if(!Job->IsUsingGraphConfiguration())
+					{
+						SMoviePipelineQueueEditor::AssignDefaultGraphPresetToJob(Job);
+					}
+				}
+			}
 		}
-
-		SMoviePipelineQueueEditor::AssignDefaultGraphPresetToJob(Job);
 	}
 
-	void OnCreateNewGraphAndAssign() const
+	void OnCreateNewGraphAndAssign()
 	{
+		FScopedTransaction Transaction(LOCTEXT("CreateNewGraphAndAssign_Transaction", "Create Graph"));
+
+		const bool bIsShot = false;
+		TSet<TSharedPtr<IMoviePipelineQueueTreeItem>> AffectedItems = GetMultiSelectAffectedItems(WeakQueueEditor, bIsShot);
+		
+		// Create only one graph for all affected items.
 		if (const UMovieGraphConfig* NewGraph = UE::MovieGraph::Private::CreateNewSavedGraphAsset())
 		{
-			UMoviePipelineExecutorJob* Job = WeakJob.Get();
-			if (ensureMsgf(Job, TEXT("Could not assign new graph to job: Job is invalid.")))
+			for (TSharedPtr<IMoviePipelineQueueTreeItem> Item : AffectedItems)
 			{
-				Job->SetGraphPreset(NewGraph);
+				TSharedPtr<FMoviePipelineQueueJobTreeItem> JobTreeItem = Item->AsJob();
+				if (JobTreeItem.IsValid())
+				{
+					UMoviePipelineExecutorJob* Job = JobTreeItem->WeakJob.Get();
+					if (Job)
+					{
+						// We only apply to jobs already using graph configurations because they have to have been converted to
+						// a graph, and now they're replacing their configuration with a new one.
+						if (Job->IsUsingGraphConfiguration())
+						{
+							Job->SetGraphPreset(NewGraph);
+						}
+					}
+				}
 			}
 		}
 	}
 
-	void OnClearGraph() const
+	void OnClearGraph()
 	{
-		if (UMoviePipelineExecutorJob* Job = WeakJob.Get())
+		FScopedTransaction Transaction(LOCTEXT("ClearJobGraph_Transaction", "Replace Graph with Config"));
+
+		const bool bIsShot = false;
+		TSet<TSharedPtr<IMoviePipelineQueueTreeItem>> AffectedItems = GetMultiSelectAffectedItems(WeakQueueEditor, bIsShot);
+		
+		for (TSharedPtr<IMoviePipelineQueueTreeItem> Item : AffectedItems)
 		{
-			FScopedTransaction Transaction(LOCTEXT("ClearJobGraph_Transaction", "Replace Graph with Config"));
-			
-			Job->SetGraphPreset(nullptr);
+			TSharedPtr<FMoviePipelineQueueJobTreeItem> JobTreeItem = Item->AsJob();
+			if (JobTreeItem.IsValid())
+			{
+				UMoviePipelineExecutorJob* Job = JobTreeItem->WeakJob.Get();
+				if (Job)
+				{
+					Job->SetGraphPreset(nullptr);
+				}
+			}
 		}
 	}
 
@@ -451,7 +564,7 @@ public:
 			WeakJob,
 			bIsShot,
 			FOnAssetSelected::CreateRaw(this, &FMoviePipelineQueueJobTreeItem::OnPickPresetFromAsset),
-			FExecuteAction::CreateRaw(this, &FMoviePipelineQueueJobTreeItem::OnPickNewPreset),
+			FExecuteAction::CreateRaw(this, &FMoviePipelineQueueJobTreeItem::OnClearNonGraphPreset),
 			FExecuteAction::CreateRaw(this, &FMoviePipelineQueueJobTreeItem::OnReplaceWithRenderGraph),
 			FExecuteAction::CreateRaw(this, &FMoviePipelineQueueJobTreeItem::OnCreateNewGraphAndAssign),
 			FExecuteAction::CreateRaw(this, &FMoviePipelineQueueJobTreeItem::OnClearGraph)
@@ -459,7 +572,7 @@ public:
 	}
 
 	static TSharedRef<SWidget> OnGenerateConfigPresetPickerMenuFromClass(UClass* InClass,
-		TWeakObjectPtr<UMoviePipelineExecutorJob> TargetJob, const bool bIsShot, FOnAssetSelected InOnAssetSelected, FExecuteAction InNewConfig,
+		TWeakObjectPtr<UMoviePipelineExecutorJob> TargetJob, const bool bIsShot, FOnAssetSelected InOnAssetSelected, FExecuteAction InClearNonGraphConfig,
 		FExecuteAction InNewRenderGraph, FExecuteAction InCreateNewGraphAndAssign, FExecuteAction InClearGraph)
 	{
 		FMenuBuilder MenuBuilder(true, nullptr);
@@ -514,7 +627,7 @@ public:
 					LOCTEXT("ClearConfig_Label", "Clear Config"),
 					LOCTEXT("ClearConfig_Tooltip", "Resets the changes to the config and goes back to the defaults."),
 					FSlateIcon(),
-					FUIAction(InNewConfig),
+					FUIAction(InClearNonGraphConfig),
 					NAME_None,
 					EUserInterfaceActionType::Button
 				);
@@ -872,6 +985,11 @@ struct FMoviePipelineShotItem : IMoviePipelineQueueTreeItem
 		return WeakShot.Get();
 	}
 
+	virtual TSharedPtr<FMoviePipelineShotItem> AsShot() override
+	{
+		return SharedThis(this);
+	}
+
 	virtual TSharedRef<ITableRow> ConstructWidget(TWeakPtr<SMoviePipelineQueueEditor> InQueueWidget, const TSharedRef<STableViewBase>& OwnerTable) override
 	{
 		WeakQueueEditor = InQueueWidget;
@@ -892,24 +1010,20 @@ struct FMoviePipelineShotItem : IMoviePipelineQueueTreeItem
 
 	void SetCheckState(ECheckBoxState InNewState)
 	{
-		if (WeakQueueEditor.IsValid() && WeakQueueEditor.Pin()->GetSelectedItems().Contains(SharedThis(this)))
+		const bool bIsShot = true;
+		TSet<TSharedPtr<IMoviePipelineQueueTreeItem>> AffectedItems = GetMultiSelectAffectedItems(WeakQueueEditor, bIsShot);
+		
+		for (TSharedPtr<IMoviePipelineQueueTreeItem> Item : AffectedItems)
 		{
-			for (TSharedPtr<IMoviePipelineQueueTreeItem> Item : WeakQueueEditor.Pin()->GetSelectedItems())
+			TSharedPtr<FMoviePipelineShotItem> ShotTreeItem = Item->AsShot();
+			if (ShotTreeItem.IsValid())
 			{
-				TSharedPtr<FMoviePipelineShotItem> ShotItem = StaticCastSharedPtr<FMoviePipelineShotItem>(Item);
-				if (ShotItem.IsValid())
+				if (UMoviePipelineExecutorShot* Shot = ShotTreeItem->WeakShot.Get())
 				{
-					if (UMoviePipelineExecutorShot* Shot = ShotItem->WeakShot.Get())
-					{
-						Shot->bEnabled = InNewState == ECheckBoxState::Checked;
-					}
+					Shot->bEnabled = InNewState == ECheckBoxState::Checked;
 				}
 			}
 		}
-		else if (UMoviePipelineExecutorShot* Shot = WeakShot.Get())
-		{
-			Shot->bEnabled = InNewState == ECheckBoxState::Checked;
-		}
 	}
 
 	FText GetShotLabel() const
@@ -1014,82 +1128,155 @@ struct FMoviePipelineShotItem : IMoviePipelineQueueTreeItem
 	{
 		// Close the dropdown menu that showed them the assets to pick from.
 		FSlateApplication::Get().DismissAllMenus();
-
-		if (UMoviePipelineExecutorShot* Shot = WeakShot.Get())
+		
+		UMoviePipelineExecutorShot* SelectedShot = WeakShot.Get();
+		if(!SelectedShot)
 		{
-			FScopedTransaction Transaction(LOCTEXT("PickShotPresetAsset_Transaction", "Set Shot Configuration Asset"));
-			Shot->Modify();
+			return;
+		}
+								
+		FScopedTransaction Transaction(LOCTEXT("PickShotPresetAsset_Transaction", "Set Shot Configuration Asset"));
 
-			if (Shot->IsUsingGraphConfiguration())
-			{
-				Shot->SetGraphPreset(CastChecked<UMovieGraphConfig>(AssetData.GetAsset()));
-			}
-			else
+		const bool bIsShot = true;
+		TSet<TSharedPtr<IMoviePipelineQueueTreeItem>> AffectedItems = GetMultiSelectAffectedItems(WeakQueueEditor, bIsShot);
+		
+		for (TSharedPtr<IMoviePipelineQueueTreeItem> Item : AffectedItems)
+		{
+			TSharedPtr<FMoviePipelineShotItem> ShotTreeItem = Item->AsShot();
+			if (ShotTreeItem.IsValid())
 			{
-				Shot->SetShotOverridePresetOrigin(CastChecked<UMoviePipelineShotConfig>(AssetData.GetAsset()));
+				UMoviePipelineExecutorShot* Shot = ShotTreeItem->WeakShot.Get();
+				if(Shot)
+				{				
+					const bool bSelectedShotIsGraph = SelectedShot->IsUsingGraphConfiguration();
+					const bool bCurrentShotIsGraph = Shot->IsUsingGraphConfiguration();
+						
+					// Only apply changes to jobs of the same configuration type as our selected item.	
+					if (bSelectedShotIsGraph && bCurrentShotIsGraph)
+					{
+						Shot->Modify();
+						Shot->SetGraphPreset(CastChecked<UMovieGraphConfig>(AssetData.GetAsset()));
+					}
+					else if (!bSelectedShotIsGraph && !bCurrentShotIsGraph)
+					{
+						Shot->Modify();
+						Shot->SetShotOverridePresetOrigin(CastChecked<UMoviePipelineShotConfig>(AssetData.GetAsset()));
+					}
+				}
 			}
 		}
 
+
 		OnChosePresetCallback.ExecuteIfBound(WeakJob, WeakShot);
 	}
 
-	void OnPickNewShotPreset()
+	void OnClearNonGraphPreset()
 	{
 		// Close the dropdown menu that showed them the assets to pick from.
 		FSlateApplication::Get().DismissAllMenus();
-
-		UMoviePipelineExecutorShot* Shot = WeakShot.Get();
-		if (Shot)
+		
+		UMoviePipelineExecutorShot* SelectedShot = WeakShot.Get();
+		if(!SelectedShot)
 		{
-			FScopedTransaction Transaction(LOCTEXT("PickShotClearConfig_Transaction", "Clear Configuration Asset"));
-			Shot->Modify();
-			Shot->SetShotOverrideConfiguration(nullptr);
+			return;
 		}
+		
+		FScopedTransaction Transaction(LOCTEXT("PickShotClearConfig_Transaction", "Clear Configuration Asset"));
+		
+		const bool bIsShot = true;
+		TSet<TSharedPtr<IMoviePipelineQueueTreeItem>> AffectedItems = GetMultiSelectAffectedItems(WeakQueueEditor, bIsShot);
+		
+		for (TSharedPtr<IMoviePipelineQueueTreeItem> Item : AffectedItems)
+		{
+			TSharedPtr<FMoviePipelineShotItem> ShotTreeItem = Item->AsShot();
+			if (ShotTreeItem.IsValid())
+			{
+				UMoviePipelineExecutorShot* Shot = ShotTreeItem->WeakShot.Get();
+				if(Shot)
+				{
+					// Only clear jobs that aren't using the graph configuration
+					if(!Shot->IsUsingGraphConfiguration())
+					{
+						Shot->Modify();
+						Shot->SetShotOverrideConfiguration(nullptr);
+					}
+				}
+			}
+		}
+					
 
 		OnChosePresetCallback.ExecuteIfBound(WeakJob, WeakShot);
 	}
 
-	void OnReplaceWithRenderGraph()
+
+
+	void OnCreateNewGraphAndAssign()
 	{
-		UMoviePipelineExecutorShot* Shot = WeakShot.Get();
-		if (!Shot)
+		UMoviePipelineExecutorShot* SelectedShot = WeakShot.Get();
+		if (!SelectedShot)
 		{
 			return;
 		}
 
-		// Note: Setting the graph preset will transition the job to use a graph-based configuration
-		// Use the default graph specified in Project Settings.
-		const UMovieRenderPipelineProjectSettings* ProjectSettings = GetDefault<UMovieRenderPipelineProjectSettings>();
-		const TSoftObjectPtr<UMovieGraphConfig> ProjectDefaultGraph = ProjectSettings->DefaultGraph;
-		if (const UMovieGraphConfig* DefaultGraph = ProjectDefaultGraph.LoadSynchronous())
-		{
-			FScopedTransaction Transaction(LOCTEXT("PickReplaceWithGraphAsset_Transaction", "Convert Job to Graph"));
-			Shot->Modify();
-			Shot->SetGraphPreset(DefaultGraph);
-		}
-	}
-
-	void OnCreateNewGraphAndAssign() const
-	{
 		if (const UMovieGraphConfig* NewGraph = UE::MovieGraph::Private::CreateNewSavedGraphAsset())
 		{
-			UMoviePipelineExecutorShot* Shot = WeakShot.Get();
-			if (ensureMsgf(Shot, TEXT("Could not assign new graph to shot: Shot is invalid.")))
+			FScopedTransaction Transaction(LOCTEXT("PickCreateNewGraphAsset_Transaction", "Assign Graph Configuration Asset"));
+
+			const bool bIsShot = true;
+			TSet<TSharedPtr<IMoviePipelineQueueTreeItem>> AffectedItems = GetMultiSelectAffectedItems(WeakQueueEditor, bIsShot);
+		
+			for (TSharedPtr<IMoviePipelineQueueTreeItem> Item : AffectedItems)
 			{
-				FScopedTransaction Transaction(LOCTEXT("PickCreateNewGraphAsset_Transaction", "Assign Graph Configuration Asset"));
-				Shot->Modify();
-				Shot->SetGraphPreset(NewGraph);
+				TSharedPtr<FMoviePipelineShotItem> ShotTreeItem = Item->AsShot();
+				if (ShotTreeItem.IsValid())
+				{
+					UMoviePipelineExecutorShot* Shot = ShotTreeItem->WeakShot.Get();
+					if (Shot)
+					{
+						// Only assign jobs that are using the graph configuration
+						if (Shot->IsUsingGraphConfiguration())
+						{
+							Shot->Modify();
+							Shot->SetGraphPreset(NewGraph);
+						}
+					}
+				}
 			}
 		}
-	}
 
-	void OnClearGraph() const
+		
+	}		
+		
+
+	void OnClearGraph()
 	{
-		if (UMoviePipelineExecutorShot* Shot = WeakShot.Get())
+		UMoviePipelineExecutorShot* SelectedShot = WeakShot.Get();
+		if (!SelectedShot)
 		{
-			FScopedTransaction Transaction(LOCTEXT("ClearShotGraph_Transaction", "Replace Graph with Config"));
-			
-			Shot->SetGraphPreset(nullptr);
+			return;
+		}
+
+		FScopedTransaction Transaction(LOCTEXT("ClearShotGraph_Transaction", "Clear Graph Config"));
+
+		const bool bIsShot = true;
+		TSet<TSharedPtr<IMoviePipelineQueueTreeItem>> AffectedItems = GetMultiSelectAffectedItems(WeakQueueEditor, bIsShot);
+		
+		for (TSharedPtr<IMoviePipelineQueueTreeItem> Item : AffectedItems)
+		{
+			TSharedPtr<FMoviePipelineShotItem> ShotTreeItem = Item->AsShot();
+			if (ShotTreeItem.IsValid())
+			{
+				UMoviePipelineExecutorShot* Shot = ShotTreeItem->WeakShot.Get();
+				if (Shot)
+				{
+					// Only clear jobs that are using the graph configuration
+					if (Shot->IsUsingGraphConfiguration())
+					{
+						Shot->Modify();
+						Shot->SetGraphPreset(nullptr);
+					}
+				}
+			}
 		}
 	}
 
@@ -1128,7 +1315,7 @@ struct FMoviePipelineShotItem : IMoviePipelineQueueTreeItem
 			WeakJob,
 			bIsShot,
 			FOnAssetSelected::CreateRaw(this, &FMoviePipelineShotItem::OnPickShotPresetFromAsset),
-			FExecuteAction::CreateRaw(this, &FMoviePipelineShotItem::OnPickNewShotPreset),
+			FExecuteAction::CreateRaw(this, &FMoviePipelineShotItem::OnClearNonGraphPreset),
 			ReplaceWithGraph,
 			FExecuteAction::CreateRaw(this, &FMoviePipelineShotItem::OnCreateNewGraphAndAssign),
 			FExecuteAction::CreateRaw(this, &FMoviePipelineShotItem::OnClearGraph)
