MoviePipeline: Added the "Is Spawnable" collection query to the graph.

diff --git a/Engine/Plugins/MovieScene/MovieRenderPipeline/Source/MovieRenderPipelineCore/Private/Graph/MovieGraphRenderLayerSubsystem.cpp b/Engine/Plugins/MovieScene/MovieRenderPipeline/Source/MovieRenderPipelineCore/Private/Graph/MovieGraphRenderLayerSubsystem.cpp
--- a/Engine/Plugins/MovieScene/MovieRenderPipeline/Source/MovieRenderPipelineCore/Private/Graph/MovieGraphRenderLayerSubsystem.cpp
+++ b/Engine/Plugins/MovieScene/MovieRenderPipeline/Source/MovieRenderPipelineCore/Private/Graph/MovieGraphRenderLayerSubsystem.cpp
@@ -12,6 +12,7 @@
 #include "Materials/MaterialInterface.h"
 #include "Modules/ModuleManager.h"
 #include "MovieRenderPipelineCoreModule.h"
+#include "MovieSceneSpawnableAnnotation.h"
 #include "Styling/AppStyle.h"
 #include "Styling/SlateIconFinder.h"
 #include "UObject/Package.h"
@@ -2298,6 +2299,30 @@ void UMovieGraphConditionGroupQuery_DataLayer::AddDataLayers(const TArray<const
 }
 #endif	// WITH_EDITOR
 
+void UMovieGraphConditionGroupQuery_IsSpawnable::Evaluate(const TArray<AActor*>& InActorsToQuery, const UWorld* InWorld, TSet<AActor*>& OutMatchingActors) const
+{
+	for (AActor* InActor : InActorsToQuery)
+	{
+		TOptional<FMovieSceneSpawnableAnnotation> Spawnable = FMovieSceneSpawnableAnnotation::Find(InActor);
+		if (Spawnable.IsSet())
+		{
+			OutMatchingActors.Add(InActor);
+		}
+	}
+}
+
+const FSlateIcon& UMovieGraphConditionGroupQuery_IsSpawnable::GetIcon() const
+{
+	static const FSlateIcon SpawnableIcon = FSlateIcon(FAppStyle::GetAppStyleSetName(), "GraphEditor.SpawnActor_16x");
+	return SpawnableIcon;
+}
+
+const FText& UMovieGraphConditionGroupQuery_IsSpawnable::GetDisplayName() const
+{
+	static const FText DisplayName = LOCTEXT("ConditionGroupQueryDisplayName_IsSpawnable", "Is Spawnable");
+	return DisplayName;
+}
+
 UMovieGraphConditionGroup::UMovieGraphConditionGroup()
 	: OpType(EMovieGraphConditionGroupOpType::Add)
 {
diff --git a/Engine/Plugins/MovieScene/MovieRenderPipeline/Source/MovieRenderPipelineCore/Public/Graph/MovieGraphRenderLayerSubsystem.h b/Engine/Plugins/MovieScene/MovieRenderPipeline/Source/MovieRenderPipelineCore/Public/Graph/MovieGraphRenderLayerSubsystem.h
--- a/Engine/Plugins/MovieScene/MovieRenderPipeline/Source/MovieRenderPipelineCore/Public/Graph/MovieGraphRenderLayerSubsystem.h
+++ b/Engine/Plugins/MovieScene/MovieRenderPipeline/Source/MovieRenderPipelineCore/Public/Graph/MovieGraphRenderLayerSubsystem.h
@@ -552,6 +552,23 @@ public:
 	TArray<TSoftObjectPtr<UDataLayerAsset>> DataLayers;
 };
 
+/** Query type which filters actors by their spawnable status. */
+UCLASS(BlueprintType)
+class MOVIERENDERPIPELINECORE_API UMovieGraphConditionGroupQuery_IsSpawnable final : public UMovieGraphConditionGroupQueryBase
+{
+	GENERATED_BODY()
+
+public:
+	virtual void Evaluate(const TArray<AActor*>& InActorsToQuery, const UWorld* InWorld, TSet<AActor*>& OutMatchingActors) const override;
+	virtual const FSlateIcon& GetIcon() const override;
+	virtual const FText& GetDisplayName() const override;
+
+public:
+	/** Whether the actor is a spawnable or not. */
+	UPROPERTY(EditAnywhere, Category="General")
+	bool bIsSpawnable;
+};
+
 /** A group of queries which can be added to a collection. */
 UCLASS(BlueprintType)
 class MOVIERENDERPIPELINECORE_API UMovieGraphConditionGroup : public UObject
