MoviePipeline: Fixed high resolution tiling not working with off-center projections. Tested by manually adjusting the off-center projection to [0, 1] and then comparing the viewport result to the MRQ result in both high resolution and non-highresolution tiling scenarios, and checked that this change doesn't affect Squeeze Factor (which uses off-center projection) and doesn't affect results when a off-center projection isn't being used.

Requires CL 35380838 as well which fixes CameraStackTypes.cpp to work properly with generating off-center projections for orthographic cameras.

diff --git a/Engine/Plugins/MovieScene/MovieRenderPipeline/Source/MovieRenderPipelineRenderPasses/Private/Graph/Renderers/MovieGraphImagePassBase.cpp b/Engine/Plugins/MovieScene/MovieRenderPipeline/Source/MovieRenderPipelineRenderPasses/Private/Graph/Renderers/MovieGraphImagePassBase.cpp
--- a/Engine/Plugins/MovieScene/MovieRenderPipeline/Source/MovieRenderPipelineRenderPasses/Private/Graph/Renderers/MovieGraphImagePassBase.cpp
+++ b/Engine/Plugins/MovieScene/MovieRenderPipeline/Source/MovieRenderPipelineRenderPasses/Private/Graph/Renderers/MovieGraphImagePassBase.cpp
@@ -321,11 +321,28 @@ void FMovieGraphImagePassBase::CalculateProjectionMatrix(UE::MovieGraph::Default
 		{
 			AspectRatioAxisConstraint = EAspectRatioAxisConstraint::AspectRatio_MaintainYFOV;
 			InOutCameraInfo.ViewInfo.OrthoWidth *= (DestAspectRatio / CameraAspectRatio);
+
+			// Off-center camera projections are calculated based on constrained aspect ratios, but those are disabled
+			// when using high-resolution tiling. This means that we need to scale the offset projection as well.
+			// 
+			// To calculate the required size change, we can look at an Aspect Ratio of 0.5 inside a square output, 
+			// ie: the rendered area is 1000 x 2000 for an output that is 2000x2000 (this is 0.5 of 1.0). With an
+			// off-center projection, an offset of 1.0 on X originally only moved by 500 pixels (1000x0.5), but with the aspect
+			// ratio constraint disabled, it now applies to the full output image (2000x0.5) resulting in a move that is twice as big.
+			// 
+			// To resolve this, we scale the offset by the CameraAspectRatio / DestAspectRatio, which is 0.5 / 1.0 for this example,
+			// meaning we multiply the user-intended offset (1.0) by 0.5, resulting in the originally desired 500px offset.
+			const double Ratio = CameraAspectRatio / DestAspectRatio; // ex: Ratio = 0.5 / 1
+			InOutCameraInfo.ViewInfo.OffCenterProjectionOffset.X *= Ratio;
 		}
 		else if (CameraAspectRatio > DestAspectRatio)
 		{
 			// Don't rescale the width and keep it X-constrained.
 			AspectRatioAxisConstraint = EAspectRatioAxisConstraint::AspectRatio_MaintainXFOV;
+
+			// Like above, off-center projections need to be rescaled too.
+			const double Ratio = DestAspectRatio / CameraAspectRatio;
+			InOutCameraInfo.ViewInfo.OffCenterProjectionOffset.Y *= Ratio;
 		}
 		InOutCameraInfo.ViewInfo.bConstrainAspectRatio = false;
 	}
@@ -390,11 +407,18 @@ void FMovieGraphImagePassBase::ModifyProjectionMatrixForTiling(const UE::MovieGr
 
 	if (bInOrthographic)
 	{
+		// Scale the off-center projection matrix too so that it's appropriately sized down for each tile.
+		InOutProjectionMatrix.M[3][0] /= ScaleX;
+		InOutProjectionMatrix.M[3][1] /= ScaleY;
 		InOutProjectionMatrix.M[3][0] += OffsetX / PadRatioX;
 		InOutProjectionMatrix.M[3][1] += OffsetY / PadRatioY;
 	}
 	else
 	{
+		// Scale the off-center projection matrix too so that it's appropriately sized down for each tile.
+		InOutProjectionMatrix.M[2][0] /= ScaleX;
+		InOutProjectionMatrix.M[2][1] /= ScaleY;
+		// Then offset it for this particular tile.
 		InOutProjectionMatrix.M[2][0] += OffsetX / PadRatioX;
 		InOutProjectionMatrix.M[2][1] += OffsetY / PadRatioY;
 	}
diff --git a/Engine/Plugins/MovieScene/MovieRenderPipeline/Source/MovieRenderPipelineRenderPasses/Private/MoviePipelineImagePassBase.cpp b/Engine/Plugins/MovieScene/MovieRenderPipeline/Source/MovieRenderPipelineRenderPasses/Private/MoviePipelineImagePassBase.cpp
--- a/Engine/Plugins/MovieScene/MovieRenderPipeline/Source/MovieRenderPipelineRenderPasses/Private/MoviePipelineImagePassBase.cpp
+++ b/Engine/Plugins/MovieScene/MovieRenderPipeline/Source/MovieRenderPipelineRenderPasses/Private/MoviePipelineImagePassBase.cpp
@@ -502,11 +502,28 @@ FSceneView* UMoviePipelineImagePassBase::GetSceneViewForSampleState(FSceneViewFa
 			{
 				AspectRatioAxisConstraint = EAspectRatioAxisConstraint::AspectRatio_MaintainYFOV;
 				CameraInfo.ViewInfo.OrthoWidth *= (DestAspectRatio / CameraAspectRatio);
+
+				// Off-center camera projections are calculated based on constrained aspect ratios, but those are disabled
+				// when using high-resolution tiling. This means that we need to scale the offset projection as well.
+				// 
+				// To calculate the required size change, we can look at an Aspect Ratio of 0.5 inside a square output, 
+				// ie: the rendered area is 1000 x 2000 for an output that is 2000x2000 (this is 0.5 of 1.0). With an
+				// off-center projection, an offset of 1.0 on X originally only moved by 500 pixels (1000x0.5), but with the aspect
+				// ratio constraint disabled, it now applies to the full output image (2000x0.5) resulting in a move that is twice as big.
+				// 
+				// To resolve this, we scale the offset by the CameraAspectRatio / DestAspectRatio, which is 0.5 / 1.0 for this example,
+				// meaning we multiply the user-intended offset (1.0) by 0.5, resulting in the originally desired 500px offset.
+				const double Ratio = CameraAspectRatio / DestAspectRatio; // ex: Ratio = 0.5 / 1
+				CameraInfo.ViewInfo.OffCenterProjectionOffset.X *= Ratio;
 			}
 			else if (CameraAspectRatio > DestAspectRatio)
 			{
 				// Don't rescale the width and keep it X-constrained.
 				AspectRatioAxisConstraint = EAspectRatioAxisConstraint::AspectRatio_MaintainXFOV;
+
+				// Like above, off-center projections need to be rescaled too.
+				const double Ratio = DestAspectRatio / CameraAspectRatio;
+				CameraInfo.ViewInfo.OffCenterProjectionOffset.Y *= Ratio;
 			}
 			CameraInfo.ViewInfo.bConstrainAspectRatio = false;
 		}
@@ -641,12 +658,19 @@ void UMoviePipelineImagePassBase::ModifyProjectionMatrixForTiling(const FMoviePi
 
 	if (bInOrthographic)
 	{
+		// Scale the off-center projection matrix too so that it's appropriately sized down for each tile.
+		InOutProjectionMatrix.M[3][0] /= ScaleX;
+		InOutProjectionMatrix.M[3][1] /= ScaleY;
 		InOutProjectionMatrix.M[3][0] += OffsetX / PadRatioX;
 		InOutProjectionMatrix.M[3][1] += OffsetY / PadRatioY;
 	}
 	else
 	{
-	InOutProjectionMatrix.M[2][0] += OffsetX / PadRatioX;
+		// Scale the off-center projection matrix too so that it's appropriately sized down for each tile.
+		InOutProjectionMatrix.M[2][0] /= ScaleX;
+		InOutProjectionMatrix.M[2][1] /= ScaleY;
+		// Then offset it for this particular tile.
+		InOutProjectionMatrix.M[2][0] += OffsetX / PadRatioX;
 		InOutProjectionMatrix.M[2][1] += OffsetY / PadRatioY;
 	}
 }
