MoviePipeline: Updated the graph and legacy MRQ to use a TSoftClassPtr rather than TSubclassOf for the Game Mode Override property. Some game modes can take a long time to load, so switching to a TSoftClassPtr will prevent the loading of it until render time.

diff --git a/Engine/Plugins/MovieScene/MovieRenderPipeline/Source/MovieRenderPipelineCore/Private/Graph/Nodes/MovieGraphGlobalGameOverrides.cpp b/Engine/Plugins/MovieScene/MovieRenderPipeline/Source/MovieRenderPipelineCore/Private/Graph/Nodes/MovieGraphGlobalGameOverrides.cpp
--- a/Engine/Plugins/MovieScene/MovieRenderPipeline/Source/MovieRenderPipelineCore/Private/Graph/Nodes/MovieGraphGlobalGameOverrides.cpp
+++ b/Engine/Plugins/MovieScene/MovieRenderPipeline/Source/MovieRenderPipelineCore/Private/Graph/Nodes/MovieGraphGlobalGameOverrides.cpp
@@ -7,7 +7,10 @@
 #include "Styling/AppStyle.h"
 
 UMovieGraphGlobalGameOverridesNode::UMovieGraphGlobalGameOverridesNode()
-	: GameModeOverride(AMoviePipelineGameMode::StaticClass())
+PRAGMA_DISABLE_DEPRECATION_WARNINGS
+	: GameModeOverride(nullptr)
+PRAGMA_ENABLE_DEPRECATION_WARNINGS
+	, SoftGameModeOverride(AMoviePipelineGameMode::StaticClass())
 	, ScalabilityQualityLevel(EMovieGraphScalabilityQualityLevel::Cinematic)
 	, bDisableTextureStreaming(false)
 	, bDisableLODs(false)
@@ -87,6 +90,22 @@ void UMovieGraphGlobalGameOverridesNode::BuildNewProcessCommandLineArgsImpl(TArr
 	InOutDeviceProfileCvars.Add(FString::Printf(TEXT("p.ChaosCloth.UseTimeStepSmoothing=%d"), 0));
 }
 
+void UMovieGraphGlobalGameOverridesNode::PostLoad()
+{
+	Super::PostLoad();
+
+	PRAGMA_DISABLE_DEPRECATION_WARNINGS
+	if (const UClass* GameModeOverrideClass = GameModeOverride.Get())
+	{
+		SoftGameModeOverride = GameModeOverrideClass;
+		bOverride_SoftGameModeOverride = bOverride_GameModeOverride;
+		
+		GameModeOverride = nullptr;
+		bOverride_GameModeOverride = false;
+	}
+	PRAGMA_ENABLE_DEPRECATION_WARNINGS
+}
+
 void UMovieGraphGlobalGameOverridesNode::ApplySettings(const bool bOverrideValues, UWorld* InWorld)
 {
 	// Apply the scalability settings
@@ -203,7 +222,7 @@ TSubclassOf<AGameModeBase> UMovieGraphGlobalGameOverridesNode::GetGameModeOverri
 			UMovieGraphGlobalGameOverridesNode* GameOverridesNode =
 				EvaluatedGraph->GetSettingForBranch<UMovieGraphGlobalGameOverridesNode>(GlobalsPinName, bIncludeCDOs, bExactMatch);
 
-			return GameOverridesNode ? GameOverridesNode->GameModeOverride : nullptr;
+			return GameOverridesNode ? GameOverridesNode->SoftGameModeOverride.LoadSynchronous() : nullptr;
 		}
 	}
 	else
@@ -218,7 +237,7 @@ TSubclassOf<AGameModeBase> UMovieGraphGlobalGameOverridesNode::GetGameModeOverri
 		{
 			if (UMoviePipelineSetting* Setting = *GameOverridesPtr)
 			{
-				return CastChecked<UMoviePipelineGameOverrideSetting>(Setting)->GameModeOverride;
+				return CastChecked<UMoviePipelineGameOverrideSetting>(Setting)->SoftGameModeOverride.LoadSynchronous();
 			}
 		}
 	}
diff --git a/Engine/Plugins/MovieScene/MovieRenderPipeline/Source/MovieRenderPipelineCore/Private/MoviePipelineGameOverrideSetting.cpp b/Engine/Plugins/MovieScene/MovieRenderPipeline/Source/MovieRenderPipelineCore/Private/MoviePipelineGameOverrideSetting.cpp
--- a/Engine/Plugins/MovieScene/MovieRenderPipeline/Source/MovieRenderPipelineCore/Private/MoviePipelineGameOverrideSetting.cpp
+++ b/Engine/Plugins/MovieScene/MovieRenderPipeline/Source/MovieRenderPipelineCore/Private/MoviePipelineGameOverrideSetting.cpp
@@ -24,6 +24,19 @@ void UMoviePipelineGameOverrideSetting::TeardownForPipelineImpl(UMoviePipeline*
 }
 
 
+void UMoviePipelineGameOverrideSetting::PostLoad()
+{
+	Super::PostLoad();
+
+PRAGMA_DISABLE_DEPRECATION_WARNINGS
+	if (const UClass* GameModeOverrideClass = GameModeOverride.Get())
+	{
+		SoftGameModeOverride = GameModeOverrideClass;
+		GameModeOverride = nullptr;
+	}
+PRAGMA_ENABLE_DEPRECATION_WARNINGS
+}
+
 void UMoviePipelineGameOverrideSetting::ApplyCVarSettings(const bool bOverrideValues)
 {
 	if (bCinematicQualitySettings)
diff --git a/Engine/Plugins/MovieScene/MovieRenderPipeline/Source/MovieRenderPipelineCore/Public/Graph/Nodes/MovieGraphGlobalGameOverrides.h b/Engine/Plugins/MovieScene/MovieRenderPipeline/Source/MovieRenderPipelineCore/Public/Graph/Nodes/MovieGraphGlobalGameOverrides.h
--- a/Engine/Plugins/MovieScene/MovieRenderPipeline/Source/MovieRenderPipelineCore/Public/Graph/Nodes/MovieGraphGlobalGameOverrides.h
+++ b/Engine/Plugins/MovieScene/MovieRenderPipeline/Source/MovieRenderPipelineCore/Public/Graph/Nodes/MovieGraphGlobalGameOverrides.h
@@ -35,6 +35,10 @@ public:
 	virtual void BuildNewProcessCommandLineArgsImpl(TArray<FString>& InOutUnrealURLParams, TArray<FString>& InOutCommandLineArgs, TArray<FString>& InOutDeviceProfileCvars, TArray<FString>& InOutExecCmds) const override;
 	// ~UMovieGraphSettingNode interface
 
+	// UObject interface
+	virtual void PostLoad() override;
+	// ~UObject interface
+
 	/**
 	 * Applies any cvars, scalability settings, etc. to reflect the properties set on the node. Remembers what the
 	 * settings are before they are set, so bOverrideValues can be set to false to revert them to their original
@@ -56,9 +60,13 @@ public:
 #endif
 
 public:
-	UPROPERTY(EditAnywhere, BlueprintReadWrite, Category = Overrides, meta = (InlineEditConditionToggle))
+	UE_DEPRECATED(5.5, "Please use the bOverride_SoftGameModeOverride property instead.")
+	UPROPERTY(EditAnywhere, BlueprintReadWrite, Category = Overrides, meta = (InlineEditConditionToggle, DeprecatedProperty, DeprecationMessage = "Please use the bOverride_SoftGameModeOverride property instead."))
 	uint8 bOverride_GameModeOverride : 1;
 
+	UPROPERTY(EditAnywhere, BlueprintReadWrite, Category = Overrides, meta = (InlineEditConditionToggle))
+	uint8 bOverride_SoftGameModeOverride : 1;
+
 	UPROPERTY(EditAnywhere, BlueprintReadWrite, Category = Overrides, meta = (InlineEditConditionToggle))
 	uint8 bOverride_ScalabilityQualityLevel : 1;
 
@@ -93,9 +101,17 @@ public:
 	 * Optional game mode to override the map's default game mode with. This can be useful if the game's normal mode
 	 * displays UI elements or loading screens that you don't want captured.
 	 */
-	UPROPERTY(EditAnywhere, BlueprintReadWrite, Category = "Game", meta = (EditCondition = "bOverride_GameModeOverride"))
+	UE_DEPRECATED(5.5, "Please use the SoftGameModeOverride property instead.")
+	UPROPERTY(BlueprintReadWrite, Category = "Game", meta = (EditCondition = "bOverride_GameModeOverride", DeprecatedProperty, DeprecationMessage = "Please use the SoftGameModeOverride property instead."))
 	TSubclassOf<AGameModeBase> GameModeOverride;
 
+	/**
+	 * Optional game mode to override the map's default game mode with. This can be useful if the game's normal mode
+	 * displays UI elements or loading screens that you don't want captured.
+	 */
+	UPROPERTY(EditAnywhere, BlueprintReadWrite, Category = "Game", DisplayName = "Game Mode Override", meta = (EditCondition = "bOverride_SoftGameModeOverride"))
+	TSoftClassPtr<AGameModeBase> SoftGameModeOverride;
+
 	/**
 	 * The scalability quality level that should be used in renders. See the Scalability Reference documentation for
 	 * information on how to edit cvars to add/change default quality values.
diff --git a/Engine/Plugins/MovieScene/MovieRenderPipeline/Source/MovieRenderPipelineCore/Public/MoviePipelineGameOverrideSetting.h b/Engine/Plugins/MovieScene/MovieRenderPipeline/Source/MovieRenderPipelineCore/Public/MoviePipelineGameOverrideSetting.h
--- a/Engine/Plugins/MovieScene/MovieRenderPipeline/Source/MovieRenderPipelineCore/Public/MoviePipelineGameOverrideSetting.h
+++ b/Engine/Plugins/MovieScene/MovieRenderPipeline/Source/MovieRenderPipelineCore/Public/MoviePipelineGameOverrideSetting.h
@@ -25,7 +25,10 @@ class MOVIERENDERPIPELINECORE_API UMoviePipelineGameOverrideSetting : public UMo
 	GENERATED_BODY()
 public:
 	UMoviePipelineGameOverrideSetting()
-		: GameModeOverride(AMoviePipelineGameMode::StaticClass())
+	PRAGMA_DISABLE_DEPRECATION_WARNINGS
+		: GameModeOverride(nullptr)
+	PRAGMA_ENABLE_DEPRECATION_WARNINGS
+		, SoftGameModeOverride(AMoviePipelineGameMode::StaticClass())
 		, bCinematicQualitySettings(true)
 		, TextureStreaming(EMoviePipelineTextureStreamingMethod::Disabled)
 		, bUseLODZero(true)
@@ -55,14 +58,23 @@ public:
 	// Used to ensure we set the default values even if the user forgets to add it.
 	virtual bool IgnoreTransientFilters() const override { return true; }
 
+	//~ Begin UObject interface
+	virtual void PostLoad() override;
+	//~ End UObject interface
+
 protected:
 	void ApplyCVarSettings(const bool bOverrideValues);
 
 public:
 	/** Optional Game Mode to override the map's default game mode with. This can be useful if the game's normal mode displays UI elements or loading screens that you don't want captured. */
-	UPROPERTY(EditAnywhere, BlueprintReadWrite, Category = "Game")
+	UE_DEPRECATED(5.5, "Please use the SoftGameModeOverride property instead.")
+	UPROPERTY(BlueprintReadWrite, Category = "Game", meta = (DeprecatedProperty, DeprecationMessage = "Please use the SoftGameModeOverride property instead."))
 	TSubclassOf<AGameModeBase> GameModeOverride;
 
+	/** Optional Game Mode to override the map's default game mode with. This can be useful if the game's normal mode displays UI elements or loading screens that you don't want captured. */
+	UPROPERTY(EditAnywhere, BlueprintReadWrite, DisplayName = "Game Mode Override", Category = "Game")
+	TSoftClassPtr<AGameModeBase> SoftGameModeOverride;
+
 	/** If true, automatically set the engine to the Cinematic Scalability quality settings during render. See the Scalability Reference documentation for information on how to edit cvars to add/change default quality values.*/
 	UPROPERTY(EditAnywhere, BlueprintReadWrite, Category = "Game")
 	bool bCinematicQualitySettings;
