MoviePipeline: Added the "Set Start/End Console Commands" node to the graph.

diff --git a/Engine/Plugins/MovieScene/MovieRenderPipeline/Source/MovieRenderPipelineCore/Private/Graph/MovieGraphCVarManager.cpp b/Engine/Plugins/MovieScene/MovieRenderPipeline/Source/MovieRenderPipelineCore/Private/Graph/MovieGraphCVarManager.cpp
--- a/Engine/Plugins/MovieScene/MovieRenderPipeline/Source/MovieRenderPipelineCore/Private/Graph/MovieGraphCVarManager.cpp
+++ b/Engine/Plugins/MovieScene/MovieRenderPipeline/Source/MovieRenderPipelineCore/Private/Graph/MovieGraphCVarManager.cpp
@@ -2,10 +2,13 @@
 
 #include "Graph/MovieGraphCVarManager.h"
 
+#include "Engine/World.h"
 #include "Graph/MovieGraphConfig.h"
 #include "Graph/Nodes/MovieGraphApplyCVarPresetNode.h"
 #include "Graph/Nodes/MovieGraphSetCVarValueNode.h"
+#include "Graph/Nodes/MovieGraphSetStartEndConsoleCommandsNode.h"
 #include "HAL/IConsoleManager.h"
+#include "Kismet/KismetSystemLibrary.h"
 #include "Misc/DefaultValueHelper.h"
 #include "MovieRenderPipelineCoreModule.h"
 #include "Sections/MovieSceneConsoleVariableTrackInterface.h"
@@ -57,6 +60,16 @@ void FMovieGraphCVarManager::AddPreset(const TScriptInterface<IMovieSceneConsole
 	}
 }
 
+void FMovieGraphCVarManager::AddStartConsoleCommands(const TArray<FString>& InStartConsoleCommands)
+{
+	StartConsoleCommands = InStartConsoleCommands;
+}
+
+void FMovieGraphCVarManager::AddEndConsoleCommands(const TArray<FString>& InEndConsoleCommands)
+{
+	EndConsoleCommands = InEndConsoleCommands;
+}
+
 void FMovieGraphCVarManager::AddEvaluatedGraph(const UMovieGraphEvaluatedConfig* InEvaluatedGraph)
 {
 	const FMovieGraphEvaluatedBranchConfig* GlobalsBranchConfig = InEvaluatedGraph->BranchConfigMapping.Find(UMovieGraphNode::GlobalsPinName);
@@ -88,6 +101,11 @@ void FMovieGraphCVarManager::AddEvaluatedGraph(const UMovieGraphEvaluatedConfig*
 				AddCVar(CvarOverride.Key, CvarOverride.Value);
 			}
 		}
+		else if (const UMovieGraphSetStartEndConsoleCommandsNode* ConsoleCommandsNode = Cast<UMovieGraphSetStartEndConsoleCommandsNode>(Node))
+		{
+			AddStartConsoleCommands(ConsoleCommandsNode->ConsoleCommands->AddStartCommands);
+			AddEndConsoleCommands(ConsoleCommandsNode->ConsoleCommands->AddEndCommands);
+		}
 	}
 }
 
@@ -136,6 +154,29 @@ void FMovieGraphCVarManager::RevertAllCVars()
 	PreviousConsoleVariableValues.Reset();
 }
 
+void FMovieGraphCVarManager::RunStartConsoleCommands()
+{
+	for (const FString& StartCommand : StartConsoleCommands)
+	{
+		UE_LOG(LogMovieRenderPipeline, Log, TEXT("Executing Console Command \"%s\" before shot starts."), *StartCommand);
+		UKismetSystemLibrary::ExecuteConsoleCommand(WorldContext, StartCommand, nullptr);
+	}
+}
+
+void FMovieGraphCVarManager::RunEndConsoleCommands()
+{
+	for (const FString& EndCommand : EndConsoleCommands)
+	{
+		UE_LOG(LogMovieRenderPipeline, Log, TEXT("Executing Console Command \"%s\" after shot ends."), *EndCommand);
+		UKismetSystemLibrary::ExecuteConsoleCommand(WorldContext, EndCommand, nullptr);
+	}
+}
+
+void FMovieGraphCVarManager::SetWorld(UWorld* InWorld)
+{
+	WorldContext = InWorld;
+}
+
 void FMovieGraphCVarManager::ApplyCVar(IConsoleVariable* InCVar, float InValue)
 {
 	check(InCVar);
diff --git a/Engine/Plugins/MovieScene/MovieRenderPipeline/Source/MovieRenderPipelineCore/Private/Graph/MovieGraphCVarManager.h b/Engine/Plugins/MovieScene/MovieRenderPipeline/Source/MovieRenderPipelineCore/Private/Graph/MovieGraphCVarManager.h
--- a/Engine/Plugins/MovieScene/MovieRenderPipeline/Source/MovieRenderPipelineCore/Private/Graph/MovieGraphCVarManager.h
+++ b/Engine/Plugins/MovieScene/MovieRenderPipeline/Source/MovieRenderPipelineCore/Private/Graph/MovieGraphCVarManager.h
@@ -13,8 +13,8 @@ class UMovieScene;
 namespace UE::MovieGraph::Private
 {
 	/**
-	 * Responsible for applying individual cvars and cvar presets. CVars can be added to the manager with the Add*() methods,
-	 * then applied via ApplyAllCVars().
+	 * Responsible for applying individual cvars, cvar presets, and console commands. CVars can be added to the manager with the Add*() methods,
+	 * then applied via ApplyAllCVars(). Console commands can be added with their associated Add*() methods, and executed via the Run*() methods.
 	 */
 	class FMovieGraphCVarManager final
 	{
@@ -33,7 +33,13 @@ namespace UE::MovieGraph::Private
 		 */
 		void AddPreset(const TScriptInterface<IMovieSceneConsoleVariableTrackInterface>& InPreset);
 
-		/** For all cvar and cvar preset nodes in InEvaluatedGraph, calls either AddCVar() or AddPreset(). */
+		/** Adds console commands that should be run before a shot starts rendering. */
+		void AddStartConsoleCommands(const TArray<FString>& InStartConsoleCommands);
+
+		/** Adds console commands that should be run after a shot finishes rendering. */
+		void AddEndConsoleCommands(const TArray<FString>& InEndConsoleCommands);
+
+		/** For all cvar, cvar preset, and console command nodes in InEvaluatedGraph, calls either AddCVar(), AddPreset(), or Add*ConsoleCommands(). */
 		void AddEvaluatedGraph(const UMovieGraphEvaluatedConfig* InEvaluatedGraph);
 
 		/** Applies all cvars that have been gathered via the Add*() methods. */
@@ -45,6 +51,15 @@ namespace UE::MovieGraph::Private
 		 */
 		void RevertAllCVars();
 
+		/** Runs all start console commands that have been added. */
+		void RunStartConsoleCommands();
+
+		/** Runs all end console commands that have been added. */
+		void RunEndConsoleCommands();
+
+		/** Sets the world context that the console commands should use when executing. */
+		void SetWorld(UWorld* InWorld);
+
 	private:
 		/** Sets the given cvar, InCVar, to InValue. */
 		static void ApplyCVar(IConsoleVariable* InCVar, float InValue);
@@ -71,5 +86,14 @@ namespace UE::MovieGraph::Private
 
 		/** The values of the gathered cvars before the the manager sets their values. */
 		TArray<float> PreviousConsoleVariableValues;
+
+		/** The start console commands that were added. */
+		TArray<FString> StartConsoleCommands;
+
+		/** The end console commands that were added. */
+		TArray<FString> EndConsoleCommands;
+
+		/** The world that should be used as the context when executing console commands. */
+		UWorld* WorldContext;
 	};
 } // namespace UE::MovieGraph::Private
\ No newline at end of file
diff --git a/Engine/Plugins/MovieScene/MovieRenderPipeline/Source/MovieRenderPipelineCore/Private/Graph/MovieGraphPipeline.cpp b/Engine/Plugins/MovieScene/MovieRenderPipeline/Source/MovieRenderPipelineCore/Private/Graph/MovieGraphPipeline.cpp
--- a/Engine/Plugins/MovieScene/MovieRenderPipeline/Source/MovieRenderPipelineCore/Private/Graph/MovieGraphPipeline.cpp
+++ b/Engine/Plugins/MovieScene/MovieRenderPipeline/Source/MovieRenderPipelineCore/Private/Graph/MovieGraphPipeline.cpp
@@ -920,8 +920,10 @@ void UMovieGraphPipeline::SetupShot(const TObjectPtr<UMoviePipelineExecutorShot>
 	}
 
 	// Apply cvars for the shot
+	CVarManager->SetWorld(GetWorld());
 	CVarManager->AddEvaluatedGraph(EvaluatedConfig);
 	CVarManager->ApplyAllCVars();
+	CVarManager->RunStartConsoleCommands();
 
 	// Setup required rendering architecture for all passes in this shot.
 	GraphRendererInstance->SetupRenderingPipelineForShot(InShot);
@@ -1000,6 +1002,7 @@ void UMovieGraphPipeline::TeardownShot(const TObjectPtr<UMoviePipelineExecutorSh
 
 	// Revert the cvar values that were initially applied for the shot
 	CVarManager->RevertAllCVars();
+	CVarManager->RunEndConsoleCommands();
 
 	// Revert cvars set by the global game overrides. Needs to be done after the CVarManager reverts (since the global
 	// game overrides are applied first in SetupShot).
diff --git a/Engine/Plugins/MovieScene/MovieRenderPipeline/Source/MovieRenderPipelineCore/Private/Graph/Nodes/MovieGraphSetStartEndConsoleCommandsNode.cpp b/Engine/Plugins/MovieScene/MovieRenderPipeline/Source/MovieRenderPipelineCore/Private/Graph/Nodes/MovieGraphSetStartEndConsoleCommandsNode.cpp
new file mode 100644
--- /dev/null
+++ b/Engine/Plugins/MovieScene/MovieRenderPipeline/Source/MovieRenderPipelineCore/Private/Graph/Nodes/MovieGraphSetStartEndConsoleCommandsNode.cpp
@@ -0,0 +1,143 @@
+// Copyright Epic Games, Inc. All Rights Reserved.
+
+#include "Graph/Nodes/MovieGraphSetStartEndConsoleCommandsNode.h"
+
+#include "Graph/MovieGraphConfig.h"
+#include "Styling/AppStyle.h"
+
+void UMovieGraphStartEndConsoleCommands::Merge(const IMovieGraphTraversableObject* InSourceObject)
+{
+	const UMovieGraphStartEndConsoleCommands* InConsoleCommands = Cast<UMovieGraphStartEndConsoleCommands>(InSourceObject);
+	checkf(InConsoleCommands, TEXT("UMovieGraphStartEndConsoleCommands cannot merge with null or an object of another type."));
+
+	// The Movie Graph is evaluated right to left, i.e., from the output to the input.
+
+	// First, merge in the "remove" console commands. This needs to be done before merging the start/end commands to ensure that commands that need to
+	// be removed are never added to the start/end commands.
+	for (const FString& RemoveStartCommand : InConsoleCommands->RemoveStartCommands)
+	{
+		if (!RemoveStartCommand.IsEmpty())
+		{
+			RemoveStartCommands.AddUnique(RemoveStartCommand);
+		}
+	}
+
+	for (const FString& RemoveEndCommand : InConsoleCommands->RemoveEndCommands)
+	{
+		if (!RemoveEndCommand.IsEmpty())
+		{
+			RemoveEndCommands.AddUnique(RemoveEndCommand);
+		}
+	}
+	
+	// Next, add in the start/end console commands, but only if they are NOT in the "remove" console commands.
+	for (const FString& StartCommand : InConsoleCommands->AddStartCommands)
+	{
+		if (!StartCommand.IsEmpty() && !RemoveStartCommands.Contains(StartCommand))
+		{
+			AddStartCommands.AddUnique(StartCommand);
+		}
+	}
+
+	for (const FString& EndCommand : InConsoleCommands->AddEndCommands)
+	{
+		if (!EndCommand.IsEmpty() && !RemoveEndCommands.Contains(EndCommand))
+		{
+			AddEndCommands.AddUnique(EndCommand);
+		}
+	}
+}
+
+TArray<TPair<FString, FString>> UMovieGraphStartEndConsoleCommands::GetMergedProperties() const
+{
+	TArray<TPair<FString, FString>> MergedProperties;
+
+	// Don't show the Remove* properties in the evaluated view. They have no meaning after evaluation is finished.
+	MergedProperties.Add(TPair<FString, FString>("Start Console Commands", FString::Join(AddStartCommands, TEXT(", "))));
+	MergedProperties.Add(TPair<FString, FString>("End Console Commands", FString::Join(AddEndCommands, TEXT(", "))));
+
+	return MergedProperties;
+}
+
+UMovieGraphSetStartEndConsoleCommandsNode::UMovieGraphSetStartEndConsoleCommandsNode()
+{
+	ConsoleCommands = CreateDefaultSubobject<UMovieGraphStartEndConsoleCommands>(TEXT("Console Commands"));
+}
+
+EMovieGraphBranchRestriction UMovieGraphSetStartEndConsoleCommandsNode::GetBranchRestriction() const
+{
+	return EMovieGraphBranchRestriction::Globals;
+}
+
+#if WITH_EDITOR
+FText UMovieGraphSetStartEndConsoleCommandsNode::GetNodeTitle(const bool bGetDescriptive) const
+{
+	static const FText StartEndConsoleCommandsNodeName = NSLOCTEXT("MovieGraphNodes", "NodeName_StartEndConsoleCommands", "Set Start/End Console Commands");
+	static const FText StartEndConsoleCommandsNodeAddDescription = NSLOCTEXT("MovieGraphNodes", "NodeDescription_AddStartEndConsoleCommands", "Add {0} Start, {1} End Command(s)");
+	static const FText StartEndConsoleCommandsNodeRemoveDescription = NSLOCTEXT("MovieGraphNodes", "NodeDescription_RemoveStartEndConsoleCommands", "Remove {0} Start, {1} End Command(s)");
+
+	if (bGetDescriptive)
+	{
+		TArray<FText> DescriptionLines = {StartEndConsoleCommandsNodeName};
+
+		// Add in the "Add" commands, if there are any
+		const int32 NumAddedStartCommands = ConsoleCommands->AddStartCommands.Num();
+		const int32 NumAddedEndCommands = ConsoleCommands->AddEndCommands.Num();
+		if ((NumAddedStartCommands > 0) || (NumAddedEndCommands > 0))
+		{
+			DescriptionLines.Add(FText::Format(StartEndConsoleCommandsNodeAddDescription, NumAddedStartCommands, NumAddedEndCommands));
+		}
+
+		// Add in the "Remove" commands, if there are any
+		const int32 NumRemovedStartCommands = ConsoleCommands->RemoveStartCommands.Num();
+		const int32 NumRemovedEndCommands = ConsoleCommands->RemoveEndCommands.Num();
+		if ((NumRemovedStartCommands > 0) || (NumRemovedEndCommands > 0))
+		{
+			DescriptionLines.Add(FText::Format(StartEndConsoleCommandsNodeRemoveDescription, NumRemovedStartCommands, NumRemovedEndCommands));
+		}
+
+		return FText::Join(INVTEXT("\n"), DescriptionLines);
+	}
+	
+	return StartEndConsoleCommandsNodeName;
+}
+
+FText UMovieGraphSetStartEndConsoleCommandsNode::GetMenuCategory() const
+{
+	return NSLOCTEXT("MovieGraphNodes", "StartEndConsoleCommandsGraphNode_Category", "Utility");
+}
+
+FText UMovieGraphSetStartEndConsoleCommandsNode::GetKeywords() const
+{
+	static const FText Keywords = NSLOCTEXT("MovieGraphNodes", "StartEndConsoleCommandsGraphNode_Keywords", "start end console commands");
+	return Keywords;
+}
+
+FLinearColor UMovieGraphSetStartEndConsoleCommandsNode::GetNodeTitleColor() const
+{
+	static const FLinearColor StartEndConsoleCommandsNodeColor = FLinearColor(0.04f, 0.22f, 0.36f);
+	return StartEndConsoleCommandsNodeColor;
+}
+
+FSlateIcon UMovieGraphSetStartEndConsoleCommandsNode::GetIconAndTint(FLinearColor& OutColor) const
+{
+	static const FSlateIcon StartEndConsoleCommandsIcon = FSlateIcon(FAppStyle::GetAppStyleSetName(), "MainFrame.BrowseCVars");
+
+	OutColor = FLinearColor::White;
+	return StartEndConsoleCommandsIcon;
+}
+
+void UMovieGraphSetStartEndConsoleCommandsNode::PostEditChangeProperty(FPropertyChangedEvent& PropertyChangedEvent)
+{
+	Super::PostEditChangeProperty(PropertyChangedEvent);
+
+	// Update the node's title once the user has committed their changes to the commands
+	if (PropertyChangedEvent.ChangeType != EPropertyChangeType::Interactive)
+	{
+		if (PropertyChangedEvent.GetPropertyName() == GET_MEMBER_NAME_CHECKED(UMovieGraphSetStartEndConsoleCommandsNode, ConsoleCommands))
+		{
+			OnNodeChangedDelegate.Broadcast(this);
+		}
+	}
+}
+#endif // WITH_EDITOR
diff --git a/Engine/Plugins/MovieScene/MovieRenderPipeline/Source/MovieRenderPipelineCore/Public/Graph/Nodes/MovieGraphSetStartEndConsoleCommandsNode.h b/Engine/Plugins/MovieScene/MovieRenderPipeline/Source/MovieRenderPipelineCore/Public/Graph/Nodes/MovieGraphSetStartEndConsoleCommandsNode.h
new file mode 100644
--- /dev/null
+++ b/Engine/Plugins/MovieScene/MovieRenderPipeline/Source/MovieRenderPipelineCore/Public/Graph/Nodes/MovieGraphSetStartEndConsoleCommandsNode.h
@@ -0,0 +1,74 @@
+// Copyright Epic Games, Inc. All Rights Reserved.
+
+#pragma once
+
+#include "Graph/MovieGraphConfig.h"
+#include "Graph/MovieGraphNode.h"
+
+#include "MovieGraphSetStartEndConsoleCommandsNode.generated.h"
+
+/** Console commands that can execute within the UMovieGraphSetStartEndConsoleCommandsNode. */
+UCLASS(BlueprintType)
+class MOVIERENDERPIPELINECORE_API UMovieGraphStartEndConsoleCommands : public UObject, public IMovieGraphTraversableObject
+{
+	GENERATED_BODY()
+
+public:
+	/**
+	 * Console commands to execute when this shot starts rendering. If the command(s) need to be undone after the shot finishes rendering, add a
+	 * matching entry to Add End Commands.
+	 */
+	UPROPERTY(EditAnywhere, BlueprintReadWrite, Category="Add Commands")
+	TArray<FString> AddStartCommands;
+
+	/** Console commands to execute when this shot finishes rendering. Used to restore changes made by Add Console Commands. */
+	UPROPERTY(EditAnywhere, BlueprintReadWrite, Category="Add Commands")
+	TArray<FString> AddEndCommands;
+
+	/** Start commands that should be removed from upstream nodes. Commands entered here must be an exact match to upstream command(s) in order to be removed. */
+	UPROPERTY(EditAnywhere, BlueprintReadWrite, Category="Remove Commands")
+	TArray<FString> RemoveStartCommands;
+
+	/** End commands that should be removed from upstream nodes. Commands entered here must be an exact match to upstream command(s) in order to be removed. */
+	UPROPERTY(EditAnywhere, BlueprintReadWrite, Category="Remove Commands")
+	TArray<FString> RemoveEndCommands;
+
+	// IMovieGraphTraversableObject interface
+	virtual void Merge(const IMovieGraphTraversableObject* InSourceObject) override;
+	virtual TArray<TPair<FString, FString>> GetMergedProperties() const override;
+	// ~IMovieGraphTraversableObject interface
+};
+
+/** A node which can set console commands to run at the start and end of a render. */
+UCLASS()
+class MOVIERENDERPIPELINECORE_API UMovieGraphSetStartEndConsoleCommandsNode : public UMovieGraphSettingNode
+{
+	GENERATED_BODY()
+public:
+	UMovieGraphSetStartEndConsoleCommandsNode();
+
+	//~ Begin UMovieGraphNode interface
+	virtual EMovieGraphBranchRestriction GetBranchRestriction() const override;
+	//~ End UMovieGraphNode interface
+
+#if WITH_EDITOR
+	//~ Begin UMovieGraphNode interface
+	virtual FText GetNodeTitle(const bool bGetDescriptive = false) const override;
+	virtual FText GetMenuCategory() const override;
+	virtual FText GetKeywords() const override;
+	virtual FLinearColor GetNodeTitleColor() const override;
+	virtual FSlateIcon GetIconAndTint(FLinearColor& OutColor) const override;
+	//~ End UMovieGraphNode interface
+
+	//~ Begin UObject interface
+	virtual void PostEditChangeProperty(FPropertyChangedEvent& PropertyChangedEvent) override;
+	//~ End UObject interface
+#endif
+
+public:
+	UPROPERTY(EditAnywhere, BlueprintReadWrite, Category = Overrides, meta = (InlineEditConditionToggle))
+	uint8 bOverride_ConsoleCommands : 1 = 1;
+
+	UPROPERTY(VisibleAnywhere, BlueprintReadWrite, Instanced, Category="Settings", meta=(EditCondition="bOverride_ConsoleCommands"))
+	TObjectPtr<UMovieGraphStartEndConsoleCommands> ConsoleCommands;
+};
\ No newline at end of file
