MoviePipeline: Added the ability to frame nodes within the graph, or zoom out to show the entire graph.

diff --git a/Engine/Plugins/MovieScene/MovieRenderPipeline/Source/MovieRenderPipelineEditor/Private/Graph/MovieGraphAssetToolkit.cpp b/Engine/Plugins/MovieScene/MovieRenderPipeline/Source/MovieRenderPipelineEditor/Private/Graph/MovieGraphAssetToolkit.cpp
--- a/Engine/Plugins/MovieScene/MovieRenderPipeline/Source/MovieRenderPipelineEditor/Private/Graph/MovieGraphAssetToolkit.cpp
+++ b/Engine/Plugins/MovieScene/MovieRenderPipeline/Source/MovieRenderPipelineEditor/Private/Graph/MovieGraphAssetToolkit.cpp
@@ -28,11 +28,13 @@
 #include "Graph/Nodes/MovieGraphSetCVarValueNode.h"
 #include "MovieEdGraphNode.h"
 #include "MovieGraphSchema.h"
+#include "MoviePipelineCommands.h"
 #include "MovieRenderPipelineSettings.h"
 #include "SMovieGraphActiveRenderSettingsTabContent.h"
 #include "SMovieGraphMembersTabContent.h"
 
 #include "Framework/Commands/GenericCommands.h"
+#include "GraphEditor.h"
 #include "IDetailRootObjectCustomization.h"
 #include "PropertyEditorModule.h"
 #include "Selection.h"
@@ -268,6 +270,7 @@ void FMovieGraphAssetToolkit::InitMovieGraphAssetToolkit(const EToolkitMode::Typ
 	InitAssetEditor(Mode, InitToolkitHost, AppIdentifier, Layout, bCreateDefaultStandaloneMenu, bCreateDefaultToolbar, InitGraph);
 
 	BindGraphCommands();
+	ExtendToolkitMenu();
 }
 
 TSharedRef<SDockTab> FMovieGraphAssetToolkit::SpawnTab_RenderGraphEditor(const FSpawnTabArgs& Args)
@@ -475,6 +478,37 @@ void FMovieGraphAssetToolkit::BindGraphCommands()
 	ToolkitCommands->MapAction(FGenericCommands::Get().Delete,
 		FExecuteAction::CreateSP(this, &FMovieGraphAssetToolkit::DeleteSelectedMembers),
 		FCanExecuteAction::CreateSP(this, &FMovieGraphAssetToolkit::CanDeleteSelectedMembers));
+
+	ToolkitCommands->MapAction(FMoviePipelineCommands::Get().ZoomToWindow,
+			FExecuteAction::CreateSP(this, &FMovieGraphAssetToolkit::OnZoomToWindow),
+			FCanExecuteAction::CreateSP(this, &FMovieGraphAssetToolkit::CanZoomToWindow));
+		
+	ToolkitCommands->MapAction(FMoviePipelineCommands::Get().ZoomToSelection,
+		FExecuteAction::CreateSP(this, &FMovieGraphAssetToolkit::OnZoomToSelection),
+		FCanExecuteAction::CreateSP(this, &FMovieGraphAssetToolkit::CanZoomToSelection));
+}
+
+void FMovieGraphAssetToolkit::ExtendToolkitMenu() const
+{
+	UToolMenus* ToolMenus = UToolMenus::Get();	
+	if (UToolMenu* MainMenu = ToolMenus->ExtendMenu(GetToolMenuName()))
+	{
+		FToolMenuSection& Section = MainMenu->FindOrAddSection(NAME_None);
+		if (!Section.FindEntry("View"))
+		{
+			Section.AddSubMenu(
+				"View",
+				LOCTEXT("ViewMenu", "View"),
+				LOCTEXT("ViewMenu_ToolTip", "Open the View menu"),
+				FNewToolMenuDelegate::CreateLambda([this](UToolMenu* InMenu)
+				{
+					FToolMenuSection& ZoomSection = InMenu->AddSection("ViewZoom", LOCTEXT("ViewMenuZoomHeading", "Zoom"));
+					ZoomSection.AddMenuEntryWithCommandList(FMoviePipelineCommands::Get().ZoomToWindow, ToolkitCommands);
+					ZoomSection.AddMenuEntryWithCommandList(FMoviePipelineCommands::Get().ZoomToSelection, ToolkitCommands);
+				})
+			).InsertPosition = FToolMenuInsert("Asset", EToolMenuInsertType::After);
+		}
+	}
 }
 
 void FMovieGraphAssetToolkit::DeleteSelectedMembers()
@@ -495,6 +529,40 @@ bool FMovieGraphAssetToolkit::CanDeleteSelectedMembers()
 	return false;
 }
 
+void FMovieGraphAssetToolkit::OnZoomToWindow() const
+{
+	if (MovieGraphWidget.IsValid())
+	{
+		if (const TSharedPtr<SGraphEditor> GraphEditor = MovieGraphWidget->GetGraphEditor().Pin())
+		{
+			constexpr bool bOnlySelection = false;
+			GraphEditor->ZoomToFit(bOnlySelection);
+		}
+	}
+}
+
+bool FMovieGraphAssetToolkit::CanZoomToWindow() const
+{
+	return MovieGraphWidget.IsValid() && MovieGraphWidget->GetGraphEditor().IsValid();
+}
+
+void FMovieGraphAssetToolkit::OnZoomToSelection() const
+{
+	if (MovieGraphWidget.IsValid())
+	{
+		if (const TSharedPtr<SGraphEditor> GraphEditor = MovieGraphWidget->GetGraphEditor().Pin())
+		{
+			constexpr bool bOnlySelection = true;
+			GraphEditor->ZoomToFit(bOnlySelection);
+		}
+	}
+}
+
+bool FMovieGraphAssetToolkit::CanZoomToSelection() const
+{
+	return MovieGraphWidget.IsValid() && MovieGraphWidget->GetGraphEditor().IsValid();
+}
+
 void FMovieGraphAssetToolkit::PersistEditorOnlyNodes() const
 {
 	if (InitialGraph && InitialGraph->PipelineEdGraph)
diff --git a/Engine/Plugins/MovieScene/MovieRenderPipeline/Source/MovieRenderPipelineEditor/Private/Graph/MovieGraphAssetToolkit.h b/Engine/Plugins/MovieScene/MovieRenderPipeline/Source/MovieRenderPipelineEditor/Private/Graph/MovieGraphAssetToolkit.h
--- a/Engine/Plugins/MovieScene/MovieRenderPipeline/Source/MovieRenderPipelineEditor/Private/Graph/MovieGraphAssetToolkit.h
+++ b/Engine/Plugins/MovieScene/MovieRenderPipeline/Source/MovieRenderPipelineEditor/Private/Graph/MovieGraphAssetToolkit.h
@@ -56,12 +56,27 @@ private:
 	/** Register commands used by the editor. */
 	void BindGraphCommands();
 
+	/** Adds additional commands to the toolkit's menu. */
+	void ExtendToolkitMenu() const;
+
 	/** Deletes the member(s) which are currently selected. */
 	void DeleteSelectedMembers();
 
 	/** Determines if the selected member(s) can be deleted. */
 	bool CanDeleteSelectedMembers();
 
+	/** Frames the entire graph so it is visible. */
+	void OnZoomToWindow() const;
+
+	/** Determines if the entire graph can be framed. */
+	bool CanZoomToWindow() const;
+
+	/** Frames the selected node(s) so they are visible. */
+	void OnZoomToSelection() const;
+
+	/** Determines if the selected node(s) can be framed. */
+	bool CanZoomToSelection() const;
+
 	/** Copies editor-only nodes to the underlying runtime graph so they can be restored later. */
 	void PersistEditorOnlyNodes() const;
 
diff --git a/Engine/Plugins/MovieScene/MovieRenderPipeline/Source/MovieRenderPipelineEditor/Private/MoviePipelineCommands.cpp b/Engine/Plugins/MovieScene/MovieRenderPipeline/Source/MovieRenderPipelineEditor/Private/MoviePipelineCommands.cpp
--- a/Engine/Plugins/MovieScene/MovieRenderPipeline/Source/MovieRenderPipelineEditor/Private/MoviePipelineCommands.cpp
+++ b/Engine/Plugins/MovieScene/MovieRenderPipeline/Source/MovieRenderPipelineEditor/Private/MoviePipelineCommands.cpp
@@ -13,6 +13,10 @@ FMoviePipelineCommands::FMoviePipelineCommands()
 void FMoviePipelineCommands::RegisterCommands()
 {
 	UI_COMMAND(ResetStatus, "ResetStatus", "Reset Status", EUserInterfaceActionType::Button, FInputChord(EModifierKey::Control | EModifierKey::Shift, EKeys::R));
+
+	// Graph commands
+	UI_COMMAND(ZoomToWindow, "Zoom to Graph Extents", "Fit the current view to the entire graph", EUserInterfaceActionType::Button, FInputChord());
+	UI_COMMAND(ZoomToSelection, "Zoom to Selection", "Fit the current view to the selection", EUserInterfaceActionType::Button, FInputChord(EKeys::Home));
 }
 
 #undef LOCTEXT_NAMESPACE
diff --git a/Engine/Plugins/MovieScene/MovieRenderPipeline/Source/MovieRenderPipelineEditor/Private/MoviePipelineCommands.h b/Engine/Plugins/MovieScene/MovieRenderPipeline/Source/MovieRenderPipelineEditor/Private/MoviePipelineCommands.h
--- a/Engine/Plugins/MovieScene/MovieRenderPipeline/Source/MovieRenderPipelineEditor/Private/MoviePipelineCommands.h
+++ b/Engine/Plugins/MovieScene/MovieRenderPipeline/Source/MovieRenderPipelineEditor/Private/MoviePipelineCommands.h
@@ -12,6 +12,10 @@ public:
 
 	TSharedPtr<FUICommandInfo> ResetStatus;
 
+	// Graph commands
+	TSharedPtr<FUICommandInfo> ZoomToWindow;
+	TSharedPtr<FUICommandInfo> ZoomToSelection;
+
 	/** Initialize commands */
 	virtual void RegisterCommands() override;
 };
diff --git a/Engine/Plugins/MovieScene/MovieRenderPipeline/Source/MovieRenderPipelineEditor/Private/Widgets/Graph/SMovieGraphConfigPanel.cpp b/Engine/Plugins/MovieScene/MovieRenderPipeline/Source/MovieRenderPipelineEditor/Private/Widgets/Graph/SMovieGraphConfigPanel.cpp
--- a/Engine/Plugins/MovieScene/MovieRenderPipeline/Source/MovieRenderPipelineEditor/Private/Widgets/Graph/SMovieGraphConfigPanel.cpp
+++ b/Engine/Plugins/MovieScene/MovieRenderPipeline/Source/MovieRenderPipelineEditor/Private/Widgets/Graph/SMovieGraphConfigPanel.cpp
@@ -646,6 +646,11 @@ void SMoviePipelineGraphPanel::ClearGraphSelection() const
 	GraphEditorWidget->ClearSelectionSet();
 }
 
+TWeakPtr<SGraphEditor> SMoviePipelineGraphPanel::GetGraphEditor() const
+{
+	return GraphEditorWidget.ToWeakPtr();
+}
+
 bool SMoviePipelineGraphPanel::MatchesContext(const FTransactionContext& InContext, const TArray<TPair<UObject*, FTransactionObjectEvent>>& TransactionObjectContexts) const
 {
 	// ToDo: If we wanted to rebuild the graph less often, we could make all of our transactions have a context
diff --git a/Engine/Plugins/MovieScene/MovieRenderPipeline/Source/MovieRenderPipelineEditor/Private/Widgets/Graph/SMovieGraphConfigPanel.h b/Engine/Plugins/MovieScene/MovieRenderPipeline/Source/MovieRenderPipelineEditor/Private/Widgets/Graph/SMovieGraphConfigPanel.h
--- a/Engine/Plugins/MovieScene/MovieRenderPipeline/Source/MovieRenderPipelineEditor/Private/Widgets/Graph/SMovieGraphConfigPanel.h
+++ b/Engine/Plugins/MovieScene/MovieRenderPipeline/Source/MovieRenderPipelineEditor/Private/Widgets/Graph/SMovieGraphConfigPanel.h
@@ -44,6 +44,9 @@ public:
 	/** Clears the selection in the graph widget being displayed. */
 	void ClearGraphSelection() const;
 
+	/** Gets the graph editor that's contained within this widget. */
+	TWeakPtr<SGraphEditor> GetGraphEditor() const;
+
 	// FEditorUndoClient interface
 	virtual bool MatchesContext(const FTransactionContext& InContext, const TArray<TPair<UObject*, FTransactionObjectEvent>>& TransactionObjectContexts) const override;
 	virtual void PostUndo(bool bSuccess) override;
