Changed MRQ/MRG supersampling to cover full area of pixel

diff --git a/Engine/Plugins/MovieScene/MovieRenderPipeline/Source/MovieRenderPipelineCore/Private/MoviePipelineUtils.cpp b/Engine/Plugins/MovieScene/MovieRenderPipeline/Source/MovieRenderPipelineCore/Private/MoviePipelineUtils.cpp
--- a/Engine/Plugins/MovieScene/MovieRenderPipeline/Source/MovieRenderPipelineCore/Private/MoviePipelineUtils.cpp
+++ b/Engine/Plugins/MovieScene/MovieRenderPipeline/Source/MovieRenderPipelineCore/Private/MoviePipelineUtils.cpp
@@ -1072,29 +1072,7 @@ namespace UE
 			float HaltonOffsetX = Halton(HaltonIndex, 2);
 			float HaltonOffsetY = Halton(HaltonIndex, 3);
 
-			float SpatialShiftX = 0.0f;
-			float SpatialShiftY = 0.0f;
-
-			static auto CVar = IConsoleManager::Get().FindConsoleVariable(TEXT("r.TemporalAAFilterSize"));
-			float FilterSize = CVar->GetFloat();
-
-			// Scale distribution to set non-unit variance
-			// Variance = Sigma^2
-			float Sigma = 0.47f * FilterSize;
-
-			// Window to [-0.5, 0.5] output
-			// Without windowing we could generate samples far away on the infinite tails.
-			float OutWindow = 0.5f;
-			float InWindow = FMath::Exp(-0.5 * FMath::Square(OutWindow / Sigma));
-
-			// Box-Muller transform
-			float Theta = 2.0f * PI * HaltonOffsetY;
-			float r = Sigma * FMath::Sqrt(-2.0f * FMath::Loge((1.0f - HaltonOffsetX) * InWindow + HaltonOffsetX));
-
-			SpatialShiftX = r * FMath::Cos(Theta);
-			SpatialShiftY = r * FMath::Sin(Theta);
-
-			return FVector2f(SpatialShiftX, SpatialShiftY);
+			return FVector2f(HaltonOffsetX-0.5, HaltonOffsetY-0.5);
 		}
 	}
 }
