MoviePipeline: Fixed a bunch of quirks in the file name format autocomplete widget. Typing a "{" is also now required in order for the token autocomplete list to show up.

diff --git a/Engine/Plugins/MovieScene/MovieRenderPipeline/Source/MovieRenderPipelineEditor/Private/Widgets/SMoviePipelineFormatTokenAutoCompleteBox.cpp b/Engine/Plugins/MovieScene/MovieRenderPipeline/Source/MovieRenderPipelineEditor/Private/Widgets/SMoviePipelineFormatTokenAutoCompleteBox.cpp
--- a/Engine/Plugins/MovieScene/MovieRenderPipeline/Source/MovieRenderPipelineEditor/Private/Widgets/SMoviePipelineFormatTokenAutoCompleteBox.cpp
+++ b/Engine/Plugins/MovieScene/MovieRenderPipeline/Source/MovieRenderPipelineEditor/Private/Widgets/SMoviePipelineFormatTokenAutoCompleteBox.cpp
@@ -106,10 +106,18 @@ FReply SMoviePipelineFormatTokenAutoCompleteBox::OnKeyDown(const FGeometry& MyGe
 				// Trigger the auto-complete for the highlighted suggestion
 				const FString SuggestionText = *Suggestions[CurrentSuggestionIndex];
 				ReplaceRelevantTextWithSuggestion(SuggestionText);
+				CloseMenuAndReset();
 				return FReply::Handled();
 			}
 		}
 	}
+	else if ((KeyEvent.GetKey() == EKeys::LeftBracket) && KeyEvent.IsShiftDown())
+	{
+		// Start showing token suggestions immediately when { is typed
+		FilterVisibleSuggestions(FString(), false);
+		return FReply::Handled();
+	}
+	
 	return FReply::Unhandled();
 }
 
@@ -132,28 +140,18 @@ void SMoviePipelineFormatTokenAutoCompleteBox::FindAutoCompletableTextAtPos(cons
 	OutStr = FString();
 	bShowAutoComplete = false;
 
-	// We want to find a { brace on or to the left of InCursorPos, but if we find a } we
-	// stop looking, because that's a brace for another text. (+1 for ::FromEnd off by one)
-	const int32 StartingBracePos = InWholeString.Find(
-		TEXT("{"), ESearchCase::Type::IgnoreCase, ESearchDir::Type::FromEnd, InCursorPos + 1);
-	const int32 PreviousEndBracePos = InWholeString.Find(
-		TEXT("}"), ESearchCase::Type::IgnoreCase, ESearchDir::Type::FromEnd, InCursorPos);
-
-	if (StartingBracePos < PreviousEndBracePos)
-	{
-		return;
-	}
+	int32 StartingBracePos, EndBracePos;
+	GetBracePositionsForCursor(InWholeString, InCursorPos, StartingBracePos, EndBracePos);
 
 	FString AutoCompleteText;
 
-	// Now that we found a {, take the substring between it and either the next }, or the end of the string.
-	const int32 NextEndBracePos = InWholeString.Find(TEXT("}"), ESearchCase::Type::IgnoreCase, ESearchDir::Type::FromStart, InCursorPos);
+	// Now that we found a {, take the substring between it and either the next } or {, or the end of the string.
 	if (StartingBracePos >= 0)
 	{
 		int32 Count = InWholeString.Len() - StartingBracePos;
-		if (NextEndBracePos >= 0)
+		if (EndBracePos >= 0)
 		{
-			Count = NextEndBracePos - StartingBracePos;
+			Count = EndBracePos - StartingBracePos;
 		}
 
 		AutoCompleteText = InWholeString.Mid(StartingBracePos + 1, Count - 1);
@@ -172,12 +170,12 @@ void SMoviePipelineFormatTokenAutoCompleteBox::ReplaceRelevantTextWithSuggestion
 	{
 		CursorPos = FMath::Clamp(CursorLoc.GetOffset(), 0, CursorPos);
 	}
-		
-	// Look for the { to the left of the cursor. We search StrPositionIndex from +1 here due to a bug in ::FromEnd being off by one.
-	const int32 StartingBracePos = TextBoxText.Find(TEXT("{"), ESearchCase::Type::IgnoreCase, ESearchDir::Type::FromEnd, CursorPos + 1);
 
-	// Now that we found a {, take the substring between it and either the next }, or the end of the string.
-	const int32 NextEndBracePos = TextBoxText.Find(TEXT("}"), ESearchCase::Type::IgnoreCase, ESearchDir::Type::FromStart, CursorPos);
+	int32 StartingBracePos, EndBracePos;
+	GetBracePositionsForCursor(TextBoxText, CursorPos, StartingBracePos, EndBracePos);
+
+	// Insert the suggestion text after the opening brace, and before the ending brace (if any). Replace any text that currently exists between the
+	// braces with the suggestion text.
 	int32 NewCursorPos = 0;
 	if (StartingBracePos >= 0)	
 	{
@@ -185,9 +183,9 @@ void SMoviePipelineFormatTokenAutoCompleteBox::ReplaceRelevantTextWithSuggestion
 		const FString Left = TextBoxText.Left(StartingBracePos+1);
 		FString Right;
 
-		if (NextEndBracePos >= 0)
+		if (EndBracePos >= 0)
 		{
-			Right = TextBoxText.RightChop(NextEndBracePos);
+			Right = TextBoxText.RightChop(EndBracePos);
 		}
 
 		// Since the user chose the suggestion ensure there's already a } brace to close off the pair.
@@ -197,6 +195,7 @@ void SMoviePipelineFormatTokenAutoCompleteBox::ReplaceRelevantTextWithSuggestion
 		}
 
 		TextBoxText = Left + InSuggestionText + Right;
+		
 		// We subtract 1 from the Right as we want to put the cursor after the automatically generated "}" token.
 		NewCursorPos = TextBoxText.Len() - (Right.Len() - 1);
 	}
@@ -210,26 +209,17 @@ void SMoviePipelineFormatTokenAutoCompleteBox::HandleTextBoxTextChanged(const FT
 	TextHandle->SetValue(InText.ToString(), EPropertyValueSetFlags::InteractiveChange);
 
 	const FString TextAsStr = InText.ToString();
-	if (TextAsStr.Len() > 0)
-	{
-		FString OutStr;
-		bool bShowAutoComplete;
-
-		int32 CursorPos = TextAsStr.Len();
-		const FTextLocation CursorLoc = TextBox->GetCursorLocation();
-		if (CursorLoc.IsValid())
-		{
-			CursorPos = FMath::Clamp(CursorLoc.GetOffset(), 0, CursorPos);
-		}
-
-		FindAutoCompletableTextAtPos(TextAsStr, CursorPos, OutStr, bShowAutoComplete);
-		FilterVisibleSuggestions(OutStr, bShowAutoComplete);
-	}	
-	else
+	int32 CursorPos = TextAsStr.Len();
+	const FTextLocation CursorLoc = TextBox->GetCursorLocation();
+	if (CursorLoc.IsValid())
 	{
-		// If they have no text, suggest all possible solutions
-		FilterVisibleSuggestions(FString(), false);
+		CursorPos = FMath::Clamp(CursorLoc.GetOffset(), 0, CursorPos);
 	}
+
+	FString OutStr;
+	bool bShowAutoComplete;
+	FindAutoCompletableTextAtPos(TextAsStr, CursorPos, OutStr, bShowAutoComplete);
+	FilterVisibleSuggestions(OutStr, bShowAutoComplete);
 }
 
 void SMoviePipelineFormatTokenAutoCompleteBox::HandleTextBoxTextCommitted(const FText& InText, ETextCommit::Type CommitInfo)
@@ -239,6 +229,35 @@ void SMoviePipelineFormatTokenAutoCompleteBox::HandleTextBoxTextCommitted(const
 
 void SMoviePipelineFormatTokenAutoCompleteBox::FilterVisibleSuggestions(const FString& StrToMatch, const bool bForceShowAll)
 {
+	// If the { is not immediately before the cursor and the suggestion string is empty (ie, the user has not started typing in a format token name),
+	// do not show the autocomplete menu at all. This is something that mostly occurs when the text box is empty.
+	if (StrToMatch.IsEmpty())
+	{
+		bool bForceHideAutocomplete = false;
+		
+		const FString TextBoxString = TextBox->GetText().ToString();
+		const int32 CursorLocation = TextBox->GetCursorLocation().GetOffset();
+		if (TextBoxString.IsValidIndex(CursorLocation - 1))
+		{
+			const TCHAR& CharBeforeCursor = TextBoxString[CursorLocation - 1];
+			if (CharBeforeCursor != '{')
+			{
+				bForceHideAutocomplete = true;
+			}
+		}
+		else
+		{
+			// Could happen if the cursor is at the very start of the text box
+			bForceHideAutocomplete = true;
+		}
+
+		if (bForceHideAutocomplete)
+		{
+			CloseMenuAndReset();
+			return;
+		}
+	}
+	
 	Suggestions.Reset();
 	for (const FString& Suggestion : AllSuggestions)
 	{
@@ -319,4 +338,19 @@ TArray<FString> SMoviePipelineFormatTokenAutoCompleteBox::GetFileNameFormatSugge
 	FileNameFormatSuggestions.Sort();
 
 	return FileNameFormatSuggestions;
+}
+
+void SMoviePipelineFormatTokenAutoCompleteBox::GetBracePositionsForCursor(const FString& InText, int32 CursorPos, int32& OutStartingBracePos, int32& OutEndBracePos)
+{
+	OutStartingBracePos = InText.Find(
+		TEXT("{"), ESearchCase::Type::IgnoreCase, ESearchDir::Type::FromEnd, CursorPos);
+	
+	const int32 NextStartingBracePos = InText.Find(
+		TEXT("{"), ESearchCase::Type::IgnoreCase, ESearchDir::Type::FromStart, CursorPos);
+	const int32 EndBracePos = InText.Find(
+		TEXT("}"), ESearchCase::Type::IgnoreCase, ESearchDir::Type::FromStart, CursorPos);
+
+	// Use the position of the next { if it comes before the next }. This could happen if, for example, the user has a format string like {layer_name},
+	// and they want to prepend a format token to the front, so the text would look like {{layer_name}, and the cursor is right after the initial {.
+	OutEndBracePos = (NextStartingBracePos != -1) && (NextStartingBracePos < EndBracePos) ? NextStartingBracePos : EndBracePos;
 }
\ No newline at end of file
diff --git a/Engine/Plugins/MovieScene/MovieRenderPipeline/Source/MovieRenderPipelineEditor/Private/Widgets/SMoviePipelineFormatTokenAutoCompleteBox.h b/Engine/Plugins/MovieScene/MovieRenderPipeline/Source/MovieRenderPipelineEditor/Private/Widgets/SMoviePipelineFormatTokenAutoCompleteBox.h
--- a/Engine/Plugins/MovieScene/MovieRenderPipeline/Source/MovieRenderPipelineEditor/Private/Widgets/SMoviePipelineFormatTokenAutoCompleteBox.h
+++ b/Engine/Plugins/MovieScene/MovieRenderPipeline/Source/MovieRenderPipelineEditor/Private/Widgets/SMoviePipelineFormatTokenAutoCompleteBox.h
@@ -55,6 +55,10 @@ public:
 	 */
 	static TArray<FString> GetFileNameFormatSuggestions();
 
+private:
+	/** Get the relevant brace positions for the given cursor position within the text. */
+	static void GetBracePositionsForCursor(const FString& InText, int32 CursorPos, int32& OutStartingBracePos, int32& OutEndBracePos);
+
 private:
 	TSharedPtr<SListView<TSharedPtr<FString>>> SuggestionListView;
 	TSharedPtr<SMultiLineEditableTextBox> TextBox;
