MoviePipeline: Updated the Apple ProRes and Avid DNx nodes to be capable of running outside the Globals branch in the graph (thus movies can be created per layer).

Also addressed some other related bugs:
* Fixed a crash that happened when a burn-in was composited on multiple video nodes.
* Fixed an issue where multiple video nodes could write to the same file because proper filename disambiguation was not in place (like it is for image output nodes).
* Fixed an issue that caused a video output node to not go through its finalize process (leaving the output file empty) if multiple video nodes were in use.

#rnx

diff --git a/Engine/Plugins/Media/AppleProResMedia/Source/AppleProResMedia/Private/MovieGraphAppleProResNode.cpp b/Engine/Plugins/Media/AppleProResMedia/Source/AppleProResMedia/Private/MovieGraphAppleProResNode.cpp
--- a/Engine/Plugins/Media/AppleProResMedia/Source/AppleProResMedia/Private/MovieGraphAppleProResNode.cpp
+++ b/Engine/Plugins/Media/AppleProResMedia/Source/AppleProResMedia/Private/MovieGraphAppleProResNode.cpp
@@ -20,11 +20,6 @@ UMovieGraphAppleProResNode::UMovieGraphAppleProResNode()
 	
 }
 
-EMovieGraphBranchRestriction UMovieGraphAppleProResNode::GetBranchRestriction() const
-{
-	return EMovieGraphBranchRestriction::Globals;
-}
-
 #if WITH_EDITOR
 FText UMovieGraphAppleProResNode::GetNodeTitle(const bool bGetDescriptive) const
 {
diff --git a/Engine/Plugins/Media/AppleProResMedia/Source/AppleProResMedia/Private/MovieGraphAppleProResNode.h b/Engine/Plugins/Media/AppleProResMedia/Source/AppleProResMedia/Private/MovieGraphAppleProResNode.h
--- a/Engine/Plugins/Media/AppleProResMedia/Source/AppleProResMedia/Private/MovieGraphAppleProResNode.h
+++ b/Engine/Plugins/Media/AppleProResMedia/Source/AppleProResMedia/Private/MovieGraphAppleProResNode.h
@@ -17,8 +17,6 @@ class UMovieGraphAppleProResNode : public UMovieGraphVideoOutputNode
 public:
 	UMovieGraphAppleProResNode();
 
-	virtual EMovieGraphBranchRestriction GetBranchRestriction() const override;
-
 #if WITH_EDITOR
 	virtual FText GetNodeTitle(const bool bGetDescriptive = false) const override;
 	virtual FText GetMenuCategory() const override;
diff --git a/Engine/Plugins/Media/AvidDNxMedia/Source/Source/Private/MovieGraphAvidDNxHRNode.cpp b/Engine/Plugins/Media/AvidDNxMedia/Source/Source/Private/MovieGraphAvidDNxHRNode.cpp
--- a/Engine/Plugins/Media/AvidDNxMedia/Source/Source/Private/MovieGraphAvidDNxHRNode.cpp
+++ b/Engine/Plugins/Media/AvidDNxMedia/Source/Source/Private/MovieGraphAvidDNxHRNode.cpp
@@ -18,11 +18,6 @@ UMovieGraphAvidDNxHRNode::UMovieGraphAvidDNxHRNode()
 {
 }
 
-EMovieGraphBranchRestriction UMovieGraphAvidDNxHRNode::GetBranchRestriction() const
-{
-	return EMovieGraphBranchRestriction::Globals;
-}
-
 #if WITH_EDITOR
 FText UMovieGraphAvidDNxHRNode::GetNodeTitle(const bool bGetDescriptive) const
 {
diff --git a/Engine/Plugins/Media/AvidDNxMedia/Source/Source/Private/MovieGraphAvidDNxHRNode.h b/Engine/Plugins/Media/AvidDNxMedia/Source/Source/Private/MovieGraphAvidDNxHRNode.h
--- a/Engine/Plugins/Media/AvidDNxMedia/Source/Source/Private/MovieGraphAvidDNxHRNode.h
+++ b/Engine/Plugins/Media/AvidDNxMedia/Source/Source/Private/MovieGraphAvidDNxHRNode.h
@@ -27,8 +27,6 @@ class UMovieGraphAvidDNxHRNode : public UMovieGraphVideoOutputNode
 public:
 	UMovieGraphAvidDNxHRNode();
 
-	virtual EMovieGraphBranchRestriction GetBranchRestriction() const override;
-
 #if WITH_EDITOR
 	virtual FText GetNodeTitle(const bool bGetDescriptive = false) const override;
 	virtual FText GetMenuCategory() const override;
diff --git a/Engine/Plugins/MovieScene/MovieRenderPipeline/Source/MovieRenderPipelineCore/Private/Graph/Nodes/MovieGraphFileOutputNode.cpp b/Engine/Plugins/MovieScene/MovieRenderPipeline/Source/MovieRenderPipelineCore/Private/Graph/Nodes/MovieGraphFileOutputNode.cpp
--- a/Engine/Plugins/MovieScene/MovieRenderPipeline/Source/MovieRenderPipelineCore/Private/Graph/Nodes/MovieGraphFileOutputNode.cpp
+++ b/Engine/Plugins/MovieScene/MovieRenderPipeline/Source/MovieRenderPipelineCore/Private/Graph/Nodes/MovieGraphFileOutputNode.cpp
@@ -3,12 +3,50 @@
 #include "Graph/Nodes/MovieGraphFileOutputNode.h"
 
 #include "Graph/MovieGraphConfig.h"
+#include "MoviePipelineUtils.h"
 
 int32 UMovieGraphFileOutputNode::GetNumFileOutputNodes(const UMovieGraphEvaluatedConfig& InEvaluatedConfig, const FName& InBranchName)
 {
 	return InEvaluatedConfig.GetSettingsForBranch(UMovieGraphFileOutputNode::StaticClass(), InBranchName, false /*bIncludeCDOs*/, false /*bExactMatch*/).Num();
 }
 
+void UMovieGraphFileOutputNode::DisambiguateFilename(FString& InOutFilenameFormatString, const UE::MovieGraph::FMovieGraphOutputMergerFrame* InRawFrameData, const FName& InNodeName, const FMovieGraphPassData& InRenderData)
+{
+	// ToDo: This is overly protective and could be relaxed later, for instance
+	// if different file write nodes have chosen a separate filepath entirely.
+	const UE::MovieGraph::FMovieGraphRenderDataValidationInfo ValidationInfo = InRawFrameData->GetValidationInfo(InRenderData.Key);
+
+	// Since there can only be one layer per branch, we restrain layer/branch validation to multi-branch graphs.
+	if (ValidationInfo.BranchCount > 1)
+	{
+		// We can run into the scenario where the users have given layers the same name, so layer_name token won't help differentiate.
+		// To resolve this, we look to see if there's multiple branches with the same layer name, and if so we force the branch name into the token too.
+		if (ValidationInfo.LayerCount < ValidationInfo.BranchCount)
+		{
+			UE::MoviePipeline::ConformOutputFormatStringToken(InOutFilenameFormatString, TEXT("{branch_name}"), InNodeName, InRenderData.Key.RootBranchName);
+		}
+		else
+		{
+			// Otherwise, we separate each branch by its unique layer name.
+			UE::MoviePipeline::ConformOutputFormatStringToken(InOutFilenameFormatString, TEXT("{layer_name}"), InNodeName, InRenderData.Key.RootBranchName);
+		}
+	}
+
+	// We only add the renderer name token if multiple (non-composited) renderers are present on the active branch (eg, in the case of optional PPMs).
+	if (ValidationInfo.ActiveBranchRendererCount > 1)
+	{
+		UE::MoviePipeline::ConformOutputFormatStringToken(InOutFilenameFormatString, TEXT("{renderer_name}"), InNodeName, InRenderData.Key.RootBranchName);
+	}
+
+	// We only add the subresource token if a (non-composited) renderer on the active branch is producing more than one subresource (eg, in the case of optional PPMs).
+	if (ValidationInfo.ActiveRendererSubresourceCount > 1)
+	{
+		UE::MoviePipeline::ConformOutputFormatStringToken(InOutFilenameFormatString, TEXT("{renderer_sub_name}"), InNodeName, InRenderData.Key.RootBranchName);
+	}
+	
+	// ToDo: Add {camera_name} validation once relevant
+}
+
 TArray<FMovieGraphPassData> UMovieGraphFileOutputNode::GetCompositedPasses(UE::MovieGraph::FMovieGraphOutputMergerFrame* InRawFrameData)
 {
 	// Gather the passes that need to be composited
diff --git a/Engine/Plugins/MovieScene/MovieRenderPipeline/Source/MovieRenderPipelineCore/Private/Graph/Nodes/MovieGraphVideoOutputNode.cpp b/Engine/Plugins/MovieScene/MovieRenderPipeline/Source/MovieRenderPipelineCore/Private/Graph/Nodes/MovieGraphVideoOutputNode.cpp
--- a/Engine/Plugins/MovieScene/MovieRenderPipeline/Source/MovieRenderPipelineCore/Private/Graph/Nodes/MovieGraphVideoOutputNode.cpp
+++ b/Engine/Plugins/MovieScene/MovieRenderPipeline/Source/MovieRenderPipelineCore/Private/Graph/Nodes/MovieGraphVideoOutputNode.cpp
@@ -1,4 +1,4 @@
-﻿// Copyright Epic Games, Inc. All Rights Reserved.
+// Copyright Epic Games, Inc. All Rights Reserved.
 
 #include "Graph/Nodes/MovieGraphVideoOutputNode.h"
 
@@ -93,23 +93,34 @@ void UMovieGraphVideoOutputNode::OnReceiveImageDataImpl(UMovieGraphPipeline* InP
 				// Match them up by camera name so multiple passes intended for different camera names work.
 				if (RenderPassData.Key.CameraName == CompositePass.Key.CameraName)
 				{
-					// Create a new composite pass (but move the actual pixel data), otherwise when the main
+					// Create a new composite pass, otherwise when the main
 					// loop tries to check if we should skip writing out this image pass in the future, it fails
 					// because we've MoveTemp'd CompositePass and thus the name checks no longer pass.
 					FMovieGraphPassData NewPassInfo;
 					NewPassInfo.Key = CompositePass.Key;
-					NewPassInfo.Value = MoveTemp(CompositePass.Value);
+
+					// Copy the pixel data if more than one node is using this composited pass
+					TUniquePtr<FImagePixelData> PixelData;
+					if (GetNumFileOutputNodes(*InRawFrameData->EvaluatedConfig, CompositePass.Key.RootBranchName) > 1)
+					{
+						NewPassInfo.Value = CompositePass.Value->CopyImageData();
+					}
+					else
+					{
+						NewPassInfo.Value = CompositePass.Value->MoveImageDataToNew();
+					}
+					
 					CompositesForThisCamera.Add(MoveTemp(NewPassInfo));
 				}
 			}
 
 			// Notify the encoder to write this frame
-			this->WriteFrame_EncodeThread(OutputWriter->CodecWriter.Get(), RawRenderPassData, MoveTemp(CompositesForThisCamera), EvaluatedConfig, RenderPassData.Key.RootBranchName.ToString());
+			WriteFrame_EncodeThread(OutputWriter->CodecWriter.Get(), RawRenderPassData, MoveTemp(CompositesForThisCamera), EvaluatedConfig, RenderPassData.Key.RootBranchName.ToString());
 		}
 		else
 		{
 			TArray<FMovieGraphPassData> Dummy;
-			this->WriteFrame_EncodeThread(OutputWriter->CodecWriter.Get(), RawRenderPassData, MoveTemp(Dummy), EvaluatedConfig, RenderPassData.Key.RootBranchName.ToString());
+			WriteFrame_EncodeThread(OutputWriter->CodecWriter.Get(), RawRenderPassData, MoveTemp(Dummy), EvaluatedConfig, RenderPassData.Key.RootBranchName.ToString());
 		}
 	}
 }
@@ -125,8 +136,11 @@ void UMovieGraphVideoOutputNode::OnAllFramesSubmittedImpl(UMovieGraphPipeline* I
 {
 	for (const FMovieGraphCodecWriterWithPromise& Writer : AllWriters)
 	{
-		MovieRenderGraph::IVideoCodecWriter* RawWriter = Writer.CodecWriter.Get();
-		this->BeginFinalize_EncodeThread(RawWriter);
+		if (Writer.NodeType == GetClass())
+		{
+			MovieRenderGraph::IVideoCodecWriter* RawWriter = Writer.CodecWriter.Get();
+			BeginFinalize_EncodeThread(RawWriter);
+		}
 	}
 }
 
@@ -134,9 +148,14 @@ void UMovieGraphVideoOutputNode::OnAllFramesFinalizedImpl(UMovieGraphPipeline* I
 {
 	for (FMovieGraphCodecWriterWithPromise& Writer : AllWriters)
 	{
+		if (Writer.NodeType != GetClass())
+		{
+			continue;
+		}
+		
 		MovieRenderGraph::IVideoCodecWriter* RawWriter = Writer.CodecWriter.Get();
 
-		this->Finalize_EncodeThread(RawWriter);
+		Finalize_EncodeThread(RawWriter);
 
 		if (!bHasError)
 		{
@@ -144,13 +163,19 @@ void UMovieGraphVideoOutputNode::OnAllFramesFinalizedImpl(UMovieGraphPipeline* I
 		}
 	}
 
-	AllWriters.Empty();
+	// AllWriters is static and shared across *all* video output node types. Only delete writers of the current node's type.
+	AllWriters.RemoveAll([this](const FMovieGraphCodecWriterWithPromise& InWriter)
+	{
+		return InWriter.NodeType == GetClass();
+	});
+	
 	bHasError = false;
 }
 
-UMovieGraphVideoOutputNode::FMovieGraphCodecWriterWithPromise::FMovieGraphCodecWriterWithPromise(TUniquePtr<MovieRenderGraph::IVideoCodecWriter>&& InWriter, TPromise<bool>&& InPromise)
+UMovieGraphVideoOutputNode::FMovieGraphCodecWriterWithPromise::FMovieGraphCodecWriterWithPromise(TUniquePtr<MovieRenderGraph::IVideoCodecWriter>&& InWriter, TPromise<bool>&& InPromise, UClass* InNodeType)
 	: CodecWriter(MoveTemp(InWriter))
 	, Promise(MoveTemp(InPromise))
+	, NodeType(InNodeType)
 {
 	
 }
@@ -181,12 +206,8 @@ void UMovieGraphVideoOutputNode::GetOutputFilePaths(const UMovieGraphPipeline* I
 	FMovieGraphResolveArgs FinalFormatArgs;
 	FString FileNameFormatString = EvaluatedNode->FileNameFormat;
 	
-	// If we're writing more than one render pass out, we need to ensure the file name has the format string in it so we don't
-	// overwrite the same file multiple times. Burn In overlays don't count because they get composited on top of an existing file.
-	const bool bIncludeRenderPass = InRawFrameData->ImageOutputData.Num() - InCompositedPasses.Num() > 1;
-	constexpr bool bTestFrameNumber = false;
-
-	UE::MoviePipeline::ValidateOutputFormatString(FileNameFormatString, bIncludeRenderPass, bTestFrameNumber);
+	// Insert tokens like {layer_name} as appropriate to make sure outputs don't clash with each other.
+	DisambiguateFilename(FileNameFormatString, InRawFrameData, EvaluatedNode->GetFName(), InRenderPassData);
 
 	// Strip any frame number tags so we don't get one video file per frame.
 	UE::MoviePipeline::RemoveFrameNumberFormatStrings(FileNameFormatString, true);
@@ -277,7 +298,7 @@ UMovieGraphVideoOutputNode::FMovieGraphCodecWriterWithPromise* UMovieGraphVideoO
 			TPromise<bool> Completed;
 			InPipeline->AddOutputFuture(Completed.GetFuture(), OutputData);
 
-			AllWriters.Add(FMovieGraphCodecWriterWithPromise(MoveTemp(NewWriter), MoveTemp(Completed)));
+			AllWriters.Add(FMovieGraphCodecWriterWithPromise(MoveTemp(NewWriter), MoveTemp(Completed), GetClass()));
 			OutputWriter = &AllWriters.Last();
 
 			// If it fails to initialize, immediately mark the promise as failed so the render queue stops.
diff --git a/Engine/Plugins/MovieScene/MovieRenderPipeline/Source/MovieRenderPipelineCore/Public/Graph/Nodes/MovieGraphFileOutputNode.h b/Engine/Plugins/MovieScene/MovieRenderPipeline/Source/MovieRenderPipelineCore/Public/Graph/Nodes/MovieGraphFileOutputNode.h
--- a/Engine/Plugins/MovieScene/MovieRenderPipeline/Source/MovieRenderPipelineCore/Public/Graph/Nodes/MovieGraphFileOutputNode.h
+++ b/Engine/Plugins/MovieScene/MovieRenderPipeline/Source/MovieRenderPipelineCore/Public/Graph/Nodes/MovieGraphFileOutputNode.h
@@ -87,6 +87,9 @@ protected:
 	/** Returns the number of evaluated (active) file nodes on the specified branch. */
 	static int32 GetNumFileOutputNodes(const UMovieGraphEvaluatedConfig& InEvaluatedConfig, const FName& InBranchName);
 
+	/** Updates InOutFilenameFormatString to include tokens (like {layer_name}) that will disambiguate the output if multiple outputs are attempting to write to the same file. */
+	static void DisambiguateFilename(FString& InOutFilenameFormatString, const UE::MovieGraph::FMovieGraphOutputMergerFrame* InRawFrameData, const FName& InNodeName, const FMovieGraphPassData& InRenderData);
+
 	/** Convenience function to get the list of active composite passes from render data. */
 	static TArray<FMovieGraphPassData> GetCompositedPasses(UE::MovieGraph::FMovieGraphOutputMergerFrame* InRawFrameData);
 };
\ No newline at end of file
diff --git a/Engine/Plugins/MovieScene/MovieRenderPipeline/Source/MovieRenderPipelineCore/Public/Graph/Nodes/MovieGraphVideoOutputNode.h b/Engine/Plugins/MovieScene/MovieRenderPipeline/Source/MovieRenderPipelineCore/Public/Graph/Nodes/MovieGraphVideoOutputNode.h
--- a/Engine/Plugins/MovieScene/MovieRenderPipeline/Source/MovieRenderPipelineCore/Public/Graph/Nodes/MovieGraphVideoOutputNode.h
+++ b/Engine/Plugins/MovieScene/MovieRenderPipeline/Source/MovieRenderPipelineCore/Public/Graph/Nodes/MovieGraphVideoOutputNode.h
@@ -65,13 +65,16 @@ protected:
 private:
 	struct FMovieGraphCodecWriterWithPromise
 	{
-		FMovieGraphCodecWriterWithPromise(TUniquePtr<MovieRenderGraph::IVideoCodecWriter>&& InWriter, TPromise<bool>&& InPromise);
+		FMovieGraphCodecWriterWithPromise(TUniquePtr<MovieRenderGraph::IVideoCodecWriter>&& InWriter, TPromise<bool>&& InPromise, UClass* InNodeType);
 		
 		/** The codec writer. */
 		TUniquePtr<MovieRenderGraph::IVideoCodecWriter> CodecWriter;
 
 		/** The promise that is provided to the pipeline that specifies whether or not the writer has finished. */
 		TPromise<bool> Promise;
+		
+		/** The type of node associated with this writer. */
+		UClass* NodeType;
 	};
 
 	/**
diff --git a/Engine/Plugins/MovieScene/MovieRenderPipeline/Source/MovieRenderPipelineRenderPasses/Private/MovieGraphImageSequenceOutputNode.cpp b/Engine/Plugins/MovieScene/MovieRenderPipeline/Source/MovieRenderPipelineRenderPasses/Private/MovieGraphImageSequenceOutputNode.cpp
--- a/Engine/Plugins/MovieScene/MovieRenderPipeline/Source/MovieRenderPipelineRenderPasses/Private/MovieGraphImageSequenceOutputNode.cpp
+++ b/Engine/Plugins/MovieScene/MovieRenderPipeline/Source/MovieRenderPipelineRenderPasses/Private/MovieGraphImageSequenceOutputNode.cpp
@@ -61,39 +61,8 @@ FString UMovieGraphImageSequenceOutputNode::CreateFileName(
 	// Generate one string that puts the directory combined with the filename format.
 	FString FileNameFormatString = OutputSettingNode->OutputDirectory.Path / InParentNode->FileNameFormat;
 
-	// ToDo: This is overly protective and could be relaxed later, for instance
-	// if different file write nodes have chosen a separate filepath entirely.
-	UE::MovieGraph::FMovieGraphRenderDataValidationInfo ValidationInfo = InRawFrameData->GetValidationInfo(InRenderData.Key);
-
-	// Since there can only be one layer per branch, we restrain layer/branch validation to multi-branch graphs.
-	if (ValidationInfo.BranchCount > 1)
-	{
-		// We can run into the scenario where the users have given layers the same name, so layer_name token won't help differentiate.
-		// To resolve this, we look to see if there's multiple branches with the same layer name, and if so we force the branch name into the token too.
-		if (ValidationInfo.LayerCount < ValidationInfo.BranchCount)
-		{
-			UE::MoviePipeline::ConformOutputFormatStringToken(FileNameFormatString, TEXT("{branch_name}"), InParentNode->GetFName(), InRenderData.Key.RootBranchName);
-		}
-		else
-		{
-			// Otherwise, we separate each branch by its unique layer name.
-			UE::MoviePipeline::ConformOutputFormatStringToken(FileNameFormatString, TEXT("{layer_name}"), InParentNode->GetFName(), InRenderData.Key.RootBranchName);
-		}
-	}
-
-	// We only add the renderer name token if multiple (non-composited) renderers are present on the active branch (eg, in the case of optional PPMs).
-	if (ValidationInfo.ActiveBranchRendererCount > 1)
-	{
-		UE::MoviePipeline::ConformOutputFormatStringToken(FileNameFormatString, TEXT("{renderer_name}"), InParentNode->GetFName(), InRenderData.Key.RootBranchName);
-	}
-
-	// We only add the subresource token if a (non-composited) renderer on the active branch is producing more than one subresource (eg, in the case of optional PPMs).
-	if (ValidationInfo.ActiveRendererSubresourceCount > 1)
-	{
-		UE::MoviePipeline::ConformOutputFormatStringToken(FileNameFormatString, TEXT("{renderer_sub_name}"), InParentNode->GetFName(), InRenderData.Key.RootBranchName);
-	}
-
-	// ToDo: Add {camera_name} validation once relevant
+	// Insert tokens like {layer_name} as appropriate to make sure outputs don't clash with each other.
+	DisambiguateFilename(FileNameFormatString, InRawFrameData, InParentNode->GetFName(), InRenderData);
 
 	// Previous method is preserved for output frame number validation.
 	constexpr bool bIncludeRenderPass = false;
