MovieGraph: Fixed a crash when using the widget renderer/burn ins due to them trying to pull settings from objects that had gotten garbage collected.

Specifically tested:
Rendering with the default burn in, and rendering with a non-default burn-in class to ensure the correct class gets instantiated, verified that the burn-in information gets up updated.

diff --git a/Engine/Plugins/MovieScene/MovieRenderPipeline/Source/MovieRenderPipelineCore/Private/Graph/Nodes/MovieGraphBurnInNode.cpp b/Engine/Plugins/MovieScene/MovieRenderPipeline/Source/MovieRenderPipelineCore/Private/Graph/Nodes/MovieGraphBurnInNode.cpp
--- a/Engine/Plugins/MovieScene/MovieRenderPipeline/Source/MovieRenderPipelineCore/Private/Graph/Nodes/MovieGraphBurnInNode.cpp
+++ b/Engine/Plugins/MovieScene/MovieRenderPipeline/Source/MovieRenderPipelineCore/Private/Graph/Nodes/MovieGraphBurnInNode.cpp
@@ -84,11 +84,18 @@ void UMovieGraphBurnInNode::FMovieGraphBurnInPass::Setup(TWeakObjectPtr<UMovieGr
 	FMovieGraphWidgetPass::Setup(InRenderer, InLayer);
 	
 	RenderDataIdentifier.SubResourceName = TEXT("burnin");
+
+	UMovieGraphBurnInNode* ParentNodeOnInitialization = Cast<UMovieGraphBurnInNode>(InLayer.RenderPassNode);
+	if (!ensureMsgf(ParentNodeOnInitialization, TEXT("FMovieGraphBurnInPass shouldn't exist without a parent node in the graph.")))
+	{
+		return;
+	}
+	CachedBurnInWidgetClass = ParentNodeOnInitialization->BurnInClass.TryLoadClass<UMovieGraphBurnInWidget>();
 }
 
-TSharedPtr<SWidget> UMovieGraphBurnInNode::FMovieGraphBurnInPass::GetWidget()
+TSharedPtr<SWidget> UMovieGraphBurnInNode::FMovieGraphBurnInPass::GetWidget(UMovieGraphWidgetRendererBaseNode* InNodeThisFrame)
 {
-	if (const TObjectPtr<UMovieGraphBurnInWidget> BurnInWidget = GetBurnInWidget())
+	if (const TObjectPtr<UMovieGraphBurnInWidget> BurnInWidget = GetBurnInWidget(InNodeThisFrame))
 	{
 		return BurnInWidget->TakeWidget();
 	}
@@ -96,7 +103,7 @@ TSharedPtr<SWidget> UMovieGraphBurnInNode::FMovieGraphBurnInPass::GetWidget()
 	return nullptr;
 }
 
-TObjectPtr<UMovieGraphBurnInWidget> UMovieGraphBurnInNode::FMovieGraphBurnInPass::GetBurnInWidget() const
+TObjectPtr<UMovieGraphBurnInWidget> UMovieGraphBurnInNode::FMovieGraphBurnInPass::GetBurnInWidget(UMovieGraphWidgetRendererBaseNode* InNodeThisFrame) const
 {
 	const UMovieGraphPipeline* Pipeline = Renderer->GetOwningGraph();
 
@@ -108,12 +115,19 @@ TObjectPtr<UMovieGraphBurnInWidget> UMovieGraphBurnInNode::FMovieGraphBurnInPass
 	}
 	
 	// The CDO contains the resources which are shared with all FMovieGraphBurnInPass instances
-	UMovieGraphBurnInNode* BurnInCDO = RenderPassNode->GetClass()->GetDefaultObject<UMovieGraphBurnInNode>();
+	UMovieGraphBurnInNode* BurnInCDO = InNodeThisFrame->GetClass()->GetDefaultObject<UMovieGraphBurnInNode>();
 	const TObjectPtr<UMovieGraphBurnInWidget> BurnInWidget = BurnInCDO->GetOrCreateBurnInWidget(LoadedBurnInClass, Pipeline->GetWorld());
 	if (!BurnInWidget)
 	{
-		const UMovieGraphBurnInNode* BurnInNode = CastChecked<UMovieGraphBurnInNode>(RenderPassNode);
-		UE_LOG(LogMovieRenderPipeline, Error, TEXT("Unable to load burn-in widget at path: %s"), *BurnInNode->BurnInClass.GetAssetPath().ToString());
+		UClass* BurnInWidgetClass = GetBurnInClass();
+		if (BurnInWidgetClass)
+		{
+			UE_LOG(LogMovieRenderPipeline, Error, TEXT("Unable to load burn-in widget at path: %s"), *BurnInWidgetClass->GetPathName());
+		}
+		else
+		{
+			UE_LOG(LogMovieRenderPipeline, Error, TEXT("Unable to load burn-in widget (nullptr)"));
+		}
 		return nullptr;
 	}
 
@@ -122,8 +136,10 @@ TObjectPtr<UMovieGraphBurnInWidget> UMovieGraphBurnInNode::FMovieGraphBurnInPass
 
 void UMovieGraphBurnInNode::FMovieGraphBurnInPass::Render(const FMovieGraphTraversalContext& InFrameTraversalContext, const FMovieGraphTimeStepData& InTimeData)
 {
+	UMovieGraphWidgetRendererBaseNode* ParentNodeThisFrame = GetParentNode(InTimeData.EvaluatedConfig);
+
 	// Update the widget with the latest frame information
-	if (const TObjectPtr<UMovieGraphBurnInWidget> BurnInWidget = GetBurnInWidget())
+	if (const TObjectPtr<UMovieGraphBurnInWidget> BurnInWidget = GetBurnInWidget(ParentNodeThisFrame))
 	{
 		UMovieGraphPipeline* Pipeline = Renderer->GetOwningGraph();
 		BurnInWidget->UpdateForGraph(Pipeline, Pipeline->GetTimeStepInstance()->GetCalculatedTimeData().EvaluatedConfig);
@@ -140,6 +156,5 @@ int32 UMovieGraphBurnInNode::FMovieGraphBurnInPass::GetCompositingSortOrder() co
 
 UClass* UMovieGraphBurnInNode::FMovieGraphBurnInPass::GetBurnInClass() const
 {
-	const UMovieGraphBurnInNode* BurnInNode = CastChecked<UMovieGraphBurnInNode>(RenderPassNode);
-	return BurnInNode->BurnInClass.TryLoadClass<UMovieGraphBurnInWidget>();
+	return CachedBurnInWidgetClass.Get();
 }
\ No newline at end of file
diff --git a/Engine/Plugins/MovieScene/MovieRenderPipeline/Source/MovieRenderPipelineCore/Private/Graph/Nodes/MovieGraphUIRendererNode.cpp b/Engine/Plugins/MovieScene/MovieRenderPipeline/Source/MovieRenderPipelineCore/Private/Graph/Nodes/MovieGraphUIRendererNode.cpp
--- a/Engine/Plugins/MovieScene/MovieRenderPipeline/Source/MovieRenderPipelineCore/Private/Graph/Nodes/MovieGraphUIRendererNode.cpp
+++ b/Engine/Plugins/MovieScene/MovieRenderPipeline/Source/MovieRenderPipelineCore/Private/Graph/Nodes/MovieGraphUIRendererNode.cpp
@@ -38,7 +38,7 @@ void UMovieGraphUIRendererNode::FMovieGraphUIPass::Setup(TWeakObjectPtr<UMovieGr
 	RenderDataIdentifier.SubResourceName = TEXT("ui");
 }
 
-TSharedPtr<SWidget> UMovieGraphUIRendererNode::FMovieGraphUIPass::GetWidget()
+TSharedPtr<SWidget> UMovieGraphUIRendererNode::FMovieGraphUIPass::GetWidget(UMovieGraphWidgetRendererBaseNode* InNodeThisFrame)
 {
 	const UMovieGraphPipeline* Pipeline = Renderer->GetOwningGraph();
 	const ULocalPlayer* LocalPlayer = Pipeline->GetWorld()->GetFirstLocalPlayerFromController();
diff --git a/Engine/Plugins/MovieScene/MovieRenderPipeline/Source/MovieRenderPipelineCore/Private/Graph/Nodes/MovieGraphWidgetRendererBaseNode.cpp b/Engine/Plugins/MovieScene/MovieRenderPipeline/Source/MovieRenderPipelineCore/Private/Graph/Nodes/MovieGraphWidgetRendererBaseNode.cpp
--- a/Engine/Plugins/MovieScene/MovieRenderPipeline/Source/MovieRenderPipelineCore/Private/Graph/Nodes/MovieGraphWidgetRendererBaseNode.cpp
+++ b/Engine/Plugins/MovieScene/MovieRenderPipeline/Source/MovieRenderPipelineCore/Private/Graph/Nodes/MovieGraphWidgetRendererBaseNode.cpp
@@ -73,11 +73,10 @@ void UMovieGraphWidgetRendererBaseNode::FMovieGraphWidgetPass::Setup(TWeakObject
 {
 	LayerData = InLayer;
 	Renderer = InRenderer;
-	RenderPassNode = CastChecked<UMovieGraphWidgetRendererBaseNode>(InLayer.RenderPassNode);
 
 	RenderDataIdentifier.RootBranchName = LayerData.BranchName;
 	RenderDataIdentifier.LayerName = LayerData.LayerName;
-	RenderDataIdentifier.RendererName = RenderPassNode->GetRendererName();
+	RenderDataIdentifier.RendererName = CastChecked<UMovieGraphWidgetRendererBaseNode>(InLayer.RenderPassNode)->GetRendererName();
 
 	// NOTE: Subclasses should specify the SubResourceName
 
@@ -90,10 +89,24 @@ void UMovieGraphWidgetRendererBaseNode::FMovieGraphWidgetPass::Teardown()
 	// Nothing to do
 }
 
+UMovieGraphWidgetRendererBaseNode* UMovieGraphWidgetRendererBaseNode::FMovieGraphWidgetPass::GetParentNode(UMovieGraphEvaluatedConfig* InConfig) const
+{
+	const bool bIncludeCDOs = true;
+	UMovieGraphWidgetRendererBaseNode* ParentNode = InConfig->GetSettingForBranch<UMovieGraphWidgetRendererBaseNode>(RenderDataIdentifier.RootBranchName, bIncludeCDOs);
+	if (!ensureMsgf(ParentNode, TEXT("FMovieGraphWidgetPass should not exist without parent node in graph.")))
+	{
+		return nullptr;
+	}
+
+	return ParentNode;
+}
+
+
 void UMovieGraphWidgetRendererBaseNode::FMovieGraphWidgetPass::Render(const FMovieGraphTraversalContext& InFrameTraversalContext, const FMovieGraphTimeStepData& InTimeData)
 {
 	const UMovieGraphPipeline* Pipeline = Renderer->GetOwningGraph();
-	
+	UMovieGraphWidgetRendererBaseNode* ParentNodeThisFrame = GetParentNode(InTimeData.EvaluatedConfig);
+
 	const FIntPoint OutputResolution = UMovieGraphBlueprintLibrary::GetEffectiveOutputResolution(InTimeData.EvaluatedConfig);
 	const int32 MaxResolution = GetMax2DTextureDimension();
 	if ((OutputResolution.X > MaxResolution) || (OutputResolution.Y > MaxResolution))
@@ -120,10 +133,10 @@ void UMovieGraphWidgetRendererBaseNode::FMovieGraphWidgetPass::Render(const FMov
 		TSharedPtr<UE::MovieGraph::IMovieGraphOutputMerger> OutputMerger = Pipeline->GetOutputMerger();
 		
 		// The CDO contains the resources which are shared with all pass instances
-		UMovieGraphWidgetRendererBaseNode* NodeCDO = RenderPassNode->GetClass()->GetDefaultObject<UMovieGraphWidgetRendererBaseNode>();
+		UMovieGraphWidgetRendererBaseNode* NodeCDO = ParentNodeThisFrame->GetClass()->GetDefaultObject<UMovieGraphWidgetRendererBaseNode>();
 		const TSharedPtr<SVirtualWindow> VirtualWindow = NodeCDO->GetOrCreateVirtualWindow(OutputResolution);
 
-		const TSharedPtr<SWidget> WidgetToRender = GetWidget();
+		const TSharedPtr<SWidget> WidgetToRender = GetWidget(ParentNodeThisFrame);
 		if (!WidgetToRender)
 		{
 			// The subclass implementation of GetWidget() is responsible for emitting error messages if getting the widget failed
@@ -139,7 +152,7 @@ void UMovieGraphWidgetRendererBaseNode::FMovieGraphWidgetPass::Render(const FMov
 
 		ENQUEUE_RENDER_COMMAND(WidgetRenderTargetResolveCommand)(
 			[InTimeData, InFrameTraversalContext, RenderTargetInitParams, RenderDataIdentifier = this->RenderDataIdentifier,
-			bComposite = RenderPassNode->bCompositeOntoFinalImage, CompositingSortOrder = GetCompositingSortOrder(),
+			bComposite = ParentNodeThisFrame->bCompositeOntoFinalImage, CompositingSortOrder = GetCompositingSortOrder(),
 			BackbufferRenderTarget, OutputMerger, OutputResolution](FRHICommandListImmediate& RHICmdList)
 			{
 				const FIntRect SourceRect = FIntRect(0, 0, BackbufferRenderTarget->GetSizeXY().X, BackbufferRenderTarget->GetSizeXY().Y);
diff --git a/Engine/Plugins/MovieScene/MovieRenderPipeline/Source/MovieRenderPipelineCore/Private/Graph/Nodes/MovieGraphWidgetRendererBaseNode.h b/Engine/Plugins/MovieScene/MovieRenderPipeline/Source/MovieRenderPipelineCore/Private/Graph/Nodes/MovieGraphWidgetRendererBaseNode.h
--- a/Engine/Plugins/MovieScene/MovieRenderPipeline/Source/MovieRenderPipelineCore/Private/Graph/Nodes/MovieGraphWidgetRendererBaseNode.h
+++ b/Engine/Plugins/MovieScene/MovieRenderPipeline/Source/MovieRenderPipelineCore/Private/Graph/Nodes/MovieGraphWidgetRendererBaseNode.h
@@ -1,4 +1,4 @@
-﻿// Copyright Epic Games, Inc. All Rights Reserved.
+// Copyright Epic Games, Inc. All Rights Reserved.
 
 #pragma once
 
@@ -38,10 +38,12 @@ protected:
 
 		virtual void Setup(TWeakObjectPtr<UMovieGraphDefaultRenderer> InRenderer, const FMovieGraphRenderPassLayerData& InLayer);
 		virtual void Teardown();
-		virtual TSharedPtr<SWidget> GetWidget() = 0;
+		virtual TSharedPtr<SWidget> GetWidget(UMovieGraphWidgetRendererBaseNode* InNodeThisFrame) = 0;
 		virtual void Render(const FMovieGraphTraversalContext& InFrameTraversalContext, const FMovieGraphTimeStepData& InTimeData);
 		virtual void GatherOutputPasses(TArray<FMovieGraphRenderDataIdentifier>& OutExpectedPasses) const;
 		virtual int32 GetCompositingSortOrder() const = 0;
+		virtual UMovieGraphWidgetRendererBaseNode* GetParentNode(UMovieGraphEvaluatedConfig* InConfig) const;
+
 
 	protected:
 		FMovieGraphRenderPassLayerData LayerData;
@@ -51,9 +53,6 @@ protected:
 
 		/** The renderer associated with this pass. */
 		TWeakObjectPtr<UMovieGraphDefaultRenderer> Renderer;
-
-		/** The node instance (not CDO) associated with this pass. */
-		TWeakObjectPtr<UMovieGraphWidgetRendererBaseNode> RenderPassNode;
 	};
 
 	/** Widget renderer shared with all FMovieGraphWidgetPass instances. */
diff --git a/Engine/Plugins/MovieScene/MovieRenderPipeline/Source/MovieRenderPipelineCore/Public/Graph/Nodes/MovieGraphBurnInNode.h b/Engine/Plugins/MovieScene/MovieRenderPipeline/Source/MovieRenderPipelineCore/Public/Graph/Nodes/MovieGraphBurnInNode.h
--- a/Engine/Plugins/MovieScene/MovieRenderPipeline/Source/MovieRenderPipelineCore/Public/Graph/Nodes/MovieGraphBurnInNode.h
+++ b/Engine/Plugins/MovieScene/MovieRenderPipeline/Source/MovieRenderPipelineCore/Public/Graph/Nodes/MovieGraphBurnInNode.h
@@ -51,7 +51,7 @@ protected:
 	struct FMovieGraphBurnInPass final : public FMovieGraphWidgetPass
 	{
 		virtual void Setup(TWeakObjectPtr<UMovieGraphDefaultRenderer> InRenderer, const FMovieGraphRenderPassLayerData& InLayer) override;
-		virtual TSharedPtr<SWidget> GetWidget() override;
+		virtual TSharedPtr<SWidget> GetWidget(UMovieGraphWidgetRendererBaseNode* InNodeThisFrame) override;
 		virtual void Render(const FMovieGraphTraversalContext& InFrameTraversalContext, const FMovieGraphTimeStepData& InTimeData) override;
 		virtual int32 GetCompositingSortOrder() const override;
 
@@ -59,7 +59,8 @@ protected:
 		UClass* GetBurnInClass() const;
 
 	private:
-		TObjectPtr<UMovieGraphBurnInWidget> GetBurnInWidget() const;
+		TObjectPtr<UMovieGraphBurnInWidget> GetBurnInWidget(UMovieGraphWidgetRendererBaseNode* InNodeThisFrame) const;
+		TSubclassOf<UMovieGraphBurnInWidget> CachedBurnInWidgetClass;
 	};
 
 private:
diff --git a/Engine/Plugins/MovieScene/MovieRenderPipeline/Source/MovieRenderPipelineCore/Public/Graph/Nodes/MovieGraphUIRendererNode.h b/Engine/Plugins/MovieScene/MovieRenderPipeline/Source/MovieRenderPipelineCore/Public/Graph/Nodes/MovieGraphUIRendererNode.h
--- a/Engine/Plugins/MovieScene/MovieRenderPipeline/Source/MovieRenderPipelineCore/Public/Graph/Nodes/MovieGraphUIRendererNode.h
+++ b/Engine/Plugins/MovieScene/MovieRenderPipeline/Source/MovieRenderPipelineCore/Public/Graph/Nodes/MovieGraphUIRendererNode.h
@@ -33,7 +33,7 @@ protected:
 	struct FMovieGraphUIPass final : public FMovieGraphWidgetPass
 	{
 		virtual void Setup(TWeakObjectPtr<UMovieGraphDefaultRenderer> InRenderer, const FMovieGraphRenderPassLayerData& InLayer) override;
-		virtual TSharedPtr<SWidget> GetWidget() override;
+		virtual TSharedPtr<SWidget> GetWidget(UMovieGraphWidgetRendererBaseNode* InNodeThisFrame) override;
 		virtual int32 GetCompositingSortOrder() const override;
 	};
 	
