MovieGraph: Fix corrupt 32->8 bit non-srgb quantization and move to 64-bit indexing for parity with TArray64 uses.

diff --git a/Engine/Plugins/MovieScene/MovieRenderPipeline/Source/MovieRenderPipelineCore/Private/MoviePipelineImageQuantization.cpp b/Engine/Plugins/MovieScene/MovieRenderPipeline/Source/MovieRenderPipelineCore/Private/MoviePipelineImageQuantization.cpp
--- a/Engine/Plugins/MovieScene/MovieRenderPipeline/Source/MovieRenderPipelineCore/Private/MoviePipelineImageQuantization.cpp
+++ b/Engine/Plugins/MovieScene/MovieRenderPipeline/Source/MovieRenderPipelineCore/Private/MoviePipelineImageQuantization.cpp
@@ -19,7 +19,7 @@ DECLARE_CYCLE_STAT(TEXT("STAT_MoviePipeline_ImageQuantization"), STAT_ImageQuant
 
 
 template<class TToColorBitDepthType, typename TFromColorBitDepthType, typename TColorChannelNumericType>
-TArray<TToColorBitDepthType> ConvertLinearToLinearBitDepth(TFromColorBitDepthType* InColor, const int32 InCount)
+TArray<TToColorBitDepthType> ConvertLinearToLinearBitDepth(TFromColorBitDepthType* InColor, const int64 InCount)
 {
 	if constexpr (std::is_same_v<TToColorBitDepthType, TFromColorBitDepthType>)
 	{
@@ -33,7 +33,7 @@ TArray<TToColorBitDepthType> ConvertLinearToLinearBitDepth(TFromColorBitDepthTyp
 	// do clamp in int in case input float is out of [0,1] range
 	int MaxValueInt = TNumericLimits<TColorChannelNumericType>::Max();
 	float MaxValue = static_cast<float>(MaxValueInt);
-	for (int32 PixelIndex = 0; PixelIndex < InCount; PixelIndex++)
+	for (int64 PixelIndex = 0; PixelIndex < InCount; PixelIndex++)
 	{
 		// Avoid the bounds checking of TArray[]
 		TToColorBitDepthType* OutColor = &OutsRGBData.GetData()[PixelIndex];
@@ -165,7 +165,7 @@ static TArray<FFloat16> GenerateSRGBTableFloat16to16()
 	return OutsRGBTable;
 }
 
-static TArray<FColor> ConvertLinearTosRGB8bppViaLookupTable(FFloat16Color* InColor, const int32 InCount)
+static TArray<FColor> ConvertLinearTosRGB8bppViaLookupTable(FFloat16Color* InColor, const int64 InCount)
 {
 	TArray<float> sRGBTable = GenerateSRGBTableFloat16toFloat();
 
@@ -198,11 +198,11 @@ static TArray<FColor> ConvertLinearTosRGB8bppViaLookupTable(FFloat16Color* InCol
 	ParallelFor(Loops,
 		[&](int32 LoopIndex)
 		{
-			const int32 Start = LoopIndex * BatchSize;
-			const int32 End = FMath::Min((LoopIndex + 1) * BatchSize, InCount);
+			const int64 Start = LoopIndex * BatchSize;
+			const int64 End = FMath::Min((LoopIndex + 1) * BatchSize, InCount);
 
 
-			for (int32 PixelIndex = Start; PixelIndex < End; PixelIndex++)
+			for (int64 PixelIndex = Start; PixelIndex < End; PixelIndex++)
 			{
 				// Avoid the bounds checking of TArray[]
 				FColor* OutColor = &OutData[PixelIndex];
@@ -218,7 +218,7 @@ static TArray<FColor> ConvertLinearTosRGB8bppViaLookupTable(FFloat16Color* InCol
 	return OutsRGBData;
 }
 
-static TArray<FColor> ConvertLinearTosRGB8bppViaLookupTable(FLinearColor* InColor, const int32 InCount)
+static TArray<FColor> ConvertLinearTosRGB8bppViaLookupTable(FLinearColor* InColor, const int64 InCount)
 {
 	TArray<uint8> sRGBTable = GenerateSRGBTable(4096); 
 
@@ -228,7 +228,7 @@ static TArray<FColor> ConvertLinearTosRGB8bppViaLookupTable(FLinearColor* InColo
 
 	int32 TableUpperBound = sRGBTable.Num() - 1;
 
-	for (int32 PixelIndex = 0; PixelIndex < InCount; PixelIndex++)
+	for (int64 PixelIndex = 0; PixelIndex < InCount; PixelIndex++)
 	{
 		// Avoid the bounds checking of TArray[]
 		FColor* OutColor = &OutsRGBData.GetData()[PixelIndex];
@@ -248,7 +248,7 @@ static TArray<FColor> ConvertLinearTosRGB8bppViaLookupTable(FLinearColor* InColo
 	return OutsRGBData;
 }
 
-static TArray<FFloat16Color> ConvertLinearTosRGB16bppViaLookupTable(FFloat16Color* InColor, const int32 InCount)
+static TArray<FFloat16Color> ConvertLinearTosRGB16bppViaLookupTable(FFloat16Color* InColor, const int64 InCount)
 {
 	TArray<FFloat16> sRGBTable = GenerateSRGBTableFloat16to16();
 
@@ -260,7 +260,7 @@ static TArray<FFloat16Color> ConvertLinearTosRGB16bppViaLookupTable(FFloat16Colo
 	FFloat16Color* OutData = static_cast<FFloat16Color*>(OutsRGBData.GetData());
 
 	QUICK_SCOPE_CYCLE_COUNTER(STAT_ImageQuant_ApplysRGB);
-	for (int32 PixelIndex = 0; PixelIndex < InCount; PixelIndex++)
+	for (int64 PixelIndex = 0; PixelIndex < InCount; PixelIndex++)
 	{
 		// Avoid the bounds checking of TArray[]
 		FFloat16Color* OutColor = &OutData[PixelIndex];
@@ -299,12 +299,12 @@ static TUniquePtr<FImagePixelData> QuantizePixelDataTo8bpp(const FImagePixelData
 
 		if (bConvertToSrgb)
 		{
-			TArray<FColor> sRGBEncoded = ConvertLinearTosRGB8bppViaLookupTable((FFloat16Color*)SrcRawDataPtr, RawSize.X * RawSize.Y);
+			TArray<FColor> sRGBEncoded = ConvertLinearTosRGB8bppViaLookupTable((FFloat16Color*)SrcRawDataPtr, RawSize.X * (int64)RawSize.Y);
 			QuantizedPixelData = MakeUnique<TImagePixelData<FColor>>(RawSize, TArray64<FColor>(MoveTemp(sRGBEncoded)), InPayload);
 		}
 		else
 		{
-			TArray<FColor> sRGBEncoded = ConvertLinearToLinearBitDepth<FColor, FFloat16Color, uint8>((FFloat16Color*)SrcRawDataPtr, RawSize.X * RawSize.Y);
+			TArray<FColor> sRGBEncoded = ConvertLinearToLinearBitDepth<FColor, FFloat16Color, uint8>((FFloat16Color*)SrcRawDataPtr, RawSize.X * (int64)RawSize.Y);
 			QuantizedPixelData = MakeUnique<TImagePixelData<FColor>>(RawSize, TArray64<FColor>(MoveTemp(sRGBEncoded)), InPayload);
 			
 		}
@@ -318,12 +318,12 @@ static TUniquePtr<FImagePixelData> QuantizePixelDataTo8bpp(const FImagePixelData
 		
 		if (bConvertToSrgb)
 		{
-			TArray<FColor> sRGBEncoded = ConvertLinearTosRGB8bppViaLookupTable((FLinearColor*)SrcRawDataPtr, RawSize.X * RawSize.Y);
+			TArray<FColor> sRGBEncoded = ConvertLinearTosRGB8bppViaLookupTable((FLinearColor*)SrcRawDataPtr, RawSize.X * (int64)RawSize.Y);
 			QuantizedPixelData = MakeUnique<TImagePixelData<FColor>>(RawSize, TArray64<FColor>(MoveTemp(sRGBEncoded)), InPayload);
 		}
 		else
 		{
-			TArray<FColor> sRGBEncoded = ConvertLinearToLinearBitDepth<FColor, FFloat16Color, uint8>((FFloat16Color*)SrcRawDataPtr, RawSize.X * RawSize.Y);
+			TArray<FColor> sRGBEncoded = ConvertLinearToLinearBitDepth<FColor, FLinearColor, uint8>((FLinearColor*)SrcRawDataPtr, RawSize.X * (int64)RawSize.Y);
 			QuantizedPixelData = MakeUnique<TImagePixelData<FColor>>(RawSize, TArray64<FColor>(MoveTemp(sRGBEncoded)), InPayload);
 		}
 		break;
@@ -355,12 +355,12 @@ static TUniquePtr<FImagePixelData> QuantizePixelDataTo16bpp(const FImagePixelDat
 
 		if (bConvertToSrgb)
 		{
-			TArray<FFloat16Color> sRGBEncoded = ConvertLinearTosRGB16bppViaLookupTable((FFloat16Color*)SrcRawDataPtr, RawSize.X * RawSize.Y);
+			TArray<FFloat16Color> sRGBEncoded = ConvertLinearTosRGB16bppViaLookupTable((FFloat16Color*)SrcRawDataPtr, RawSize.X * (int64)RawSize.Y);
 			QuantizedPixelData = MakeUnique<TImagePixelData<FFloat16Color>>(RawSize, TArray64<FFloat16Color>(MoveTemp(sRGBEncoded)), InPayload);
 		}
 		else
 		{
-			QuantizedPixelData = MakeUnique<TImagePixelData<FFloat16Color>>(RawSize, TArray64<FFloat16Color>((FFloat16Color*)SrcRawDataPtr, RawSize.X * RawSize.Y), InPayload);
+			QuantizedPixelData = MakeUnique<TImagePixelData<FFloat16Color>>(RawSize, TArray64<FFloat16Color>((FFloat16Color*)SrcRawDataPtr, RawSize.X * (int64)RawSize.Y), InPayload);
 		}
 		break;
 	}
@@ -372,13 +372,13 @@ static TUniquePtr<FImagePixelData> QuantizePixelDataTo16bpp(const FImagePixelDat
 
 		// FColor is assumed to be in sRGB while FFloat16Color is assumed to be linear so we need to convert back.
 		TArray64<FFloat16Color> OutColors;
-		OutColors.SetNumUninitialized(RawSize.X * RawSize.Y);
+		OutColors.SetNumUninitialized(RawSize.X * (int64)RawSize.Y);
 
 		TArray<float> LookupTable = GenerateInverseSRGBTable(256);
 		const float* LookupTablePtr = LookupTable.GetData();
 		
 		const FColor* DataPtr = (FColor*)(SrcRawDataPtr);
-		for (int64 Index = 0; Index < (RawSize.X * RawSize.Y); Index++)
+		for (int64 Index = 0; Index < (RawSize.X * (int64)RawSize.Y); Index++)
 		if (bConvertToSrgb)
 		{
 			OutColors[Index].R = LookupTablePtr[DataPtr[Index].R];
