MoviePipeline: Expose shot version number to Blueprints.

Specifically tested:
1) A sequence with 5 shots that has Auto Version set to false, and then using Saved/MovieRenders/{version} as the filename.
2) The above, but with Auto Version pinned to a specific version
3) Using Saved/MovieRenders/{shot_name}/{version} with and without version pinning
4) The above, but enabling Auto-version and rendering some shots but not all of them (which creates unique version tokens within the shot folders)
5) All of the above in Movie Graph too.

diff --git a/Engine/Plugins/MovieScene/MovieRenderPipeline/Source/MovieRenderPipelineCore/Private/Graph/MovieGraphBlueprintLibrary.cpp b/Engine/Plugins/MovieScene/MovieRenderPipeline/Source/MovieRenderPipelineCore/Private/Graph/MovieGraphBlueprintLibrary.cpp
--- a/Engine/Plugins/MovieScene/MovieRenderPipeline/Source/MovieRenderPipelineCore/Private/Graph/MovieGraphBlueprintLibrary.cpp
+++ b/Engine/Plugins/MovieScene/MovieRenderPipeline/Source/MovieRenderPipelineCore/Private/Graph/MovieGraphBlueprintLibrary.cpp
@@ -342,6 +342,28 @@ int32 UMovieGraphBlueprintLibrary::ResolveVersionNumber(FMovieGraphFilenameResol
 	return HighestVersion + (bGetNextVersion ? 1 : 0);
 }
 
+int32 UMovieGraphBlueprintLibrary::GetCurrentVersionNumber(const UMovieGraphPipeline* InMovieGraphPipeline)
+{
+	// This is effectively identical to UMoviePipelineBlueprintLibrary::GetCurrentVersionNumber (since the
+	// version is stored on the shot structure which is shared between the two pipelines) but provided here
+	// via a UMovieGraphPipeline specific pointer for consistency in fetching values in the UI to go along
+	// with the other functions here get things like focal length, etc.
+	if (!InMovieGraphPipeline)
+	{
+		FFrame::KismetExecutionMessage(TEXT("Cannot get version number from null pipeline!"), ELogVerbosity::Error);
+		return 0;
+	}
+	
+	int32 CurrentShotIndex = InMovieGraphPipeline->GetCurrentShotIndex();
+	if(!InMovieGraphPipeline->GetActiveShotList().IsValidIndex(CurrentShotIndex))
+	{
+		FFrame::KismetExecutionMessage(TEXT("No shot is currently active to get the version number from."), ELogVerbosity::Error);
+		return 0;
+	}
+
+	return InMovieGraphPipeline->GetActiveShotList()[CurrentShotIndex]->ShotInfo.VersionNumber;
+}
+
 FIntPoint UMovieGraphBlueprintLibrary::GetEffectiveOutputResolution(UMovieGraphEvaluatedConfig* InEvaluatedGraph)
 {
 	if (!InEvaluatedGraph)
@@ -629,4 +651,19 @@ bool UMovieGraphBlueprintLibrary::IsNamedResolutionValid(const FName& InResoluti
 FMovieGraphNamedResolution UMovieGraphBlueprintLibrary::NamedResolutionFromSize(const int32 InResX, const int32 InResY)
 {
 	return FMovieGraphNamedResolution(FMovieGraphNamedResolution::CustomEntryName, FIntPoint(InResX, InResY), FString());
+}
+
+UMoviePipelineExecutorShot* UMovieGraphBlueprintLibrary::GetCurrentExecutorShot(const UMovieGraphPipeline* InMoviePipeline)
+{
+	if (InMoviePipeline)
+	{
+		const TArray<UMoviePipelineExecutorShot*>& ActiveShotList = InMoviePipeline->GetActiveShotList();
+		int32 CurrentShotIndex = InMoviePipeline->GetCurrentShotIndex();
+		if (ActiveShotList.IsValidIndex(CurrentShotIndex))
+		{
+			return ActiveShotList[CurrentShotIndex];
+		}
+	}
+
+	return nullptr;
 }
\ No newline at end of file
diff --git a/Engine/Plugins/MovieScene/MovieRenderPipeline/Source/MovieRenderPipelineCore/Private/MoviePipelineBlueprintLibrary.cpp b/Engine/Plugins/MovieScene/MovieRenderPipeline/Source/MovieRenderPipelineCore/Private/MoviePipelineBlueprintLibrary.cpp
--- a/Engine/Plugins/MovieScene/MovieRenderPipeline/Source/MovieRenderPipelineCore/Private/MoviePipelineBlueprintLibrary.cpp
+++ b/Engine/Plugins/MovieScene/MovieRenderPipeline/Source/MovieRenderPipelineCore/Private/MoviePipelineBlueprintLibrary.cpp
@@ -882,6 +882,24 @@ int32 UMoviePipelineBlueprintLibrary::ResolveVersionNumber(FMoviePipelineFilenam
 	return 0;
 }
 
+int32 UMoviePipelineBlueprintLibrary::GetCurrentVersionNumber(const UMoviePipeline* InMoviePipeline)
+{
+	if (!InMoviePipeline)
+	{
+		FFrame::KismetExecutionMessage(TEXT("Cannot get version number from null pipeline!"), ELogVerbosity::Error);
+		return 0;
+	}
+	
+	int32 CurrentShotIndex = InMoviePipeline->GetCurrentShotIndex();
+	if(!InMoviePipeline->GetActiveShotList().IsValidIndex(CurrentShotIndex))
+	{
+		FFrame::KismetExecutionMessage(TEXT("No shot is currently active to get the version number from."), ELogVerbosity::Error);
+		return 0;
+	}
+
+	return InMoviePipeline->GetActiveShotList()[CurrentShotIndex]->ShotInfo.VersionNumber;
+}
+
 FIntPoint UMoviePipelineBlueprintLibrary::GetEffectiveOutputResolution(UMoviePipelinePrimaryConfig* InPrimaryConfig, UMoviePipelineExecutorShot* InPipelineExecutorShot)
 {
 	if (InPrimaryConfig && InPipelineExecutorShot)
diff --git a/Engine/Plugins/MovieScene/MovieRenderPipeline/Source/MovieRenderPipelineCore/Public/Graph/MovieGraphBlueprintLibrary.h b/Engine/Plugins/MovieScene/MovieRenderPipeline/Source/MovieRenderPipelineCore/Public/Graph/MovieGraphBlueprintLibrary.h
--- a/Engine/Plugins/MovieScene/MovieRenderPipeline/Source/MovieRenderPipelineCore/Public/Graph/MovieGraphBlueprintLibrary.h
+++ b/Engine/Plugins/MovieScene/MovieRenderPipeline/Source/MovieRenderPipelineCore/Public/Graph/MovieGraphBlueprintLibrary.h
@@ -52,6 +52,15 @@ public:
 	UFUNCTION(BlueprintCallable, Category = "Movie Graph")
 	static int32 ResolveVersionNumber(FMovieGraphFilenameResolveParams InParams, const bool bGetNextVersion = true);
 
+	/**
+	* Retrieves the cached version number calculated for the current shot, which depends on where the version token was used in the File Name Output
+	* ie: If {version} comes before {shot_name} then all shots will use the same version number, but if it comes afterwards then each shot may
+	* have a different version (which is the highest number found of that particular shot). This function should retrieve what is used in the
+	* filename writing step either way.
+	*/
+	UFUNCTION(BlueprintPure, Category = "Movie Graph")
+	static int32 GetCurrentVersionNumber(const UMovieGraphPipeline* InMovieGraphPipeline);
+
 	/**
 	* In case of overscan percentage being higher than 0, additional pixels are rendered. This function returns the resolution with overscan taken into account.
 	* @param	InEvaluatedGraph	- The evaluated graph that will provide context for resolving the resolution
@@ -219,4 +228,10 @@ public:
 	UFUNCTION(BlueprintPure, Category = "Movie Graph")
 	static FMovieGraphNamedResolution NamedResolutionFromSize(const int32 InResX, const int32 InResY);
 
+	/**
+	* Gets the current shot being rendered by the graph (could be nullptr if rendering hasn't started or has moved to Finalize!)
+	*/
+	UFUNCTION(BlueprintPure, Category = "Movie Graph")
+	static UMoviePipelineExecutorShot* GetCurrentExecutorShot(const UMovieGraphPipeline* InMoviePipeline);
+
 };
\ No newline at end of file
diff --git a/Engine/Plugins/MovieScene/MovieRenderPipeline/Source/MovieRenderPipelineCore/Public/MoviePipelineBlueprintLibrary.h b/Engine/Plugins/MovieScene/MovieRenderPipeline/Source/MovieRenderPipelineCore/Public/MoviePipelineBlueprintLibrary.h
--- a/Engine/Plugins/MovieScene/MovieRenderPipeline/Source/MovieRenderPipelineCore/Public/MoviePipelineBlueprintLibrary.h
+++ b/Engine/Plugins/MovieScene/MovieRenderPipeline/Source/MovieRenderPipelineCore/Public/MoviePipelineBlueprintLibrary.h
@@ -138,6 +138,15 @@ public:
 	UFUNCTION(BlueprintCallable, Category = "Movie Render Pipeline")
 	static int32 ResolveVersionNumber(FMoviePipelineFilenameResolveParams InParams, bool bGetNextVersion = true);
 
+	/**
+	* Retrieves the cached version number calculated for the current shot, which depends on where the version token was used in the File Name Output
+	* ie: If {version} comes before {shot_name} then all shots will use the same version number, but if it comes afterwards then each shot may
+	* have a different version (which is the highest number found of that particular shot). This function should retrieve what is used in the
+	* filename writing step either way.
+	*/
+	UFUNCTION(BlueprintPure, Category = "Movie Render Pipeline")
+	static int32 GetCurrentVersionNumber(const UMoviePipeline* InMoviePipeline);
+	
 	/** In case of Overscan percentage being higher than 0 we render additional pixels. This function returns the resolution with overscan taken into account. */
 	UFUNCTION(BlueprintCallable, Category = "Movie Render Pipeline")
 	static FIntPoint GetEffectiveOutputResolution(UMoviePipelinePrimaryConfig* InPrimaryConfig, UMoviePipelineExecutorShot* InPipelineExecutorShot);
