MovieGraph: Actor-reference collections no longer accept transient actors (ie: Spawnables) in their drag-and-drop function.

Specifically tested:
Drag and drop a single spawnable (blocked)
Drag and drop multiple items that contain a spawnable (allowed, but spawnable is removed from the actual list created)
Drag and drop multiple items that don't contain a spawnable
Drag and drop a single item that isn't a spawnable
The above tests but with "Add Selected in Outliner" inside the dropdown menu.

diff --git a/Engine/Plugins/MovieScene/MovieRenderPipeline/Source/MovieRenderPipelineCore/Private/Graph/MovieGraphRenderLayerSubsystem.cpp b/Engine/Plugins/MovieScene/MovieRenderPipeline/Source/MovieRenderPipelineCore/Private/Graph/MovieGraphRenderLayerSubsystem.cpp
--- a/Engine/Plugins/MovieScene/MovieRenderPipeline/Source/MovieRenderPipelineCore/Private/Graph/MovieGraphRenderLayerSubsystem.cpp
+++ b/Engine/Plugins/MovieScene/MovieRenderPipeline/Source/MovieRenderPipelineCore/Private/Graph/MovieGraphRenderLayerSubsystem.cpp
@@ -602,26 +602,44 @@ TArray<TSharedRef<SWidget>> UMovieGraphConditionGroupQuery_Actor::GetWidgets()
 		ListDataSource.Add(MakeShared<TSoftObjectPtr<AActor>>(Actor));
 	}
 
+	auto GetValidActorsFromOperation = [](TSharedPtr<FDragDropOperation> InOperation, TArray<AActor*>& OutActors) {
+
+		// Support dragging both actors and folders from the Outliner (dragging a folder will add all actors in the folder)
+		if (InOperation->IsOfType<FActorDragDropGraphEdOp>())
+		{
+			const TSharedPtr<FActorDragDropGraphEdOp> ActorOperation = StaticCastSharedPtr<FActorDragDropGraphEdOp>(InOperation);
+			Algo::Transform(ActorOperation->Actors, OutActors, [](const TWeakObjectPtr<AActor>& Actor) { return Actor.IsValid() ? Actor.Get() : nullptr; });
+		}
+
+		if (InOperation->IsOfType<FSceneOutlinerDragDropOp>())
+		{
+			const TSharedPtr<FSceneOutlinerDragDropOp> SceneOperation = StaticCastSharedPtr<FSceneOutlinerDragDropOp>(InOperation);
+			UE::MovieGraph::Private::GetActorsFromSceneDragDropOp(SceneOperation, OutActors);
+		}
+
+		// Prevent any transient actors (ie: spawnables) from being added as they won't exist later
+		for (int32 Index = OutActors.Num() - 1; Index >= 0; Index--)
+		{
+			if(OutActors[Index]->HasAnyFlags(RF_Transient))
+			{
+				OutActors.RemoveAt(Index);
+			}
+		}
+	};
+
 	Widgets.Add(
 		SNew(SDropTarget)
-		.OnAllowDrop_Lambda([](TSharedPtr<FDragDropOperation> InDragOperation)
+		.OnAllowDrop_Lambda([GetValidActorsFromOperation](TSharedPtr<FDragDropOperation> InDragOperation)
 		{
-			// Support dragging both actors and folders from the Outliner (dragging a folder will add all actors in the folder)
-			return InDragOperation->IsOfType<FActorDragDropGraphEdOp>() || InDragOperation->IsOfType<FSceneOutlinerDragDropOp>();
+				TArray<AActor*> DroppedActors;
+				GetValidActorsFromOperation(InDragOperation, DroppedActors);
+
+				return DroppedActors.Num() > 0;
 		})
-		.OnDropped_Lambda([this](const FGeometry& Geometry, const FDragDropEvent& DragDropEvent)
+		.OnDropped_Lambda([this, GetValidActorsFromOperation](const FGeometry& Geometry, const FDragDropEvent& DragDropEvent)
 		{
 			TArray<AActor*> DroppedActors;
-			
-			if (const TSharedPtr<FActorDragDropGraphEdOp> ActorOperation = DragDropEvent.GetOperationAs<FActorDragDropGraphEdOp>())
-			{
-				Algo::Transform(ActorOperation->Actors, DroppedActors, [](const TWeakObjectPtr<AActor>& Actor) { return Actor.IsValid() ? Actor.Get() : nullptr; });
-			}
-
-			if (const TSharedPtr<FSceneOutlinerDragDropOp> SceneOperation = DragDropEvent.GetOperationAs<FSceneOutlinerDragDropOp>())
-			{
-				UE::MovieGraph::Private::GetActorsFromSceneDragDropOp(SceneOperation, DroppedActors);
-			}
+			GetValidActorsFromOperation(DragDropEvent.GetOperation(), DroppedActors);
 
 			const FMovieGraphConditionGroupQueryContentsChanged OnAddFinished = nullptr;
 			AddActors(DroppedActors, OnAddFinished);
@@ -667,6 +685,7 @@ TSharedRef<SWidget> UMovieGraphConditionGroupQuery_Actor::GetAddMenuContents(con
 	SceneOutlinerInitOptions.bShowCreateNewFolder = false;
 	SceneOutlinerInitOptions.bFocusSearchBoxWhenOpened = true;
 
+
 	// Show the custom "Add" column, as well as the built-in name/label/type columns
 	SceneOutlinerInitOptions.ColumnMap.Add(
 		FActorSelectionColumn::GetID(),
@@ -704,11 +723,32 @@ TSharedRef<SWidget> UMovieGraphConditionGroupQuery_Actor::GetAddMenuContents(con
 					TArray<AActor*> SelectedActors;
 					GEditor->GetSelectedActors()->GetSelectedObjects<AActor>(SelectedActors);
 
+					for (int32 Index = SelectedActors.Num() - 1; Index >= 0; Index--)
+					{
+						if(SelectedActors[Index]->HasAnyFlags(RF_Transient))
+						{
+							SelectedActors.RemoveAt(Index);
+						}
+					}
+
 					AddActors(SelectedActors, OnAddFinished);
 				}),
 				FCanExecuteAction::CreateLambda([]()
 				{
-					return GEditor->GetSelectedActors()->Num() > 0;
+					// Assume we have only transient actors until we prove we don't.
+					bool bHasNonTransientActorsSelected = false;
+					TArray<AActor*> SelectedActors;
+					GEditor->GetSelectedActors()->GetSelectedObjects<AActor>(SelectedActors);
+
+					for (AActor* SelectedActor : SelectedActors)
+					{
+						if(!SelectedActor->HasAnyFlags(RF_Transient))
+						{
+							bHasNonTransientActorsSelected = true;
+							break;
+						}
+					}
+					return bHasNonTransientActorsSelected;
 				})
 			)
 		);
