MoviePipeline: Fixed some 32-bit integer overflows when using extremely large renders.

Specifically tested:
Rendered a 64,000 x 64,000 image to .exr file with Anti-Aliasing set to None, warm-up frames and Engine Warm Up set to Zero, and High Res Tiles set to 32.

diff --git a/Engine/Plugins/MovieScene/MovieRenderPipeline/Source/MovieRenderPipelineCore/Private/MovieRenderOverlappedImage.cpp b/Engine/Plugins/MovieScene/MovieRenderPipeline/Source/MovieRenderPipelineCore/Private/MovieRenderOverlappedImage.cpp
--- a/Engine/Plugins/MovieScene/MovieRenderPipeline/Source/MovieRenderPipelineCore/Private/MovieRenderOverlappedImage.cpp
+++ b/Engine/Plugins/MovieScene/MovieRenderPipeline/Source/MovieRenderPipelineCore/Private/MovieRenderOverlappedImage.cpp
@@ -24,7 +24,7 @@ void FImageOverlappedPlane::Init(FIntPoint InSize)
 	Size = InSize;
 
 	// Don't initialize now as we'll call ZeroPlane later when appropriate to reset the values.
-	ChannelData.SetNumUninitialized(Size.X * Size.Y);
+	ChannelData.SetNumUninitialized(Size.X * (int64)Size.Y);
 }
 
 void FImageOverlappedPlane::ZeroPlane()
@@ -49,7 +49,7 @@ void FImageOverlappedPlane::AccumulateSinglePlane(const TArray64<float>& InRawDa
 {
 	SCOPE_CYCLE_COUNTER(STAT_AccumulateSinglePlane);
 
-	check(InRawData.Num() == InSize.X * InSize.Y);
+	check(InRawData.Num() == InSize.X * (int64)InSize.Y);
 	check(WeightDataX.Num() == InSize.X);
 	check(WeightDataY.Num() == InSize.Y);
 
@@ -106,11 +106,11 @@ void FImageOverlappedPlane::AccumulateSinglePlane(const TArray64<float>& InRawDa
 						{
 							for (int32 OffsetX = 0; OffsetX < 2; OffsetX++)
 							{
-								int32 SrcX = FMath::Max<int32>(CurrX - 1 + OffsetX,0);
-								int32 SrcY = FMath::Max<int32>(CurrY - 1 + OffsetY,0);
+								int64 SrcX = FMath::Max<int64>(CurrX - 1 + OffsetX,0);
+								int64 SrcY = FMath::Max<int64>(CurrY - 1 + OffsetY,0);
 
 
-								float Val = InRawData[SrcY * InSize.X + SrcX];
+								float Val = InRawData[SrcY * (int64)InSize.X + SrcX];
 								
 								float WX = WeightDataX[SrcX];
 								float WY = WeightDataY[SrcY];
@@ -118,7 +118,7 @@ void FImageOverlappedPlane::AccumulateSinglePlane(const TArray64<float>& InRawDa
 
 								float Weight = BaseWeight * PixelWeight[1-OffsetY][1-OffsetX];
 
-								ChannelData[DstY * Size.X + DstX] += Weight * Val;
+								ChannelData[DstY * (int64)Size.X + DstX] += Weight * Val;
 							}
 						}
 					}
@@ -128,11 +128,11 @@ void FImageOverlappedPlane::AccumulateSinglePlane(const TArray64<float>& InRawDa
 		else
 		{
 			// Slow, reference version. This is the main one.
-			for (int32 CurrY = 0; CurrY < InSize.Y; CurrY++)
+			for (int64 CurrY = 0; CurrY < InSize.Y; CurrY++)
 			{
-				for (int32 CurrX = 0; CurrX < InSize.X; CurrX++)
+				for (int64 CurrX = 0; CurrX < InSize.X; CurrX++)
 				{
-					float Val = InRawData[CurrY * InSize.X + CurrX];
+					float Val = InRawData[CurrY * (int64)InSize.X + CurrX];
 					float WX = WeightDataX[CurrX];
 					float WY = WeightDataX[CurrY];
 
@@ -142,15 +142,15 @@ void FImageOverlappedPlane::AccumulateSinglePlane(const TArray64<float>& InRawDa
 					{
 						for (int32 OffsetX = 0; OffsetX < 2; OffsetX++)
 						{
-							int32 DstY = StartY + CurrY + OffsetY;
-							int32 DstX = StartX + CurrX + OffsetX;
+							int64 DstY = StartY + CurrY + OffsetY;
+							int64 DstX = StartX + CurrX + OffsetX;
 
 							float Weight = BaseWeight * PixelWeight[OffsetY][OffsetX];
 
 							if (DstX >= 0 && DstY >= 0 &&
 								DstX < Size.X && DstY < Size.Y)
 							{
-								ChannelData[DstY * Size.X + DstX] += Weight * Val;
+								ChannelData[DstY * (int64)Size.X + DstX] += Weight * Val;
 							}
 						}
 					}
@@ -186,8 +186,8 @@ void FImageOverlappedPlane::AccumulateSinglePlane(const TArray64<float>& InRawDa
 				int32 SrcY0 = FMath::Max<int32>(0,CurrY + (0 - 1));
 				int32 SrcY1 = FMath::Max<int32>(0,CurrY + (1 - 1));
 
-				const float * SrcLineRaw0 = &InRawData[SrcY0 * InSize.X];
-				const float * SrcLineRaw1 = &InRawData[SrcY1 * InSize.X];
+				const float * SrcLineRaw0 = &InRawData[SrcY0 * (int64)InSize.X];
+				const float * SrcLineRaw1 = &InRawData[SrcY1 * (int64)InSize.X];
 
 
 				const float * SrcLineWeight = WeightDataX.GetData();
@@ -195,7 +195,7 @@ void FImageOverlappedPlane::AccumulateSinglePlane(const TArray64<float>& InRawDa
 				const float RowWeight0 = WeightDataY[SrcY0];
 				const float RowWeight1 = WeightDataY[SrcY1];
 
-				float * DstLine = &ChannelData[DstY * Size.X];
+				float * DstLine = &ChannelData[DstY * (int64)Size.X];
 				for (int32 DstX = ActualDstX0; DstX < ActualDstX1; DstX++)
 				{
 					int32 CurrX = DstX - StartX;
@@ -232,8 +232,8 @@ void FImageOverlappedPlane::AccumulateSinglePlane(const TArray64<float>& InRawDa
 				int32 SrcY0 = FMath::Max<int32>(CurrY + (0 - 1), 0); // we need the max for the first pixel which could go off the edge
 				int32 SrcY1 = FMath::Max<int32>(CurrY + (1 - 1), 0); 
 
-				const float * SrcLineRaw0 = &InRawData[SrcY0 * InSize.X];
-				const float * SrcLineRaw1 = &InRawData[SrcY1 * InSize.X];
+				const float * SrcLineRaw0 = &InRawData[SrcY0 * (int64)InSize.X];
+				const float * SrcLineRaw1 = &InRawData[SrcY1 * (int64)InSize.X];
 
 				//const float * SrcLineWeight0 = &InWeightData[SrcY0 * InSize.X];
 				//const float * SrcLineWeight1 = &InWeightData[SrcY1 * InSize.X];
@@ -247,7 +247,7 @@ void FImageOverlappedPlane::AccumulateSinglePlane(const TArray64<float>& InRawDa
 				VectorRegister VecRowWeight0 = VectorSetFloat1(RowWeight0);
 				VectorRegister VecRowWeight1 = VectorSetFloat1(RowWeight1);
 
-				float * DstLine = &ChannelData[DstY * Size.X];
+				float * DstLine = &ChannelData[DstY * (int64)Size.X];
 
 				check(ActualDstX0 - StartX >= 0);
 
@@ -325,7 +325,7 @@ void FImageOverlappedPlane::AccumulateSinglePlane(const TArray64<float>& InRawDa
 
 				for (int32 IterExtra = 0; IterExtra < ActualWidthExtra4; IterExtra++)
 				{
-					int32 DstX = ActualDstX0 + NumSkipX + ActualWidthGroup4 * 4 + IterExtra;
+					int64 DstX = ActualDstX0 + NumSkipX + (int64)ActualWidthGroup4 * 4 + IterExtra;
 
 					int32 CurrX = DstX - StartX;
 
@@ -532,7 +532,7 @@ void FImageOverlappedAccumulator::AccumulatePixelData(const FImagePixelData& InP
 
 		{
 			// We always make room for all 4 channels (RGBA) even if we only accumulate 3 of them.
-			if (UnpackedDataStorage[0].Num() != RawSize.X * RawSize.Y)
+			if (UnpackedDataStorage[0].Num() != RawSize.X * (int64)RawSize.Y)
 			{
 				TRACE_CPUPROFILER_EVENT_SCOPE(AllocDataStorage);
 				ParallelFor(4, [&](int32 InChannelIndex)
@@ -542,15 +542,15 @@ void FImageOverlappedAccumulator::AccumulatePixelData(const FImagePixelData& InP
 			}
 
 			// Initialize the weight channel with 1's so that we can accumulate a bunch of 1 values together
-			if (WeightChannelStorage.Num() != RawSize.X * RawSize.Y)
+			if (WeightChannelStorage.Num() != RawSize.X * (int64)RawSize.Y)
 			{
 				TRACE_CPUPROFILER_EVENT_SCOPE(AllocWeightChannelStorage);
-				WeightChannelStorage.SetNumUninitialized(RawSize.X * RawSize.Y);
+				WeightChannelStorage.SetNumUninitialized(RawSize.X * (int64)RawSize.Y);
 				ParallelFor(RawSize.Y, [&](int32 InIndexY)
 				{
 					for (int32 IndexX = 0; IndexX < RawSize.X; IndexX++)
 					{
-						WeightChannelStorage[(InIndexY * RawSize.X) + IndexX] = 1.f;
+						WeightChannelStorage[(InIndexY * (int64)RawSize.X) + IndexX] = 1.f;
 					}
 				});
 			}
@@ -574,10 +574,10 @@ void FImageOverlappedAccumulator::AccumulatePixelData(const FImagePixelData& InP
 						int32 ChanReorder[4] = { 2, 1, 0, 3 };
 						for (int32 ChanIter = 0; ChanIter < 4; ChanIter++)
 						{
-							int32 RawValue = RawDataPtr[(Y*RawSize.X + X)*RawNumChan + ChanIter];
+							int32 RawValue = RawDataPtr[(Y*(int64)RawSize.X + X)*RawNumChan + ChanIter];
 							float Value = float(RawValue) / 255.0f;
 							int32 Reorder = ChanReorder[ChanIter];
-							UnpackedDataStorage[Reorder][Y*RawSize.X + X] = Value;
+							UnpackedDataStorage[Reorder][Y*(int64)RawSize.X + X] = Value;
 						}
 					}
 				}
@@ -587,17 +587,17 @@ void FImageOverlappedAccumulator::AccumulatePixelData(const FImagePixelData& InP
 				// slightly optimized, takes about 3-7ms on a 1080p image
 				for (int32 Y = 0; Y < RawSize.Y; Y++)
 				{
-					const uint8* SrcRowDataPtr = &RawDataPtr[Y*RawSize.X*RawNumChan];
+					const uint8* SrcRowDataPtr = &RawDataPtr[Y*(int64)RawSize.X*RawNumChan];
 
-					float* DstRowDataR = &UnpackedDataStorage[0][Y*RawSize.X];
-					float* DstRowDataG = &UnpackedDataStorage[1][Y*RawSize.X];
-					float* DstRowDataB = &UnpackedDataStorage[2][Y*RawSize.X];
-					float* DstRowDataA = &UnpackedDataStorage[3][Y*RawSize.X];
+					float* DstRowDataR = &UnpackedDataStorage[0][Y*(int64)RawSize.X];
+					float* DstRowDataG = &UnpackedDataStorage[1][Y*(int64)RawSize.X];
+					float* DstRowDataB = &UnpackedDataStorage[2][Y*(int64)RawSize.X];
+					float* DstRowDataA = &UnpackedDataStorage[3][Y*(int64)RawSize.X];
 
 					VectorRegister ColorScale = MakeVectorRegister(1.0f / 255.0f, 1.0f / 255.0f, 1.0f / 255.0f, 1.0f / 255.0f);
 
 					// simple, one pixel at a time vectorized version, we could do better
-					for (int32 X = 0; X < RawSize.X; X++)
+					for (int32 X = 0; X < (int64)RawSize.X; X++)
 					{
 						VectorRegister Color = VectorLoadByte4(&SrcRowDataPtr[X*RawNumChan]);
 						Color = VectorMultiply(Color, ColorScale); // multiply by 1/255
@@ -618,15 +618,15 @@ void FImageOverlappedAccumulator::AccumulatePixelData(const FImagePixelData& InP
 
 			ParallelFor(RawSize.Y, [&](int32 Y)
 			{
-				for (int32 X = 0; X < RawSize.X; X++)
+				for (int32 X = 0; X < (int64)RawSize.X; X++)
 				{
 					float UnpackedColors[4];
-					uint16_t* HalfColor = (uint16_t*)&RawDataPtr[(Y * RawSize.X + X) * RawNumChan];
+					uint16_t* HalfColor = (uint16_t*)&RawDataPtr[(Y * (int64)RawSize.X + X) * RawNumChan];
 					FPlatformMath::VectorLoadHalf(&UnpackedColors[0], HalfColor);
 
 					for (int32 ChanIter = 0; ChanIter < 4; ChanIter++)
 					{
-						UnpackedDataStorage[ChanIter][Y * RawSize.X + X] = UnpackedColors[ChanIter];
+						UnpackedDataStorage[ChanIter][Y * (int64)RawSize.X + X] = UnpackedColors[ChanIter];
 					}
 				}
 			});
@@ -642,12 +642,12 @@ void FImageOverlappedAccumulator::AccumulatePixelData(const FImagePixelData& InP
 				// reference version for float
 				for (int32 Y = 0; Y < RawSize.Y; Y++)
 				{
-					for (int32 X = 0; X < RawSize.X; X++)
+					for (int32 X = 0; X < (int64)RawSize.X; X++)
 					{
 						for (int32 ChanIter = 0; ChanIter < 4; ChanIter++)
 						{
-							float Value = RawDataPtr[(Y*RawSize.X + X)*RawNumChan + ChanIter];
-							UnpackedDataStorage[ChanIter][Y*RawSize.X + X] = Value;
+							float Value = RawDataPtr[(Y*(int64)RawSize.X + X)*RawNumChan + ChanIter];
+							UnpackedDataStorage[ChanIter][Y*(int64)RawSize.X + X] = Value;
 						}
 					}
 				}
@@ -659,15 +659,15 @@ void FImageOverlappedAccumulator::AccumulatePixelData(const FImagePixelData& InP
 				// slightly optimized, takes about 3-7ms on a 1080p image
 				for (int32 Y = 0; Y < RawSize.Y; Y++)
 				{
-					const float* SrcRowDataPtr = &RawDataPtr[Y*RawSize.X*RawNumChan];
+					const float* SrcRowDataPtr = &RawDataPtr[Y*(int64)RawSize.X*RawNumChan];
 
-					float* DstRowDataR = &UnpackedDataStorage[0][Y*RawSize.X];
-					float* DstRowDataG = &UnpackedDataStorage[1][Y*RawSize.X];
-					float* DstRowDataB = &UnpackedDataStorage[2][Y*RawSize.X];
-					float* DstRowDataA = &UnpackedDataStorage[3][Y*RawSize.X];
+					float* DstRowDataR = &UnpackedDataStorage[0][Y*(int64)RawSize.X];
+					float* DstRowDataG = &UnpackedDataStorage[1][Y*(int64)RawSize.X];
+					float* DstRowDataB = &UnpackedDataStorage[2][Y*(int64)RawSize.X];
+					float* DstRowDataA = &UnpackedDataStorage[3][Y*(int64)RawSize.X];
 
 					// simple, one pixel at a time vectorized version, we could do better
-					for (int32 X = 0; X < RawSize.X; X++)
+					for (int32 X = 0; X < (int64)RawSize.X; X++)
 					{
 						DstRowDataR[X] = SrcRowDataPtr[X*4+0];
 						DstRowDataG[X] = SrcRowDataPtr[X*4+1];
@@ -701,7 +701,7 @@ void FImageOverlappedAccumulator::AccumulatePixelData(const FImagePixelData& InP
 			{
 				float* DstData = UnpackedDataStorage[ChanIter].GetData();
 
-				int32 DstSize = RawSize.X * RawSize.Y;
+				int32 DstSize = (int64)RawSize.X * RawSize.Y;
 
 				static bool EnableVectorizedPow = true;
 				if (EnableVectorizedPow)
@@ -789,13 +789,13 @@ void FImageOverlappedAccumulator::FetchFullImageValue(float Rgba[4], int32 FullX
 	Rgba[2] = 0.0f;
 	Rgba[3] = 1.0f;
 
-	float RawWeight = WeightPlane.ChannelData[FullY * PlaneSize.X + FullX];
+	float RawWeight = WeightPlane.ChannelData[FullY * (int64)PlaneSize.X + FullX];
 
 	float Scale = 1.0f / FMath::Max<float>(RawWeight, 0.0001f);
 
 	for (int64 Chan = 0; Chan < NumChannels; Chan++)
 	{
-		float Val = ChannelPlanes[Chan].ChannelData[FullY * PlaneSize.X + FullX];
+		float Val = ChannelPlanes[Chan].ChannelData[FullY * (int64)PlaneSize.X + FullX];
 
 		Rgba[Chan] = Val * Scale;
 	}
@@ -836,8 +836,8 @@ void FImageOverlappedAccumulator::FetchFinalPixelDataByte(TArray64<FColor> & Out
 void FImageOverlappedAccumulator::FetchFinalPixelDataHalfFloat(TArray64<FFloat16Color>& OutPixelData) const
 {
 	SCOPE_CYCLE_COUNTER(STAT_FetchFinalPixelData16bit);
-	int32 FullSizeX = PlaneSize.X;
-	int32 FullSizeY = PlaneSize.Y;
+	int64 FullSizeX = PlaneSize.X;
+	int64 FullSizeY = PlaneSize.Y;
 	OutPixelData.SetNumUninitialized(FullSizeX * FullSizeY);
 
 	ParallelFor(FullSizeY, [&](int32 InIndexY)
