MoviePipeline: Added the ability to override cvars at the job/shot level for graph-based configurations. These cvars will override any cvars set within the graph.

diff --git a/Engine/Plugins/MovieScene/MovieRenderPipeline/Source/MovieRenderPipelineCore/Private/Graph/MovieGraphCVarManager.cpp b/Engine/Plugins/MovieScene/MovieRenderPipeline/Source/MovieRenderPipelineCore/Private/Graph/MovieGraphCVarManager.cpp
--- a/Engine/Plugins/MovieScene/MovieRenderPipeline/Source/MovieRenderPipelineCore/Private/Graph/MovieGraphCVarManager.cpp
+++ b/Engine/Plugins/MovieScene/MovieRenderPipeline/Source/MovieRenderPipelineCore/Private/Graph/MovieGraphCVarManager.cpp
@@ -10,6 +10,7 @@
 #include "HAL/IConsoleManager.h"
 #include "Kismet/KismetSystemLibrary.h"
 #include "Misc/DefaultValueHelper.h"
+#include "MoviePipelineQueue.h"
 #include "MovieRenderPipelineCoreModule.h"
 #include "Sections/MovieSceneConsoleVariableTrackInterface.h"
 
@@ -109,6 +110,32 @@ void FMovieGraphCVarManager::AddEvaluatedGraph(const UMovieGraphEvaluatedConfig*
 	}
 }
 
+void FMovieGraphCVarManager::AddShot(const TObjectPtr<UMoviePipelineExecutorShot>& InShot)
+{
+// UMoviePipelineExecutorShot and UMoviePipelineExecutorJob don't share a common base class, hence the need for a define vs. a lambda.
+#define ADD_CVARS(Owner) \
+	for (const FMoviePipelineConsoleVariableEntry& CVarEntry : Owner->ConsoleVariableOverrides) \
+	{ \
+		if (CVarEntry.bIsEnabled) \
+		{ \
+			AddCVar(CVarEntry.Name, CVarEntry.Value); \
+		} \
+	}
+	
+	if (InShot)
+	{
+		// Add cvars specified on the parent job first. The shot's cvars will have a chance to override these if needed.
+		if (const UMoviePipelineExecutorJob* ParentJob = InShot->GetTypedOuter<UMoviePipelineExecutorJob>())
+		{
+			ADD_CVARS(ParentJob)
+		}
+
+		// Then add the shot's cvar overrides.
+		ADD_CVARS(InShot)
+	}
+#undef ADD_CVARS
+}
+
 void FMovieGraphCVarManager::ApplyAllCVars()
 {
 	PreviousConsoleVariableValues.Reset();
diff --git a/Engine/Plugins/MovieScene/MovieRenderPipeline/Source/MovieRenderPipelineCore/Private/Graph/MovieGraphCVarManager.h b/Engine/Plugins/MovieScene/MovieRenderPipeline/Source/MovieRenderPipelineCore/Private/Graph/MovieGraphCVarManager.h
--- a/Engine/Plugins/MovieScene/MovieRenderPipeline/Source/MovieRenderPipelineCore/Private/Graph/MovieGraphCVarManager.h
+++ b/Engine/Plugins/MovieScene/MovieRenderPipeline/Source/MovieRenderPipelineCore/Private/Graph/MovieGraphCVarManager.h
@@ -8,6 +8,7 @@
 
 class IMovieSceneConsoleVariableTrackInterface;
 class UMovieGraphEvaluatedConfig;
+class UMoviePipelineExecutorShot;
 class UMovieScene;
 
 namespace UE::MovieGraph::Private
@@ -42,6 +43,13 @@ namespace UE::MovieGraph::Private
 		/** For all cvar, cvar preset, and console command nodes in InEvaluatedGraph, calls either AddCVar(), AddPreset(), or Add*ConsoleCommands(). */
 		void AddEvaluatedGraph(const UMovieGraphEvaluatedConfig* InEvaluatedGraph);
 
+		/**
+		 * Calls AddCvar() for all console variables that have been added via overrides on the shot, as well as the shot's parent job. Console
+		 * variables set on the shot and job have priority over all other cvars that have been set (with shot-based cvars having the highest
+		 * priority).
+		 */
+		void AddShot(const TObjectPtr<UMoviePipelineExecutorShot>& InShot);
+
 		/** Applies all cvars that have been gathered via the Add*() methods. */
 		void ApplyAllCVars();
 
diff --git a/Engine/Plugins/MovieScene/MovieRenderPipeline/Source/MovieRenderPipelineCore/Private/Graph/MovieGraphPipeline.cpp b/Engine/Plugins/MovieScene/MovieRenderPipeline/Source/MovieRenderPipelineCore/Private/Graph/MovieGraphPipeline.cpp
--- a/Engine/Plugins/MovieScene/MovieRenderPipeline/Source/MovieRenderPipelineCore/Private/Graph/MovieGraphPipeline.cpp
+++ b/Engine/Plugins/MovieScene/MovieRenderPipeline/Source/MovieRenderPipelineCore/Private/Graph/MovieGraphPipeline.cpp
@@ -919,9 +919,10 @@ void UMovieGraphPipeline::SetupShot(const TObjectPtr<UMoviePipelineExecutorShot>
 		GlobalGameOverridesNode->ApplySettings(bOverrideValues, GetWorld());
 	}
 
-	// Apply cvars for the shot
+	// Apply cvars and console commands for the shot
 	CVarManager->SetWorld(GetWorld());
 	CVarManager->AddEvaluatedGraph(EvaluatedConfig);
+	CVarManager->AddShot(InShot);
 	CVarManager->ApplyAllCVars();
 	CVarManager->RunStartConsoleCommands();
 
diff --git a/Engine/Plugins/MovieScene/MovieRenderPipeline/Source/MovieRenderPipelineCore/Public/MoviePipelineQueue.h b/Engine/Plugins/MovieScene/MovieRenderPipeline/Source/MovieRenderPipelineCore/Public/MoviePipelineQueue.h
--- a/Engine/Plugins/MovieScene/MovieRenderPipeline/Source/MovieRenderPipelineCore/Public/MoviePipelineQueue.h
+++ b/Engine/Plugins/MovieScene/MovieRenderPipeline/Source/MovieRenderPipelineCore/Public/MoviePipelineQueue.h
@@ -257,6 +257,12 @@ public:
 	/** Called when the graph preset assigned to the shot changes. */
 	FOnMoviePipelineShotGraphPresetChanged OnShotGraphPresetChanged;
 
+	/**
+	 * (Optional) Console variable overrides which are applied after cvars set via nodes. Only applies to graph-based configs.
+	 */
+	UPROPERTY(BlueprintReadWrite, EditAnywhere, Category = "Console Variable Overrides", meta = (DisplayName = "Console Variables", EditCondition="IsUsingGraphConfiguration()", EditConditionHides))
+	TArray<FMoviePipelineConsoleVariableEntry> ConsoleVariableOverrides;
+
 public:
 	/** Transient information used by the active Movie Pipeline working on this shot. */
 	FMoviePipelineCameraCutInfo ShotInfo;
@@ -592,6 +598,12 @@ public:
 	/** Called when the graph preset assigned to the job changes. */
 	FOnMoviePipelineJobGraphPresetChanged OnJobGraphPresetChanged;
 
+	/**
+	 * (Optional) Console variable overrides which are applied after cvars set via nodes. Only applies to graph-based configs.
+	 */
+	UPROPERTY(BlueprintReadWrite, EditAnywhere, Category = "Console Variable Overrides", meta = (DisplayName = "Console Variables", EditCondition="IsUsingGraphConfiguration()", EditConditionHides))
+	TArray<FMoviePipelineConsoleVariableEntry> ConsoleVariableOverrides;
+
 private:
 	UPROPERTY(Transient)
 	FString StatusMessage;
diff --git a/Engine/Plugins/MovieScene/MovieRenderPipeline/Source/MovieRenderPipelineCore/Public/MovieRenderPipelineDataTypes.cpp b/Engine/Plugins/MovieScene/MovieRenderPipeline/Source/MovieRenderPipelineCore/Public/MovieRenderPipelineDataTypes.cpp
--- a/Engine/Plugins/MovieScene/MovieRenderPipeline/Source/MovieRenderPipelineCore/Public/MovieRenderPipelineDataTypes.cpp
+++ b/Engine/Plugins/MovieScene/MovieRenderPipeline/Source/MovieRenderPipelineCore/Public/MovieRenderPipelineDataTypes.cpp
@@ -33,6 +33,19 @@ void FMoviePipelineCameraCutInfo::CalculateWorkMetrics()
 	WorkMetrics.TotalEngineWarmUpFrameCount = NumEngineWarmUpFramesRemaining;
 }
 
+FMoviePipelineConsoleVariableEntry::FMoviePipelineConsoleVariableEntry(const FString& InName, const float InValue, const bool bInIsEnabled)
+	: Name(InName)
+	, Value(InValue)
+	, bIsEnabled(bInIsEnabled)
+{
+
+}
+
+FMoviePipelineConsoleVariableEntry::FMoviePipelineConsoleVariableEntry()
+	: Value(0), bIsEnabled(true)
+{
+}
+
 void FMoviePipelineCameraCutInfo::SetNextStateAfter(const EMovieRenderShotState InCurrentState)
 {
 	switch (InCurrentState)
diff --git a/Engine/Plugins/MovieScene/MovieRenderPipeline/Source/MovieRenderPipelineCore/Public/MovieRenderPipelineDataTypes.h b/Engine/Plugins/MovieScene/MovieRenderPipeline/Source/MovieRenderPipelineCore/Public/MovieRenderPipelineDataTypes.h
--- a/Engine/Plugins/MovieScene/MovieRenderPipeline/Source/MovieRenderPipelineCore/Public/MovieRenderPipelineDataTypes.h
+++ b/Engine/Plugins/MovieScene/MovieRenderPipeline/Source/MovieRenderPipelineCore/Public/MovieRenderPipelineDataTypes.h
@@ -137,6 +137,30 @@ public:
 	FString CameraName;
 };
 
+/**
+ * Represents a console variable override that can be enabled/disabled.
+ */
+USTRUCT(BlueprintType)
+struct MOVIERENDERPIPELINECORE_API FMoviePipelineConsoleVariableEntry
+{
+	GENERATED_BODY()
+	
+	FMoviePipelineConsoleVariableEntry(const FString& InName, const float InValue, const bool bInIsEnabled = true);
+	FMoviePipelineConsoleVariableEntry();
+
+	/* The name of the console variable. */
+	UPROPERTY(EditAnywhere, BlueprintReadOnly, Category = "Settings")
+	FString Name;
+
+	/* The value of the console variable. */
+	UPROPERTY(EditAnywhere, BlueprintReadOnly, Category = "Settings")
+	float Value;
+
+	/* Enable state. If disabled, this cvar entry will be ignored when resolving the final value of the cvar. */
+	UPROPERTY(EditAnywhere, BlueprintReadOnly, Category = "Settings")
+	bool bIsEnabled;
+};
+
 namespace MoviePipeline
 {
 	/**
diff --git a/Engine/Plugins/MovieScene/MovieRenderPipeline/Source/MovieRenderPipelineEditor/Private/Customizations/ConsoleVariableCustomization.h b/Engine/Plugins/MovieScene/MovieRenderPipeline/Source/MovieRenderPipelineEditor/Private/Customizations/ConsoleVariableCustomization.h
--- a/Engine/Plugins/MovieScene/MovieRenderPipeline/Source/MovieRenderPipelineEditor/Private/Customizations/ConsoleVariableCustomization.h
+++ b/Engine/Plugins/MovieScene/MovieRenderPipeline/Source/MovieRenderPipelineEditor/Private/Customizations/ConsoleVariableCustomization.h
@@ -137,6 +137,12 @@ protected:
 		);
 	}
 
+	TWeakPtr<FConsoleVariablesEditorCommandInfo> GetCommandInfoForConsoleVariableEntry(const FMoviePipelineConsoleVariableEntry* InEntry) const
+	{
+		FConsoleVariablesEditorModule& ConsoleVariablesEditorModule = FConsoleVariablesEditorModule::Get();
+		return ConsoleVariablesEditorModule.FindCommandInfoByName(InEntry->Name);
+	}
+
 	ECheckBoxState GetIsEnabledCheckBoxState(TSharedRef<IPropertyHandle> PropertyHandle) const
 	{
 		bool bValue = false;
@@ -178,23 +184,6 @@ protected:
 
 		return FText();
 	}
-	
-	FText GetConsoleVariableHelpText(TSharedRef<IPropertyHandle> PropertyHandle) const
-	{
-		// Get the handle to the parent struct property
-		const TSharedPtr<IPropertyHandle> ParentPropertyHandle = PropertyHandle->GetParentHandle();
-		
-		if (const FMoviePipelineConsoleVariableEntry* CVarEntry = GetEntryForHandle(ParentPropertyHandle))
-		{
-			const TSharedPtr<FConsoleVariablesEditorCommandInfo> CommandInfo = CVarEntry->CommandInfo.Pin();
-			if (CommandInfo.IsValid())
-			{
-				return FText::FromString(CommandInfo->GetHelpText());
-			}
-		}
-
-		return FText();
-	}
 
 	void OnConsoleVariableNameTextChanged(const FText& NewText, TSharedRef<IPropertyHandle> PropertyHandle)
 	{
@@ -230,7 +219,7 @@ protected:
 			// Fall back to the startup value of the cvar if we're not in a settings object
 			else
 			{
-				const TSharedPtr<FConsoleVariablesEditorCommandInfo> CommandInfo = CVarEntry->CommandInfo.Pin();
+				const TSharedPtr<FConsoleVariablesEditorCommandInfo> CommandInfo = GetCommandInfoForConsoleVariableEntry(CVarEntry).Pin();
 				if (CommandInfo.IsValid())
 				{
 					FDefaultValueHelper::ParseFloat(CommandInfo->StartupValueAsString, CVarFloatValue);
@@ -273,7 +262,7 @@ protected:
 			return FText();
 		}
 		
-		const TSharedPtr<FConsoleVariablesEditorCommandInfo> CommandInfo = CVarEntry->CommandInfo.Pin();
+		const TSharedPtr<FConsoleVariablesEditorCommandInfo> CommandInfo = GetCommandInfoForConsoleVariableEntry(CVarEntry).Pin();
 		if (!CommandInfo.IsValid())
 		{
 			return FText();
diff --git a/Engine/Plugins/MovieScene/MovieRenderPipeline/Source/MovieRenderPipelineEditor/Private/Customizations/JobCustomization.h b/Engine/Plugins/MovieScene/MovieRenderPipeline/Source/MovieRenderPipelineEditor/Private/Customizations/JobCustomization.h
--- a/Engine/Plugins/MovieScene/MovieRenderPipeline/Source/MovieRenderPipelineEditor/Private/Customizations/JobCustomization.h
+++ b/Engine/Plugins/MovieScene/MovieRenderPipeline/Source/MovieRenderPipelineEditor/Private/Customizations/JobCustomization.h
@@ -128,13 +128,13 @@ protected:
 			SelectedJob->OnJobGraphPresetChanged.AddSP(this, &FJobDetailsCustomization::RefreshLayout);
 		}
 
-		// Set up the categories for variable assignments. Set as "Uncommon" priority to push variables down below the other properties.
+		// Set up the categories for variable assignments.
 		IDetailCategoryBuilder& PrimaryGraphVariablesCategory = InDetailBuilder.EditCategory(
-			"PrimaryGraphVariables", LOCTEXT("PrimaryGraphVariablesCategory", "Primary Graph Variables"), ECategoryPriority::Uncommon);
+			"PrimaryGraphVariables", LOCTEXT("PrimaryGraphVariablesCategory", "Primary Graph Variables"));
 		IDetailCategoryBuilder& PrimaryGraphVariablesShotOverridesCategory = InDetailBuilder.EditCategory(
-			"PrimaryGraphVariablesShotOverrides", LOCTEXT("PrimaryGraphVariablesShotOverridesCategory", "Primary Graph Variables (shot overrides)"), ECategoryPriority::Uncommon);
+			"PrimaryGraphVariablesShotOverrides", LOCTEXT("PrimaryGraphVariablesShotOverridesCategory", "Primary Graph Variables (shot overrides)"));
 		IDetailCategoryBuilder& ShotGraphVariablesCategory = InDetailBuilder.EditCategory(
-			"ShotGraphVariables", LOCTEXT("ShotGraphVariablesCategory", "Shot Graph Variables"), ECategoryPriority::Uncommon);
+			"ShotGraphVariables", LOCTEXT("ShotGraphVariablesCategory", "Shot Graph Variables"));
 
 		// Set all categories as hidden by default. Individual categories will be made visible if variables are added under them.
 		PrimaryGraphVariablesCategory.SetCategoryVisibility(false);
@@ -150,6 +150,20 @@ protected:
 		{
 			AddVariableAssignments(SelectedJob->GetGraphVariableAssignments(), PrimaryGraphVariablesCategory);
 		}
+
+		IDetailCategoryBuilder& MovieRenderPipelineCategory = InDetailBuilder.EditCategory("Movie Render Pipeline");
+
+		// The Console Variable Overrides category should be collapsed by default. It's considered advanced, and most people won't need to use it.
+		IDetailCategoryBuilder& ConsoleVariableOverridesCategory = InDetailBuilder.EditCategory("Console Variable Overrides");
+		ConsoleVariableOverridesCategory.InitiallyCollapsed(true);
+
+		// Give the categories a specific ordering.
+		int32 SortOrder = 0;
+		MovieRenderPipelineCategory.SetSortOrder(SortOrder);
+		PrimaryGraphVariablesCategory.SetSortOrder(++SortOrder);
+		PrimaryGraphVariablesShotOverridesCategory.SetSortOrder(++SortOrder);
+		ShotGraphVariablesCategory.SetSortOrder(++SortOrder);
+		ConsoleVariableOverridesCategory.SetSortOrder(++SortOrder);
 	}
 	//~ End IDetailCustomization interface
 
diff --git a/Engine/Plugins/MovieScene/MovieRenderPipeline/Source/MovieRenderPipelineSettings/Private/MoviePipelineConsoleVariableSetting.cpp b/Engine/Plugins/MovieScene/MovieRenderPipeline/Source/MovieRenderPipelineSettings/Private/MoviePipelineConsoleVariableSetting.cpp
--- a/Engine/Plugins/MovieScene/MovieRenderPipeline/Source/MovieRenderPipelineSettings/Private/MoviePipelineConsoleVariableSetting.cpp
+++ b/Engine/Plugins/MovieScene/MovieRenderPipeline/Source/MovieRenderPipelineSettings/Private/MoviePipelineConsoleVariableSetting.cpp
@@ -185,8 +185,11 @@ FString UMoviePipelineConsoleVariableSetting::ResolveDisabledValue(const FMovieP
 		return PresetValue;
 	}
 
+	FConsoleVariablesEditorModule& ConsoleVariablesEditorModule = FConsoleVariablesEditorModule::Get();
+	const TWeakPtr<FConsoleVariablesEditorCommandInfo> WeakCommandInfo = ConsoleVariablesEditorModule.FindCommandInfoByName(InEntry.Name);
+
 	// Fall back to the startup value of the cvar
-	const TSharedPtr<FConsoleVariablesEditorCommandInfo> CommandInfo = InEntry.CommandInfo.Pin();
+	const TSharedPtr<FConsoleVariablesEditorCommandInfo> CommandInfo = WeakCommandInfo.Pin();
 	if (CommandInfo.IsValid())
 	{
 		return CommandInfo->StartupValueAsString;
@@ -344,51 +347,6 @@ bool UMoviePipelineConsoleVariableSetting::UpdateConsoleVariableEnableState(cons
 	return false;
 }
 
-void UMoviePipelineConsoleVariableSetting::PostLoad()
-{
-	Super::PostLoad();
-
-#if WITH_EDITOR
-	// Initialize the CommandInfo structs
-	for (FMoviePipelineConsoleVariableEntry& Entry : CVars)
-	{
-		Entry.UpdateCommandInfo();
-	}
-#endif	// WITH_EDITOR
-}
-
-#if WITH_EDITOR
-void UMoviePipelineConsoleVariableSetting::PostEditChangeChainProperty(FPropertyChangedChainEvent& PropertyChangedEvent)
-{
-	Super::PostEditChangeChainProperty(PropertyChangedEvent);
-
-	const FProperty* Property = PropertyChangedEvent.Property;
-	const FProperty* MemberProperty = nullptr;
-	if (PropertyChangedEvent.PropertyChain.GetActiveMemberNode())
-	{
-		MemberProperty = PropertyChangedEvent.PropertyChain.GetActiveMemberNode()->GetValue();
-	}
-
-	if (!MemberProperty || !Property)
-	{
-		return;
-	}
-
-	// If the name of one of the cvar overrides changes, generate a new CommandInfo for it
-	if (MemberProperty->GetFName() == GET_MEMBER_NAME_CHECKED(UMoviePipelineConsoleVariableSetting, CVars))
-	{
-		if (Property->GetFName() == GET_MEMBER_NAME_CHECKED(FMoviePipelineConsoleVariableEntry, Name))
-		{
-			const int32 ArrayIndex = PropertyChangedEvent.GetArrayIndex(MemberProperty->GetName());
-			if (CVars.IsValidIndex(ArrayIndex))
-			{
-				CVars[ArrayIndex].UpdateCommandInfo();
-			}
-		}
-	}
-}
-#endif	// WITH_EDITOR
-
 void UMoviePipelineConsoleVariableSetting::MergeInPresetConsoleVariables()
 {
 	MergedConsoleVariables.Reset();
diff --git a/Engine/Plugins/MovieScene/MovieRenderPipeline/Source/MovieRenderPipelineSettings/Public/MoviePipelineConsoleVariableSetting.h b/Engine/Plugins/MovieScene/MovieRenderPipeline/Source/MovieRenderPipelineSettings/Public/MoviePipelineConsoleVariableSetting.h
--- a/Engine/Plugins/MovieScene/MovieRenderPipeline/Source/MovieRenderPipelineSettings/Public/MoviePipelineConsoleVariableSetting.h
+++ b/Engine/Plugins/MovieScene/MovieRenderPipeline/Source/MovieRenderPipelineSettings/Public/MoviePipelineConsoleVariableSetting.h
@@ -7,62 +7,13 @@
 #endif
 
 #include "MoviePipelineSetting.h"
+#include "MovieRenderPipelineDataTypes.h"
 #include "HAL/IConsoleManager.h"
 
 #include "MoviePipelineConsoleVariableSetting.generated.h"
 
 class IMovieSceneConsoleVariableTrackInterface;
 
-/**
- * Represents a console variable override within the Console Variable setting.
- */
-USTRUCT(BlueprintType)
-struct FMoviePipelineConsoleVariableEntry
-{
-	GENERATED_BODY()
-	
-	FMoviePipelineConsoleVariableEntry(const FString& InName, const float InValue, const bool bInIsEnabled = true)
-		: Name(InName)
-		, Value(InValue)
-		, bIsEnabled(bInIsEnabled)
-	{
-#if WITH_EDITOR
-		UpdateCommandInfo();
-#endif
-	}
-
-	FMoviePipelineConsoleVariableEntry()
-		: Value(0), bIsEnabled(true)
-	{
-	}
-
-#if WITH_EDITOR
-	/**
-	 * Updates the CommandInfo pointer. Should be called in PostLoad to ensure CommandInfo is properly set. In other
-	 * cases, CommandInfo will be kept up-to-date.
-	 */
-	void UpdateCommandInfo()
-	{
-		FConsoleVariablesEditorModule& ConsoleVariablesEditorModule = FConsoleVariablesEditorModule::Get();
-		CommandInfo = ConsoleVariablesEditorModule.FindCommandInfoByName(Name);
-	}
-
-	TWeakPtr<FConsoleVariablesEditorCommandInfo> CommandInfo;
-#endif
-
-	/* The name of the console variable. */
-	UPROPERTY(EditAnywhere, BlueprintReadOnly, Category = "Settings")
-	FString Name;
-
-	/* The value of the console variable. */
-	UPROPERTY(EditAnywhere, BlueprintReadOnly, Category = "Settings")
-	float Value;
-
-	/* Enable state. If disabled, this cvar entry will be ignored when resolving the final value of the cvar. */
-	UPROPERTY(EditAnywhere, BlueprintReadOnly, Category = "Settings")
-	bool bIsEnabled;
-};
-
 UCLASS(BlueprintType)
 class MOVIERENDERPIPELINESETTINGS_API UMoviePipelineConsoleVariableSetting : public UMoviePipelineSetting
 {
@@ -166,14 +117,6 @@ public:
 	UFUNCTION(BlueprintCallable, Category = "Settings")
 	bool UpdateConsoleVariableEnableState(const FString& Name, const bool bIsEnabled);
 
-	// Begin UObject interface
-	virtual void PostLoad() override;
-
-#if WITH_EDITOR
-	virtual void PostEditChangeChainProperty(FPropertyChangedChainEvent& PropertyChangedEvent) override;
-#endif
-	// End UObject interface
-
 private:
 	/** Merge together preset and override cvars into MergedConsoleVariables. Discards result of a prior merge (if any). */
 	void MergeInPresetConsoleVariables();
