MoviePipeline: Added a way to drag a pin to the Inputs/Outputs nodes in the graph and have the input/output created (rather than creating the input/output manually, then doing the connection).

diff --git a/Engine/Plugins/MovieScene/MovieRenderPipeline/Source/MovieRenderPipelineCore/Private/Graph/MovieGraphConfig.cpp b/Engine/Plugins/MovieScene/MovieRenderPipeline/Source/MovieRenderPipelineCore/Private/Graph/MovieGraphConfig.cpp
--- a/Engine/Plugins/MovieScene/MovieRenderPipeline/Source/MovieRenderPipelineCore/Private/Graph/MovieGraphConfig.cpp
+++ b/Engine/Plugins/MovieScene/MovieRenderPipeline/Source/MovieRenderPipelineCore/Private/Graph/MovieGraphConfig.cpp
@@ -865,11 +865,13 @@ UMovieGraphVariable* UMovieGraphConfig::AddVariable(const FName InCustomBaseName
 	return NewVariable;
 }
 
-UMovieGraphInput* UMovieGraphConfig::AddInput()
+UMovieGraphInput* UMovieGraphConfig::AddInput(const FText& InBaseName)
 {
 	static const FText InputBaseName = LOCTEXT("InputBaseName", "Input");
 
-	UMovieGraphInput* NewInput = AddMember<UMovieGraphInput>(Inputs, FName(*InputBaseName.ToString()));
+	const FText NewInputName = !InBaseName.IsEmpty() ? InBaseName : InputBaseName;
+	
+	UMovieGraphInput* NewInput = AddMember<UMovieGraphInput>(Inputs, FName(*NewInputName.ToString()));
 	InputNode->UpdatePins();
 	
 #if WITH_EDITOR
@@ -879,11 +881,13 @@ UMovieGraphInput* UMovieGraphConfig::AddInput()
 	return NewInput;
 }
 
-UMovieGraphOutput* UMovieGraphConfig::AddOutput()
+UMovieGraphOutput* UMovieGraphConfig::AddOutput(const FText& InBaseName)
 {
 	static const FText OutputBaseName = LOCTEXT("OutputBaseName", "Output");
 	
-	UMovieGraphOutput* NewOutput = AddMember<UMovieGraphOutput>(Outputs, FName(*OutputBaseName.ToString()));
+	const FText NewOutputName = !InBaseName.IsEmpty() ? InBaseName : OutputBaseName;
+	
+	UMovieGraphOutput* NewOutput = AddMember<UMovieGraphOutput>(Outputs, FName(*NewOutputName.ToString()));
 	OutputNode->UpdatePins();
 
 #if WITH_EDITOR
diff --git a/Engine/Plugins/MovieScene/MovieRenderPipeline/Source/MovieRenderPipelineCore/Public/Graph/MovieGraphConfig.h b/Engine/Plugins/MovieScene/MovieRenderPipeline/Source/MovieRenderPipelineCore/Public/Graph/MovieGraphConfig.h
--- a/Engine/Plugins/MovieScene/MovieRenderPipeline/Source/MovieRenderPipelineCore/Public/Graph/MovieGraphConfig.h
+++ b/Engine/Plugins/MovieScene/MovieRenderPipeline/Source/MovieRenderPipelineCore/Public/Graph/MovieGraphConfig.h
@@ -677,13 +677,23 @@ public:
 	UFUNCTION(BlueprintCallable, Category="Movie Graph")
 	UMovieGraphVariable* AddVariable(const FName InCustomBaseName = NAME_None);
 
-	/** Adds a new input member to the graph. Returns the new input on success, else nullptr. */
+	/**
+	 * Adds a new input member to the graph. Returns the new input on success, else nullptr.
+	 *
+	 * The default name of the input is "Input". Optionally, InBaseName can be specified to add the input with a specific name. If the name "Input"
+	 * (or the custom InBaseName) isn't available, a numerical suffix will be added.
+	 */
 	UFUNCTION(BlueprintCallable, Category = "Movie Graph")
-	UMovieGraphInput* AddInput();
+	UMovieGraphInput* AddInput(const FText& InBaseName = FText::GetEmpty());
 
-	/** Adds a new output member to the graph. Returns the new output on success, else nullptr. */
+	/**
+	 * Adds a new output member to the graph. Returns the new output on success, else nullptr.
+	 *
+	 * The default name of the output is "Output". Optionally, InBaseName can be specified to add the output with a specific name. If the name "Output"
+	 * (or the custom InBaseName) isn't available, a numerical suffix will be added.
+	 */
 	UFUNCTION(BlueprintCallable, Category = "Movie Graph")
-	UMovieGraphOutput* AddOutput();
+	UMovieGraphOutput* AddOutput(const FText& InBaseName = FText::GetEmpty());
 
 	/** Gets the variable in the graph with the specified GUID, else nullptr if one could not be found. */
 	UMovieGraphVariable* GetVariableByGuid(const FGuid& InGuid) const;
diff --git a/Engine/Plugins/MovieScene/MovieRenderPipeline/Source/MovieRenderPipelineEditor/Private/Graph/MovieEdGraphNode.cpp b/Engine/Plugins/MovieScene/MovieRenderPipeline/Source/MovieRenderPipelineEditor/Private/Graph/MovieEdGraphNode.cpp
--- a/Engine/Plugins/MovieScene/MovieRenderPipeline/Source/MovieRenderPipelineEditor/Private/Graph/MovieEdGraphNode.cpp
+++ b/Engine/Plugins/MovieScene/MovieRenderPipeline/Source/MovieRenderPipelineEditor/Private/Graph/MovieEdGraphNode.cpp
@@ -136,7 +136,7 @@ FEdGraphPinType UMoviePipelineEdGraphNodeBase::GetPinType(EMovieGraphValueType V
 
 FEdGraphPinType UMoviePipelineEdGraphNodeBase::GetPinType(const UMovieGraphPin* InPin)
 {
-	return GetPinType(InPin->Properties.Type, InPin->Properties.bIsBranch);
+	return GetPinType(InPin->Properties.Type, InPin->Properties.bIsBranch, InPin->Properties.TypeObject);
 }
 
 EMovieGraphValueType UMoviePipelineEdGraphNodeBase::GetValueTypeFromPinType(const FEdGraphPinType& InPinType)
@@ -565,6 +565,20 @@ void UMoviePipelineEdGraphNodeBase::PostLoad()
 {
 	Super::PostLoad();
 
+	// Some older nodes did not have the pin type properly set on the editor pin (specifically the value type object).
+	for (UEdGraphPin* Pin : GetAllPins())
+	{
+		const UMovieGraphPin* RuntimePin = Pin->Direction == EGPD_Input
+			? RuntimeNode->GetInputPin(Pin->PinName)
+			: RuntimeNode->GetOutputPin(Pin->PinName);
+		
+		if (RuntimePin)
+		{
+			UObject* NonConstValueTypeObject = const_cast<UObject*>(RuntimePin->Properties.TypeObject.Get());
+			Pin->PinType.PinSubCategoryObject = MakeWeakObjectPtr(NonConstValueTypeObject);
+		}
+	}
+
 	RegisterDelegates();	
 }
 
diff --git a/Engine/Plugins/MovieScene/MovieRenderPipeline/Source/MovieRenderPipelineEditor/Private/Graph/MovieGraphSchema.cpp b/Engine/Plugins/MovieScene/MovieRenderPipeline/Source/MovieRenderPipelineEditor/Private/Graph/MovieGraphSchema.cpp
--- a/Engine/Plugins/MovieScene/MovieRenderPipeline/Source/MovieRenderPipelineEditor/Private/Graph/MovieGraphSchema.cpp
+++ b/Engine/Plugins/MovieScene/MovieRenderPipeline/Source/MovieRenderPipelineEditor/Private/Graph/MovieGraphSchema.cpp
@@ -16,6 +16,8 @@
 #include "UObject/UObjectIterator.h"
 #include "ScopedTransaction.h"
 #include "GraphEditor.h"
+#include "MovieEdGraphInputNode.h"
+#include "MovieEdGraphOutputNode.h"
 #include "MovieEdGraphVariableNode.h"
 #include "MoviePipelineEdGraphSubgraphNode.h"
 
@@ -118,6 +120,79 @@ bool UMovieGraphSchema::ShouldHidePinDefaultValue(UEdGraphPin* Pin) const
 	return true;
 }
 
+bool UMovieGraphSchema::SupportsDropPinOnNode(UEdGraphNode* InTargetNode, const FEdGraphPinType& InSourcePinType, EEdGraphPinDirection InSourcePinDirection, FText& OutErrorMessage) const
+{
+	bool bIsSupported = false;
+
+	if ((InSourcePinDirection == EGPD_Input) && Cast<UMoviePipelineEdGraphNodeInput>(InTargetNode))
+	{
+		bIsSupported = true;
+		OutErrorMessage = LOCTEXT("AddPinToInputNode", "Add Pin to Input Node");
+	}
+	else if ((InSourcePinDirection == EGPD_Output) && Cast<UMoviePipelineEdGraphNodeOutput>(InTargetNode))
+	{
+		bIsSupported = true;
+		OutErrorMessage = LOCTEXT("AddPinToOutputNode", "Add Pin to Output Node");
+	}
+	
+	return bIsSupported;
+}
+
+UEdGraphPin* UMovieGraphSchema::DropPinOnNode(UEdGraphNode* InTargetNode, const FName& InSourcePinName, const FEdGraphPinType& InSourcePinType, EEdGraphPinDirection InSourcePinDirection) const
+{
+	UEdGraphPin* NewEdPin = nullptr;
+	
+	const UMoviePipelineEdGraphNodeBase* EdNode = Cast<UMoviePipelineEdGraphNodeBase>(InTargetNode);
+	if (!EdNode)
+	{
+		return nullptr;
+	}
+
+	const UMovieGraphNode* RuntimeNode = EdNode->GetRuntimeNode();
+	if (!RuntimeNode)
+	{
+		return nullptr;
+	}
+
+	if (UMovieGraphConfig* GraphConfig = RuntimeNode->GetGraph())
+	{
+		FText NewMemberName;
+		if (InSourcePinName == NAME_None)
+		{
+			NewMemberName = (InSourcePinDirection == EGPD_Input) ? LOCTEXT("NewInputName", "NewInput") : LOCTEXT("NewOutputName", "NewOutput");
+		}
+		else
+		{
+			NewMemberName = FText::FromName(InSourcePinName);
+		}
+
+		UMovieGraphInterfaceBase* NewMember;
+		if (InSourcePinDirection == EGPD_Input)
+		{
+			NewMember = GraphConfig->AddInput(NewMemberName);
+		}
+		else
+		{
+			NewMember = GraphConfig->AddOutput(NewMemberName);
+		}
+		
+		if (NewMember)
+		{
+			NewMember->bIsBranch = InSourcePinType.PinCategory == PC_Branch;
+
+			if (!NewMember->bIsBranch)
+			{
+				NewMember->SetValueType(UMoviePipelineEdGraphNode::GetValueTypeFromPinType(InSourcePinType), InSourcePinType.PinSubCategoryObject.Get());
+			}
+
+			// Return the last pin on the node (which was just added above)
+			NewEdPin = EdNode->GetPinAt(EdNode->GetAllPins().Num() - 1);
+		}
+	}
+
+	return NewEdPin; 
+}
+
 void UMovieGraphSchema::InitMoviePipelineNodeClasses()
 {
 	if (MoviePipelineNodeClasses.Num() > 0)
diff --git a/Engine/Plugins/MovieScene/MovieRenderPipeline/Source/MovieRenderPipelineEditor/Private/Graph/MovieGraphSchema.h b/Engine/Plugins/MovieScene/MovieRenderPipeline/Source/MovieRenderPipelineEditor/Private/Graph/MovieGraphSchema.h
--- a/Engine/Plugins/MovieScene/MovieRenderPipeline/Source/MovieRenderPipelineEditor/Private/Graph/MovieGraphSchema.h
+++ b/Engine/Plugins/MovieScene/MovieRenderPipeline/Source/MovieRenderPipelineEditor/Private/Graph/MovieGraphSchema.h
@@ -18,24 +18,19 @@ public:
 	//~ Begin EdGraphSchema Interface
 	virtual void CreateDefaultNodesForGraph(UEdGraph& Graph) const override;
 	virtual void GetGraphContextActions(FGraphContextMenuBuilder& ContextMenuBuilder) const override;
-	//virtual void GetContextMenuActions(class UToolMenu* Menu, class UGraphNodeContextMenuContext* Context) const override;
 	virtual const FPinConnectionResponse CanCreateConnection(const UEdGraphPin* A, const UEdGraphPin* B) const override;
 	virtual bool TryCreateConnection(UEdGraphPin* InA, UEdGraphPin* InB) const override;
-	//virtual const FPinConnectionResponse CanMergeNodes(const UEdGraphNode* A, const UEdGraphNode* B) const override;
 	static FLinearColor GetTypeColor(const FName& InPinCategory, const FName& InPinSubCategory);
 	virtual FLinearColor GetPinTypeColor(const FEdGraphPinType& PinType) const override;
 	virtual FConnectionDrawingPolicy* CreateConnectionDrawingPolicy(int32 InBackLayerID, int32 InFrontLayerID, float InZoomFactor, const FSlateRect& InClippingRect, FSlateWindowElementList& InDrawElements, UEdGraph* InGraphObj) const override;
-
-	//virtual int32 GetNodeSelectionCount(const UEdGraph* Graph) const override;
-	//virtual bool IsCacheVisualizationOutOfDate(int32 InVisualizationCacheID) const override;
-	//virtual int32 GetCurrentVisualizationCacheID() const override;
-	//virtual void ForceVisualizationCacheClear() const override;
-	//virtual TSharedPtr<FEdGraphSchemaAction> GetCreateCommentAction() const override;
 	virtual void BreakPinLinks(UEdGraphPin& TargetPin, bool bSendsNodeNotification) const override;
 	virtual void BreakSinglePinLink(UEdGraphPin* SourcePin, UEdGraphPin* TargetPin) const override;
 	virtual bool SupportsPinTypeContainer(TWeakPtr<const FEdGraphSchemaAction> SchemaAction, const FEdGraphPinType& PinType, const EPinContainerType& ContainerType) const override;
 	virtual bool ShouldHidePinDefaultValue(UEdGraphPin* Pin) const override;
+	virtual bool SupportsDropPinOnNode(UEdGraphNode* InTargetNode, const FEdGraphPinType& InSourcePinType, EEdGraphPinDirection InSourcePinDirection, FText& OutErrorMessage) const override;
+	virtual UEdGraphPin* DropPinOnNode(UEdGraphNode* InTargetNode, const FName& InSourcePinName, const FEdGraphPinType& InSourcePinType, EEdGraphPinDirection InSourcePinDirection) const override;
 	//~ End EdGraphSchema Interface
+	
 	static void InitMoviePipelineNodeClasses();
 
 private:
diff --git a/Engine/Plugins/MovieScene/MovieRenderPipeline/Source/MovieRenderPipelineEditor/Private/Graph/SMovieGraphMembersTabContent.cpp b/Engine/Plugins/MovieScene/MovieRenderPipeline/Source/MovieRenderPipelineEditor/Private/Graph/SMovieGraphMembersTabContent.cpp
--- a/Engine/Plugins/MovieScene/MovieRenderPipeline/Source/MovieRenderPipelineEditor/Private/Graph/SMovieGraphMembersTabContent.cpp
+++ b/Engine/Plugins/MovieScene/MovieRenderPipeline/Source/MovieRenderPipelineEditor/Private/Graph/SMovieGraphMembersTabContent.cpp
@@ -71,6 +71,17 @@ void SMovieGraphMembersTabContent::Construct(const FArguments& InArgs)
 	// Update the UI whenever the graph adds/updates variables. In this case it's not known which variable is added/updated, so just pass nullptr.
 	UMovieGraphMember* UpdatedMember = nullptr;
 	CurrentGraph->OnGraphVariablesChangedDelegate.AddSP(this, &SMovieGraphMembersTabContent::RefreshMemberActions, UpdatedMember);
+
+	// Also update the UI when inputs/outputs are added to the graph.
+	CurrentGraph->OnGraphInputAddedDelegate.AddSPLambda(this, [this](UMovieGraphInput* InInput)
+	{
+		RefreshMemberActions(InInput);
+	});
+
+	CurrentGraph->OnGraphOutputAddedDelegate.AddSPLambda(this, [this](UMovieGraphOutput* InOutput)
+	{
+		RefreshMemberActions(InOutput);
+	});
 	
 	ChildSlot
 	[
@@ -432,8 +443,6 @@ FReply SMovieGraphMembersTabContent::OnAddButtonClickedOnSection(const int32 InS
 		CurrentGraph->AddVariable();
 	}
 
-	RefreshMemberActions();
-
 	return FReply::Handled();
 }
 
