MoviePipeline: Updated the graph's preview window to hide burn-in and widget render layers (composited or not).

#rnx

diff --git a/Engine/Plugins/MovieScene/MovieRenderPipeline/Source/MovieRenderPipelineCore/Private/Graph/MovieGraphDefaultRenderer.cpp b/Engine/Plugins/MovieScene/MovieRenderPipeline/Source/MovieRenderPipelineCore/Private/Graph/MovieGraphDefaultRenderer.cpp
--- a/Engine/Plugins/MovieScene/MovieRenderPipeline/Source/MovieRenderPipelineCore/Private/Graph/MovieGraphDefaultRenderer.cpp
+++ b/Engine/Plugins/MovieScene/MovieRenderPipeline/Source/MovieRenderPipelineCore/Private/Graph/MovieGraphDefaultRenderer.cpp
@@ -3,11 +3,12 @@
 #include "Graph/MovieGraphDefaultRenderer.h"
 #include "Graph/MovieGraphPipeline.h"
 #include "Graph/MovieGraphConfig.h"
+#include "Graph/Nodes/MovieGraphBurnInNode.h"
+#include "Graph/Nodes/MovieGraphDebugNode.h"
 #include "Graph/Nodes/MovieGraphGlobalGameOverrides.h"
 #include "Graph/Nodes/MovieGraphRenderLayerNode.h"
 #include "Graph/Nodes/MovieGraphRenderPassNode.h"
-#include "Graph/Nodes/MovieGraphGlobalGameOverrides.h"
-#include "Graph/Nodes/MovieGraphDebugNode.h"
+#include "Graph/Nodes/MovieGraphUIRendererNode.h"
 #include "MovieRenderPipelineCoreModule.h"
 #include "RenderingThread.h"
 #include "Engine/TextureRenderTarget2D.h"
@@ -472,6 +473,14 @@ TArray<FMovieGraphImagePreviewData> UMovieGraphDefaultRenderer::GetPreviewData()
 
 	for (const TPair<UE::MovieGraph::DefaultRenderer::FMovieGraphImagePreviewDataPoolParams, TObjectPtr<UTextureRenderTarget2D>>& RenderTarget : PooledViewRenderTargets)
 	{
+		const FString RendererName = RenderTarget.Key.Identifier.RendererName;
+		
+		// Skip the burn-in and widget renderer outputs. They clog up the preview window and seeing them doesn't provide any value.
+		if ((RendererName == UMovieGraphBurnInNode::RendererName) || (RendererName == UMovieGraphUIRendererNode::RendererName))
+		{
+			continue;
+		}
+		
 		FMovieGraphImagePreviewData& Data = Results.AddDefaulted_GetRef();
 		Data.Identifier = RenderTarget.Key.Identifier;
 		Data.Texture = RenderTarget.Value.Get();
