MovieGraph: Fixed GameOverrides node not restoring previous values after rendering finishes.

diff --git a/Engine/Plugins/MovieScene/MovieRenderPipeline/Source/MovieRenderPipelineCore/Private/Graph/Nodes/MovieGraphGlobalGameOverrides.cpp b/Engine/Plugins/MovieScene/MovieRenderPipeline/Source/MovieRenderPipelineCore/Private/Graph/Nodes/MovieGraphGlobalGameOverrides.cpp
--- a/Engine/Plugins/MovieScene/MovieRenderPipeline/Source/MovieRenderPipelineCore/Private/Graph/Nodes/MovieGraphGlobalGameOverrides.cpp
+++ b/Engine/Plugins/MovieScene/MovieRenderPipeline/Source/MovieRenderPipelineCore/Private/Graph/Nodes/MovieGraphGlobalGameOverrides.cpp
@@ -108,11 +108,18 @@ void UMovieGraphGlobalGameOverridesNode::PostLoad()
 
 void UMovieGraphGlobalGameOverridesNode::ApplySettings(const bool bOverrideValues, UWorld* InWorld)
 {
+	// We use a different node instance at the start/end of the shot, so we need to store the cached values on the CDO.
+	// It's okay to use the CDO here because these are unserialized properties (which means it won't affect delta-diff
+	// serialization with the CDO). There is potentially an issue where if the user changed a property on the node
+	// between the start of the shot and the end that we don't restore correctly, but we don't currently provide a way
+	// to do that, and even if we did, we have multiple scenarios where changing values mid-render won't work.
+	UMovieGraphGlobalGameOverridesNode* NodeCDO = GetMutableDefault<UMovieGraphGlobalGameOverridesNode>();
+
 	// Apply the scalability settings
 	if (bOverrideValues)
 	{
 		// Store the current scalability settings so we can revert back to them
-		PreviousQualityLevels = Scalability::GetQualityLevels();
+		NodeCDO->PreviousQualityLevels = Scalability::GetQualityLevels();
 
 		// Set to the chosen scalability quality level
 		Scalability::FQualityLevels QualityLevels;
@@ -130,16 +137,16 @@ void UMovieGraphGlobalGameOverridesNode::ApplySettings(const bool bOverrideValue
 
 	if (bDisableTextureStreaming)
 	{
-		MOVIEPIPELINE_STORE_AND_OVERRIDE_CVAR_INT(PreviousTextureStreaming, TEXT("r.TextureStreaming"), 0, bOverrideValues);
+		MOVIEPIPELINE_STORE_AND_OVERRIDE_CVAR_INT(NodeCDO->PreviousTextureStreaming, TEXT("r.TextureStreaming"), 0, bOverrideValues);
 	}
 	
 	if (bDisableLODs)
 	{
-		MOVIEPIPELINE_STORE_AND_OVERRIDE_CVAR_INT(PreviousForceLOD, TEXT("r.ForceLOD"), 0, bOverrideValues);
-		MOVIEPIPELINE_STORE_AND_OVERRIDE_CVAR_INT(PreviousSkeletalMeshBias, TEXT("r.SkeletalMeshLODBias"), -10, bOverrideValues);
-		MOVIEPIPELINE_STORE_AND_OVERRIDE_CVAR_INT(PreviousParticleLODBias, TEXT("r.ParticleLODBias"), -10, bOverrideValues);
-		MOVIEPIPELINE_STORE_AND_OVERRIDE_CVAR_INT(PreviousFoliageDitheredLOD, TEXT("foliage.DitheredLOD"), 0, bOverrideValues);
-		MOVIEPIPELINE_STORE_AND_OVERRIDE_CVAR_INT(PreviousFoliageForceLOD, TEXT("foliage.ForceLOD"), 0, bOverrideValues);
+		MOVIEPIPELINE_STORE_AND_OVERRIDE_CVAR_INT(NodeCDO->PreviousForceLOD, TEXT("r.ForceLOD"), 0, bOverrideValues);
+		MOVIEPIPELINE_STORE_AND_OVERRIDE_CVAR_INT(NodeCDO->PreviousSkeletalMeshBias, TEXT("r.SkeletalMeshLODBias"), -10, bOverrideValues);
+		MOVIEPIPELINE_STORE_AND_OVERRIDE_CVAR_INT(NodeCDO->PreviousParticleLODBias, TEXT("r.ParticleLODBias"), -10, bOverrideValues);
+		MOVIEPIPELINE_STORE_AND_OVERRIDE_CVAR_INT(NodeCDO->PreviousFoliageDitheredLOD, TEXT("foliage.DitheredLOD"), 0, bOverrideValues);
+		MOVIEPIPELINE_STORE_AND_OVERRIDE_CVAR_INT(NodeCDO->PreviousFoliageForceLOD, TEXT("foliage.ForceLOD"), 0, bOverrideValues);
 	}
 
 	if (bDisableHLODs)
@@ -147,13 +154,14 @@ void UMovieGraphGlobalGameOverridesNode::ApplySettings(const bool bOverrideValue
 		// It's a command and not an integer cvar (despite taking 1/0), so we can't cache it 
 		if (GEngine && InWorld)
 		{
+			// ToDo: We don't restore this right now, because we don't have a way to fetch the current value to cache it.
 			GEngine->Exec(InWorld, TEXT("r.HLOD 0"));
 		}
 	}
 
 	if (bFlushStreamingManagers)
 	{
-		MOVIEPIPELINE_STORE_AND_OVERRIDE_CVAR_INT(PreviousStreamingManagerSyncState, TEXT("r.Streaming.SyncStatesWhenBlocking"), 1, bOverrideValues);
+		MOVIEPIPELINE_STORE_AND_OVERRIDE_CVAR_INT(NodeCDO->PreviousStreamingManagerSyncState, TEXT("r.Streaming.SyncStatesWhenBlocking"), 1, bOverrideValues);
 	}
 
 	// Set cvars that allow MRQ to run correctly. These are not exposed as properties on the node, but the chance of a
@@ -161,36 +169,36 @@ void UMovieGraphGlobalGameOverridesNode::ApplySettings(const bool bOverrideValue
 	{
 #if WITH_EDITOR
 		// To make sure the GeometryCache streamer doesn't skip frames and doesn't pop up notification during rendering
-		MOVIEPIPELINE_STORE_AND_OVERRIDE_CVAR_INT_IF_EXIST(PreviousGeoCacheStreamerBlockTillFinish, TEXT("GeometryCache.Streamer.BlockTillFinishStreaming"), 1, bOverrideValues);
-		MOVIEPIPELINE_STORE_AND_OVERRIDE_CVAR_INT_IF_EXIST(PreviousGeoCacheStreamerShowNotification, TEXT("GeometryCache.Streamer.ShowNotification"), 0, bOverrideValues);
+		MOVIEPIPELINE_STORE_AND_OVERRIDE_CVAR_INT_IF_EXIST(NodeCDO->PreviousGeoCacheStreamerBlockTillFinish, TEXT("GeometryCache.Streamer.BlockTillFinishStreaming"), 1, bOverrideValues);
+		MOVIEPIPELINE_STORE_AND_OVERRIDE_CVAR_INT_IF_EXIST(NodeCDO->PreviousGeoCacheStreamerShowNotification, TEXT("GeometryCache.Streamer.ShowNotification"), 0, bOverrideValues);
 #endif
 
 		// Disable systems that try to preserve performance in runtime games.
-		MOVIEPIPELINE_STORE_AND_OVERRIDE_CVAR_INT(PreviousAnimationUROEnabled, TEXT("a.URO.Enable"), 0, bOverrideValues);
+		MOVIEPIPELINE_STORE_AND_OVERRIDE_CVAR_INT(NodeCDO->PreviousAnimationUROEnabled, TEXT("a.URO.Enable"), 0, bOverrideValues);
 
 		// To make sure that the skylight is always valid and consistent across capture sessions, we enforce a full capture each frame, accepting a small GPU cost.
-		MOVIEPIPELINE_STORE_AND_OVERRIDE_CVAR_INT(PreviousSkyLightRealTimeReflectionCaptureTimeSlice, TEXT("r.SkyLight.RealTimeReflectionCapture.TimeSlice"), 0, bOverrideValues);
+		MOVIEPIPELINE_STORE_AND_OVERRIDE_CVAR_INT(NodeCDO->PreviousSkyLightRealTimeReflectionCaptureTimeSlice, TEXT("r.SkyLight.RealTimeReflectionCapture.TimeSlice"), 0, bOverrideValues);
 
 		// Cloud are rendered using high quality volumetric render target mode 3: per pixel tracing and composition on screen, while supporting cloud on translucent.
-		MOVIEPIPELINE_STORE_AND_OVERRIDE_CVAR_INT(PreviousVolumetricRenderTarget, TEXT("r.VolumetricRenderTarget"), 1, bOverrideValues);
-		MOVIEPIPELINE_STORE_AND_OVERRIDE_CVAR_INT(PreviousVolumetricRenderTargetMode, TEXT("r.VolumetricRenderTarget.Mode"), 3, bOverrideValues);
+		MOVIEPIPELINE_STORE_AND_OVERRIDE_CVAR_INT(NodeCDO->PreviousVolumetricRenderTarget, TEXT("r.VolumetricRenderTarget"), 1, bOverrideValues);
+		MOVIEPIPELINE_STORE_AND_OVERRIDE_CVAR_INT(NodeCDO->PreviousVolumetricRenderTargetMode, TEXT("r.VolumetricRenderTarget.Mode"), 3, bOverrideValues);
 
 		// To make sure that the world partition streaming doesn't end up in critical streaming performances and stops streaming low priority cells.
-		MOVIEPIPELINE_STORE_AND_OVERRIDE_CVAR_INT(PreviousIgnoreStreamingPerformance, TEXT("wp.Runtime.BlockOnSlowStreaming"), 0, bOverrideValues);
+		MOVIEPIPELINE_STORE_AND_OVERRIDE_CVAR_INT(NodeCDO->PreviousIgnoreStreamingPerformance, TEXT("wp.Runtime.BlockOnSlowStreaming"), 0, bOverrideValues);
 
 		// Remove any minimum delta time requirements from Chaos Physics to ensure accuracy at high Temporal Sample counts
-		MOVIEPIPELINE_STORE_AND_OVERRIDE_CVAR_FLOAT(PreviousChaosImmPhysicsMinStepTime, TEXT("p.Chaos.ImmPhys.MinStepTime"), 0, bOverrideValues);
+		MOVIEPIPELINE_STORE_AND_OVERRIDE_CVAR_FLOAT(NodeCDO->PreviousChaosImmPhysicsMinStepTime, TEXT("p.Chaos.ImmPhys.MinStepTime"), 0, bOverrideValues);
 
 		// MRQ's 0 -> 0.99 -> 0 evaluation for motion blur emulation can occasionally cause it to be detected as a redundant update and thus never updated
 		// which causes objects to render in the wrong position on the first frame (and without motion blur). This disables an optimization that detects
 		// the redundant updates so the update will get sent through anyways even though it thinks it's a duplicate (but it's not).
-		MOVIEPIPELINE_STORE_AND_OVERRIDE_CVAR_INT(PreviousSkipRedundantTransformUpdate, TEXT("r.SkipRedundantTransformUpdate"), 0, bOverrideValues);
+		MOVIEPIPELINE_STORE_AND_OVERRIDE_CVAR_INT(NodeCDO->PreviousSkipRedundantTransformUpdate, TEXT("r.SkipRedundantTransformUpdate"), 0, bOverrideValues);
 
 		// Cloth's time step smoothing messes up the change in number of simulation substeps that fixes the cloth simulation behavior when using Temporal Samples.
-		MOVIEPIPELINE_STORE_AND_OVERRIDE_CVAR_INT(PreviousChaosClothUseTimeStepSmoothing, TEXT("p.ChaosCloth.UseTimeStepSmoothing"), 0, bOverrideValues);
+		MOVIEPIPELINE_STORE_AND_OVERRIDE_CVAR_INT(NodeCDO->PreviousChaosClothUseTimeStepSmoothing, TEXT("p.ChaosCloth.UseTimeStepSmoothing"), 0, bOverrideValues);
 
 		// Water skips water info texture when the world's game viewport rendering is disabled so we need to prevent this from happening.
-		MOVIEPIPELINE_STORE_AND_OVERRIDE_CVAR_INT_IF_EXIST(PreviousSkipWaterInfoTextureRenderWhenWorldRenderingDisabled, TEXT("r.Water.SkipWaterInfoTextureRenderWhenWorldRenderingDisabled"), 0, bOverrideValues);
+		MOVIEPIPELINE_STORE_AND_OVERRIDE_CVAR_INT_IF_EXIST(NodeCDO->PreviousSkipWaterInfoTextureRenderWhenWorldRenderingDisabled, TEXT("r.Water.SkipWaterInfoTextureRenderWhenWorldRenderingDisabled"), 0, bOverrideValues);
 
 	}
 	
@@ -198,7 +206,7 @@ void UMovieGraphGlobalGameOverridesNode::ApplySettings(const bool bOverrideValue
 	// restore to the value in the original scalability level, not the value we cached in the scalability level (if applied).
 	if (!bOverrideValues)
 	{
-		Scalability::SetQualityLevels(PreviousQualityLevels);
+		Scalability::SetQualityLevels(NodeCDO->PreviousQualityLevels);
 	}
 }
 
