MovieGraph: Updated tooltips/toast notifications around trying to add Spawnables to Actor Condition Groups.

diff --git a/Engine/Plugins/MovieScene/MovieRenderPipeline/Source/MovieRenderPipelineCore/Private/Graph/MovieGraphRenderLayerSubsystem.cpp b/Engine/Plugins/MovieScene/MovieRenderPipeline/Source/MovieRenderPipelineCore/Private/Graph/MovieGraphRenderLayerSubsystem.cpp
--- a/Engine/Plugins/MovieScene/MovieRenderPipeline/Source/MovieRenderPipelineCore/Private/Graph/MovieGraphRenderLayerSubsystem.cpp
+++ b/Engine/Plugins/MovieScene/MovieRenderPipeline/Source/MovieRenderPipelineCore/Private/Graph/MovieGraphRenderLayerSubsystem.cpp
@@ -39,6 +39,7 @@
 #include "DragAndDrop/LevelDragDropOp.h"
 #include "Editor.h"
 #include "EditorActorFolders.h"
+#include "Framework/Notifications/NotificationManager.h"
 #include "Graph/MovieGraphSharedWidgets.h"
 #include "IContentBrowserSingleton.h"
 #include "ISceneOutliner.h"
@@ -51,6 +52,7 @@
 #include "ScopedTransaction.h"
 #include "SDropTarget.h"
 #include "Selection.h"
+#include "Widgets/Notifications/SNotificationList.h"
 #endif
 
 #define LOCTEXT_NAMESPACE "MovieGraph"
@@ -602,7 +604,7 @@ TArray<TSharedRef<SWidget>> UMovieGraphConditionGroupQuery_Actor::GetWidgets()
 		ListDataSource.Add(MakeShared<TSoftObjectPtr<AActor>>(Actor));
 	}
 
-	auto GetValidActorsFromOperation = [](TSharedPtr<FDragDropOperation> InOperation, TArray<AActor*>& OutActors) {
+	auto GetValidActorsFromOperation = [](TSharedPtr<FDragDropOperation> InOperation, TArray<AActor*>& OutActors, bool& bHadTransient) {
 
 		// Support dragging both actors and folders from the Outliner (dragging a folder will add all actors in the folder)
 		if (InOperation->IsOfType<FActorDragDropGraphEdOp>())
@@ -618,11 +620,13 @@ TArray<TSharedRef<SWidget>> UMovieGraphConditionGroupQuery_Actor::GetWidgets()
 		}
 
 		// Prevent any transient actors (ie: spawnables) from being added as they won't exist later
+		bHadTransient = false;
 		for (int32 Index = OutActors.Num() - 1; Index >= 0; Index--)
 		{
 			if(OutActors[Index]->HasAnyFlags(RF_Transient))
 			{
 				OutActors.RemoveAt(Index);
+				bHadTransient = true;
 			}
 		}
 	};
@@ -632,14 +636,33 @@ TArray<TSharedRef<SWidget>> UMovieGraphConditionGroupQuery_Actor::GetWidgets()
 		.OnAllowDrop_Lambda([GetValidActorsFromOperation](TSharedPtr<FDragDropOperation> InDragOperation)
 		{
 				TArray<AActor*> DroppedActors;
-				GetValidActorsFromOperation(InDragOperation, DroppedActors);
+				bool bHadTransient;
+				GetValidActorsFromOperation(InDragOperation, DroppedActors, bHadTransient);
 
 				return DroppedActors.Num() > 0;
 		})
 		.OnDropped_Lambda([this, GetValidActorsFromOperation](const FGeometry& Geometry, const FDragDropEvent& DragDropEvent)
 		{
 			TArray<AActor*> DroppedActors;
-			GetValidActorsFromOperation(DragDropEvent.GetOperation(), DroppedActors);
+			bool bHadTransient;
+
+			GetValidActorsFromOperation(DragDropEvent.GetOperation(), DroppedActors, bHadTransient);
+			if (bHadTransient)
+			{
+				// If we go this far, we have some non-spawnables that will actually get added to the list,
+				// but apparently they also had a transient actor selected, so we'll toast notify them that
+				// that one in particular won't be added.
+				FNotificationInfo Info(LOCTEXT("TransientActorsUnsupported_Notification", "Actor Conditions do not support Spawnable (Transient) actors"));
+				Info.SubText = LOCTEXT("TransientActorsUnsupported_NotificationSubtext", "Use the \"Actor Name\" Condition to add Spawnable actors to Collections.");
+				Info.Image = FAppStyle::GetBrush(TEXT("MessageLog.Warning"));
+				Info.Image = FAppStyle::GetBrush(TEXT("Icons.Warning"));
+
+				//Set a default expire duration
+				Info.ExpireDuration = 5.0f;
+
+				//And call Add Notification, this is pretty much it!
+				FSlateNotificationManager::Get().AddNotification(Info);
+			}
 
 			const FMovieGraphConditionGroupQueryContentsChanged OnAddFinished = nullptr;
 			AddActors(DroppedActors, OnAddFinished);
@@ -708,12 +731,40 @@ TSharedRef<SWidget> UMovieGraphConditionGroupQuery_Actor::GetAddMenuContents(con
 	const FSceneOutlinerModule& SceneOutlinerModule = FModuleManager::LoadModuleChecked<FSceneOutlinerModule>("SceneOutliner");
 
 	FMenuBuilder MenuBuilder(false, MakeShared<FUICommandList>());
+
+	auto CanExecuteAction = []() -> bool {
+		// Assume we have only transient actors until we prove we don't.
+		bool bHasNonTransientActorsSelected = false;
+		TArray<AActor*> SelectedActors;
+		GEditor->GetSelectedActors()->GetSelectedObjects<AActor>(SelectedActors);
+
+		for (AActor* SelectedActor : SelectedActors)
+		{
+			if (!SelectedActor->HasAnyFlags(RF_Transient))
+			{
+				bHasNonTransientActorsSelected = true;
+				break;
+			}
+		}
+		return bHasNonTransientActorsSelected;
+	};
 	
 	MenuBuilder.BeginSection("AddActor", LOCTEXT("AddActor", "Add Actor"));
 	{
 		MenuBuilder.AddMenuEntry(
 			LOCTEXT("AddSelectedInOutliner", "Add Selected In Outliner"),
-			LOCTEXT("AddSelectedInOutlinerTooltip", "Add actors currently selected in the level editor's scene outliner."),
+			TAttribute<FText>::CreateLambda([CanExecuteAction]()
+				{
+					bool bCanExecute = CanExecuteAction();
+					if (bCanExecute)
+					{
+						return LOCTEXT("AddSelectedInOutlinerTooltip_Success", "Add actors currently selected in the level editor's scene outliner.");
+					}
+					else
+					{
+						return LOCTEXT("AddSelectedInOutlinerTooltip_Failure", "Actor Conditions do not support Spawnable (Transient) actors.  Select one or more non-transient actors or use the \"Actor Name\" Condition to add Spawnable actors to Collections.");
+					}
+				}),
 			FSlateIcon(FAppStyle::GetAppStyleSetName(),"FoliageEditMode.SetSelect"),
 			FUIAction(
 				FExecuteAction::CreateLambda([this, OnAddFinished]()
@@ -733,22 +784,9 @@ TSharedRef<SWidget> UMovieGraphConditionGroupQuery_Actor::GetAddMenuContents(con
 
 					AddActors(SelectedActors, OnAddFinished);
 				}),
-				FCanExecuteAction::CreateLambda([]()
+				FCanExecuteAction::CreateLambda([CanExecuteAction]()
 				{
-					// Assume we have only transient actors until we prove we don't.
-					bool bHasNonTransientActorsSelected = false;
-					TArray<AActor*> SelectedActors;
-					GEditor->GetSelectedActors()->GetSelectedObjects<AActor>(SelectedActors);
-
-					for (AActor* SelectedActor : SelectedActors)
-					{
-						if(!SelectedActor->HasAnyFlags(RF_Transient))
-						{
-							bHasNonTransientActorsSelected = true;
-							break;
-						}
-					}
-					return bHasNonTransientActorsSelected;
+					return CanExecuteAction();
 				})
 			)
 		);
