MoviePipeline: Fixed missing RHI resource transition before reading back from a render target resource.

diff --git a/Engine/Plugins/MovieScene/MovieRenderPipeline/Source/MovieRenderPipelineCore/Private/Graph/Nodes/MovieGraphWidgetRendererBaseNode.cpp b/Engine/Plugins/MovieScene/MovieRenderPipeline/Source/MovieRenderPipelineCore/Private/Graph/Nodes/MovieGraphWidgetRendererBaseNode.cpp
--- a/Engine/Plugins/MovieScene/MovieRenderPipeline/Source/MovieRenderPipelineCore/Private/Graph/Nodes/MovieGraphWidgetRendererBaseNode.cpp
+++ b/Engine/Plugins/MovieScene/MovieRenderPipeline/Source/MovieRenderPipelineCore/Private/Graph/Nodes/MovieGraphWidgetRendererBaseNode.cpp
@@ -155,6 +155,9 @@ void UMovieGraphWidgetRendererBaseNode::FMovieGraphWidgetPass::Render(const FMov
 			bComposite = ParentNodeThisFrame->bCompositeOntoFinalImage, CompositingSortOrder = GetCompositingSortOrder(),
 			BackbufferRenderTarget, OutputMerger, OutputResolution](FRHICommandListImmediate& RHICmdList)
 			{
+				// Transition our render target from a render target view to a shader resource view to allow a shader to read from this Render Target.
+				RHICmdList.Transition(FRHITransitionInfo(BackbufferRenderTarget->GetRenderTargetTexture(), ERHIAccess::RTV, ERHIAccess::SRVGraphicsPixel));
+
 				const FIntRect SourceRect = FIntRect(0, 0, BackbufferRenderTarget->GetSizeXY().X, BackbufferRenderTarget->GetSizeXY().Y);
 
 				// Read the data back to the CPU
diff --git a/Engine/Plugins/MovieScene/MovieRenderPipeline/Source/MovieRenderPipelineRenderPasses/Private/Graph/Renderers/MovieGraphImagePassBase.cpp b/Engine/Plugins/MovieScene/MovieRenderPipeline/Source/MovieRenderPipelineRenderPasses/Private/Graph/Renderers/MovieGraphImagePassBase.cpp
--- a/Engine/Plugins/MovieScene/MovieRenderPipeline/Source/MovieRenderPipelineRenderPasses/Private/Graph/Renderers/MovieGraphImagePassBase.cpp
+++ b/Engine/Plugins/MovieScene/MovieRenderPipeline/Source/MovieRenderPipelineRenderPasses/Private/Graph/Renderers/MovieGraphImagePassBase.cpp
@@ -484,6 +484,9 @@ void FMovieGraphImagePassBase::PostRendererSubmission(
 	ENQUEUE_RENDER_COMMAND(CanvasRenderTargetResolveCommand)(
 		[LocalSurfaceQueue, OnSurfaceReadbackFinished, RenderTarget](FRHICommandListImmediate& RHICmdList) mutable
 		{
+			// Transition our render target from a render target view to a shader resource view to allow a shader to read from this Render Target.
+			RHICmdList.Transition(FRHITransitionInfo(RenderTarget->GetRenderTargetTexture(), ERHIAccess::RTV, ERHIAccess::SRVGraphicsPixel));
+
 			// The legacy surface reader takes the payload just so it can shuffle it into our callback, but we can just include the data
 			// directly in the callback, so this is just a dummy payload.
 			TSharedRef<FImagePixelDataPayload, ESPMode::ThreadSafe> FramePayload = MakeShared<FImagePixelDataPayload, ESPMode::ThreadSafe>();
diff --git a/Engine/Plugins/MovieScene/MovieRenderPipeline/Source/MovieRenderPipelineRenderPasses/Private/MoviePipelineDeferredPasses.cpp b/Engine/Plugins/MovieScene/MovieRenderPipeline/Source/MovieRenderPipelineRenderPasses/Private/MoviePipelineDeferredPasses.cpp
--- a/Engine/Plugins/MovieScene/MovieRenderPipeline/Source/MovieRenderPipelineRenderPasses/Private/MoviePipelineDeferredPasses.cpp
+++ b/Engine/Plugins/MovieScene/MovieRenderPipeline/Source/MovieRenderPipelineRenderPasses/Private/MoviePipelineDeferredPasses.cpp
@@ -1019,6 +1019,9 @@ void UMoviePipelineDeferredPassBase::PostRendererSubmission(const FMoviePipeline
 	ENQUEUE_RENDER_COMMAND(CanvasRenderTargetResolveCommand)(
 		[LocalSurfaceQueue, FramePayload, Callback, RenderTarget](FRHICommandListImmediate& RHICmdList) mutable
 		{
+			// Transition our render target from a render target view to a shader resource view to allow a shader to read from this Render Target.
+			RHICmdList.Transition(FRHITransitionInfo(RenderTarget->GetRenderTargetTexture(), ERHIAccess::RTV, ERHIAccess::SRVGraphicsPixel));
+
 			// Enqueue a encode for this frame onto our worker thread.
 			LocalSurfaceQueue->OnRenderTargetReady_RenderThread(RenderTarget->GetRenderTargetTexture(), FramePayload, MoveTemp(Callback));
 		});
diff --git a/Engine/Plugins/MovieScene/MovieRenderPipeline/Source/MovieRenderPipelineSettings/Private/MoviePipelineBurnInSetting.cpp b/Engine/Plugins/MovieScene/MovieRenderPipeline/Source/MovieRenderPipelineSettings/Private/MoviePipelineBurnInSetting.cpp
--- a/Engine/Plugins/MovieScene/MovieRenderPipeline/Source/MovieRenderPipelineSettings/Private/MoviePipelineBurnInSetting.cpp
+++ b/Engine/Plugins/MovieScene/MovieRenderPipeline/Source/MovieRenderPipelineSettings/Private/MoviePipelineBurnInSetting.cpp
@@ -90,6 +90,9 @@ void UMoviePipelineBurnInSetting::RenderSample_GameThreadImpl(const FMoviePipeli
 			ENQUEUE_RENDER_COMMAND(BurnInRenderTargetResolveCommand)(
 				[InSampleState, PassIdentifierForCurrentCamera, bComposite = bCompositeOntoFinalImage, BackbufferRenderTarget, OutputBuilder](FRHICommandListImmediate& RHICmdList)
 				{
+					// Transition our render target from a render target view to a shader resource view to allow a shader to read from this Render Target.
+					RHICmdList.Transition(FRHITransitionInfo(BackbufferRenderTarget->GetRenderTargetTexture(), ERHIAccess::RTV, ERHIAccess::SRVGraphicsPixel));
+
 					FIntRect SourceRect = FIntRect(0, 0, BackbufferRenderTarget->GetSizeXY().X, BackbufferRenderTarget->GetSizeXY().Y);
 
 					// Read the data back to the CPU
diff --git a/Engine/Plugins/MovieScene/MovieRenderPipeline/Source/MovieRenderPipelineSettings/Private/MoviePipelineWidgetRenderSetting.cpp b/Engine/Plugins/MovieScene/MovieRenderPipeline/Source/MovieRenderPipelineSettings/Private/MoviePipelineWidgetRenderSetting.cpp
--- a/Engine/Plugins/MovieScene/MovieRenderPipeline/Source/MovieRenderPipelineSettings/Private/MoviePipelineWidgetRenderSetting.cpp
+++ b/Engine/Plugins/MovieScene/MovieRenderPipeline/Source/MovieRenderPipelineSettings/Private/MoviePipelineWidgetRenderSetting.cpp
@@ -80,6 +80,9 @@ void UMoviePipelineWidgetRenderer::RenderSample_GameThreadImpl(const FMoviePipel
 			ENQUEUE_RENDER_COMMAND(BurnInRenderTargetResolveCommand)(
 				[InSampleState, PassIdentifierForCurrentCamera, bComposite = bCompositeOntoFinalImage, BackbufferRenderTarget, OutputBuilder](FRHICommandListImmediate& RHICmdList)
 				{
+					// Transition our render target from a render target view to a shader resource view to allow a shader to read from this Render Target.
+					RHICmdList.Transition(FRHITransitionInfo(BackbufferRenderTarget->GetRenderTargetTexture(), ERHIAccess::RTV, ERHIAccess::SRVGraphicsPixel));
+
 					FIntRect SourceRect = FIntRect(0, 0, BackbufferRenderTarget->GetSizeXY().X, BackbufferRenderTarget->GetSizeXY().Y);
 
 					// Read the data back to the CPU
