MoviePipeline: Added duplicate for graph variables.

diff --git a/Engine/Plugins/MovieScene/MovieRenderPipeline/Source/MovieRenderPipelineCore/Private/Graph/MovieGraphConfig.cpp b/Engine/Plugins/MovieScene/MovieRenderPipeline/Source/MovieRenderPipelineCore/Private/Graph/MovieGraphConfig.cpp
--- a/Engine/Plugins/MovieScene/MovieRenderPipeline/Source/MovieRenderPipelineCore/Private/Graph/MovieGraphConfig.cpp
+++ b/Engine/Plugins/MovieScene/MovieRenderPipeline/Source/MovieRenderPipelineCore/Private/Graph/MovieGraphConfig.cpp
@@ -980,6 +980,30 @@ bool UMovieGraphConfig::DeleteMember(UMovieGraphMember* MemberToDelete)
 	return false;
 }
 
+UMovieGraphVariable* UMovieGraphConfig::DuplicateVariable(UMovieGraphVariable* InVariableToDuplicate)
+{
+	if (!InVariableToDuplicate)
+	{
+		return nullptr;
+	}
+
+	Modify();
+
+	// AddVariable() does lots of heavy lifting to make sure variables are added/named correctly. Instead of duplicating lots of that boilerplate
+	// here, just manually copy over the value/category/etc from the source variable. DuplicateObject() would work too, but has its own set of
+	// clean-up procedures that are needed.
+	UMovieGraphVariable* NewVariable = AddVariable(FName(InVariableToDuplicate->GetMemberName()));
+	if (NewVariable)
+	{
+		NewVariable->SetValueType(InVariableToDuplicate->GetValueType(), const_cast<UObject*>(InVariableToDuplicate->GetValueTypeObject()));
+		NewVariable->SetValueSerializedString(InVariableToDuplicate->GetValueSerializedString());
+		NewVariable->SetCategory(InVariableToDuplicate->GetCategory());
+		NewVariable->Description = InVariableToDuplicate->Description;
+	}
+	
+	return NewVariable;
+}
+
 bool UMovieGraphConfig::DeleteVariableMember(UMovieGraphVariable* VariableMemberToDelete)
 {
 	if (!VariableMemberToDelete)
diff --git a/Engine/Plugins/MovieScene/MovieRenderPipeline/Source/MovieRenderPipelineCore/Public/Graph/MovieGraphConfig.h b/Engine/Plugins/MovieScene/MovieRenderPipeline/Source/MovieRenderPipelineCore/Public/Graph/MovieGraphConfig.h
--- a/Engine/Plugins/MovieScene/MovieRenderPipeline/Source/MovieRenderPipelineCore/Public/Graph/MovieGraphConfig.h
+++ b/Engine/Plugins/MovieScene/MovieRenderPipeline/Source/MovieRenderPipelineCore/Public/Graph/MovieGraphConfig.h
@@ -721,6 +721,9 @@ public:
 	UFUNCTION(BlueprintCallable, Category = "Movie Graph")
 	bool DeleteMember(UMovieGraphMember* MemberToDelete);
 
+	/** Duplicates the provided variable. Returns the new variable on success, else nullptr. */
+	UMovieGraphVariable* DuplicateVariable(UMovieGraphVariable* InVariableToDuplicate);
+
 #if WITH_EDITOR
 	/** Gets the editor-only nodes in this graph. Editor-only nodes do not have an equivalent runtime node. */
 	const TArray<TObjectPtr<UObject>>& GetEditorOnlyNodes() const { return EditorOnlyNodes; }
diff --git a/Engine/Plugins/MovieScene/MovieRenderPipeline/Source/MovieRenderPipelineEditor/Private/Graph/MovieGraphAssetToolkit.cpp b/Engine/Plugins/MovieScene/MovieRenderPipeline/Source/MovieRenderPipelineEditor/Private/Graph/MovieGraphAssetToolkit.cpp
--- a/Engine/Plugins/MovieScene/MovieRenderPipeline/Source/MovieRenderPipelineEditor/Private/Graph/MovieGraphAssetToolkit.cpp
+++ b/Engine/Plugins/MovieScene/MovieRenderPipeline/Source/MovieRenderPipelineEditor/Private/Graph/MovieGraphAssetToolkit.cpp
@@ -479,6 +479,12 @@ void FMovieGraphAssetToolkit::BindGraphCommands()
 		FExecuteAction::CreateSP(this, &FMovieGraphAssetToolkit::DeleteSelectedMembers),
 		FCanExecuteAction::CreateSP(this, &FMovieGraphAssetToolkit::CanDeleteSelectedMembers));
 
+	ToolkitCommands->MapAction(FGenericCommands::Get().Duplicate,
+		FExecuteAction::CreateSP(this, &FMovieGraphAssetToolkit::DuplicateSelectedMembers),
+		FCanExecuteAction::CreateSP(this, &FMovieGraphAssetToolkit::CanDuplicateSelectedMembers),
+		FIsActionChecked::CreateLambda([]() { return false; }),
+		FCanExecuteAction::CreateSP(this, &FMovieGraphAssetToolkit::IsDuplicateVisible));
+
 	ToolkitCommands->MapAction(FMoviePipelineCommands::Get().ZoomToWindow,
 			FExecuteAction::CreateSP(this, &FMovieGraphAssetToolkit::OnZoomToWindow),
 			FCanExecuteAction::CreateSP(this, &FMovieGraphAssetToolkit::CanZoomToWindow));
@@ -529,6 +535,29 @@ bool FMovieGraphAssetToolkit::CanDeleteSelectedMembers()
 	return false;
 }
 
+void FMovieGraphAssetToolkit::DuplicateSelectedMembers()
+{
+	if (MembersTabContent.IsValid())
+	{
+		MembersTabContent->DuplicateSelectedMembers();
+	}
+}
+
+bool FMovieGraphAssetToolkit::CanDuplicateSelectedMembers()
+{
+	if (MembersTabContent.IsValid())
+	{
+		return MembersTabContent->CanDuplicateSelectedMembers();
+	}
+
+	return false;
+}
+
+bool FMovieGraphAssetToolkit::IsDuplicateVisible()
+{
+	return CanDuplicateSelectedMembers();
+}
+
 void FMovieGraphAssetToolkit::OnZoomToWindow() const
 {
 	if (MovieGraphWidget.IsValid())
diff --git a/Engine/Plugins/MovieScene/MovieRenderPipeline/Source/MovieRenderPipelineEditor/Private/Graph/MovieGraphAssetToolkit.h b/Engine/Plugins/MovieScene/MovieRenderPipeline/Source/MovieRenderPipelineEditor/Private/Graph/MovieGraphAssetToolkit.h
--- a/Engine/Plugins/MovieScene/MovieRenderPipeline/Source/MovieRenderPipelineEditor/Private/Graph/MovieGraphAssetToolkit.h
+++ b/Engine/Plugins/MovieScene/MovieRenderPipeline/Source/MovieRenderPipelineEditor/Private/Graph/MovieGraphAssetToolkit.h
@@ -65,6 +65,15 @@ private:
 	/** Determines if the selected member(s) can be deleted. */
 	bool CanDeleteSelectedMembers();
 
+	/** Duplicates the member(s) which are currently selected. */
+	void DuplicateSelectedMembers();
+
+	/** Determines if the selected member(s) can be duplicated. */
+	bool CanDuplicateSelectedMembers();
+
+	/** Determines if the Duplicate action is visible for the selected members. */
+	bool IsDuplicateVisible();
+
 	/** Frames the entire graph so it is visible. */
 	void OnZoomToWindow() const;
 
diff --git a/Engine/Plugins/MovieScene/MovieRenderPipeline/Source/MovieRenderPipelineEditor/Private/Graph/SMovieGraphMembersTabContent.cpp b/Engine/Plugins/MovieScene/MovieRenderPipeline/Source/MovieRenderPipelineEditor/Private/Graph/SMovieGraphMembersTabContent.cpp
--- a/Engine/Plugins/MovieScene/MovieRenderPipeline/Source/MovieRenderPipelineEditor/Private/Graph/SMovieGraphMembersTabContent.cpp
+++ b/Engine/Plugins/MovieScene/MovieRenderPipeline/Source/MovieRenderPipelineEditor/Private/Graph/SMovieGraphMembersTabContent.cpp
@@ -36,6 +36,34 @@ namespace UE::MovieGraph::Private
 		return nullptr;
 	}
 
+	/** Gets all selected graph members of the specified MemberType within the action menu. */
+	template<typename MemberType>
+	TArray<MemberType*> GetAllSelectedMembers(const TSharedPtr<SGraphActionMenu> InActionMenu)
+	{
+		TArray<MemberType*> SelectedMembers;
+		
+		if (!InActionMenu.IsValid())
+		{
+			return SelectedMembers;
+		}
+
+		TArray<TSharedPtr<FEdGraphSchemaAction>> SelectedActions;
+		InActionMenu->GetSelectedActions(SelectedActions);
+	
+		for (TSharedPtr<FEdGraphSchemaAction> SelectedAction : SelectedActions)
+		{
+			if (UMovieGraphMember* GraphMember = GetMemberFromAction(SelectedAction.Get()))
+			{
+				if (MemberType* Member = Cast<MemberType>(GraphMember))
+				{
+					SelectedMembers.Add(Member);
+				}
+			}
+		}
+
+		return SelectedMembers;
+	}
+
 	static bool GetIconAndColorFromDataType(const UMovieGraphVariable* InGraphVariable, const FSlateBrush*& OutPrimaryBrush, FSlateColor& OutIconColor, const FSlateBrush*& OutSecondaryBrush, FSlateColor& OutSecondaryColor)
 	{
 		if (!InGraphVariable)
@@ -179,22 +207,17 @@ void SMovieGraphMembersTabContent::ClearSelection() const
 
 void SMovieGraphMembersTabContent::DeleteSelectedMembers()
 {
-	if (!ActionMenu.IsValid() || !CurrentGraph)
+	if (!CurrentGraph)
 	{
 		return;
 	}
 
-	TArray<TSharedPtr<FEdGraphSchemaAction>> SelectedActions;
-	ActionMenu->GetSelectedActions(SelectedActions);
-	for (TSharedPtr<FEdGraphSchemaAction> SelectedAction : SelectedActions)
+	FScopedTransaction Transaction(LOCTEXT("DeleteGraphMembers", "Delete Graph Member(s)"));
+	
+	for (UMovieGraphMember* Member : UE::MovieGraph::Private::GetAllSelectedMembers<UMovieGraphMember>(ActionMenu))
 	{
-		if (UMovieGraphMember* GraphMember = UE::MovieGraph::Private::GetMemberFromAction(SelectedAction.Get()))
-		{
-			FScopedTransaction Transaction(LOCTEXT("DeleteGraphMember", "Delete Graph Member"));
-			
-			MemberChangedHandles.Remove(GraphMember);
-			CurrentGraph->DeleteMember(GraphMember);
-		}
+		MemberChangedHandles.Remove(Member);
+		CurrentGraph->DeleteMember(Member);
 	}
 
 	RefreshMemberActions();
@@ -202,27 +225,50 @@ void SMovieGraphMembersTabContent::DeleteSelectedMembers()
 
 bool SMovieGraphMembersTabContent::CanDeleteSelectedMembers() const
 {
-	if (!ActionMenu.IsValid())
+	// Don't allow deletion if the member was explicitly marked as non-deletable
+	for (const UMovieGraphMember* Member : UE::MovieGraph::Private::GetAllSelectedMembers<UMovieGraphMember>(ActionMenu))
 	{
-		return false;
+		if (!Member->IsDeletable())
+		{
+			return false;
+		}
 	}
 
-	TArray<TSharedPtr<FEdGraphSchemaAction>> SelectedActions;
-	ActionMenu->GetSelectedActions(SelectedActions);
+	return true;
+}
+
+void SMovieGraphMembersTabContent::DuplicateSelectedMembers()
+{
+	if (!CurrentGraph)
+	{
+		return;
+	}
 	
-	for (const TSharedPtr<FEdGraphSchemaAction>& SelectedAction : SelectedActions)
+	FScopedTransaction Transaction(LOCTEXT("DuplicateGraphMembers", "Duplicate Graph Member(s)"));
+
+	for (UMovieGraphVariable* GraphVariable : UE::MovieGraph::Private::GetAllSelectedMembers<UMovieGraphVariable>(ActionMenu))
 	{
-		// Don't allow deletion if the member was explicitly marked as non-deletable
-		if (const UMovieGraphMember* Member = UE::MovieGraph::Private::GetMemberFromAction(SelectedAction.Get()))
+		if (GraphVariable && !GraphVariable->IsGlobal())
 		{
-			if (!Member->IsDeletable())
-			{
-				return false;
-			}
+			CurrentGraph->DuplicateVariable(GraphVariable);
 		}
 	}
 
-	return true;
+	RefreshMemberActions();
+}
+
+bool SMovieGraphMembersTabContent::CanDuplicateSelectedMembers() const
+{
+	// Only allow duplication of non-global variables for now
+	for (const UMovieGraphVariable* GraphVariable : UE::MovieGraph::Private::GetAllSelectedMembers<UMovieGraphVariable>(ActionMenu))
+	{
+		if (GraphVariable && !GraphVariable->IsGlobal())
+		{
+			return true;
+		}
+	}
+
+	return false;
 }
 
 void SMovieGraphMembersTabContent::PostUndo(bool bSuccess)
@@ -419,6 +465,7 @@ TSharedPtr<SWidget> SMovieGraphMembersTabContent::OnContextMenuOpening()
 
 	FMenuBuilder MenuBuilder(true, PinnedToolkit->GetToolkitCommands());
 	MenuBuilder.AddMenuEntry(FGenericCommands::Get().Delete);
+	MenuBuilder.AddMenuEntry(FGenericCommands::Get().Duplicate);
 	
 	return MenuBuilder.MakeWidget();
 }
diff --git a/Engine/Plugins/MovieScene/MovieRenderPipeline/Source/MovieRenderPipelineEditor/Private/Graph/SMovieGraphMembersTabContent.h b/Engine/Plugins/MovieScene/MovieRenderPipeline/Source/MovieRenderPipelineEditor/Private/Graph/SMovieGraphMembersTabContent.h
--- a/Engine/Plugins/MovieScene/MovieRenderPipeline/Source/MovieRenderPipelineEditor/Private/Graph/SMovieGraphMembersTabContent.h
+++ b/Engine/Plugins/MovieScene/MovieRenderPipeline/Source/MovieRenderPipelineEditor/Private/Graph/SMovieGraphMembersTabContent.h
@@ -59,6 +59,12 @@ public:
 	/** Determines if all selected member(s) can be deleted. */
 	bool CanDeleteSelectedMembers() const;
 
+	/** Duplicates the member(s) which are currently selected in the UI. */
+	void DuplicateSelectedMembers();
+
+	/** Determines if at least one selected member can be duplicated. */
+	bool CanDuplicateSelectedMembers() const;
+
 	//~ Begin FSelfRegisteringEditorUndoClient Interface
 	virtual void PostUndo(bool bSuccess) override;
 	virtual void PostRedo(bool bSuccess) override;
