MoviePipeline: Updated the category for Script Tags in the graph.

#rnx

diff --git a/Engine/Plugins/MovieScene/MovieRenderPipeline/Source/MovieRenderPipelineCore/Public/Graph/MovieGraphNode.h b/Engine/Plugins/MovieScene/MovieRenderPipeline/Source/MovieRenderPipelineCore/Public/Graph/MovieGraphNode.h
--- a/Engine/Plugins/MovieScene/MovieRenderPipeline/Source/MovieRenderPipelineCore/Public/Graph/MovieGraphNode.h
+++ b/Engine/Plugins/MovieScene/MovieRenderPipeline/Source/MovieRenderPipelineCore/Public/Graph/MovieGraphNode.h
@@ -340,7 +340,7 @@ public:
 	 * Tags that can be used to identify this node within a pre/post render script. Tags can be unique in order to identify this specific node,
 	 * or the same tag can be applied to multiple nodes in order to identify a grouping of nodes.
 	 */
-	UPROPERTY(EditAnywhere, BlueprintReadWrite, Category = "Settings", AdvancedDisplay)
+	UPROPERTY(EditAnywhere, BlueprintReadWrite, Category = "Tags")
 	TArray<FString> ScriptTags;
 
 #if WITH_EDITORONLY_DATA
diff --git a/Engine/Plugins/MovieScene/MovieRenderPipeline/Source/MovieRenderPipelineEditor/Private/Customizations/Graph/MovieGraphNodeCustomization.h b/Engine/Plugins/MovieScene/MovieRenderPipeline/Source/MovieRenderPipelineEditor/Private/Customizations/Graph/MovieGraphNodeCustomization.h
--- a/Engine/Plugins/MovieScene/MovieRenderPipeline/Source/MovieRenderPipelineEditor/Private/Customizations/Graph/MovieGraphNodeCustomization.h
+++ b/Engine/Plugins/MovieScene/MovieRenderPipeline/Source/MovieRenderPipelineEditor/Private/Customizations/Graph/MovieGraphNodeCustomization.h
@@ -2,6 +2,7 @@
 
 #pragma once
 
+#include "DetailCategoryBuilder.h"
 #include "DetailLayoutBuilder.h"
 #include "IDetailCustomization.h"
 #include "Graph/MovieGraphConfig.h"
@@ -50,6 +51,9 @@ protected:
 		{
 			DetailBuilder.HideCategory("Properties");
 		}
+
+		// Collapse the Tags category by default because it's not used by most users.
+		DetailBuilder.EditCategory("Tags").InitiallyCollapsed(true).SetSortOrder(1000 * ECategoryPriority::Uncommon);
 	}
 	//~ End IDetailCustomization interface
 };
