MoviePipeline: Updated the category for View Mode Index in the graph.

#rnx

diff --git a/Engine/Plugins/MovieScene/MovieRenderPipeline/Source/MovieRenderPipelineRenderPasses/Public/Graph/Nodes/MovieGraphDeferredPassNode.h b/Engine/Plugins/MovieScene/MovieRenderPipeline/Source/MovieRenderPipelineRenderPasses/Public/Graph/Nodes/MovieGraphDeferredPassNode.h
--- a/Engine/Plugins/MovieScene/MovieRenderPipeline/Source/MovieRenderPipelineRenderPasses/Public/Graph/Nodes/MovieGraphDeferredPassNode.h
+++ b/Engine/Plugins/MovieScene/MovieRenderPipeline/Source/MovieRenderPipelineRenderPasses/Public/Graph/Nodes/MovieGraphDeferredPassNode.h
@@ -108,7 +108,7 @@ public:
 	* have to be used in combination with certain Show Flags to produce a result similar to what you see in
 	* the viewport.
 	*/
-	UPROPERTY(EditAnywhere, BlueprintReadWrite, Category = "Settings", meta=(EditCondition="bOverride_ViewModeIndex", InvalidEnumValues = "VMI_PathTracing,VMI_VisualizeBuffer,VMI_Unknown"))
+	UPROPERTY(EditAnywhere, BlueprintReadWrite, Category = "View Mode", meta=(EditCondition="bOverride_ViewModeIndex", InvalidEnumValues = "VMI_PathTracing,VMI_VisualizeBuffer,VMI_Unknown"))
 	TEnumAsByte<EViewModeIndex> ViewModeIndex;
 
 	/**
