MoviePipeline: Fixed several instances of search not working with details customizations done in the graph.

#rnx

diff --git a/Engine/Plugins/MovieScene/MovieRenderPipeline/Source/MovieRenderPipelineEditor/Private/Customizations/Graph/MovieGraphApplyCVarPresetNodeCustomization.h b/Engine/Plugins/MovieScene/MovieRenderPipeline/Source/MovieRenderPipelineEditor/Private/Customizations/Graph/MovieGraphApplyCVarPresetNodeCustomization.h
--- a/Engine/Plugins/MovieScene/MovieRenderPipeline/Source/MovieRenderPipelineEditor/Private/Customizations/Graph/MovieGraphApplyCVarPresetNodeCustomization.h
+++ b/Engine/Plugins/MovieScene/MovieRenderPipeline/Source/MovieRenderPipelineEditor/Private/Customizations/Graph/MovieGraphApplyCVarPresetNodeCustomization.h
@@ -140,6 +140,7 @@ protected:
 
 			// The group's header should look like the default preset property w/ the asset picker widget
 			PresetGroup.HeaderRow()
+			.FilterString(ConsoleVariablePresetHandle->GetPropertyDisplayName())
 			.PropertyHandleList({ConsoleVariablePresetHandle})
 			.ShouldAutoExpand(true)
 			.EditCondition(
diff --git a/Engine/Plugins/MovieScene/MovieRenderPipeline/Source/MovieRenderPipelineEditor/Private/Customizations/Graph/MovieGraphMemberCustomization.h b/Engine/Plugins/MovieScene/MovieRenderPipeline/Source/MovieRenderPipelineEditor/Private/Customizations/Graph/MovieGraphMemberCustomization.h
--- a/Engine/Plugins/MovieScene/MovieRenderPipeline/Source/MovieRenderPipelineEditor/Private/Customizations/Graph/MovieGraphMemberCustomization.h
+++ b/Engine/Plugins/MovieScene/MovieRenderPipeline/Source/MovieRenderPipelineEditor/Private/Customizations/Graph/MovieGraphMemberCustomization.h
@@ -66,6 +66,10 @@ protected:
 
 	virtual void CustomizeDetails(IDetailLayoutBuilder& DetailBuilder) override
 	{
+		static const FText NameLabel = LOCTEXT("MemberPropertyLabel_Name", "Name");
+		static const FText CategoryLabel = LOCTEXT("MemberPropertyLabel_Category", "Category");
+		static const FText TypeLabel = LOCTEXT("MemberPropertyLabel_Type", "Type");
+		
 		TSharedRef<IPropertyUtilities> PropUtils = DetailBuilder.GetPropertyUtilities();
 		
 		auto GetFilteredVariableTypeTree = [this](TArray<TSharedPtr<UEdGraphSchema_K2::FPinTypeTreeInfo>>& TypeTree, ETypeTreeFilter TypeTreeFilter)
@@ -161,11 +165,12 @@ protected:
 			// Add a custom row for the Name property (to allow for proper validation)
 			IDetailCategoryBuilder& GeneralCategory = DetailBuilder.EditCategory("General");
 			GeneralCategory.AddCustomRow(FText::GetEmpty())
+			.FilterString(NameLabel)
 			.IsEnabled(bIsEditable)
 			.NameContent()
 			[
 				SNew(STextBlock)
-				.Text(LOCTEXT("MemberPropertyLabel_Name", "Name"))
+				.Text(NameLabel)
 				.Font(DetailBuilder.GetDetailFont())
 			]
 			.ValueContent()
@@ -183,11 +188,12 @@ protected:
 			if (UMovieGraphVariable* VariableMember = Cast<UMovieGraphVariable>(MemberObject))
 			{
 				GeneralCategory.AddCustomRow(FText::GetEmpty())
+				.FilterString(CategoryLabel)
 				.IsEnabled(bIsEditable)
 				.NameContent()
 				[
 					SNew(STextBlock)
-					.Text(LOCTEXT("MemberPropertyLabel_Category", "Category"))
+					.Text(CategoryLabel)
 					.ToolTipText(LOCTEXT("MemberPropertyTooltip_Category", "The category assigned to the variable. Use a '|' to separate category names to create a category hierarchy (eg, Settings|Resolution)."))
 					.Font(DetailBuilder.GetDetailFont())
 				]
@@ -216,11 +222,12 @@ protected:
 			// Add a PinTypeSelector widget to pick the data type the member uses
 			IDetailCategoryBuilder& ValueCategory = DetailBuilder.EditCategory("Value");
 			ValueCategory.AddCustomRow(FText::GetEmpty())
+			.FilterString(TypeLabel)
 			.IsEnabled(bIsTypeRowEnabled)
 			.NameContent()
 			[
 				SNew(STextBlock)
-				.Text(LOCTEXT("MemberPropertyLabel_Type", "Type"))
+				.Text(TypeLabel)
 				.Font(DetailBuilder.GetDetailFont())
 			]
 			.ValueContent()
diff --git a/Engine/Plugins/MovieScene/MovieRenderPipeline/Source/MovieRenderPipelineEditor/Private/Customizations/Graph/MovieGraphVersioningSettingsCustomization.h b/Engine/Plugins/MovieScene/MovieRenderPipeline/Source/MovieRenderPipelineEditor/Private/Customizations/Graph/MovieGraphVersioningSettingsCustomization.h
--- a/Engine/Plugins/MovieScene/MovieRenderPipeline/Source/MovieRenderPipelineEditor/Private/Customizations/Graph/MovieGraphVersioningSettingsCustomization.h
+++ b/Engine/Plugins/MovieScene/MovieRenderPipeline/Source/MovieRenderPipelineEditor/Private/Customizations/Graph/MovieGraphVersioningSettingsCustomization.h
@@ -35,6 +35,7 @@ protected:
 			GET_MEMBER_NAME_CHECKED(FMovieGraphVersioningSettings, VersionNumber)).ToSharedRef();
 
 		HeaderRow
+		.FilterString(VersionNumberProp->GetPropertyDisplayName())
 		.NameContent()
 		[
 			// Show "Version Number" as the row name
