MovieGraph: Adds toast notification for any missing project setting the first time users try to load or use a holdout modifier.

#rnx

diff --git a/Engine/Plugins/MovieScene/MovieRenderPipeline/Source/MovieRenderPipelineCore/MovieRenderPipelineCore.Build.cs b/Engine/Plugins/MovieScene/MovieRenderPipeline/Source/MovieRenderPipelineCore/MovieRenderPipelineCore.Build.cs
--- a/Engine/Plugins/MovieScene/MovieRenderPipeline/Source/MovieRenderPipelineCore/MovieRenderPipelineCore.Build.cs
+++ b/Engine/Plugins/MovieScene/MovieRenderPipeline/Source/MovieRenderPipelineCore/MovieRenderPipelineCore.Build.cs
@@ -54,7 +54,10 @@ public class MovieRenderPipelineCore : ModuleRules
 			        "Layers",
 			        "SceneOutliner",
 			        "UnrealEd",
-		        });
+					"SettingsEditor",
+					"Slate",
+					"SlateCore",
+				});
 	        
 			PublicDependencyModuleNames.Add("MovieSceneTools");
         }
diff --git a/Engine/Plugins/MovieScene/MovieRenderPipeline/Source/MovieRenderPipelineCore/Private/Graph/MovieGraphRenderLayerSubsystem.cpp b/Engine/Plugins/MovieScene/MovieRenderPipeline/Source/MovieRenderPipelineCore/Private/Graph/MovieGraphRenderLayerSubsystem.cpp
--- a/Engine/Plugins/MovieScene/MovieRenderPipeline/Source/MovieRenderPipelineCore/Private/Graph/MovieGraphRenderLayerSubsystem.cpp
+++ b/Engine/Plugins/MovieScene/MovieRenderPipeline/Source/MovieRenderPipelineCore/Private/Graph/MovieGraphRenderLayerSubsystem.cpp
@@ -7,7 +7,6 @@
 #include "Components/VolumetricCloudComponent.h"
 #include "Components/ExponentialHeightFogComponent.h"
 #include "EngineUtils.h"
-#include "Engine/RendererSettings.h"
 #include "Framework/Application/SlateApplication.h"
 #include "Materials/MaterialInterface.h"
 #include "Modules/ModuleManager.h"
@@ -46,6 +45,7 @@
 #include "ISceneOutlinerTreeItem.h"
 #include "LayersDragDropOp.h"
 #include "Layers/LayersSubsystem.h"
+#include "MovieGraphUtils.h"
 #include "SceneOutlinerModule.h"
 #include "SceneOutlinerPublicTypes.h"
 #include "SClassViewer.h"
@@ -183,6 +183,25 @@ UMovieGraphRenderPropertyModifier::UMovieGraphRenderPropertyModifier()
 	// override will initially be a no-op due to the defaults being the same.
 }
 
+void UMovieGraphRenderPropertyModifier::PostLoad()
+{
+	Super::PostLoad();
+
+	ValidateProjectSettings();
+}
+
+#if WITH_EDITOR
+void UMovieGraphRenderPropertyModifier::PostEditChangeProperty(FPropertyChangedEvent& PropertyChangedEvent)
+{
+	Super::PostEditChangeProperty(PropertyChangedEvent);
+	
+	if (PropertyChangedEvent.GetPropertyName() == GET_MEMBER_NAME_CHECKED(UMovieGraphRenderPropertyModifier, bHoldout))
+	{
+		ValidateProjectSettings();
+	}
+}
+#endif
+
 void UMovieGraphRenderPropertyModifier::ApplyModifier(const UWorld* World)
 {
 	ModifiedActors.Empty();
@@ -192,14 +211,7 @@ void UMovieGraphRenderPropertyModifier::ApplyModifier(const UWorld* World)
 	FActorVisibilityState NewVisibilityState;
 	NewVisibilityState.bIsHidden = bIsHidden;
 
-	// Generate a warning if holdout is being used, but alpha is not enabled in post processing. Without that setting enabled, holdout will not work.
-	const URendererSettings* RendererSettings = GetDefault<URendererSettings>();
-	if (bHoldout && (RendererSettings->bEnableAlphaChannelInPostProcessing == false))
-	{
-		// TODO: Ideally this is called in a general-purpose validation step instead, but that framework does not exist yet.
-		UE_LOG(LogMovieRenderPipeline, Warning, TEXT("A modifier with 'Holdout' is active, but 'Alpha Output' in "
-													 "the project's Rendering settings is disabled: holdout will not work properly."));
-	}
+	ValidateProjectSettings();
 	
 	for (const UMovieGraphCollection* Collection : Collections)
 	{
@@ -384,6 +396,25 @@ void UMovieGraphRenderPropertyModifier::SetActorVisibilityState(const FActorVisi
 	}
 }
 
+void UMovieGraphRenderPropertyModifier::ValidateProjectSettings() const
+{
+	static const FText HoldoutModifierLabel = LOCTEXT("ConditionGroupQueryHoldoutModifier", "Holdout Modifier");
+
+	if (bHoldout)
+	{
+#if WITH_EDITOR
+		UE_CALL_ONCE([]
+			{
+				constexpr bool bMandatePrimitiveAlphaHoldout = true;
+				UE::MovieGraph::ValidateAlphaProjectSettings(HoldoutModifierLabel, bMandatePrimitiveAlphaHoldout);
+			}
+		);
+#else
+		UE_LOG(LogMovieRenderPipeline, Warning, TEXT("Both \"Alpha Output\" and \"Support Primitive Alpha Holdout\" project settings must be enabled, otherwise holdout modifiers will not work properly."));
+#endif
+	}
+}
+
 void UMovieGraphCollectionModifier::AddCollection(UMovieGraphCollection* Collection)
 {
 	// Don't allow adding a duplicate collection
diff --git a/Engine/Plugins/MovieScene/MovieRenderPipeline/Source/MovieRenderPipelineCore/Private/Graph/MovieGraphUtils.cpp b/Engine/Plugins/MovieScene/MovieRenderPipeline/Source/MovieRenderPipelineCore/Private/Graph/MovieGraphUtils.cpp
--- a/Engine/Plugins/MovieScene/MovieRenderPipeline/Source/MovieRenderPipelineCore/Private/Graph/MovieGraphUtils.cpp
+++ b/Engine/Plugins/MovieScene/MovieRenderPipeline/Source/MovieRenderPipelineCore/Private/Graph/MovieGraphUtils.cpp
@@ -1,4 +1,4 @@
-﻿// Copyright Epic Games, Inc. All Rights Reserved.
+// Copyright Epic Games, Inc. All Rights Reserved.
 
 #include "MovieGraphUtils.h"
 
@@ -6,6 +6,17 @@
 #include "AudioThread.h"
 #include "Engine/Engine.h"
 
+#if WITH_EDITOR
+#include "Engine/RendererSettings.h"
+#include "Framework/Notifications/NotificationManager.h"
+#include "ISettingsEditorModule.h"
+#include "Misc/ConfigCacheIni.h"
+#include "Modules/ModuleManager.h"
+#include "Widgets/Notifications/SNotificationList.h"
+#endif // WITH_EDITOR
+
+#define LOCTEXT_NAMESPACE "MovieGraphUtils"
+
 namespace UE::MovieGraph
 {
 	FString GetUniqueName(const TArray<FString>& InExistingNames, const FString& InBaseName)
@@ -66,4 +77,129 @@ namespace UE::MovieGraph
 			return false;
 		}
 	}
+
+#if WITH_EDITOR
+	void ValidateAlphaProjectSettings(const FText& InRequestingFeatureName, bool bMandatePrimitiveAlphaHoldout)
+	{
+		URendererSettings* RendererSettings = GetMutableDefault<URendererSettings>();
+		check(RendererSettings);
+
+		const bool bAlphaOutputMissing = !RendererSettings->bEnableAlphaChannelInPostProcessing;
+		const bool bPrimitiveHoldoutMissing = bMandatePrimitiveAlphaHoldout ? !RendererSettings->bDeferredSupportPrimitiveAlphaHoldout : false;
+
+		if (bAlphaOutputMissing || bPrimitiveHoldoutMissing)
+		{
+			static TWeakPtr<SNotificationItem> NotificationItem;
+
+			const FText AlphaOutputSettingOption = LOCTEXT("MovieAlphaSetting_AlphaOutput", "\n- Alpha Output");
+			const FText PrimitiveHoldoutSettingOption = LOCTEXT("MovieAlphaSetting_PrimitiveHoldout", "\n- Support Primitive Alpha Holdout");
+			const FText MovieAlphaText = FText::Format(LOCTEXT("MovieAlphaSettingPrompt", "The following project setting(s) must be enabled for Movie Render Graph's {0}:{1}{2}"), InRequestingFeatureName, bAlphaOutputMissing ? AlphaOutputSettingOption : FText::GetEmpty(), bPrimitiveHoldoutMissing ? PrimitiveHoldoutSettingOption : FText::GetEmpty());
+			const FText MovieAlphaConfirmText = LOCTEXT("MovieAlphaSettingConfirm", "Enable");
+			const FText MovieAlphaCancelText = LOCTEXT("MovieAlphaSettingCancel", "Not Now");
+
+			/** Utility functions for notifications */
+			struct FSuppressDialogOptions
+			{
+				static bool ShouldSuppressModal()
+				{
+					bool bSuppressNotification = false;
+					GConfig->GetBool(TEXT("MovieRenderPipeline"), TEXT("SuppressMovieRenderPipelineAlphaPromptNotification"), bSuppressNotification, GEditorPerProjectIni);
+					return bSuppressNotification;
+				}
+
+				static ECheckBoxState GetDontAskAgainCheckBoxState()
+				{
+					return ShouldSuppressModal() ? ECheckBoxState::Checked : ECheckBoxState::Unchecked;
+				}
+
+				static void OnDontAskAgainCheckBoxStateChanged(ECheckBoxState NewState)
+				{
+					// If the user selects to not show this again, set that in the config so we know about it in between sessions
+					const bool bSuppressNotification = (NewState == ECheckBoxState::Checked);
+					GConfig->SetBool(TEXT("MovieRenderPipeline"), TEXT("SuppressMovieRenderPipelineAlphaPromptNotification"), bSuppressNotification, GEditorPerProjectIni);
+				}
+			};
+
+			// If the user has specified to supress this pop up, then just early out and exit	
+			if (FSuppressDialogOptions::ShouldSuppressModal())
+			{
+				return;
+			}
+
+			FSimpleDelegate OnConfirmDelegate = FSimpleDelegate::CreateLambda(
+				[RendererSettings, bAlphaOutputMissing, bPrimitiveHoldoutMissing]()
+				{
+					if (IsValid(RendererSettings))
+					{
+						if (bAlphaOutputMissing)
+						{
+							FProperty* Property = RendererSettings->GetClass()->FindPropertyByName(GET_MEMBER_NAME_CHECKED(URendererSettings, bEnableAlphaChannelInPostProcessing));
+							RendererSettings->PreEditChange(Property);
+
+							RendererSettings->bEnableAlphaChannelInPostProcessing = true;
+
+							FPropertyChangedEvent PropertyChangedEvent(Property, EPropertyChangeType::ValueSet, { RendererSettings });
+							RendererSettings->PostEditChangeProperty(PropertyChangedEvent);
+							RendererSettings->UpdateSinglePropertyInConfigFile(Property, RendererSettings->GetDefaultConfigFilename());
+						}
+
+						if (bPrimitiveHoldoutMissing)
+						{
+							FProperty* Property = RendererSettings->GetClass()->FindPropertyByName(GET_MEMBER_NAME_CHECKED(URendererSettings, bDeferredSupportPrimitiveAlphaHoldout));
+							RendererSettings->PreEditChange(Property);
+
+							RendererSettings->bDeferredSupportPrimitiveAlphaHoldout = true;
+
+							FPropertyChangedEvent PropertyChangedEvent(Property, EPropertyChangeType::ValueSet, { RendererSettings });
+							RendererSettings->PostEditChangeProperty(PropertyChangedEvent);
+							RendererSettings->UpdateSinglePropertyInConfigFile(Property, RendererSettings->GetDefaultConfigFilename());
+
+							// SupportPrimitiveAlphaHoldout requires shader recompilation, ask for a restart.
+							FModuleManager::GetModuleChecked<ISettingsEditorModule>("SettingsEditor").OnApplicationRestartRequired();
+						}
+					}
+				}
+			);
+
+			FSimpleDelegate OnCancelDelegate = FSimpleDelegate::CreateLambda([]
+				{
+					if (TSharedPtr<SNotificationItem> Item = NotificationItem.Pin())
+					{
+						Item->SetCompletionState(SNotificationItem::CS_None);
+						Item->ExpireAndFadeout();
+					}
+
+					NotificationItem.Reset();
+				}
+			);
+
+			FNotificationInfo Info(MovieAlphaText);
+			Info.bFireAndForget = false;
+			Info.bUseLargeFont = false;
+			Info.bUseThrobber = false;
+			Info.bUseSuccessFailIcons = false;
+			Info.ButtonDetails.Add(FNotificationButtonInfo(MovieAlphaConfirmText, FText(), OnConfirmDelegate));
+			Info.ButtonDetails.Add(FNotificationButtonInfo(MovieAlphaCancelText, FText(), OnCancelDelegate));
+
+			// Add a "Don't show this again" option
+			Info.CheckBoxState = TAttribute<ECheckBoxState>::Create(&FSuppressDialogOptions::GetDontAskAgainCheckBoxState);
+			Info.CheckBoxStateChanged = FOnCheckStateChanged::CreateStatic(&FSuppressDialogOptions::OnDontAskAgainCheckBoxStateChanged);
+			Info.CheckBoxText = LOCTEXT("DefaultCheckBoxMessage", "Don't show this again");
+
+			if (TSharedPtr<SNotificationItem> Item = NotificationItem.Pin())
+			{
+				Item->ExpireAndFadeout();
+			}
+
+			NotificationItem = FSlateNotificationManager::Get().AddNotification(Info);
+
+			if (TSharedPtr<SNotificationItem> Item = NotificationItem.Pin())
+			{
+				Item->SetCompletionState(SNotificationItem::CS_Pending);
+			}
+		}
+	}
+#endif
 }
+
+#undef LOCTEXT_NAMESPACE // "MovieGraphUtils"
diff --git a/Engine/Plugins/MovieScene/MovieRenderPipeline/Source/MovieRenderPipelineCore/Private/Graph/MovieGraphUtils.h b/Engine/Plugins/MovieScene/MovieRenderPipeline/Source/MovieRenderPipelineCore/Private/Graph/MovieGraphUtils.h
--- a/Engine/Plugins/MovieScene/MovieRenderPipeline/Source/MovieRenderPipelineCore/Private/Graph/MovieGraphUtils.h
+++ b/Engine/Plugins/MovieScene/MovieRenderPipeline/Source/MovieRenderPipelineCore/Private/Graph/MovieGraphUtils.h
@@ -31,4 +31,8 @@ namespace UE::MovieGraph
 		/** Determines if the pipeline can generate audio. */
 		bool IsMoviePipelineAudioOutputSupported(const UObject* InWorldContextObject);
 	}
+
+#if WITH_EDITOR
+	void ValidateAlphaProjectSettings(const FText& InRequestingFeatureLabel, bool bMandatePrimitiveAlphaHoldout = false);
+#endif
 }
diff --git a/Engine/Plugins/MovieScene/MovieRenderPipeline/Source/MovieRenderPipelineCore/Public/Graph/MovieGraphRenderLayerSubsystem.h b/Engine/Plugins/MovieScene/MovieRenderPipeline/Source/MovieRenderPipelineCore/Public/Graph/MovieGraphRenderLayerSubsystem.h
--- a/Engine/Plugins/MovieScene/MovieRenderPipeline/Source/MovieRenderPipelineCore/Public/Graph/MovieGraphRenderLayerSubsystem.h
+++ b/Engine/Plugins/MovieScene/MovieRenderPipeline/Source/MovieRenderPipelineCore/Public/Graph/MovieGraphRenderLayerSubsystem.h
@@ -798,6 +798,13 @@ class MOVIERENDERPIPELINECORE_API UMovieGraphRenderPropertyModifier : public UMo
 
 public:
 	UMovieGraphRenderPropertyModifier();
+
+	// ~UObject interface
+	virtual void PostLoad() override;
+#if WITH_EDITOR
+	virtual void PostEditChangeProperty(FPropertyChangedEvent& PropertyChangedEvent) override;
+#endif
+	// ~UObject Interface
 	
 	UFUNCTION(BlueprintCallable, Category = "Settings")
 	void SetHidden(const bool bInIsHidden) { bIsHidden = bInIsHidden; }
@@ -836,6 +843,9 @@ private:
 	/** Updates an actor's visibility state to the state contained in NewVisibilityState. */
 	void SetActorVisibilityState(const FActorVisibilityState& NewVisibilityState);
 
+	/** Convienence function to ensure that output alpha and primitive alpha holdout settings are enabled if required. */
+	void ValidateProjectSettings() const;
+
 private:
 	/** Tracks actor visibility state prior to having the modifier applied. */
 	TArray<FActorVisibilityState> ModifiedActors;
