MoviePipeline: Updated the "Write All Samples (Debug)" feature to work for the deferred and path traced renderers in the graph.

diff --git a/Engine/Plugins/MovieScene/MoviePipelineMaskRenderPass/Source/MoviePipelineMaskRenderPass/Private/Graph/MovieGraphObjectIdPass.cpp b/Engine/Plugins/MovieScene/MoviePipelineMaskRenderPass/Source/MoviePipelineMaskRenderPass/Private/Graph/MovieGraphObjectIdPass.cpp
--- a/Engine/Plugins/MovieScene/MoviePipelineMaskRenderPass/Source/MoviePipelineMaskRenderPass/Private/Graph/MovieGraphObjectIdPass.cpp
+++ b/Engine/Plugins/MovieScene/MoviePipelineMaskRenderPass/Source/MoviePipelineMaskRenderPass/Private/Graph/MovieGraphObjectIdPass.cpp
@@ -33,6 +33,8 @@ namespace UE::MovieGraph
 		const bool bIsWellFormed = SamplePixelData->IsDataWellFormed();
 		check(bIsWellFormed);
 
+		// Note: Object ID doesn't currently have a property to control writing samples to disk, so this condition will never be true. One can be
+		// added in the future though, if the need to write samples ever comes up.
 		if (SampleStatePayload->bWriteSampleToDisk)
 		{
 			// Debug Feature: Write the raw sample to disk for debugging purposes. We copy the data here,
diff --git a/Engine/Plugins/MovieScene/MovieRenderPipeline/Source/MovieRenderPipelineCore/Private/Graph/MovieGraphOutputMerger.cpp b/Engine/Plugins/MovieScene/MovieRenderPipeline/Source/MovieRenderPipelineCore/Private/Graph/MovieGraphOutputMerger.cpp
--- a/Engine/Plugins/MovieScene/MovieRenderPipeline/Source/MovieRenderPipelineCore/Private/Graph/MovieGraphOutputMerger.cpp
+++ b/Engine/Plugins/MovieScene/MovieRenderPipeline/Source/MovieRenderPipelineCore/Private/Graph/MovieGraphOutputMerger.cpp
@@ -1,15 +1,23 @@
 // Copyright Epic Games, Inc. All Rights Reserved.
 
 #include "MovieGraphOutputMerger.h"
+
+#include "Async/Async.h"
+#include "ImageWriteQueue.h"
+#include "ImageWriteTask.h"
+#include "Graph/MovieGraphBlueprintLibrary.h"
 #include "Graph/MovieGraphPipeline.h"
 #include "Graph/MovieGraphRenderDataIdentifier.h"
-#include "Async/Async.h"
+#include "Graph/Nodes/MovieGraphGlobalOutputSettingNode.h"
+#include "Modules/ModuleManager.h"
+#include "MoviePipelineQueue.h"
 
 namespace UE::MovieGraph
 {
 	FMovieGraphOutputMerger::FMovieGraphOutputMerger(UMovieGraphPipeline* InOwningMoviePipeline)
 		: WeakMoviePipeline(MakeWeakObjectPtr(InOwningMoviePipeline))
 	{
+		Debug_ImageWriteQueue = &FModuleManager::Get().LoadModuleChecked<IImageWriteQueueModule>("ImageWriteQueue").GetWriteQueue();
 	}
 
 	void FMovieGraphOutputMerger::AbandonOutstandingWork()
@@ -35,25 +43,64 @@ namespace UE::MovieGraph
 		check(PendingData.Find(InRenderedFrameNumber));
 		return *PendingData.Find(InRenderedFrameNumber);
 	}
-
 	
 	void FMovieGraphOutputMerger::OnSingleSampleDataAvailable_AnyThread(TUniquePtr<FImagePixelData>&& InData)
 	{
 		// This is to support outputting individual samples (skipping accumulation) for debug reasons,
 		// or because you want to post-process them yourself. We just forward this directly on for output to disk.
 
-		TWeakObjectPtr<UMovieGraphPipeline> LocalWeakPipeline = WeakMoviePipeline;
-
-		AsyncTask(ENamedThreads::GameThread, [LocalData = MoveTemp(InData), LocalWeakPipeline]() mutable
+		AsyncTask(ENamedThreads::GameThread, [LocalData = MoveTemp(InData), LocalWeakPipeline = WeakMoviePipeline, LocalDebugQueue = Debug_ImageWriteQueue]() mutable
 		{
-			if (ensureAlwaysMsgf(LocalWeakPipeline.IsValid(), TEXT("A memory lifespan issue has left an output builder alive without an owning Movie Pipeline.")))
+			if (!ensureAlwaysMsgf(LocalWeakPipeline.IsValid(), TEXT("A memory lifespan issue has left an output builder alive without an owning Movie Pipeline.")))
 			{
-				// LocalWeakPipeline->Debug_OnSampleRendered(MoveTemp(LocalData));
-				// ToDo: Read the finished frames from MRQ
-				// also implement the above function.
+				return;
 			}
-		}
-		);
+			
+			const UE::MovieGraph::FMovieGraphSampleState* Payload = LocalData->GetPayload<UE::MovieGraph::FMovieGraphSampleState>();
+			const FMovieGraphTraversalContext& TraversalContext = Payload->TraversalContext;
+			const TObjectPtr<UMovieGraphEvaluatedConfig> EvaluatedConfig = TraversalContext.Time.EvaluatedConfig;
+
+			// Resolve the file output path
+			FString FinalFilePath;
+			{
+				// TODO: Add Tile X/Y when tiling is available
+				const FString OutputName = FString::Printf(TEXT("%s_%s_%s_%s_SS_%d_TS_%d.%d"),
+					*TraversalContext.Shot->OuterName,
+					*TraversalContext.RenderDataIdentifier.LayerName,
+					*TraversalContext.RenderDataIdentifier.SubResourceName,
+					*TraversalContext.RenderDataIdentifier.CameraName,
+					TraversalContext.Time.SpatialSampleIndex,
+					TraversalContext.Time.TemporalSampleIndex,
+					TraversalContext.Time.OutputFrameNumber);
+
+				const UMovieGraphGlobalOutputSettingNode* OutputSettings =
+					EvaluatedConfig->GetSettingForBranch<UMovieGraphGlobalOutputSettingNode>(UMovieGraphNode::GlobalsPinName);
+
+				const FString OutputDirectory = OutputSettings->OutputDirectory.Path;
+				const FString FileNameFormatString = OutputDirectory / OutputName;
+
+				TMap<FString, FString> AdditionalFormatArgs;
+				AdditionalFormatArgs.Add(TEXT("ext"), TEXT("exr"));
+
+				FMovieGraphResolveArgs MergedFormatArgs;
+				const FMovieGraphFilenameResolveParams ResolveParams = FMovieGraphFilenameResolveParams::MakeResolveParams(
+					TraversalContext.RenderDataIdentifier,
+					LocalWeakPipeline.Get(),
+					EvaluatedConfig.Get(),
+					Payload->TraversalContext,
+					AdditionalFormatArgs);
+
+				FinalFilePath = UMovieGraphBlueprintLibrary::ResolveFilenameFormatArguments(FileNameFormatString, ResolveParams, MergedFormatArgs);
+			}
+
+			TUniquePtr<FImageWriteTask> TileImageTask = MakeUnique<FImageWriteTask>();
+			TileImageTask->Format = EImageFormat::EXR;
+			TileImageTask->CompressionQuality = static_cast<int32>(EImageCompressionQuality::Default);
+			TileImageTask->Filename = FinalFilePath;
+			TileImageTask->PixelData = MoveTemp(LocalData);	// The pixel data is currently owned by the async task; transfer ownership to the image write task
+
+			LocalDebugQueue->Enqueue(MoveTemp(TileImageTask));
+		});
 	}
 	
 	void FMovieGraphOutputMerger::OnCompleteRenderPassDataAvailable_AnyThread(TUniquePtr<FImagePixelData>&& InData)
diff --git a/Engine/Plugins/MovieScene/MovieRenderPipeline/Source/MovieRenderPipelineCore/Private/Graph/MovieGraphOutputMerger.h b/Engine/Plugins/MovieScene/MovieRenderPipeline/Source/MovieRenderPipelineCore/Private/Graph/MovieGraphOutputMerger.h
--- a/Engine/Plugins/MovieScene/MovieRenderPipeline/Source/MovieRenderPipelineCore/Private/Graph/MovieGraphOutputMerger.h
+++ b/Engine/Plugins/MovieScene/MovieRenderPipeline/Source/MovieRenderPipelineCore/Private/Graph/MovieGraphOutputMerger.h
@@ -5,6 +5,7 @@
 #include "Graph/MovieGraphDataTypes.h"
 #include "Containers/Queue.h"
 
+class IImageWriteQueue;
 class UMovieGraphPipeline;
 
 namespace UE::MovieGraph
@@ -37,6 +38,7 @@ namespace UE::MovieGraph
 		* All frames in the queue are processed on the next Game Thread tick and moved onto the output containers.
 		*/
 		TQueue<UE::MovieGraph::FMovieGraphOutputMergerFrame> FinishedFrames;
+
 	private:
 		/** The Movie Pipeline that owns us. */
 		TWeakObjectPtr<UMovieGraphPipeline> WeakMoviePipeline;
@@ -46,5 +48,8 @@ namespace UE::MovieGraph
 
 		/** Mutex that protects adding/updating/removing from PendingData */
 		FCriticalSection PendingDataMutex;
+
+		/** A debug image sequence writer (EXR-only). Typically used for writing out individual samples, bypassing accumulation and output containers. */
+		IImageWriteQueue* Debug_ImageWriteQueue;
 	};
 }
\ No newline at end of file
diff --git a/Engine/Plugins/MovieScene/MovieRenderPipeline/Source/MovieRenderPipelineCore/Private/Graph/MovieGraphPipeline.cpp b/Engine/Plugins/MovieScene/MovieRenderPipeline/Source/MovieRenderPipelineCore/Private/Graph/MovieGraphPipeline.cpp
--- a/Engine/Plugins/MovieScene/MovieRenderPipeline/Source/MovieRenderPipelineCore/Private/Graph/MovieGraphPipeline.cpp
+++ b/Engine/Plugins/MovieScene/MovieRenderPipeline/Source/MovieRenderPipelineCore/Private/Graph/MovieGraphPipeline.cpp
@@ -26,7 +26,6 @@
 #include "Graph/Nodes/MovieGraphRenderPassNode.h"
 
 #include "HAL/PlatformFileManager.h"
-#include "ImageWriteQueue.h"
 #include "RenderingThread.h"
 #include "Misc/CoreDelegates.h"
 #include "Modules/ModuleManager.h"
@@ -42,8 +41,6 @@ UMovieGraphPipeline::UMovieGraphPipeline()
 {
 	OutputMerger = MakeShared<UE::MovieGraph::FMovieGraphOutputMerger>(this);
 	CustomEngineTimeStep = CreateDefaultSubobject<UMovieGraphEngineTimeStep>("MovieGraphEngineTimeStep");
-
-	Debug_ImageWriteQueue = &FModuleManager::Get().LoadModuleChecked<IImageWriteQueueModule>("ImageWriteQueue").GetWriteQueue();
 }
 
 void UMovieGraphPipeline::Initialize(UMoviePipelineExecutorJob* InJob, const FMovieGraphInitConfig& InitConfig)
diff --git a/Engine/Plugins/MovieScene/MovieRenderPipeline/Source/MovieRenderPipelineCore/Public/Graph/MovieGraphPipeline.h b/Engine/Plugins/MovieScene/MovieRenderPipeline/Source/MovieRenderPipelineCore/Public/Graph/MovieGraphPipeline.h
--- a/Engine/Plugins/MovieScene/MovieRenderPipeline/Source/MovieRenderPipelineCore/Public/Graph/MovieGraphPipeline.h
+++ b/Engine/Plugins/MovieScene/MovieRenderPipeline/Source/MovieRenderPipelineCore/Public/Graph/MovieGraphPipeline.h
@@ -304,9 +304,6 @@ protected:
 	/** This gathers all of the produced data for an output frame (which may come in async, many frames later) before passing them onto the Output Containers. */
 	TSharedPtr<UE::MovieGraph::IMovieGraphOutputMerger> OutputMerger;
 
-	/** A debug image sequence writer in the event they want to dump every sample generated on its own. */
-	IImageWriteQueue* Debug_ImageWriteQueue;
-
 	/** 
 	* An array of Output Futures for files that have started writing to disk, but have not finished. 
 	* Each frame we take each future that finishes writing to disk and push it into the GeneratedOutputData array.
diff --git a/Engine/Plugins/MovieScene/MovieRenderPipeline/Source/MovieRenderPipelineCore/Public/Graph/MovieGraphTimeStepData.h b/Engine/Plugins/MovieScene/MovieRenderPipeline/Source/MovieRenderPipelineCore/Public/Graph/MovieGraphTimeStepData.h
--- a/Engine/Plugins/MovieScene/MovieRenderPipeline/Source/MovieRenderPipelineCore/Public/Graph/MovieGraphTimeStepData.h
+++ b/Engine/Plugins/MovieScene/MovieRenderPipeline/Source/MovieRenderPipelineCore/Public/Graph/MovieGraphTimeStepData.h
@@ -29,6 +29,8 @@ public:
 		, FrameRate(0, 0)
 		, TemporalSampleIndex(0)
 		, TemporalSampleCount(0)
+		, SpatialSampleIndex(0)
+		, SpatialSampleCount(0)
 		, bIsFirstTemporalSampleForFrame(false)
 		, bIsLastTemporalSampleForFrame(false)
 		, bDiscardOutput(false)
@@ -78,6 +80,14 @@ public:
 	UPROPERTY(BlueprintReadWrite, EditAnywhere, Category = "Movie Graph")
 	int32 TemporalSampleCount;
 
+	/** Index out of SpatialSampleCount that we're on. */
+	UPROPERTY(BlueprintReadWrite, EditAnywhere, Category = "Movie Graph")
+	int32 SpatialSampleIndex;
+
+	/** The maximum number of Spatial Samples this frame is expected to be able to process. */
+	UPROPERTY(BlueprintReadWrite, EditAnywhere, Category = "Movie Graph")
+	int32 SpatialSampleCount;
+
 	/** 
 	* Should be set to true for the first sample of each output frame. Used to determine
 	* if various systems need to reset or gather data for a new frame. Can be true at
diff --git a/Engine/Plugins/MovieScene/MovieRenderPipeline/Source/MovieRenderPipelineRenderPasses/Private/Graph/Renderers/MovieGraphDeferredPass.cpp b/Engine/Plugins/MovieScene/MovieRenderPipeline/Source/MovieRenderPipelineRenderPasses/Private/Graph/Renderers/MovieGraphDeferredPass.cpp
--- a/Engine/Plugins/MovieScene/MovieRenderPipeline/Source/MovieRenderPipelineRenderPasses/Private/Graph/Renderers/MovieGraphDeferredPass.cpp
+++ b/Engine/Plugins/MovieScene/MovieRenderPipeline/Source/MovieRenderPipelineRenderPasses/Private/Graph/Renderers/MovieGraphDeferredPass.cpp
@@ -255,6 +255,8 @@ void FMovieGraphDeferredPass::Render(const FMovieGraphTraversalContext& InFrameT
 			// Take our per-frame Traversal Context and update it with context specific to this sample.
 			FMovieGraphTraversalContext UpdatedTraversalContext = InFrameTraversalContext;
 			UpdatedTraversalContext.Time = InTimeData;
+			UpdatedTraversalContext.Time.SpatialSampleIndex = SpatialIndex;
+			UpdatedTraversalContext.Time.SpatialSampleCount = NumSpatialSamples;
 			UpdatedTraversalContext.RenderDataIdentifier = RenderDataIdentifier;
 
 			SampleState.TraversalContext = MoveTemp(UpdatedTraversalContext);
@@ -319,8 +321,8 @@ void FMovieGraphDeferredPass::Render(const FMovieGraphTraversalContext& InFrameT
 		// Submit the renderer to be rendered
 		GetRendererModule().BeginRenderingViewFamily(&Canvas, ViewFamily.ToSharedPtr().Get());
 
-		// If this was just to contribute to the history buffer, no need to go any further.
-		bool bDiscardOutput = InTimeData.bDiscardOutput || ShouldDiscardOutput(ViewFamily, CameraInfo);
+		// If this was just to contribute to the history buffer, no need to go any further. Never discard if we're writing individual samples, though.
+		bool bDiscardOutput = (InTimeData.bDiscardOutput || ShouldDiscardOutput(ViewFamily, CameraInfo)) && !SampleState.bWriteSampleToDisk;
 		if (bDiscardOutput)
 		{
 			continue;
@@ -372,7 +374,6 @@ void FMovieGraphDeferredPass::Render(const FMovieGraphTraversalContext& InFrameT
 		{
 			FramesToDelayPostSubmission--;
 		}
-
 	}
 }
 
diff --git a/Engine/Plugins/MovieScene/MovieRenderPipeline/Source/MovieRenderPipelineRenderPasses/Public/Graph/Nodes/MovieGraphPathTracerPassNode.h b/Engine/Plugins/MovieScene/MovieRenderPipeline/Source/MovieRenderPipelineRenderPasses/Public/Graph/Nodes/MovieGraphPathTracerPassNode.h
--- a/Engine/Plugins/MovieScene/MovieRenderPipeline/Source/MovieRenderPipelineRenderPasses/Public/Graph/Nodes/MovieGraphPathTracerPassNode.h
+++ b/Engine/Plugins/MovieScene/MovieRenderPipeline/Source/MovieRenderPipelineRenderPasses/Public/Graph/Nodes/MovieGraphPathTracerPassNode.h
@@ -156,9 +156,11 @@ public:
 	int32 FrameCount;
 
 	/**
-	* Debug Feature. Not currently marked BlueprintReadWrite/EditAnywhere as it's not totally implemented on the Path Tracer right now.
+	* Debug Feature. This can be used to write out each individual spatial sample rendered by this render pass,
+	* which allows you to see which images are being accumulated together. Can be useful for debugging incorrect looking
+	* frames to see which sub-frame evaluations were incorrect.
 	*/
-	UPROPERTY(DisplayName = "Write All Samples (Debug)", meta = (EditCondition = "bOverride_bWriteAllsamples"))
+	UPROPERTY(EditAnywhere, BlueprintReadWrite, Category = "Sampling", AdvancedDisplay, DisplayName = "Write All Samples (Debug)", meta = (EditCondition = "bOverride_bWriteAllsamples"))
 	bool bWriteAllSamples;
 
 	/**
