MoveRenderQueue: Fix pixel offset in path traced renders from the spatial jitter being applied during accumulation

Override CalculateViewFamily for the path tracer so we can zero out the fields that impact how the image is accumulated.

#rnx

diff --git a/Engine/Plugins/MovieScene/MovieRenderPipeline/Source/MovieRenderPipelineRenderPasses/Private/MoviePipelineDeferredPasses.cpp b/Engine/Plugins/MovieScene/MovieRenderPipeline/Source/MovieRenderPipelineRenderPasses/Private/MoviePipelineDeferredPasses.cpp
--- a/Engine/Plugins/MovieScene/MovieRenderPipeline/Source/MovieRenderPipelineRenderPasses/Private/MoviePipelineDeferredPasses.cpp
+++ b/Engine/Plugins/MovieScene/MovieRenderPipeline/Source/MovieRenderPipelineRenderPasses/Private/MoviePipelineDeferredPasses.cpp
@@ -1088,6 +1088,15 @@ void UMoviePipelineDeferredPass_PathTracer::SetupImpl(const MoviePipeline::FMovi
 	Super::SetupImpl(InPassInitSettings);
 }
 
+TSharedPtr<FSceneViewFamilyContext> UMoviePipelineDeferredPass_PathTracer::CalculateViewFamily(FMoviePipelineRenderPassMetrics& InOutSampleState, IViewCalcPayload* OptPayload)
+{
+	// remove sub-pixel shift, since the path tracer does its own anti-aliasing
+	InOutSampleState.SpatialShiftX = 0;
+	InOutSampleState.SpatialShiftY = 0;
+	InOutSampleState.OverlappedSubpixelShift = FVector2d(0.5, 0.5);
+	return Super::CalculateViewFamily(InOutSampleState, OptPayload);
+}
+
 bool UMoviePipelineDeferredPassBase::IsUsingDataLayers() const
 {
 	int32 NumDataLayers = 0;
diff --git a/Engine/Plugins/MovieScene/MovieRenderPipeline/Source/MovieRenderPipelineRenderPasses/Public/MoviePipelineDeferredPasses.h b/Engine/Plugins/MovieScene/MovieRenderPipeline/Source/MovieRenderPipelineRenderPasses/Public/MoviePipelineDeferredPasses.h
--- a/Engine/Plugins/MovieScene/MovieRenderPipeline/Source/MovieRenderPipelineRenderPasses/Public/MoviePipelineDeferredPasses.h
+++ b/Engine/Plugins/MovieScene/MovieRenderPipeline/Source/MovieRenderPipelineRenderPasses/Public/MoviePipelineDeferredPasses.h
@@ -327,6 +327,8 @@ public:
 	virtual void ValidateStateImpl() override;
 	virtual void SetupImpl(const MoviePipeline::FMoviePipelineRenderPassInitSettings& InPassInitSettings) override;
 
+	virtual TSharedPtr<FSceneViewFamilyContext> CalculateViewFamily(FMoviePipelineRenderPassMetrics& InOutSampleState, IViewCalcPayload* OptPayload) override;
+
 	virtual bool NeedsFrameThrottle() const override { return true; }
 
 	/** When enabled, the path tracer will blend all spatial and temporal samples prior to the denoising and will disable post-processed motion blur.
