MoviePipeline: Fixed a possible issue when doing 16->8 bit quantization where dithering could result in a 8 bit overflow, wrapping around to zero.

This would happen on rare occasions (1-2 times out of a million) where the 16-bit float had a value of 255.0, and the random [0, 1) value we added to that (to create dithering during conversion for softer gradients) picked a very high number within the 0-1 range, ie: 0.999996185.

As a decimal, this value is totally valid (floor(255.999996185) cast to uint8 would be 255, no overflow allowed). However due to the precision of floating point numbers, 255.999996185 can't actually be represented, the hardware rounds it up to 256, which is then floored (to 256) and cast to uint8 - wrapping it around to zero.

diff --git a/Engine/Plugins/MovieScene/MovieRenderPipeline/Source/MovieRenderPipelineCore/Private/MoviePipelineImageQuantization.cpp b/Engine/Plugins/MovieScene/MovieRenderPipeline/Source/MovieRenderPipelineCore/Private/MoviePipelineImageQuantization.cpp
--- a/Engine/Plugins/MovieScene/MovieRenderPipeline/Source/MovieRenderPipelineCore/Private/MoviePipelineImageQuantization.cpp
+++ b/Engine/Plugins/MovieScene/MovieRenderPipeline/Source/MovieRenderPipelineCore/Private/MoviePipelineImageQuantization.cpp
@@ -182,7 +182,20 @@ static TArray<FColor> ConvertLinearTosRGB8bppViaLookupTable(FFloat16Color* InCol
 	RandStreamTable.SetNumUninitialized(TableSize);
 	for (int32 Index = 0; Index < TableSize; Index++)
 	{
-		RandStreamTable[Index] = RandStream.GetFraction();
+		// GetFraction() returns [0,1) as desired, but it occasionally contains
+		// too much precision - the sRGBTable will return values of 255.0f, and
+		// GetFraction() can return values like 0.999996185, which when taken
+		// as decimals is a valid value (255.99999996185), and then below we floor it
+		// before casting to uint8, which should be fine.
+
+		// In practice however, we can't actually represent 255.999996185 with a 
+		// single-point precision float, it gets rounded up to 256 by the hardware,
+		// then floored, and then cast to uint8 overflowing to zero.
+
+		// To fix this, we limit this float to 15 bits of mantissa precision.
+		float Value = RandStream.GetFraction();
+		*(uint32*)&Value &= 0xFFFFFF00U;
+		RandStreamTable[Index] = Value;
 	}
 	// Convert all of our pixels.
 	TArray<FColor> OutsRGBData;
