MoviePipeline: Added the ability to toggle Sections in the details panel in the graph.

#rnx

diff --git a/Engine/Plugins/Media/AppleProResMedia/Source/AppleProResMedia/Private/MovieGraphAppleProResNode.h b/Engine/Plugins/Media/AppleProResMedia/Source/AppleProResMedia/Private/MovieGraphAppleProResNode.h
--- a/Engine/Plugins/Media/AppleProResMedia/Source/AppleProResMedia/Private/MovieGraphAppleProResNode.h
+++ b/Engine/Plugins/Media/AppleProResMedia/Source/AppleProResMedia/Private/MovieGraphAppleProResNode.h
@@ -9,7 +9,7 @@
 #include "MovieGraphAppleProResNode.generated.h"
 
 /** A node which can output Apple ProRes movies. */
-UCLASS(BlueprintType, PrioritizeCategories=("FileOutput"))
+UCLASS(BlueprintType)
 class UMovieGraphAppleProResNode : public UMovieGraphVideoOutputNode
 {
 	GENERATED_BODY()
diff --git a/Engine/Plugins/Media/AvidDNxMedia/Source/Source/Private/MovieGraphAvidDNxHRNode.h b/Engine/Plugins/Media/AvidDNxMedia/Source/Source/Private/MovieGraphAvidDNxHRNode.h
--- a/Engine/Plugins/Media/AvidDNxMedia/Source/Source/Private/MovieGraphAvidDNxHRNode.h
+++ b/Engine/Plugins/Media/AvidDNxMedia/Source/Source/Private/MovieGraphAvidDNxHRNode.h
@@ -11,7 +11,7 @@
 class FAvidDNxEncoder;
 
 /** A node which can output Avid DNxHR movies. */
-UCLASS(BlueprintType, PrioritizeCategories=("FileOutput"))
+UCLASS(BlueprintType)
 class UMovieGraphAvidDNxHRNode : public UMovieGraphVideoOutputNode
 {
 	GENERATED_BODY()
diff --git a/Engine/Plugins/MovieScene/MovieRenderPipeline/Source/MovieRenderPipelineEditor/Private/Customizations/Graph/MovieGraphNodeCustomization.h b/Engine/Plugins/MovieScene/MovieRenderPipeline/Source/MovieRenderPipelineEditor/Private/Customizations/Graph/MovieGraphNodeCustomization.h
--- a/Engine/Plugins/MovieScene/MovieRenderPipeline/Source/MovieRenderPipelineEditor/Private/Customizations/Graph/MovieGraphNodeCustomization.h
+++ b/Engine/Plugins/MovieScene/MovieRenderPipeline/Source/MovieRenderPipelineEditor/Private/Customizations/Graph/MovieGraphNodeCustomization.h
@@ -54,6 +54,9 @@ protected:
 
 		// Collapse the Tags category by default because it's not used by most users.
 		DetailBuilder.EditCategory("Tags").InitiallyCollapsed(true).SetSortOrder(1000 * ECategoryPriority::Uncommon);
+
+		// The File Output category should be shown first in all output nodes that have this category.
+		DetailBuilder.EditCategory("File Output").SetSortOrder(ECategoryPriority::Important);
 	}
 	//~ End IDetailCustomization interface
 };
diff --git a/Engine/Plugins/MovieScene/MovieRenderPipeline/Source/MovieRenderPipelineEditor/Private/Graph/MovieGraphAssetToolkit.cpp b/Engine/Plugins/MovieScene/MovieRenderPipeline/Source/MovieRenderPipelineEditor/Private/Graph/MovieGraphAssetToolkit.cpp
--- a/Engine/Plugins/MovieScene/MovieRenderPipeline/Source/MovieRenderPipelineEditor/Private/Graph/MovieGraphAssetToolkit.cpp
+++ b/Engine/Plugins/MovieScene/MovieRenderPipeline/Source/MovieRenderPipelineEditor/Private/Graph/MovieGraphAssetToolkit.cpp
@@ -36,6 +36,7 @@
 #include "Framework/Commands/GenericCommands.h"
 #include "GraphEditor.h"
 #include "IDetailRootObjectCustomization.h"
+#include "ObjectEditorUtils.h"
 #include "PropertyEditorModule.h"
 #include "Selection.h"
 #include "ToolMenus.h"
@@ -369,6 +370,9 @@ TSharedRef<SDockTab> FMovieGraphAssetToolkit::SpawnTab_RenderGraphDetails(const
 	DetailsViewArgs.bAllowMultipleTopLevelObjects = true;
 	DetailsViewArgs.ViewIdentifier = "MovieGraphSettings";
 	DetailsViewArgs.bLockable = false;
+	DetailsViewArgs.bShowSectionSelector = true;
+
+	RegisterDetailsViewSections();
 
 	SelectedGraphObjectsDetailsWidget = PropertyEditorModule.CreateDetailView(DetailsViewArgs);
 	SelectedGraphObjectsDetailsWidget->SetRootObjectCustomizationInstance(
@@ -473,6 +477,34 @@ TSharedRef<SDockTab> FMovieGraphAssetToolkit::SpawnTab_RenderGraphActiveRenderSe
 		];
 }
 
+void FMovieGraphAssetToolkit::RegisterDetailsViewSections()
+{
+	static const FName PropertyEditor("PropertyEditor");
+	static const FName CategoryMetadataKey("Category");
+	static const TArray<FString> CategoriesToExclude = { FString(TEXT("Tags")) };
+	
+	FPropertyEditorModule& PropertyModule = FModuleManager::GetModuleChecked<FPropertyEditorModule>(PropertyEditor);
+
+	// Do a 1:1 mapping of node categories to section names
+	for (const UClass* NodeClass : UMovieGraphSchema::GetNodeClasses())
+	{
+		for (TFieldIterator<FProperty> PropertyIt(NodeClass); PropertyIt; ++PropertyIt)
+		{
+			const FProperty* NodeProperty = *PropertyIt;
+
+			const FString* CategoryMetadata = NodeProperty->FindMetaData(CategoryMetadataKey);
+			if (CategoryMetadata && !CategoryMetadata->IsEmpty() && !CategoriesToExclude.Contains(*CategoryMetadata))
+			{
+				const FName SectionName = FName(*CategoryMetadata);
+				const FText LocalizedCategory = FObjectEditorUtils::GetCategoryText(NodeProperty);
+				
+				const TSharedRef<FPropertySection> Section = PropertyModule.FindOrCreateSection(NodeClass->GetFName(), SectionName, LocalizedCategory);
+				Section->AddCategory(SectionName);
+			}
+		}
+	}
+}
+
 void FMovieGraphAssetToolkit::BindGraphCommands()
 {
 	ToolkitCommands->MapAction(FGenericCommands::Get().Delete,
diff --git a/Engine/Plugins/MovieScene/MovieRenderPipeline/Source/MovieRenderPipelineEditor/Private/Graph/MovieGraphAssetToolkit.h b/Engine/Plugins/MovieScene/MovieRenderPipeline/Source/MovieRenderPipelineEditor/Private/Graph/MovieGraphAssetToolkit.h
--- a/Engine/Plugins/MovieScene/MovieRenderPipeline/Source/MovieRenderPipelineEditor/Private/Graph/MovieGraphAssetToolkit.h
+++ b/Engine/Plugins/MovieScene/MovieRenderPipeline/Source/MovieRenderPipelineEditor/Private/Graph/MovieGraphAssetToolkit.h
@@ -53,6 +53,9 @@ private:
 	TSharedRef<SDockTab> SpawnTab_RenderGraphDetails(const FSpawnTabArgs& Args);
 	TSharedRef<SDockTab> SpawnTab_RenderGraphActiveRenderSettings(const FSpawnTabArgs& Args);
 
+	/** Register sections used by the graph's details panel. */
+	void RegisterDetailsViewSections();
+
 	/** Register commands used by the editor. */
 	void BindGraphCommands();
 
diff --git a/Engine/Plugins/MovieScene/MovieRenderPipeline/Source/MovieRenderPipelineEditor/Private/Graph/MovieGraphSchema.cpp b/Engine/Plugins/MovieScene/MovieRenderPipeline/Source/MovieRenderPipelineEditor/Private/Graph/MovieGraphSchema.cpp
--- a/Engine/Plugins/MovieScene/MovieRenderPipeline/Source/MovieRenderPipelineEditor/Private/Graph/MovieGraphSchema.cpp
+++ b/Engine/Plugins/MovieScene/MovieRenderPipeline/Source/MovieRenderPipelineEditor/Private/Graph/MovieGraphSchema.cpp
@@ -212,6 +212,16 @@ void UMovieGraphSchema::InitMoviePipelineNodeClasses()
 	MoviePipelineNodeClasses.Sort();
 }
 
+const TArray<UClass*>& UMovieGraphSchema::GetNodeClasses()
+{
+	if (MoviePipelineNodeClasses.IsEmpty())
+	{
+		InitMoviePipelineNodeClasses();
+	}
+	
+	return MoviePipelineNodeClasses;
+}
+
 bool UMovieGraphSchema::IsConnectionToBranchAllowed(const UEdGraphPin* InputPin, const UEdGraphPin* OutputPin, FText& OutError) const
 {
 	const UMovieGraphPin* ToPin = UE::MovieGraph::Private::GetGraphPinFromEdPin(InputPin);
diff --git a/Engine/Plugins/MovieScene/MovieRenderPipeline/Source/MovieRenderPipelineEditor/Private/Graph/MovieGraphSchema.h b/Engine/Plugins/MovieScene/MovieRenderPipeline/Source/MovieRenderPipelineEditor/Private/Graph/MovieGraphSchema.h
--- a/Engine/Plugins/MovieScene/MovieRenderPipeline/Source/MovieRenderPipelineEditor/Private/Graph/MovieGraphSchema.h
+++ b/Engine/Plugins/MovieScene/MovieRenderPipeline/Source/MovieRenderPipelineEditor/Private/Graph/MovieGraphSchema.h
@@ -33,6 +33,9 @@ public:
 	
 	static void InitMoviePipelineNodeClasses();
 
+	/** Gets all node classes that the graph is aware of. Should only be called after all plugins that contain nodes have been loaded. */
+	static const TArray<UClass*>& GetNodeClasses();
+
 private:
 	/**
 	 * Determines if the connection between InputPin and OutputPin follows branch restriction rules. OutError is populated
diff --git a/Engine/Plugins/MovieScene/MovieRenderPipeline/Source/MovieRenderPipelineRenderPasses/Public/MovieGraphImageSequenceOutputNode.h b/Engine/Plugins/MovieScene/MovieRenderPipeline/Source/MovieRenderPipelineRenderPasses/Public/MovieGraphImageSequenceOutputNode.h
--- a/Engine/Plugins/MovieScene/MovieRenderPipeline/Source/MovieRenderPipelineRenderPasses/Public/MovieGraphImageSequenceOutputNode.h
+++ b/Engine/Plugins/MovieScene/MovieRenderPipeline/Source/MovieRenderPipelineRenderPasses/Public/MovieGraphImageSequenceOutputNode.h
@@ -19,7 +19,7 @@ class IImageWriteQueue;
 * a series of jpeg, png, bmp, or .exr images. Create an instance of the appropriate class (such as 
 * UMovieGraphImageSequenceOutputNode_JPG) instead of this abstract base class.
 */
-UCLASS(Abstract, PrioritizeCategories=("FileOutput"))
+UCLASS(Abstract)
 class UMovieGraphImageSequenceOutputNode : public UMovieGraphFileOutputNode
 {
 	GENERATED_BODY()
