MoviePipeline: Added additional telemetry for MRQ and MRG.

#rnx

diff --git a/Engine/Plugins/MovieScene/MoviePipelineMaskRenderPass/Source/MoviePipelineMaskRenderPass/Private/Graph/MovieGraphObjectIdNode.cpp b/Engine/Plugins/MovieScene/MoviePipelineMaskRenderPass/Source/MoviePipelineMaskRenderPass/Private/Graph/MovieGraphObjectIdNode.cpp
--- a/Engine/Plugins/MovieScene/MoviePipelineMaskRenderPass/Source/MoviePipelineMaskRenderPass/Private/Graph/MovieGraphObjectIdNode.cpp
+++ b/Engine/Plugins/MovieScene/MoviePipelineMaskRenderPass/Source/MoviePipelineMaskRenderPass/Private/Graph/MovieGraphObjectIdNode.cpp
@@ -5,6 +5,7 @@
 #include "Editor/EditorPerProjectUserSettings.h"
 #include "MovieGraphObjectIdPass.h"
 #include "MoviePipelineObjectIdUtils.h"
+#include "MoviePipelineTelemetry.h"
 
 FUObjectAnnotationSparse<UE::MoviePipeline::FObjectIdAccelerationData, true>& UMovieGraphObjectIdNode::GetManifestAnnotation()
 {
@@ -54,6 +55,11 @@ FSlateIcon UMovieGraphObjectIdNode::GetIconAndTint(FLinearColor& OutColor) const
 }
 #endif	// WITH_EDITOR
 
+void UMovieGraphObjectIdNode::UpdateTelemetry(FMoviePipelineShotRenderTelemetry* InTelemetry) const
+{
+	InTelemetry->bUsesObjectID = true;
+}
+
 TUniquePtr<UE::MovieGraph::Rendering::FMovieGraphImagePassBase> UMovieGraphObjectIdNode::CreateInstance() const
 {
 	return MakeUnique<FMovieGraphObjectIdPass>();
diff --git a/Engine/Plugins/MovieScene/MoviePipelineMaskRenderPass/Source/MoviePipelineMaskRenderPass/Private/Graph/MovieGraphObjectIdNode.h b/Engine/Plugins/MovieScene/MoviePipelineMaskRenderPass/Source/MoviePipelineMaskRenderPass/Private/Graph/MovieGraphObjectIdNode.h
--- a/Engine/Plugins/MovieScene/MoviePipelineMaskRenderPass/Source/MoviePipelineMaskRenderPass/Private/Graph/MovieGraphObjectIdNode.h
+++ b/Engine/Plugins/MovieScene/MoviePipelineMaskRenderPass/Source/MoviePipelineMaskRenderPass/Private/Graph/MovieGraphObjectIdNode.h
@@ -34,6 +34,8 @@ public:
 #endif
 	// ~UMovieGraphNode Interface
 
+	virtual void UpdateTelemetry(FMoviePipelineShotRenderTelemetry* InTelemetry) const override;
+
 protected:
 	// UMovieGraphImagePassBaseNode Interface
 	virtual TUniquePtr<UE::MovieGraph::Rendering::FMovieGraphImagePassBase> CreateInstance() const override;
diff --git a/Engine/Plugins/MovieScene/MoviePipelineMaskRenderPass/Source/MoviePipelineMaskRenderPass/Private/MoviePipelineObjectIdPass.cpp b/Engine/Plugins/MovieScene/MoviePipelineMaskRenderPass/Source/MoviePipelineMaskRenderPass/Private/MoviePipelineObjectIdPass.cpp
--- a/Engine/Plugins/MovieScene/MoviePipelineMaskRenderPass/Source/MoviePipelineMaskRenderPass/Private/MoviePipelineObjectIdPass.cpp
+++ b/Engine/Plugins/MovieScene/MoviePipelineMaskRenderPass/Source/MoviePipelineMaskRenderPass/Private/MoviePipelineObjectIdPass.cpp
@@ -9,6 +9,7 @@
 #include "MoviePipelineSurfaceReader.h"
 #include "Serialization/JsonSerializer.h"
 #include "MoviePipelineHashUtils.h"
+#include "MoviePipelineTelemetry.h"
 #include "EngineModule.h"
 #include "Materials/Material.h"
 #include "Materials/MaterialInstanceDynamic.h"
@@ -52,6 +53,11 @@ UMoviePipelineObjectIdRenderPass::UMoviePipelineObjectIdRenderPass()
 	}
 }
 
+void UMoviePipelineObjectIdRenderPass::UpdateTelemetry(FMoviePipelineShotRenderTelemetry* InTelemetry) const
+{
+	InTelemetry->bUsesObjectID = true;
+}
+
 TWeakObjectPtr<UTextureRenderTarget2D> UMoviePipelineObjectIdRenderPass::CreateViewRenderTargetImpl(const FIntPoint& InSize, IViewCalcPayload* OptPayload) const
 {
 	// Crete the render target with the correct bit depth.
diff --git a/Engine/Plugins/MovieScene/MoviePipelineMaskRenderPass/Source/MoviePipelineMaskRenderPass/Private/MoviePipelineObjectIdPass.h b/Engine/Plugins/MovieScene/MoviePipelineMaskRenderPass/Source/MoviePipelineMaskRenderPass/Private/MoviePipelineObjectIdPass.h
--- a/Engine/Plugins/MovieScene/MoviePipelineMaskRenderPass/Source/MoviePipelineMaskRenderPass/Private/MoviePipelineObjectIdPass.h
+++ b/Engine/Plugins/MovieScene/MoviePipelineMaskRenderPass/Source/MoviePipelineMaskRenderPass/Private/MoviePipelineObjectIdPass.h
@@ -22,6 +22,7 @@ public:
 	virtual void GatherOutputPassesImpl(TArray<FMoviePipelinePassIdentifier>& ExpectedRenderPasses) override;
 	virtual bool IsScreenPercentageSupported() const override { return false; }
 	virtual int32 GetOutputFileSortingOrder() const override { return 10; }
+	virtual void UpdateTelemetry(FMoviePipelineShotRenderTelemetry* InTelemetry) const override;
 	
 	virtual TWeakObjectPtr<UTextureRenderTarget2D> CreateViewRenderTargetImpl(const FIntPoint& InSize, IViewCalcPayload* OptPayload = nullptr) const override;
 	virtual TSharedPtr<FMoviePipelineSurfaceQueue, ESPMode::ThreadSafe> CreateSurfaceQueueImpl(const FIntPoint& InSize, IViewCalcPayload* OptPayload = nullptr) const override;
diff --git a/Engine/Plugins/MovieScene/MoviePipelineMaskRenderPass/Source/MoviePipelineMaskRenderPass/Private/MoviePipelinePanoramicPass.cpp b/Engine/Plugins/MovieScene/MoviePipelineMaskRenderPass/Source/MoviePipelineMaskRenderPass/Private/MoviePipelinePanoramicPass.cpp
--- a/Engine/Plugins/MovieScene/MoviePipelineMaskRenderPass/Source/MoviePipelineMaskRenderPass/Private/MoviePipelinePanoramicPass.cpp
+++ b/Engine/Plugins/MovieScene/MoviePipelineMaskRenderPass/Source/MoviePipelineMaskRenderPass/Private/MoviePipelinePanoramicPass.cpp
@@ -12,6 +12,7 @@
 #include "MovieRenderOverlappedImage.h"
 #include "MoviePipelineOutputBuilder.h"
 #include "MoviePipelinePanoramicBlender.h"
+#include "MoviePipelineTelemetry.h"
 #include "OpenColorIODisplayExtension.h"
 #include "TextureResource.h"
 #include "SceneUtils.h"
@@ -558,3 +559,8 @@ void UMoviePipelinePanoramicPass::ScheduleReadbackAndAccumulation(const FMoviePi
 		});
 
 }
+
+void UMoviePipelinePanoramicPass::UpdateTelemetry(FMoviePipelineShotRenderTelemetry* InTelemetry) const
+{
+	InTelemetry->bUsesPanoramic = true;
+}
diff --git a/Engine/Plugins/MovieScene/MoviePipelineMaskRenderPass/Source/MoviePipelineMaskRenderPass/Private/MoviePipelinePanoramicPass.h b/Engine/Plugins/MovieScene/MoviePipelineMaskRenderPass/Source/MoviePipelineMaskRenderPass/Private/MoviePipelinePanoramicPass.h
--- a/Engine/Plugins/MovieScene/MoviePipelineMaskRenderPass/Source/MoviePipelineMaskRenderPass/Private/MoviePipelinePanoramicPass.h
+++ b/Engine/Plugins/MovieScene/MoviePipelineMaskRenderPass/Source/MoviePipelineMaskRenderPass/Private/MoviePipelinePanoramicPass.h
@@ -142,6 +142,9 @@ protected:
 	void GetFieldOfView(float& OutHorizontal, float& OutVertical, const bool bInStereo) const;
 	FIntPoint GetPaneResolution(const FIntPoint& InSize) const;
 	FIntPoint GetPayloadPaneResolution(const FIntPoint& InSize, IViewCalcPayload* OptPayload) const;
+
+	virtual void UpdateTelemetry(FMoviePipelineShotRenderTelemetry* InTelemetry) const override;
+
 public:
 
 	/**
diff --git a/Engine/Plugins/MovieScene/MovieRenderPipeline/Source/MovieRenderPipelineCore/Private/Graph/MovieGraphPipeline.cpp b/Engine/Plugins/MovieScene/MovieRenderPipeline/Source/MovieRenderPipelineCore/Private/Graph/MovieGraphPipeline.cpp
--- a/Engine/Plugins/MovieScene/MovieRenderPipeline/Source/MovieRenderPipelineCore/Private/Graph/MovieGraphPipeline.cpp
+++ b/Engine/Plugins/MovieScene/MovieRenderPipeline/Source/MovieRenderPipelineCore/Private/Graph/MovieGraphPipeline.cpp
@@ -3,6 +3,7 @@
 #include "Graph/MovieGraphPipeline.h"
 
 #include "MoviePipelineQueue.h"
+#include "MoviePipelineTelemetry.h"
 #include "MoviePipelineUtils.h"
 #include "MovieRenderPipelineCoreModule.h"
 #include "MovieScene.h"
@@ -903,7 +904,7 @@ void UMovieGraphPipeline::SetupShot(const TObjectPtr<UMoviePipelineExecutorShot>
 	//}
 
 	const FMovieGraphTimeStepData& TimeStepData = GetTimeStepInstance()->GetCalculatedTimeData();
-	const UMovieGraphEvaluatedConfig* EvaluatedConfig = TimeStepData.EvaluatedConfig;
+	UMovieGraphEvaluatedConfig* EvaluatedConfig = TimeStepData.EvaluatedConfig;
 
 	// Apply any global game overrides, which includes cvars. This needs to be done before the CVarManager sets cvars
 	// so any user-specified cvars can override cvars set via the global game overrides. Note that the CDO is intentionally
@@ -925,6 +926,8 @@ void UMovieGraphPipeline::SetupShot(const TObjectPtr<UMoviePipelineExecutorShot>
 
 	// Setup required rendering architecture for all passes in this shot.
 	GraphRendererInstance->SetupRenderingPipelineForShot(InShot);
+
+	FMoviePipelineTelemetry::SendBeginShotRenderTelemetry(InShot, EvaluatedConfig);
 }
 
 void UMovieGraphPipeline::TeardownShot(const TObjectPtr<UMoviePipelineExecutorShot>& InShot)
@@ -1012,6 +1015,9 @@ void UMovieGraphPipeline::TeardownShot(const TObjectPtr<UMoviePipelineExecutorSh
 		GlobalGameOverridesNode->ApplySettings(bOverrideValues, GetWorld());
 	}
 
+	constexpr bool bIsGraph = true;
+	FMoviePipelineTelemetry::SendEndShotRenderTelemetry(bIsGraph, !bShutdownSetErrorFlag, bShutdownRequested);
+
 	if (IsPostShotCallbackNeeded())
 	{
 		ProcessOutstandingFinishedFrames();
diff --git a/Engine/Plugins/MovieScene/MovieRenderPipeline/Source/MovieRenderPipelineCore/Private/Graph/Nodes/MovieGraphAudioOutputNode.cpp b/Engine/Plugins/MovieScene/MovieRenderPipeline/Source/MovieRenderPipelineCore/Private/Graph/Nodes/MovieGraphAudioOutputNode.cpp
--- a/Engine/Plugins/MovieScene/MovieRenderPipeline/Source/MovieRenderPipelineCore/Private/Graph/Nodes/MovieGraphAudioOutputNode.cpp
+++ b/Engine/Plugins/MovieScene/MovieRenderPipeline/Source/MovieRenderPipelineCore/Private/Graph/Nodes/MovieGraphAudioOutputNode.cpp
@@ -7,6 +7,7 @@
 #include "Graph/MovieGraphPipeline.h"
 #include "Graph/MovieGraphUtils.h"
 #include "Graph/Nodes/MovieGraphGlobalOutputSettingNode.h"
+#include "MoviePipelineTelemetry.h"
 #include "MoviePipelineUtils.h"
 #include "MovieRenderPipelineCoreModule.h"
 #include "Styling/AppStyle.h"
@@ -30,6 +31,11 @@ void UMovieGraphAudioOutputNode::BuildNewProcessCommandLineArgsImpl(TArray<FStri
 	InOutCommandLineArgs.Add("-deterministicaudio");
 }
 
+void UMovieGraphAudioOutputNode::UpdateTelemetry(FMoviePipelineShotRenderTelemetry* InTelemetry) const
+{
+	InTelemetry->bUsesAudio = true;
+}
+
 #if WITH_EDITOR
 FText UMovieGraphAudioOutputNode::GetNodeTitle(const bool bGetDescriptive) const
 {
diff --git a/Engine/Plugins/MovieScene/MovieRenderPipeline/Source/MovieRenderPipelineCore/Private/Graph/Nodes/MovieGraphExecuteScriptNode.cpp b/Engine/Plugins/MovieScene/MovieRenderPipeline/Source/MovieRenderPipelineCore/Private/Graph/Nodes/MovieGraphExecuteScriptNode.cpp
--- a/Engine/Plugins/MovieScene/MovieRenderPipeline/Source/MovieRenderPipelineCore/Private/Graph/Nodes/MovieGraphExecuteScriptNode.cpp
+++ b/Engine/Plugins/MovieScene/MovieRenderPipeline/Source/MovieRenderPipelineCore/Private/Graph/Nodes/MovieGraphExecuteScriptNode.cpp
@@ -1,6 +1,8 @@
 // Copyright Epic Games, Inc. All Rights Reserved.
 
 #include "Graph/Nodes/MovieGraphExecuteScriptNode.h"
+
+#include "MoviePipelineTelemetry.h"
 #include "Styling/AppStyle.h"
 #include "UObject/Package.h"
 
@@ -40,4 +42,9 @@ UMovieGraphScriptBase* UMovieGraphExecuteScriptNode::AllocateScriptInstance() co
 	}
 
 	return nullptr;
+}
+
+void UMovieGraphExecuteScriptNode::UpdateTelemetry(FMoviePipelineShotRenderTelemetry* InTelemetry) const
+{
+	InTelemetry->bUsesScripting = true;
 }
\ No newline at end of file
diff --git a/Engine/Plugins/MovieScene/MovieRenderPipeline/Source/MovieRenderPipelineCore/Private/Graph/Nodes/MovieGraphGlobalOutputSettingNode.cpp b/Engine/Plugins/MovieScene/MovieRenderPipeline/Source/MovieRenderPipelineCore/Private/Graph/Nodes/MovieGraphGlobalOutputSettingNode.cpp
--- a/Engine/Plugins/MovieScene/MovieRenderPipeline/Source/MovieRenderPipelineCore/Private/Graph/Nodes/MovieGraphGlobalOutputSettingNode.cpp
+++ b/Engine/Plugins/MovieScene/MovieRenderPipeline/Source/MovieRenderPipelineCore/Private/Graph/Nodes/MovieGraphGlobalOutputSettingNode.cpp
@@ -1,10 +1,12 @@
 // Copyright Epic Games, Inc. All Rights Reserved.
 
 #include "Graph/Nodes/MovieGraphGlobalOutputSettingNode.h"
+
+#include "Algo/Find.h"
 #include "Graph/MovieGraphProjectSettings.h"
 #include "Graph/MovieGraphBlueprintLibrary.h"
+#include "MoviePipelineTelemetry.h"
 #include "Styling/AppStyle.h"
-#include "Algo/Find.h"
 
 UMovieGraphGlobalOutputSettingNode::UMovieGraphGlobalOutputSettingNode()
 	: OutputFrameRate(FFrameRate(24, 1))
@@ -48,6 +50,11 @@ void UMovieGraphGlobalOutputSettingNode::GetFormatResolveArgs(FMovieGraphResolve
 	// We don't resolve the version here because that's handled on a per-file/shot basis
 }
 
+void UMovieGraphGlobalOutputSettingNode::UpdateTelemetry(FMoviePipelineShotRenderTelemetry* InTelemetry) const
+{
+	InTelemetry->HandleFrameCount = HandleFrameCount;
+}
+
 #if WITH_EDITOR
 FText UMovieGraphGlobalOutputSettingNode::GetNodeTitle(const bool bGetDescriptive) const
 {
diff --git a/Engine/Plugins/MovieScene/MovieRenderPipeline/Source/MovieRenderPipelineCore/Private/Graph/Nodes/MovieGraphRenderLayerNode.cpp b/Engine/Plugins/MovieScene/MovieRenderPipeline/Source/MovieRenderPipelineCore/Private/Graph/Nodes/MovieGraphRenderLayerNode.cpp
--- a/Engine/Plugins/MovieScene/MovieRenderPipeline/Source/MovieRenderPipelineCore/Private/Graph/Nodes/MovieGraphRenderLayerNode.cpp
+++ b/Engine/Plugins/MovieScene/MovieRenderPipeline/Source/MovieRenderPipelineCore/Private/Graph/Nodes/MovieGraphRenderLayerNode.cpp
@@ -3,6 +3,7 @@
 #include "Graph/Nodes/MovieGraphRenderLayerNode.h"
 
 #include "Graph/MovieGraphConfig.h"
+#include "MoviePipelineTelemetry.h"
 #include "Styling/AppStyle.h"
 
 UMovieGraphRenderLayerNode::UMovieGraphRenderLayerNode()
@@ -52,4 +53,9 @@ void UMovieGraphRenderLayerNode::PostEditChangeProperty(FPropertyChangedEvent& P
 		OnNodeChangedDelegate.Broadcast(this);
 	}
 }
-#endif // WITH_EDITOR
\ No newline at end of file
+#endif // WITH_EDITOR
+
+void UMovieGraphRenderLayerNode::UpdateTelemetry(FMoviePipelineShotRenderTelemetry* InTelemetry) const
+{
+	++InTelemetry->RenderLayerCount;
+}
\ No newline at end of file
diff --git a/Engine/Plugins/MovieScene/MovieRenderPipeline/Source/MovieRenderPipelineCore/Private/Graph/Nodes/MovieGraphSamplingMethodNode.cpp b/Engine/Plugins/MovieScene/MovieRenderPipeline/Source/MovieRenderPipelineCore/Private/Graph/Nodes/MovieGraphSamplingMethodNode.cpp
--- a/Engine/Plugins/MovieScene/MovieRenderPipeline/Source/MovieRenderPipelineCore/Private/Graph/Nodes/MovieGraphSamplingMethodNode.cpp
+++ b/Engine/Plugins/MovieScene/MovieRenderPipeline/Source/MovieRenderPipelineCore/Private/Graph/Nodes/MovieGraphSamplingMethodNode.cpp
@@ -4,6 +4,7 @@
 
 #include "Graph/MovieGraphConfig.h"
 #include "Graph/MovieGraphLinearTimeStep.h"
+#include "MoviePipelineTelemetry.h"
 #include "Styling/AppStyle.h"
 
 UMovieGraphSamplingMethodNode::UMovieGraphSamplingMethodNode()
@@ -65,4 +66,9 @@ FSlateIcon UMovieGraphSamplingMethodNode::GetIconAndTint(FLinearColor& OutColor)
 	OutColor = FLinearColor::White;
 	return SamplingMethodPresetIcon;
 }
-#endif // WITH_EDITOR
\ No newline at end of file
+#endif // WITH_EDITOR
+
+void UMovieGraphSamplingMethodNode::UpdateTelemetry(FMoviePipelineShotRenderTelemetry* InTelemetry) const
+{
+	InTelemetry->TemporalSampleCount = TemporalSampleCount;
+}
\ No newline at end of file
diff --git a/Engine/Plugins/MovieScene/MovieRenderPipeline/Source/MovieRenderPipelineCore/Private/MoviePipeline.cpp b/Engine/Plugins/MovieScene/MovieRenderPipeline/Source/MovieRenderPipelineCore/Private/MoviePipeline.cpp
--- a/Engine/Plugins/MovieScene/MovieRenderPipeline/Source/MovieRenderPipelineCore/Private/MoviePipeline.cpp
+++ b/Engine/Plugins/MovieScene/MovieRenderPipeline/Source/MovieRenderPipelineCore/Private/MoviePipeline.cpp
@@ -51,6 +51,7 @@
 #include "ChaosClothAsset/ClothComponent.h"
 #include "ChaosClothAsset/ClothAssetInteractor.h"
 #include "Components/SkeletalMeshComponent.h"
+#include "MoviePipelineTelemetry.h"
 #include "MovieSceneCommonHelpers.h"
 
 #if WITH_EDITOR
@@ -1013,6 +1014,8 @@ void UMoviePipeline::InitializeShot(UMoviePipelineExecutorShot* InShot)
 
 	// Setup required rendering architecture for all passes in this shot.
 	SetupRenderingPipelineForShot(InShot);
+
+	FMoviePipelineTelemetry::SendBeginShotRenderTelemetry(InShot);
 }
 
 void UMoviePipeline::TeardownShot(UMoviePipelineExecutorShot* InShot)
@@ -1069,6 +1072,9 @@ void UMoviePipeline::TeardownShot(UMoviePipelineExecutorShot* InShot)
 
 	RestoreSkeletalMeshClothSubSteps();
 	ClothSimCache.Reset();
+
+	constexpr bool bIsGraph = false;
+	FMoviePipelineTelemetry::SendEndShotRenderTelemetry(bIsGraph, !bFatalError, bShutdownRequested);
 	
 	if (IsFlushDiskWritesPerShot())
 	{
diff --git a/Engine/Plugins/MovieScene/MovieRenderPipeline/Source/MovieRenderPipelineCore/Private/MoviePipelineAntiAliasingSetting.cpp b/Engine/Plugins/MovieScene/MovieRenderPipeline/Source/MovieRenderPipelineCore/Private/MoviePipelineAntiAliasingSetting.cpp
--- a/Engine/Plugins/MovieScene/MovieRenderPipeline/Source/MovieRenderPipelineCore/Private/MoviePipelineAntiAliasingSetting.cpp
+++ b/Engine/Plugins/MovieScene/MovieRenderPipeline/Source/MovieRenderPipelineCore/Private/MoviePipelineAntiAliasingSetting.cpp
@@ -2,6 +2,8 @@
 
 #include "MoviePipelineAntiAliasingSetting.h"
 
+#include "MoviePipelineTelemetry.h"
+
 #include UE_INLINE_GENERATED_CPP_BY_NAME(MoviePipelineAntiAliasingSetting)
 
 void UMoviePipelineAntiAliasingSetting::ValidateStateImpl()
@@ -53,3 +55,9 @@ void UMoviePipelineAntiAliasingSetting::GetFormatArguments(FMoviePipelineFormatA
 	InOutFormatArgs.FileMetadata.Add(TEXT("unreal/aa/temporalSampleCount"), FString::FromInt(TemporalSampleCount));
 	InOutFormatArgs.FileMetadata.Add(TEXT("unreal/aa/spatialSampleCount"), FString::FromInt(SpatialSampleCount));
 }
+
+void UMoviePipelineAntiAliasingSetting::UpdateTelemetry(FMoviePipelineShotRenderTelemetry* InTelemetry) const
+{
+	InTelemetry->TemporalSampleCount = FMath::Max(InTelemetry->TemporalSampleCount, TemporalSampleCount);
+	InTelemetry->SpatialSampleCount = FMath::Max(InTelemetry->SpatialSampleCount, SpatialSampleCount);
+}
diff --git a/Engine/Plugins/MovieScene/MovieRenderPipeline/Source/MovieRenderPipelineCore/Private/MoviePipelineCameraSetting.cpp b/Engine/Plugins/MovieScene/MovieRenderPipeline/Source/MovieRenderPipelineCore/Private/MoviePipelineCameraSetting.cpp
new file mode 100644
--- /dev/null
+++ b/Engine/Plugins/MovieScene/MovieRenderPipeline/Source/MovieRenderPipelineCore/Private/MoviePipelineCameraSetting.cpp
@@ -0,0 +1,22 @@
+﻿// Copyright Epic Games, Inc. All Rights Reserved.
+
+#include "MoviePipelineCameraSetting.h"
+
+#include "MoviePipelineTelemetry.h"
+
+#include UE_INLINE_GENERATED_CPP_BY_NAME(MoviePipelineCameraSetting)
+
+void UMoviePipelineCameraSetting::GetFormatArguments(FMoviePipelineFormatArgs& InOutFormatArgs) const
+{
+	Super::GetFormatArguments(InOutFormatArgs);
+
+	InOutFormatArgs.FilenameArguments.Add(TEXT("shutter_timing"), StaticEnum<EMoviePipelineShutterTiming>()->GetNameStringByValue((int64)ShutterTiming));
+	InOutFormatArgs.FilenameArguments.Add(TEXT("overscan_percentage"), FString::SanitizeFloat(OverscanPercentage));
+	InOutFormatArgs.FileMetadata.Add(TEXT("unreal/camera/shutterTiming"), StaticEnum<EMoviePipelineShutterTiming>()->GetNameStringByValue((int64)ShutterTiming));
+	InOutFormatArgs.FileMetadata.Add(TEXT("unreal/camera/overscanPercentage"), FString::SanitizeFloat(OverscanPercentage));
+}
+
+void UMoviePipelineCameraSetting::UpdateTelemetry(FMoviePipelineShotRenderTelemetry* InTelemetry) const
+{
+	InTelemetry->bUsesMultiCamera |= bRenderAllCameras;
+}
\ No newline at end of file
diff --git a/Engine/Plugins/MovieScene/MovieRenderPipeline/Source/MovieRenderPipelineCore/Private/MoviePipelineOutputSetting.cpp b/Engine/Plugins/MovieScene/MovieRenderPipeline/Source/MovieRenderPipelineCore/Private/MoviePipelineOutputSetting.cpp
--- a/Engine/Plugins/MovieScene/MovieRenderPipeline/Source/MovieRenderPipelineCore/Private/MoviePipelineOutputSetting.cpp
+++ b/Engine/Plugins/MovieScene/MovieRenderPipeline/Source/MovieRenderPipelineCore/Private/MoviePipelineOutputSetting.cpp
@@ -6,6 +6,7 @@
 #include "MoviePipelineQueue.h"
 #include "MovieRenderPipelineDataTypes.h"
 #include "MoviePipeline.h"
+#include "MoviePipelineTelemetry.h"
 
 #include UE_INLINE_GENERATED_CPP_BY_NAME(MoviePipelineOutputSetting)
 
@@ -105,3 +106,7 @@ void UMoviePipelineOutputSetting::SetupForPipelineImpl(UMoviePipeline* InPipelin
 	}
 }
 
+void UMoviePipelineOutputSetting::UpdateTelemetry(FMoviePipelineShotRenderTelemetry* InTelemetry) const
+{
+	InTelemetry->HandleFrameCount = FMath::Max(InTelemetry->HandleFrameCount, HandleFrameCount);
+}
diff --git a/Engine/Plugins/MovieScene/MovieRenderPipeline/Source/MovieRenderPipelineCore/Private/MoviePipelineTelemetry.cpp b/Engine/Plugins/MovieScene/MovieRenderPipeline/Source/MovieRenderPipelineCore/Private/MoviePipelineTelemetry.cpp
new file mode 100644
--- /dev/null
+++ b/Engine/Plugins/MovieScene/MovieRenderPipeline/Source/MovieRenderPipelineCore/Private/MoviePipelineTelemetry.cpp
@@ -0,0 +1,211 @@
+﻿// Copyright Epic Games, Inc. All Rights Reserved.
+
+#include "MoviePipelineTelemetry.h"
+
+#include "EngineAnalytics.h"
+#include "Graph/MovieGraphBlueprintLibrary.h"
+#include "Graph/MovieGraphConfig.h"
+#include "MoviePipelineBlueprintLibrary.h"
+#include "MoviePipelineQueue.h"
+
+#if WITH_EDITOR
+#include "Editor.h"
+#endif
+
+void FMoviePipelineTelemetry::SendRendersRequestedTelemetry(const bool bIsLocal, TArray<UMoviePipelineExecutorJob*>&& InJobs)
+{
+	if (!FEngineAnalytics::IsAvailable())
+	{
+		return;
+	}
+	
+	int32 EnabledJobCount = 0;
+	int32 DisabledJobCount = 0;
+	int32 EnabledShotCount = 0;
+	int32 DisabledShotCount = 0;
+
+	for (const UMoviePipelineExecutorJob* Job : InJobs)
+	{
+		if (Job->IsEnabled())
+		{
+			++EnabledJobCount;
+		}
+		else
+		{
+			++DisabledJobCount;
+		}
+
+		for (const TObjectPtr<UMoviePipelineExecutorShot>& Shot : Job->ShotInfo)
+		{
+			if (Shot->bEnabled && Job->IsEnabled())
+			{
+				++EnabledShotCount;
+			}
+			else
+			{
+				++DisabledShotCount;
+			}
+		}
+	}
+	
+	TArray<FAnalyticsEventAttribute> EventAttributes;
+	EventAttributes.Add(FAnalyticsEventAttribute(TEXT("IsLocal"), bIsLocal));
+	EventAttributes.Add(FAnalyticsEventAttribute(TEXT("SessionType"), GetSessionType()));
+	EventAttributes.Add(FAnalyticsEventAttribute(TEXT("EnabledJobCount"), EnabledJobCount));
+	EventAttributes.Add(FAnalyticsEventAttribute(TEXT("DisabledJobCount"), DisabledJobCount));
+	EventAttributes.Add(FAnalyticsEventAttribute(TEXT("EnabledShotCount"), EnabledShotCount));
+	EventAttributes.Add(FAnalyticsEventAttribute(TEXT("DisabledShotCount"), DisabledShotCount));
+
+	FEngineAnalytics::GetProvider().RecordEvent(TEXT("MoviePipeline.RendersRequested"), EventAttributes);
+}
+
+/** Sends out telemetry that includes information about the shot being rendered (the type of rendering being done, which types of settings/nodes are being used, etc). */
+void FMoviePipelineTelemetry::SendBeginShotRenderTelemetry(UMoviePipelineExecutorShot* InShot, UMovieGraphEvaluatedConfig* InEvaluatedConfig)
+{
+	if (!FEngineAnalytics::IsAvailable())
+    {
+        return;
+    }
+	
+	const FMoviePipelineShotRenderTelemetry ShotRenderTelemetry = InEvaluatedConfig
+		? GatherShotRenderTelemetryForGraph(InShot, InEvaluatedConfig)
+		: GatherShotRenderTelemetryForLegacy(InShot);
+	
+	TArray<FAnalyticsEventAttribute> EventAttributes;
+	EventAttributes.Add(FAnalyticsEventAttribute(TEXT("IsGraph"), ShotRenderTelemetry.bIsGraph));
+	EventAttributes.Add(FAnalyticsEventAttribute(TEXT("UsesDeferred"), ShotRenderTelemetry.bUsesDeferred));
+	EventAttributes.Add(FAnalyticsEventAttribute(TEXT("UsesPathTracer"), ShotRenderTelemetry.bUsesPathTracer));
+	EventAttributes.Add(FAnalyticsEventAttribute(TEXT("UsesPanoramic"), ShotRenderTelemetry.bUsesPanoramic));
+	EventAttributes.Add(FAnalyticsEventAttribute(TEXT("UsesHighResTiling"), ShotRenderTelemetry.bUsesHighResTiling));
+	EventAttributes.Add(FAnalyticsEventAttribute(TEXT("UsesNDisplay"), ShotRenderTelemetry.bUsesNDisplay));
+	EventAttributes.Add(FAnalyticsEventAttribute(TEXT("UsesObjectID"), ShotRenderTelemetry.bUsesObjectID));
+	EventAttributes.Add(FAnalyticsEventAttribute(TEXT("UsesMultiCamera"), ShotRenderTelemetry.bUsesMultiCamera));
+	EventAttributes.Add(FAnalyticsEventAttribute(TEXT("UsesScripting"), ShotRenderTelemetry.bUsesScripting));
+	EventAttributes.Add(FAnalyticsEventAttribute(TEXT("UsesSubgraphs"), ShotRenderTelemetry.bUsesSubgraphs));
+	EventAttributes.Add(FAnalyticsEventAttribute(TEXT("UsesPPMs"), ShotRenderTelemetry.bUsesPPMs));
+	EventAttributes.Add(FAnalyticsEventAttribute(TEXT("UsesAudio"), ShotRenderTelemetry.bUsesAudio));
+	EventAttributes.Add(FAnalyticsEventAttribute(TEXT("ResolutionX"), ShotRenderTelemetry.ResolutionX));
+	EventAttributes.Add(FAnalyticsEventAttribute(TEXT("ResolutionY"), ShotRenderTelemetry.ResolutionY));
+	EventAttributes.Add(FAnalyticsEventAttribute(TEXT("HandleFrameCount"), ShotRenderTelemetry.HandleFrameCount));
+	EventAttributes.Add(FAnalyticsEventAttribute(TEXT("TemporalSampleCount"), ShotRenderTelemetry.TemporalSampleCount));
+	EventAttributes.Add(FAnalyticsEventAttribute(TEXT("SpatialSampleCount"), ShotRenderTelemetry.SpatialSampleCount));
+	EventAttributes.Add(FAnalyticsEventAttribute(TEXT("RenderLayerCount"), ShotRenderTelemetry.RenderLayerCount));
+	
+	FEngineAnalytics::GetProvider().RecordEvent(TEXT("MoviePipeline.BeginShotRender"), EventAttributes);
+}
+
+/** Sends out telemetry that indicates whether the shot was rendered successfully. */
+void FMoviePipelineTelemetry::SendEndShotRenderTelemetry(const bool bIsGraph, const bool bWasSuccessful, const bool bWasCanceled)
+{
+	if (!FEngineAnalytics::IsAvailable())
+	{
+		return;
+	}
+
+	TArray<FAnalyticsEventAttribute> EventAttributes;
+	EventAttributes.Add(FAnalyticsEventAttribute(TEXT("IsGraph"), bIsGraph));
+	EventAttributes.Add(FAnalyticsEventAttribute(TEXT("WasSuccessful"), bWasSuccessful));
+	EventAttributes.Add(FAnalyticsEventAttribute(TEXT("WasCanceled"), bWasCanceled));
+
+	FEngineAnalytics::GetProvider().RecordEvent(TEXT("MoviePipeline.EndShotRender"), EventAttributes);
+}
+
+FString FMoviePipelineTelemetry::GetSessionType()
+{
+	FString SessionType;
+
+#if WITH_EDITOR
+	if (GEditor != nullptr)
+	{
+		SessionType = TEXT("Editor");
+	}
+	else
+	{
+		SessionType = TEXT("DashGame");
+	}
+#else
+	SessionType = TEXT("Shipping");
+#endif
+
+	return SessionType;
+}
+
+FMoviePipelineShotRenderTelemetry FMoviePipelineTelemetry::GatherShotRenderTelemetryForGraph(const UMoviePipelineExecutorShot* InShot, UMovieGraphEvaluatedConfig* InEvaluatedConfig)
+{
+	FMoviePipelineShotRenderTelemetry Telemetry;
+	Telemetry.bIsGraph = true;
+
+	const FIntPoint Resolution = UMovieGraphBlueprintLibrary::GetEffectiveOutputResolution(InEvaluatedConfig);
+	Telemetry.ResolutionX = Resolution.X;
+	Telemetry.ResolutionY = Resolution.Y;
+
+	// The evaluated graph won't include subgraph nodes, so peek into the non-evaluated graph(s) to see if there are any in there
+	{
+		TSet<UMovieGraphConfig*> Subgraphs;
+		if (const UMovieGraphConfig* ShotGraphConfig = InShot->GetGraphPreset())
+		{
+			ShotGraphConfig->GetAllContainedSubgraphs(Subgraphs);
+		}
+
+		if (Subgraphs.IsEmpty())
+		{
+			if (const UMoviePipelineExecutorJob* ParentJob = InShot->GetTypedOuter<UMoviePipelineExecutorJob>())
+			{
+				if (const UMovieGraphConfig* JobGraphConfig = ParentJob->GetGraphPreset())
+				{
+					JobGraphConfig->GetAllContainedSubgraphs(Subgraphs);
+				}
+			}
+		}
+
+		Telemetry.bUsesSubgraphs = !Subgraphs.IsEmpty();
+	}
+
+	// Iterate through all of the nodes on every branch. Ask the nodes to update the telemetry data.
+	for (const TPair<FName, FMovieGraphEvaluatedBranchConfig>& BranchMapping : InEvaluatedConfig->BranchConfigMapping)
+	{
+		for (const TObjectPtr<UMovieGraphNode>& Node : BranchMapping.Value.GetNodes())
+		{
+			const UMovieGraphSettingNode* SettingNode = Cast<UMovieGraphSettingNode>(Node);
+			if (!SettingNode || SettingNode->IsDisabled())
+			{
+				continue;
+			}
+
+			SettingNode->UpdateTelemetry(&Telemetry);
+		}
+	}
+
+	return Telemetry;
+}
+
+FMoviePipelineShotRenderTelemetry FMoviePipelineTelemetry::GatherShotRenderTelemetryForLegacy(UMoviePipelineExecutorShot* InShot)
+{
+	FMoviePipelineShotRenderTelemetry Telemetry;
+	Telemetry.bIsGraph = false;
+
+	UMoviePipelinePrimaryConfig* PrimaryConfig = InShot->GetTypedOuter<UMoviePipelineExecutorJob>()->GetConfiguration();
+	const FIntPoint Resolution = UMoviePipelineBlueprintLibrary::GetEffectiveOutputResolution(PrimaryConfig, InShot);
+	Telemetry.ResolutionX = Resolution.X;
+	Telemetry.ResolutionY = Resolution.Y;
+	
+	// Merge settings from the primary and shot configs
+	TArray<UMoviePipelineSetting*> AllSettings = PrimaryConfig->GetAllSettings();
+	if (const UMoviePipelineShotConfig* ShotConfig = InShot->GetShotOverrideConfiguration())
+	{
+		AllSettings.Append(ShotConfig->GetUserSettings());
+	}
+
+	// Iterate through all of the settings. Ask the setting to update the telemetry data.
+	for (const UMoviePipelineSetting* Setting : AllSettings)
+	{
+		if (!Setting || !Setting->IsEnabled())
+		{
+			continue;
+		}
+
+		Setting->UpdateTelemetry(&Telemetry);
+	}
+
+	return Telemetry;
+}
\ No newline at end of file
diff --git a/Engine/Plugins/MovieScene/MovieRenderPipeline/Source/MovieRenderPipelineCore/Public/Graph/MovieGraphNode.h b/Engine/Plugins/MovieScene/MovieRenderPipeline/Source/MovieRenderPipelineCore/Public/Graph/MovieGraphNode.h
--- a/Engine/Plugins/MovieScene/MovieRenderPipeline/Source/MovieRenderPipelineCore/Public/Graph/MovieGraphNode.h
+++ b/Engine/Plugins/MovieScene/MovieRenderPipeline/Source/MovieRenderPipelineCore/Public/Graph/MovieGraphNode.h
@@ -26,6 +26,7 @@ class UMovieGraphPipeline;
 class UMovieGraphVariable;
 struct FMovieGraphEvaluationContext;
 struct FMovieGraphTraversalContext;
+struct FMoviePipelineShotRenderTelemetry;
 
 #if WITH_EDITOR
 class UEdGraphNode;
@@ -458,6 +459,9 @@ public:
 	
 	/** Modify the Unreal URL and command line arguments when the node will be run in a new process. Only applies to nodes in the Globals branch. */
 	virtual void BuildNewProcessCommandLineArgsImpl(TArray<FString>& InOutUnrealURLParams, TArray<FString>& InOutCommandLineArgs, TArray<FString>& InOutDeviceProfileCvars, TArray<FString>& InOutExecCmds) const { }
+
+	/** Updates telemetry data for this node. Should only be used by nodes that ship with Movie Render Graph. Called on the fully-evaluated node in the flattened graph. */
+	virtual void UpdateTelemetry(FMoviePipelineShotRenderTelemetry* InTelemetry) const { }
 };
 
 UINTERFACE()
diff --git a/Engine/Plugins/MovieScene/MovieRenderPipeline/Source/MovieRenderPipelineCore/Public/Graph/Nodes/MovieGraphAudioOutputNode.h b/Engine/Plugins/MovieScene/MovieRenderPipeline/Source/MovieRenderPipelineCore/Public/Graph/Nodes/MovieGraphAudioOutputNode.h
--- a/Engine/Plugins/MovieScene/MovieRenderPipeline/Source/MovieRenderPipelineCore/Public/Graph/Nodes/MovieGraphAudioOutputNode.h
+++ b/Engine/Plugins/MovieScene/MovieRenderPipeline/Source/MovieRenderPipelineCore/Public/Graph/Nodes/MovieGraphAudioOutputNode.h
@@ -20,6 +20,7 @@ public:
 
 	// UMovieGraphSettingNode Interface
 	virtual void BuildNewProcessCommandLineArgsImpl(TArray<FString>& InOutUnrealURLParams, TArray<FString>& InOutCommandLineArgs, TArray<FString>& InOutDeviceProfileCvars, TArray<FString>& InOutExecCmds) const override;
+	virtual void UpdateTelemetry(FMoviePipelineShotRenderTelemetry* InTelemetry) const override;
 	// ~UMovieGraphSettingNode Interface
 
 	// UMovieGraphNode Interface
diff --git a/Engine/Plugins/MovieScene/MovieRenderPipeline/Source/MovieRenderPipelineCore/Public/Graph/Nodes/MovieGraphExecuteScriptNode.h b/Engine/Plugins/MovieScene/MovieRenderPipeline/Source/MovieRenderPipelineCore/Public/Graph/Nodes/MovieGraphExecuteScriptNode.h
--- a/Engine/Plugins/MovieScene/MovieRenderPipeline/Source/MovieRenderPipelineCore/Public/Graph/Nodes/MovieGraphExecuteScriptNode.h
+++ b/Engine/Plugins/MovieScene/MovieRenderPipeline/Source/MovieRenderPipelineCore/Public/Graph/Nodes/MovieGraphExecuteScriptNode.h
@@ -190,6 +190,7 @@ public:
 	UMovieGraphExecuteScriptNode() = default;
 
 	virtual EMovieGraphBranchRestriction GetBranchRestriction() const override { return EMovieGraphBranchRestriction::Globals; }
+	virtual void UpdateTelemetry(FMoviePipelineShotRenderTelemetry* InTelemetry) const override;
 
 #if WITH_EDITOR
 	virtual FText GetNodeTitle(const bool bGetDescriptive = false) const override;
diff --git a/Engine/Plugins/MovieScene/MovieRenderPipeline/Source/MovieRenderPipelineCore/Public/Graph/Nodes/MovieGraphGlobalOutputSettingNode.h b/Engine/Plugins/MovieScene/MovieRenderPipeline/Source/MovieRenderPipelineCore/Public/Graph/Nodes/MovieGraphGlobalOutputSettingNode.h
--- a/Engine/Plugins/MovieScene/MovieRenderPipeline/Source/MovieRenderPipelineCore/Public/Graph/Nodes/MovieGraphGlobalOutputSettingNode.h
+++ b/Engine/Plugins/MovieScene/MovieRenderPipeline/Source/MovieRenderPipelineCore/Public/Graph/Nodes/MovieGraphGlobalOutputSettingNode.h
@@ -38,6 +38,7 @@ public:
 
 	// UMovieGraphSettingNode Interface
 	virtual void GetFormatResolveArgs(FMovieGraphResolveArgs& OutMergedFormatArgs, const FMovieGraphRenderDataIdentifier& InRenderDataIdentifier) const override;
+	virtual void UpdateTelemetry(FMoviePipelineShotRenderTelemetry* InTelemetry) const override;
 
 #if WITH_EDITOR
 	virtual FText GetNodeTitle(const bool bGetDescriptive) const override;
diff --git a/Engine/Plugins/MovieScene/MovieRenderPipeline/Source/MovieRenderPipelineCore/Public/Graph/Nodes/MovieGraphRenderLayerNode.h b/Engine/Plugins/MovieScene/MovieRenderPipeline/Source/MovieRenderPipelineCore/Public/Graph/Nodes/MovieGraphRenderLayerNode.h
--- a/Engine/Plugins/MovieScene/MovieRenderPipeline/Source/MovieRenderPipelineCore/Public/Graph/Nodes/MovieGraphRenderLayerNode.h
+++ b/Engine/Plugins/MovieScene/MovieRenderPipeline/Source/MovieRenderPipelineCore/Public/Graph/Nodes/MovieGraphRenderLayerNode.h
@@ -25,6 +25,8 @@ public:
 	//~ End UObject interface
 #endif
 
+	virtual void UpdateTelemetry(FMoviePipelineShotRenderTelemetry* InTelemetry) const override;
+
 public:
 	UPROPERTY(EditAnywhere, BlueprintReadWrite, Category = Overrides, meta = (InlineEditConditionToggle))
 	uint8 bOverride_LayerName : 1;
diff --git a/Engine/Plugins/MovieScene/MovieRenderPipeline/Source/MovieRenderPipelineCore/Public/Graph/Nodes/MovieGraphSamplingMethodNode.h b/Engine/Plugins/MovieScene/MovieRenderPipeline/Source/MovieRenderPipelineCore/Public/Graph/Nodes/MovieGraphSamplingMethodNode.h
--- a/Engine/Plugins/MovieScene/MovieRenderPipeline/Source/MovieRenderPipelineCore/Public/Graph/Nodes/MovieGraphSamplingMethodNode.h
+++ b/Engine/Plugins/MovieScene/MovieRenderPipeline/Source/MovieRenderPipelineCore/Public/Graph/Nodes/MovieGraphSamplingMethodNode.h
@@ -25,9 +25,9 @@ public:
 	virtual FSlateIcon GetIconAndTint(FLinearColor& OutColor) const override;
 #endif
 
-public:
-
+	virtual void UpdateTelemetry(FMoviePipelineShotRenderTelemetry* InTelemetry) const override;
 
+public:
 	UPROPERTY(EditAnywhere, BlueprintReadWrite, Category = "Settings", meta = (InlineEditConditionToggle))
 	uint8 bOverride_SamplingMethodClass : 1;
 
diff --git a/Engine/Plugins/MovieScene/MovieRenderPipeline/Source/MovieRenderPipelineCore/Public/MoviePipelineAntiAliasingSetting.h b/Engine/Plugins/MovieScene/MovieRenderPipeline/Source/MovieRenderPipelineCore/Public/MoviePipelineAntiAliasingSetting.h
--- a/Engine/Plugins/MovieScene/MovieRenderPipeline/Source/MovieRenderPipelineCore/Public/MoviePipelineAntiAliasingSetting.h
+++ b/Engine/Plugins/MovieScene/MovieRenderPipeline/Source/MovieRenderPipelineCore/Public/MoviePipelineAntiAliasingSetting.h
@@ -25,15 +25,16 @@ public:
 	{
 	}
 
-public:
 #if WITH_EDITOR
 	virtual FText GetDisplayText() const override { return NSLOCTEXT("MovieRenderPipeline", "AntiAliasingSettingDisplayName", "Anti-aliasing"); }
 #endif
+
 protected:
 	virtual bool IsValidOnShots() const override { return true; }
 	virtual bool IsValidOnPrimary() const override { return true; }
 	virtual void ValidateStateImpl() override;
 	virtual void GetFormatArguments(FMoviePipelineFormatArgs& InOutFormatArgs) const override;
+	virtual void UpdateTelemetry(FMoviePipelineShotRenderTelemetry* InTelemetry) const override;
 
 public:
 
diff --git a/Engine/Plugins/MovieScene/MovieRenderPipeline/Source/MovieRenderPipelineCore/Public/MoviePipelineCameraSetting.h b/Engine/Plugins/MovieScene/MovieRenderPipeline/Source/MovieRenderPipelineCore/Public/MoviePipelineCameraSetting.h
--- a/Engine/Plugins/MovieScene/MovieRenderPipeline/Source/MovieRenderPipelineCore/Public/MoviePipelineCameraSetting.h
+++ b/Engine/Plugins/MovieScene/MovieRenderPipeline/Source/MovieRenderPipelineCore/Public/MoviePipelineCameraSetting.h
@@ -23,15 +23,9 @@ protected:
 	virtual bool IsValidOnShots() const override { return true; }
 	virtual bool IsValidOnPrimary() const override { return true; }
 	
-	virtual void GetFormatArguments(FMoviePipelineFormatArgs& InOutFormatArgs) const override
-	{
-		Super::GetFormatArguments(InOutFormatArgs);
+	virtual void GetFormatArguments(FMoviePipelineFormatArgs& InOutFormatArgs) const override;
+	virtual void UpdateTelemetry(FMoviePipelineShotRenderTelemetry* InTelemetry) const override;
 
-		InOutFormatArgs.FilenameArguments.Add(TEXT("shutter_timing"), StaticEnum<EMoviePipelineShutterTiming>()->GetNameStringByValue((int64)ShutterTiming));
-		InOutFormatArgs.FilenameArguments.Add(TEXT("overscan_percentage"), FString::SanitizeFloat(OverscanPercentage));
-		InOutFormatArgs.FileMetadata.Add(TEXT("unreal/camera/shutterTiming"), StaticEnum<EMoviePipelineShutterTiming>()->GetNameStringByValue((int64)ShutterTiming));
-		InOutFormatArgs.FileMetadata.Add(TEXT("unreal/camera/overscanPercentage"), FString::SanitizeFloat(OverscanPercentage));
-	}
 public:	
 	/**
 	* Shutter Timing allows you to bias the timing of your shutter angle to either be before, during, or after
diff --git a/Engine/Plugins/MovieScene/MovieRenderPipeline/Source/MovieRenderPipelineCore/Public/MoviePipelineHighResSetting.h b/Engine/Plugins/MovieScene/MovieRenderPipeline/Source/MovieRenderPipelineCore/Public/MoviePipelineHighResSetting.h
--- a/Engine/Plugins/MovieScene/MovieRenderPipeline/Source/MovieRenderPipelineCore/Public/MoviePipelineHighResSetting.h
+++ b/Engine/Plugins/MovieScene/MovieRenderPipeline/Source/MovieRenderPipelineCore/Public/MoviePipelineHighResSetting.h
@@ -7,6 +7,8 @@
 #include "MoviePipelineQueue.h"
 #include "MoviePipelineOutputSetting.h"
 #include "MoviePipelineAntiAliasingSetting.h"
+#include "MoviePipelineTelemetry.h"
+
 #include "MoviePipelineHighResSetting.generated.h"
 
 UCLASS(Blueprintable)
@@ -63,6 +65,11 @@ public:
 		InOutFormatArgs.FileMetadata.Add(TEXT("unreal/highres/overlapPercent"), FString::SanitizeFloat(OverlapRatio));
 	}
 
+	virtual void UpdateTelemetry(FMoviePipelineShotRenderTelemetry* InTelemetry) const override
+	{
+		InTelemetry->bUsesHighResTiling = true;
+	}
+
 #if WITH_EDITOR
 	virtual FText GetFooterText(UMoviePipelineExecutorJob* InJob) const override
 	{
diff --git a/Engine/Plugins/MovieScene/MovieRenderPipeline/Source/MovieRenderPipelineCore/Public/MoviePipelineOutputSetting.h b/Engine/Plugins/MovieScene/MovieRenderPipeline/Source/MovieRenderPipelineCore/Public/MoviePipelineOutputSetting.h
--- a/Engine/Plugins/MovieScene/MovieRenderPipeline/Source/MovieRenderPipelineCore/Public/MoviePipelineOutputSetting.h
+++ b/Engine/Plugins/MovieScene/MovieRenderPipeline/Source/MovieRenderPipelineCore/Public/MoviePipelineOutputSetting.h
@@ -13,7 +13,7 @@ class MOVIERENDERPIPELINECORE_API UMoviePipelineOutputSetting : public UMoviePip
 	GENERATED_BODY()
 public:
 	UMoviePipelineOutputSetting();
-public:
+
 #if WITH_EDITOR
 	virtual FText GetDisplayText() const override { return NSLOCTEXT("MovieRenderPipeline", "OutputSettingDisplayName", "Output"); }
 	virtual FText GetFooterText(UMoviePipelineExecutorJob* InJob) const override;
@@ -27,6 +27,9 @@ public:
 	// UObject Interface
 	virtual void PostLoad() override;
 	// ~UObject Interface
+
+	virtual void UpdateTelemetry(FMoviePipelineShotRenderTelemetry* InTelemetry) const override;
+
 public:
 	/** What directory should all of our output files be relative to. */
 	UPROPERTY(EditAnywhere, BlueprintReadWrite, Category = "File Output")
diff --git a/Engine/Plugins/MovieScene/MovieRenderPipeline/Source/MovieRenderPipelineCore/Public/MoviePipelineSetting.h b/Engine/Plugins/MovieScene/MovieRenderPipeline/Source/MovieRenderPipelineCore/Public/MoviePipelineSetting.h
--- a/Engine/Plugins/MovieScene/MovieRenderPipeline/Source/MovieRenderPipelineCore/Public/MoviePipelineSetting.h
+++ b/Engine/Plugins/MovieScene/MovieRenderPipeline/Source/MovieRenderPipelineCore/Public/MoviePipelineSetting.h
@@ -8,6 +8,7 @@
 class UMoviePipeline;
 struct FSlateBrush;
 struct FMoviePipelineFormatArgs;
+struct FMoviePipelineShotRenderTelemetry;
 class UMoviePipelineExecutorJob;
 class UMoviePipelineExecutorShot;
 
@@ -77,6 +78,9 @@ public:
 	void BeginExport() { BeginExportImpl(); }
 	// ~Post Finalize Export
 
+	/** Updates telemetry data for this setting. Should only be used by settings that ship with Movie Render Queue. */
+	virtual void UpdateTelemetry(FMoviePipelineShotRenderTelemetry* InTelemetry) const { }
+
 protected:
 	UMoviePipeline* GetPipeline() const;
 
diff --git a/Engine/Plugins/MovieScene/MovieRenderPipeline/Source/MovieRenderPipelineCore/Public/MoviePipelineTelemetry.h b/Engine/Plugins/MovieScene/MovieRenderPipeline/Source/MovieRenderPipelineCore/Public/MoviePipelineTelemetry.h
new file mode 100644
--- /dev/null
+++ b/Engine/Plugins/MovieScene/MovieRenderPipeline/Source/MovieRenderPipelineCore/Public/MoviePipelineTelemetry.h
@@ -0,0 +1,57 @@
+﻿// Copyright Epic Games, Inc. All Rights Reserved.
+
+#pragma once
+
+#include "Containers/Array.h"
+#include "HAL/Platform.h"
+
+class UMovieGraphEvaluatedConfig;
+class UMoviePipelineExecutorJob;
+class UMoviePipelineExecutorShot;
+
+/** Telemetry data that is captured when a shot begins rendering. Only for use with settings/nodes shipped with Movie Render Queue/Graph. */
+struct MOVIERENDERPIPELINECORE_API FMoviePipelineShotRenderTelemetry
+{
+	bool bIsGraph = false;
+	bool bUsesDeferred = false;
+	bool bUsesPathTracer = false;
+	bool bUsesPanoramic = false;
+	bool bUsesHighResTiling = false;
+	bool bUsesNDisplay = false;
+	bool bUsesObjectID = false;
+	bool bUsesMultiCamera = false;
+	bool bUsesScripting = false;
+	bool bUsesSubgraphs = false;
+	bool bUsesPPMs = false;
+	bool bUsesAudio = false;
+	int32 ResolutionX = 0;
+	int32 ResolutionY = 0;
+	int32 HandleFrameCount = 0;
+	int32 TemporalSampleCount = 0;
+	int32 SpatialSampleCount = 0;
+	int32 RenderLayerCount = 0;
+};
+
+/** Responsible for sending out telemetry for both Movie Render Queue and Movie Render Graph. */
+class MOVIERENDERPIPELINECORE_API FMoviePipelineTelemetry
+{
+public:
+	/** Sends out telemetry that captures queue-level information. */
+	static void SendRendersRequestedTelemetry(const bool bIsLocal, TArray<UMoviePipelineExecutorJob*>&& InJobs);
+
+	/** Sends out telemetry that includes information about the shot being rendered (the type of rendering being done, which types of settings/nodes are being used, etc). */
+	static void SendBeginShotRenderTelemetry(UMoviePipelineExecutorShot* InShot, UMovieGraphEvaluatedConfig* InEvaluatedConfig = nullptr);
+
+	/** Sends out telemetry that indicates whether the shot was rendered successfully. */
+	static void SendEndShotRenderTelemetry(const bool bIsGraph, const bool bWasSuccessful, const bool bWasCanceled);
+
+private:
+	/** Gets the current session type (Editor, DashGame, Shipping). */
+	static FString GetSessionType();
+
+	/** Returns a populated telemetry object for the shot and graph that's specified. */
+	static FMoviePipelineShotRenderTelemetry GatherShotRenderTelemetryForGraph(const UMoviePipelineExecutorShot* InShot, UMovieGraphEvaluatedConfig* InEvaluatedConfig);
+
+	/** Returns a populated telemetry object for the shot that's specified. For use with the legacy system, not the graph. */
+	static FMoviePipelineShotRenderTelemetry GatherShotRenderTelemetryForLegacy(UMoviePipelineExecutorShot* InShot);
+};
\ No newline at end of file
diff --git a/Engine/Plugins/MovieScene/MovieRenderPipeline/Source/MovieRenderPipelineEditor/Private/Graph/MovieGraphSchema.h b/Engine/Plugins/MovieScene/MovieRenderPipeline/Source/MovieRenderPipelineEditor/Private/Graph/MovieGraphSchema.h
--- a/Engine/Plugins/MovieScene/MovieRenderPipeline/Source/MovieRenderPipelineEditor/Private/Graph/MovieGraphSchema.h
+++ b/Engine/Plugins/MovieScene/MovieRenderPipeline/Source/MovieRenderPipelineEditor/Private/Graph/MovieGraphSchema.h
@@ -33,7 +33,10 @@ public:
 	
 	static void InitMoviePipelineNodeClasses();
 
-	/** Gets all node classes that the graph is aware of. Should only be called after all plugins that contain nodes have been loaded. */
+	/**
+	 * Gets all node classes that the graph is aware of and are usable within the graph (including those from plugins).
+	 * Should only be called after all plugins that contain nodes have been loaded.
+	 */
 	static const TArray<UClass*>& GetNodeClasses();
 
 private:
diff --git a/Engine/Plugins/MovieScene/MovieRenderPipeline/Source/MovieRenderPipelineEditor/Private/Widgets/SMoviePipelineQueuePanel.cpp b/Engine/Plugins/MovieScene/MovieRenderPipeline/Source/MovieRenderPipelineEditor/Private/Widgets/SMoviePipelineQueuePanel.cpp
--- a/Engine/Plugins/MovieScene/MovieRenderPipeline/Source/MovieRenderPipelineEditor/Private/Widgets/SMoviePipelineQueuePanel.cpp
+++ b/Engine/Plugins/MovieScene/MovieRenderPipeline/Source/MovieRenderPipelineEditor/Private/Widgets/SMoviePipelineQueuePanel.cpp
@@ -14,6 +14,7 @@
 #include "MoviePipelineQueue.h"
 #include "MoviePipelinePrimaryConfig.h"
 #include "MoviePipelineQueueSubsystem.h"
+#include "MoviePipelineTelemetry.h"
 #include "Graph/MovieGraphConfig.h"
 #include "Graph/MovieGraphAssetToolkit.h"
 #include "Graph/MovieGraphConfigFactory.h"
@@ -285,6 +286,10 @@ FReply SMoviePipelineQueuePanel::OnRenderLocalRequested()
 	// OnRenderLocalRequested should only get called if IsRenderLocalEnabled() returns true, meaning there's a valid class.
 	check(ExecutorClass != nullptr);
 	Subsystem->RenderQueueWithExecutor(ExecutorClass);
+
+	constexpr bool bIsLocal = true;
+	FMoviePipelineTelemetry::SendRendersRequestedTelemetry(bIsLocal, Subsystem->GetQueue()->GetJobs());
+	
 	return FReply::Handled();
 }
 
@@ -322,8 +327,11 @@ FReply SMoviePipelineQueuePanel::OnRenderRemoteRequested()
 
 	// OnRenderRemoteRequested should only get called if IsRenderRemoteEnabled() returns true, meaning there's a valid class.
 	check(ExecutorClass != nullptr);
-
 	Subsystem->RenderQueueWithExecutor(ExecutorClass);
+
+	constexpr bool bIsLocal = false;
+	FMoviePipelineTelemetry::SendRendersRequestedTelemetry(bIsLocal, Subsystem->GetQueue()->GetJobs());
+	
 	return FReply::Handled();
 }
 
diff --git a/Engine/Plugins/MovieScene/MovieRenderPipeline/Source/MovieRenderPipelineRenderPasses/Private/Graph/Nodes/MovieGraphDeferredPassNode.cpp b/Engine/Plugins/MovieScene/MovieRenderPipeline/Source/MovieRenderPipelineRenderPasses/Private/Graph/Nodes/MovieGraphDeferredPassNode.cpp
--- a/Engine/Plugins/MovieScene/MovieRenderPipeline/Source/MovieRenderPipelineRenderPasses/Private/Graph/Nodes/MovieGraphDeferredPassNode.cpp
+++ b/Engine/Plugins/MovieScene/MovieRenderPipeline/Source/MovieRenderPipelineRenderPasses/Private/Graph/Nodes/MovieGraphDeferredPassNode.cpp
@@ -1,7 +1,9 @@
 ﻿// Copyright Epic Games, Inc. All Rights Reserved.
 
 #include "Graph/Nodes/MovieGraphDeferredPassNode.h"
+
 #include "Graph/Renderers/MovieGraphDeferredPass.h"
+#include "MoviePipelineTelemetry.h"
 
 TUniquePtr<UE::MovieGraph::Rendering::FMovieGraphImagePassBase> UMovieGraphDeferredRenderPassNode::CreateInstance() const
 {
@@ -37,6 +39,13 @@ void UMovieGraphDeferredRenderPassNode::GetFormatResolveArgs(FMovieGraphResolveA
 	OutMergedFormatArgs.FileMetadata.Add(TEXT("unreal/sampling/spatialSampleCount"), FString::FromInt(SpatialSampleCount));
 }
 
+void UMovieGraphDeferredRenderPassNode::UpdateTelemetry(FMoviePipelineShotRenderTelemetry* InTelemetry) const
+{
+	InTelemetry->bUsesDeferred = true;
+	InTelemetry->bUsesPPMs |= Algo::AnyOf(AdditionalPostProcessMaterials, [](const FMoviePipelinePostProcessPass& Pass) { return Pass.bEnabled; });
+	InTelemetry->SpatialSampleCount = FMath::Max(InTelemetry->SpatialSampleCount, SpatialSampleCount);
+}
+
 #if WITH_EDITOR
 FText UMovieGraphDeferredRenderPassNode::GetNodeTitle(const bool bGetDescriptive) const
 {
diff --git a/Engine/Plugins/MovieScene/MovieRenderPipeline/Source/MovieRenderPipelineRenderPasses/Private/Graph/Nodes/MovieGraphPathTracerPassNode.cpp b/Engine/Plugins/MovieScene/MovieRenderPipeline/Source/MovieRenderPipelineRenderPasses/Private/Graph/Nodes/MovieGraphPathTracerPassNode.cpp
--- a/Engine/Plugins/MovieScene/MovieRenderPipeline/Source/MovieRenderPipelineRenderPasses/Private/Graph/Nodes/MovieGraphPathTracerPassNode.cpp
+++ b/Engine/Plugins/MovieScene/MovieRenderPipeline/Source/MovieRenderPipelineRenderPasses/Private/Graph/Nodes/MovieGraphPathTracerPassNode.cpp
@@ -1,8 +1,10 @@
 // Copyright Epic Games, Inc. All Rights Reserved.
 
 #include "Graph/Nodes/MovieGraphPathTracerPassNode.h"
+
 #include "Graph/Renderers/MovieGraphPathTracerPass.h"
 #include "Engine/EngineBaseTypes.h"
+#include "MoviePipelineTelemetry.h"
 #include "RenderUtils.h"
 #include "ShowFlags.h"
 #include "PathTracingDenoiser.h"
@@ -225,4 +227,11 @@ FEngineShowFlags UMovieGraphPathTracerRenderPassNode::GetShowFlags() const
 	OutShowFlag.SetMotionBlur(!bEnableReferenceMotionBlur);
 
 	return OutShowFlag;
+}
+
+void UMovieGraphPathTracerRenderPassNode::UpdateTelemetry(FMoviePipelineShotRenderTelemetry* InTelemetry) const
+{
+	InTelemetry->bUsesPathTracer = true;
+	InTelemetry->bUsesPPMs |= Algo::AnyOf(AdditionalPostProcessMaterials, [](const FMoviePipelinePostProcessPass& Pass) { return Pass.bEnabled; });
+	InTelemetry->SpatialSampleCount = FMath::Max(InTelemetry->SpatialSampleCount, SpatialSampleCount);
 }
\ No newline at end of file
diff --git a/Engine/Plugins/MovieScene/MovieRenderPipeline/Source/MovieRenderPipelineRenderPasses/Private/MoviePipelineDeferredPasses.cpp b/Engine/Plugins/MovieScene/MovieRenderPipeline/Source/MovieRenderPipelineRenderPasses/Private/MoviePipelineDeferredPasses.cpp
--- a/Engine/Plugins/MovieScene/MovieRenderPipeline/Source/MovieRenderPipelineRenderPasses/Private/MoviePipelineDeferredPasses.cpp
+++ b/Engine/Plugins/MovieScene/MovieRenderPipeline/Source/MovieRenderPipelineRenderPasses/Private/MoviePipelineDeferredPasses.cpp
@@ -35,6 +35,7 @@
 #include "Camera/CameraComponent.h"
 #include "CineCameraComponent.h"
 #include "Interfaces/Interface_PostProcessVolume.h"
+#include "MoviePipelineTelemetry.h"
 #include "MoviePipelineUtils.h"
 #include "WorldPartition/DataLayer/DataLayerAsset.h"
 #include "TextureResource.h"
@@ -1097,6 +1098,12 @@ TSharedPtr<FSceneViewFamilyContext> UMoviePipelineDeferredPass_PathTracer::Calcu
 	return Super::CalculateViewFamily(InOutSampleState, OptPayload);
 }
 
+void UMoviePipelineDeferredPass_PathTracer::UpdateTelemetry(FMoviePipelineShotRenderTelemetry* InTelemetry) const
+{
+	InTelemetry->bUsesPathTracer = true;
+	InTelemetry->bUsesPPMs |= Algo::AnyOf(AdditionalPostProcessMaterials, [](const FMoviePipelinePostProcessPass& Pass) { return Pass.bEnabled; });
+}
+
 bool UMoviePipelineDeferredPassBase::IsUsingDataLayers() const
 {
 	int32 NumDataLayers = 0;
@@ -1209,3 +1216,8 @@ bool UMoviePipelineDeferredPassBase::IsActorInAnyStencilLayer(AActor* InActor) c
 	return bInLayer;
 }
 
+void UMoviePipelineDeferredPassBase::UpdateTelemetry(FMoviePipelineShotRenderTelemetry* InTelemetry) const
+{
+	InTelemetry->bUsesDeferred = true;
+	InTelemetry->bUsesPPMs |= Algo::AnyOf(AdditionalPostProcessMaterials, [](const FMoviePipelinePostProcessPass& Pass) { return Pass.bEnabled; });
+}
diff --git a/Engine/Plugins/MovieScene/MovieRenderPipeline/Source/MovieRenderPipelineRenderPasses/Private/MoviePipelineWaveOutput.cpp b/Engine/Plugins/MovieScene/MovieRenderPipeline/Source/MovieRenderPipelineRenderPasses/Private/MoviePipelineWaveOutput.cpp
--- a/Engine/Plugins/MovieScene/MovieRenderPipeline/Source/MovieRenderPipelineRenderPasses/Private/MoviePipelineWaveOutput.cpp
+++ b/Engine/Plugins/MovieScene/MovieRenderPipeline/Source/MovieRenderPipelineRenderPasses/Private/MoviePipelineWaveOutput.cpp
@@ -13,6 +13,7 @@
 #include "AudioDevice.h"
 #include "AudioMixerDevice.h"
 #include "AudioThread.h"
+#include "MoviePipelineTelemetry.h"
 #include "MoviePipelineUtils.h"
 
 #include UE_INLINE_GENERATED_CPP_BY_NAME(MoviePipelineWaveOutput)
@@ -219,3 +220,7 @@ void UMoviePipelineWaveOutput::BuildNewProcessCommandLineArgsImpl(TArray<FString
 	InOutCommandLineArgs.Add("-deterministicaudio");
 }
 
+void UMoviePipelineWaveOutput::UpdateTelemetry(FMoviePipelineShotRenderTelemetry* InTelemetry) const
+{
+	InTelemetry->bUsesAudio = true;
+}
diff --git a/Engine/Plugins/MovieScene/MovieRenderPipeline/Source/MovieRenderPipelineRenderPasses/Public/Graph/Nodes/MovieGraphDeferredPassNode.h b/Engine/Plugins/MovieScene/MovieRenderPipeline/Source/MovieRenderPipelineRenderPasses/Public/Graph/Nodes/MovieGraphDeferredPassNode.h
--- a/Engine/Plugins/MovieScene/MovieRenderPipeline/Source/MovieRenderPipelineRenderPasses/Public/Graph/Nodes/MovieGraphDeferredPassNode.h
+++ b/Engine/Plugins/MovieScene/MovieRenderPipeline/Source/MovieRenderPipelineRenderPasses/Public/Graph/Nodes/MovieGraphDeferredPassNode.h
@@ -17,6 +17,7 @@ public:
 	UMovieGraphDeferredRenderPassNode();
 
 	virtual void GetFormatResolveArgs(FMovieGraphResolveArgs& OutMergedFormatArgs, const FMovieGraphRenderDataIdentifier& InRenderDataIdentifier) const override;
+	virtual void UpdateTelemetry(FMoviePipelineShotRenderTelemetry* InTelemetry) const override;
 
 #if WITH_EDITOR
 	virtual FText GetNodeTitle(const bool bGetDescriptive = false) const override;
diff --git a/Engine/Plugins/MovieScene/MovieRenderPipeline/Source/MovieRenderPipelineRenderPasses/Public/Graph/Nodes/MovieGraphPathTracerPassNode.h b/Engine/Plugins/MovieScene/MovieRenderPipeline/Source/MovieRenderPipelineRenderPasses/Public/Graph/Nodes/MovieGraphPathTracerPassNode.h
--- a/Engine/Plugins/MovieScene/MovieRenderPipeline/Source/MovieRenderPipelineRenderPasses/Public/Graph/Nodes/MovieGraphPathTracerPassNode.h
+++ b/Engine/Plugins/MovieScene/MovieRenderPipeline/Source/MovieRenderPipelineRenderPasses/Public/Graph/Nodes/MovieGraphPathTracerPassNode.h
@@ -57,6 +57,8 @@ public:
 	virtual FEngineShowFlags GetShowFlags() const override;
 	// ~UMovieGraphImagePassBaseNode Interface
 
+	virtual void UpdateTelemetry(FMoviePipelineShotRenderTelemetry* InTelemetry) const override;
+
 protected:
 	// UMovieGraphRenderPassNode Interface
 	virtual FString GetRendererNameImpl() const override;
diff --git a/Engine/Plugins/MovieScene/MovieRenderPipeline/Source/MovieRenderPipelineRenderPasses/Public/MoviePipelineDeferredPasses.h b/Engine/Plugins/MovieScene/MovieRenderPipeline/Source/MovieRenderPipelineRenderPasses/Public/MoviePipelineDeferredPasses.h
--- a/Engine/Plugins/MovieScene/MovieRenderPipeline/Source/MovieRenderPipelineRenderPasses/Public/MoviePipelineDeferredPasses.h
+++ b/Engine/Plugins/MovieScene/MovieRenderPipeline/Source/MovieRenderPipelineRenderPasses/Public/MoviePipelineDeferredPasses.h
@@ -98,6 +98,8 @@ protected:
 	bool CheckIfPathTracerIsSupported() const;
 	void PathTracerValidationImpl();
 
+	virtual void UpdateTelemetry(FMoviePipelineShotRenderTelemetry* InTelemetry) const override;
+
 public:
 	/**
 	* Should multiple temporal/spatial samples accumulate the alpha channel? This requires r.PostProcessing.PropagateAlpha
@@ -330,6 +332,7 @@ public:
 	virtual TSharedPtr<FSceneViewFamilyContext> CalculateViewFamily(FMoviePipelineRenderPassMetrics& InOutSampleState, IViewCalcPayload* OptPayload) override;
 
 	virtual bool NeedsFrameThrottle() const override { return true; }
+	virtual void UpdateTelemetry(FMoviePipelineShotRenderTelemetry* InTelemetry) const override;
 
 	/** When enabled, the path tracer will blend all spatial and temporal samples prior to the denoising and will disable post-processed motion blur.
 	 *  In this mode it is possible to use higher temporal sample counts to improve the motion blur quality.
diff --git a/Engine/Plugins/MovieScene/MovieRenderPipeline/Source/MovieRenderPipelineRenderPasses/Public/MoviePipelineWaveOutput.h b/Engine/Plugins/MovieScene/MovieRenderPipeline/Source/MovieRenderPipelineRenderPasses/Public/MoviePipelineWaveOutput.h
--- a/Engine/Plugins/MovieScene/MovieRenderPipeline/Source/MovieRenderPipelineRenderPasses/Public/MoviePipelineWaveOutput.h
+++ b/Engine/Plugins/MovieScene/MovieRenderPipeline/Source/MovieRenderPipelineRenderPasses/Public/MoviePipelineWaveOutput.h
@@ -25,6 +25,7 @@ protected:
 	virtual bool HasFinishedProcessingImpl() override;
 	virtual void ValidateStateImpl() override;
 	virtual void BuildNewProcessCommandLineArgsImpl(TArray<FString>& InOutUnrealURLParams, TArray<FString>& InOutCommandLineArgs, TArray<FString>& InOutDeviceProfileCvars, TArray<FString>& InOutExecCmds) const override;
+	virtual void UpdateTelemetry(FMoviePipelineShotRenderTelemetry* InTelemetry) const override;
 
 public:
 	/* File name format string override. If specified it will override the FileNameFormat from the Output setting. */
diff --git a/Engine/Plugins/Runtime/nDisplay/Source/DisplayClusterMoviePipeline/Private/DisplayClusterMoviePipelineViewportPass.cpp b/Engine/Plugins/Runtime/nDisplay/Source/DisplayClusterMoviePipeline/Private/DisplayClusterMoviePipelineViewportPass.cpp
--- a/Engine/Plugins/Runtime/nDisplay/Source/DisplayClusterMoviePipeline/Private/DisplayClusterMoviePipelineViewportPass.cpp
+++ b/Engine/Plugins/Runtime/nDisplay/Source/DisplayClusterMoviePipeline/Private/DisplayClusterMoviePipelineViewportPass.cpp
@@ -139,6 +139,11 @@ FIntPoint UDisplayClusterMoviePipelineViewportPassBase::GetEffectiveOutputResolu
 	return OutputResolution;
 }
 
+void UDisplayClusterMoviePipelineViewportPassBase::UpdateTelemetry(FMoviePipelineShotRenderTelemetry* InTelemetry) const
+{
+	InTelemetry->bUsesNDisplay = true;
+}
+
 /////////////////////////////////////////////////////////////////////////////////////////////////////////////////////////////////
 UE::MoviePipeline::FImagePassCameraViewData UDisplayClusterMoviePipelineViewportPassBase::GetCameraInfo(FMoviePipelineRenderPassMetrics& InOutSampleState, IViewCalcPayload* OptPayload) const
 {
diff --git a/Engine/Plugins/Runtime/nDisplay/Source/DisplayClusterMoviePipeline/Public/DisplayClusterMoviePipelineViewportPass.h b/Engine/Plugins/Runtime/nDisplay/Source/DisplayClusterMoviePipeline/Public/DisplayClusterMoviePipelineViewportPass.h
--- a/Engine/Plugins/Runtime/nDisplay/Source/DisplayClusterMoviePipeline/Public/DisplayClusterMoviePipelineViewportPass.h
+++ b/Engine/Plugins/Runtime/nDisplay/Source/DisplayClusterMoviePipeline/Public/DisplayClusterMoviePipelineViewportPass.h
@@ -54,6 +54,8 @@ public:
 	virtual int32 GetOutputFileSortingOrder() const override { return 2; }
 	virtual bool bIsEnabledWarpBlend() const { return true; }
 
+	virtual void UpdateTelemetry(FMoviePipelineShotRenderTelemetry* InTelemetry) const override;
+
 protected:
 	bool GetViewportId(int32 InViewportIndex, FString& OutViewportId) const;
 
