MovieGraph: Added support for rendering multiple cameras at once (feature parity with MRQ). This uses the same rules for multiple cameras as MRQ did (ie: you need a Camera Cut Track pointed to one camera, and then any other cameras found adjacent to the Camera Cut track will also be rendered when the option is enabled).

Tested multiple camera views that used Depth of Field (which rely on scene history) and different post process effects per camera to ensure no rendering information was being incorrectly shared between multiple views.

#virtualized

diff --git a/Engine/Plugins/MovieScene/MoviePipelineMaskRenderPass/Source/MoviePipelineMaskRenderPass/Private/Graph/MovieGraphObjectIdPass.cpp b/Engine/Plugins/MovieScene/MoviePipelineMaskRenderPass/Source/MoviePipelineMaskRenderPass/Private/Graph/MovieGraphObjectIdPass.cpp
--- a/Engine/Plugins/MovieScene/MoviePipelineMaskRenderPass/Source/MoviePipelineMaskRenderPass/Private/Graph/MovieGraphObjectIdPass.cpp
+++ b/Engine/Plugins/MovieScene/MoviePipelineMaskRenderPass/Source/MoviePipelineMaskRenderPass/Private/Graph/MovieGraphObjectIdPass.cpp
@@ -160,9 +160,7 @@ void FMovieGraphObjectIdPass::Setup(TWeakObjectPtr<UMovieGraphDefaultRenderer> I
 		NewIdentifier.LayerName = LayerData.LayerName;
 		NewIdentifier.RendererName = InRenderPassNode->GetRendererName();
 		NewIdentifier.SubResourceName = NewIdentifier.RendererName + FString::Printf(TEXT("%02d"), Index);
-
-		UE::MovieGraph::DefaultRenderer::FCameraInfo CameraInfo = InRenderer->GetCameraInfo(LayerData.CameraIdentifier);
-		NewIdentifier.CameraName =  CameraInfo.CameraName;
+		NewIdentifier.CameraName = InLayer.CameraName;
 		
 		RenderDataIdentifiers.Add(NewIdentifier);
 	}
diff --git a/Engine/Plugins/MovieScene/MovieRenderPipeline/Source/MovieRenderPipelineCore/Private/Graph/MovieGraphCoreTimeStep.cpp b/Engine/Plugins/MovieScene/MovieRenderPipeline/Source/MovieRenderPipelineCore/Private/Graph/MovieGraphCoreTimeStep.cpp
--- a/Engine/Plugins/MovieScene/MovieRenderPipeline/Source/MovieRenderPipelineCore/Private/Graph/MovieGraphCoreTimeStep.cpp
+++ b/Engine/Plugins/MovieScene/MovieRenderPipeline/Source/MovieRenderPipelineCore/Private/Graph/MovieGraphCoreTimeStep.cpp
@@ -77,12 +77,16 @@ void UMovieGraphCoreTimeStep::TickProducingFrames()
 		
 		// Get ready to begin tracking relative shot frame count
 		CurrentTimeStepData.ShotOutputFrameNumber = -1;
+		
+		FFrameTime EvalTime = CurrentCameraCut->ShotInfo.CurrentTimeInRoot;
 
+		// Seeks the external data source to match. This is done before setting up for renders
+		// so that Spawnables, etc. will be spawned.
+		GetOwningGraph()->GetDataSourceInstance()->InitializeShot(CurrentCameraCut, EvalTime);
+		
 		// Sets up the render state, etc.
 		GetOwningGraph()->SetupShot(CurrentCameraCut);
 
-		// Seeks the external data source to match
-		GetOwningGraph()->GetDataSourceInstance()->InitializeShot(CurrentCameraCut);
 
 		// Generate render layers from the evaluated graph
 		GetOwningGraph()->CreateLayersInRenderLayerSubsystem(CurrentTimeStepData.EvaluatedConfig);
diff --git a/Engine/Plugins/MovieScene/MovieRenderPipeline/Source/MovieRenderPipelineCore/Private/Graph/MovieGraphDataTypes.cpp b/Engine/Plugins/MovieScene/MovieRenderPipeline/Source/MovieRenderPipelineCore/Private/Graph/MovieGraphDataTypes.cpp
--- a/Engine/Plugins/MovieScene/MovieRenderPipeline/Source/MovieRenderPipelineCore/Private/Graph/MovieGraphDataTypes.cpp
+++ b/Engine/Plugins/MovieScene/MovieRenderPipeline/Source/MovieRenderPipelineCore/Private/Graph/MovieGraphDataTypes.cpp
@@ -43,25 +43,28 @@ UE::MovieGraph::FMovieGraphRenderDataValidationInfo UE::MovieGraph::FMovieGraphO
 {
 	const FName& ActiveBranchName = InRenderID.RootBranchName;
 	const FString& ActiveRendererName = InRenderID.RendererName;
+	const FString& ActiveCameraName = InRenderID.CameraName;
 
 	FMovieGraphRenderDataValidationInfo ValidationInfo;
 
 	TSet<FString>	LayerCounts;
 	TSet<FName>		BranchCounts;
 	TSet<FString>	BranchRendererCounts;
+	TSet<FString>	CameraNameCounts;
 
 	const int32 ReserveNum = ExpectedRenderPasses.Num();
 	LayerCounts.Reserve(ReserveNum);
 	BranchCounts.Reserve(ReserveNum);
 	BranchRendererCounts.Reserve(ReserveNum);
+	CameraNameCounts.Reserve(ReserveNum);
 
 	for (const FMovieGraphRenderDataIdentifier& PassIdentifier : ExpectedRenderPasses)
 	{
 		LayerCounts.Add(PassIdentifier.LayerName);
 		BranchCounts.Add(PassIdentifier.RootBranchName);
 
-		// We only count renderers on the active branch
-		if (PassIdentifier.RootBranchName == ActiveBranchName)
+		// We only count renderers on the active branch and active camera
+		if (PassIdentifier.RootBranchName == ActiveBranchName && PassIdentifier.CameraName == ActiveCameraName)
 		{
 			BranchRendererCounts.Add(PassIdentifier.RendererName);
 
@@ -71,6 +74,17 @@ UE::MovieGraph::FMovieGraphRenderDataValidationInfo UE::MovieGraph::FMovieGraphO
 				ValidationInfo.ActiveRendererSubresourceCount++;
 			}
 		}
+
+		// Figure out how many cameras exist for this branch
+		if (PassIdentifier.RootBranchName == ActiveBranchName)
+		{
+			if (CameraNameCounts.Contains(PassIdentifier.CameraName))
+			{
+				ValidationInfo.ActiveCameraCount++;
+			}
+			CameraNameCounts.Add(PassIdentifier.CameraName);
+
+		}
 	}
 
 	// Remove any renderers that will be composited on later
@@ -82,7 +96,7 @@ UE::MovieGraph::FMovieGraphRenderDataValidationInfo UE::MovieGraph::FMovieGraphO
 			check(Payload);
 			if (Payload->bCompositeOnOtherRenders)
 			{
-				if (RenderData.Key.RootBranchName == ActiveBranchName)
+				if (RenderData.Key.RootBranchName == ActiveBranchName && RenderData.Key.CameraName == ActiveCameraName)
 				{
 					BranchRendererCounts.Remove(RenderData.Key.RendererName);
 
diff --git a/Engine/Plugins/MovieScene/MovieRenderPipeline/Source/MovieRenderPipelineCore/Private/Graph/MovieGraphDefaultRenderer.cpp b/Engine/Plugins/MovieScene/MovieRenderPipeline/Source/MovieRenderPipelineCore/Private/Graph/MovieGraphDefaultRenderer.cpp
--- a/Engine/Plugins/MovieScene/MovieRenderPipeline/Source/MovieRenderPipelineCore/Private/Graph/MovieGraphDefaultRenderer.cpp
+++ b/Engine/Plugins/MovieScene/MovieRenderPipeline/Source/MovieRenderPipelineCore/Private/Graph/MovieGraphDefaultRenderer.cpp
@@ -8,6 +8,9 @@
 #include "Graph/Nodes/MovieGraphGlobalGameOverrides.h"
 #include "Graph/Nodes/MovieGraphRenderLayerNode.h"
 #include "Graph/Nodes/MovieGraphRenderPassNode.h"
+#include "Graph/Nodes/MovieGraphGlobalGameOverrides.h"
+#include "Graph/Nodes/MovieGraphDebugNode.h"
+#include "Graph/Nodes/MovieGraphCameraNode.h"
 #include "Graph/Nodes/MovieGraphUIRendererNode.h"
 #include "MovieRenderPipelineCoreModule.h"
 #include "RenderingThread.h"
@@ -31,9 +34,11 @@
 
 void UMovieGraphDefaultRenderer::SetupRenderingPipelineForShot(UMoviePipelineExecutorShot* InShot)
 {
+
 	// Iterate through the graph config and look for Render Layers.
 	UMovieGraphConfig* RootGraph = GetOwningGraph()->GetRootGraphForShot(InShot);
 
+
 	struct FMovieGraphPass
 	{
 		TSubclassOf<UMovieGraphRenderPassNode> ClassType;
@@ -49,6 +54,24 @@ void UMovieGraphDefaultRenderer::SetupRenderingPipelineForShot(UMoviePipelineExe
 	UMovieGraphEvaluatedConfig* EvaluatedConfig = TimeStepData.EvaluatedConfig;
 	TArray<FName> GraphBranches = EvaluatedConfig->GetBranchNames();
 
+	// Figure out how many cameras we're rendering and generate fguids so render layers can fetch
+	// the appropriate camera later.
+	UMovieGraphCameraSettingNode* CameraNode = EvaluatedConfig->GetSettingForBranch<UMovieGraphCameraSettingNode>(UMovieGraphNode::GlobalsPinName);
+	TArray<int32> CameraIndexes;
+	if (CameraNode->bRenderAllCameras)
+	{
+		// When using sidecar cameras, the primary one is included as one of them.
+		for (int32 Index = 0; Index < InShot->SidecarCameras.Num(); Index++)
+		{
+			CameraIndexes.Add(Index);
+		}
+	}
+	else
+	{
+		// "-1" means primary camera instead of any sidecar cameras
+		CameraIndexes.Add(-1);
+	}
+
 	for (const FName& Branch : GraphBranches)
 	{
 		// We follow each branch looking for Render Layer nodes to figure out what render layer this should be. We assume a render layer is named
@@ -92,6 +115,7 @@ void UMovieGraphDefaultRenderer::SetupRenderingPipelineForShot(UMoviePipelineExe
 
 	//UE_LOG(LogMovieRenderPipeline, Warning, TEXT("Found: %d Render Passes:"), OutputPasses.Num());
 	int32 TotalLayerCount = 0;
+	int32 CameraCount = 1;
 	for (const FMovieGraphPass& Pass : OutputPasses)
 	{
 		// ToDo: This should probably come from the Renderers themselves, as they can internally produce multiple
@@ -103,10 +127,20 @@ void UMovieGraphDefaultRenderer::SetupRenderingPipelineForShot(UMoviePipelineExe
 		//UE_LOG(LogMovieRenderPipeline, Warning, TEXT("\tRenderer Class: %s"), *Pass.ClassType->GetName());
 		for (const TTuple<FMovieGraphRenderDataIdentifier, TWeakObjectPtr<UMovieGraphRenderPassNode>>& BranchRenderer : Pass.BranchRenderers)
 		{
-			FMovieGraphRenderPassLayerData& LayerData = SetupData.Layers.AddDefaulted_GetRef();
-			LayerData.BranchName = BranchRenderer.Key.RootBranchName;
-			LayerData.LayerName = BranchRenderer.Key.LayerName;
-			LayerData.RenderPassNode = BranchRenderer.Value;
+			for (int32 CameraIndex = 0; CameraIndex < CameraIndexes.Num(); CameraIndex++)
+			{
+				FMovieGraphRenderPassLayerData& LayerData = SetupData.Layers.AddDefaulted_GetRef();
+				LayerData.BranchName = BranchRenderer.Key.RootBranchName;
+				LayerData.LayerName = BranchRenderer.Key.LayerName;
+				LayerData.RenderPassNode = BranchRenderer.Value;
+				// Result could be -1 or 0...n
+				LayerData.CameraIndex = CameraIndexes[CameraIndex]; 
+				
+				// We provide the name here so that the setup functions of the renderers can create RenderLayerIdentifiers.
+				UE::MovieGraph::DefaultRenderer::FCameraInfo CameraInfo = GetCameraInfo(EvaluatedConfig, LayerData.CameraIndex);
+				LayerData.CameraName = CameraInfo.CameraName;
+			}
+				
 			// UE_LOG(LogMovieRenderPipeline, Warning, TEXT("\t\tBranch Name: %s"), *LayerBranchName.ToString());
 		}
 
@@ -116,7 +150,7 @@ void UMovieGraphDefaultRenderer::SetupRenderingPipelineForShot(UMoviePipelineExe
 		RenderPassesInUse.Add(RenderPassCDO);
 	}
 
-	UE_LOG(LogMovieRenderPipeline, Log, TEXT("Finished initializing %d Render Passes (with %d total layers)."), OutputPasses.Num(), TotalLayerCount);
+	UE_LOG(LogMovieRenderPipeline, Log, TEXT("Finished initializing %d Render Passes (with %d total layers and %d cameras)."), OutputPasses.Num(), TotalLayerCount, CameraCount);
 }
 
 void UMovieGraphDefaultRenderer::TeardownRenderingPipelineForShot(UMoviePipelineExecutorShot* InShot)
@@ -301,6 +335,12 @@ void UMovieGraphDefaultRenderer::FlushAsyncEngineSystems(const TObjectPtr<UMovie
 		return;
 	}
 
+	const UMovieGraphCameraSettingNode* CameraSettingNode = InConfig->GetSettingForBranch<UMovieGraphCameraSettingNode>(UMovieGraphNode::GlobalsPinName);
+	if (!CameraSettingNode)
+	{
+		return;
+	}
+
 	// Block until level streaming is completed, we do this at the end of the frame
 	// so that level streaming requests made by Sequencer level visibility tracks are
 	// accounted for.
@@ -356,24 +396,22 @@ void UMovieGraphDefaultRenderer::FlushAsyncEngineSystems(const TObjectPtr<UMovie
 			constexpr bool bFlushGrass = false; // Flush means a different thing to grass system
 			constexpr bool bInForceSync = true;
 
+			UMoviePipelineExecutorShot* CurrentShot = GetOwningGraph()->GetActiveShotList()[GetOwningGraph()->GetCurrentShotIndex()];
+
 			TArray<FVector> CameraLocations;
-			GetCameraLocationsForFrame(CameraLocations);
+			GetCameraLocationsForFrame(CameraLocations, CurrentShot, CameraSettingNode->bRenderAllCameras);
 
 			LandscapeSubsystem->RegenerateGrass(bFlushGrass, bInForceSync, MakeArrayView(CameraLocations));
 		}
 	}
 }
 
-void UMovieGraphDefaultRenderer::GetCameraLocationsForFrame(TArray<FVector>& OutLocations) const
+void UMovieGraphDefaultRenderer::GetCameraLocationsForFrame(TArray<FVector>& OutLocations, UMoviePipelineExecutorShot* InShot, bool bIncludeSidecar) const
 {
-	// ToDo: Multi-camera support
-	if (const APlayerController* LocalPlayerController = GetOwningGraph()->GetWorld()->GetFirstPlayerController())
+	TArray<FMinimalViewInfo> ViewInfos = GetOwningGraph()->GetDataSourceInstance()->GetCameraInformation(InShot, bIncludeSidecar);
+	for (const FMinimalViewInfo& ViewInfo : ViewInfos)
 	{
-		FVector PrimaryCameraLoc;
-		FRotator PrimaryCameraRot;
-
-		LocalPlayerController->GetPlayerViewPoint(PrimaryCameraLoc, PrimaryCameraRot);
-		OutLocations.Add(PrimaryCameraLoc);
+		OutLocations.Add(ViewInfo.Location);
 	}
 }
 
@@ -385,29 +423,46 @@ void UMovieGraphDefaultRenderer::AddOutstandingRenderTask_AnyThread(UE::Tasks::F
 	OutstandingTasks.Add(MoveTemp(InTask));
 }
 
-UE::MovieGraph::DefaultRenderer::FCameraInfo UMovieGraphDefaultRenderer::GetCameraInfo(const FGuid& InCameraIdentifier) const
+UE::MovieGraph::DefaultRenderer::FCameraInfo UMovieGraphDefaultRenderer::GetCameraInfo(UMovieGraphEvaluatedConfig* InConfig, const int32 InCameraIndex) const
 {
 	UE::MovieGraph::DefaultRenderer::FCameraInfo CameraInfo;
-
-	// We only support the primary camera right now
-	APlayerController* LocalPlayerController = GetWorld()->GetFirstPlayerController();
-	if (LocalPlayerController && LocalPlayerController->PlayerCameraManager)
+	
+	const bool bIncludeCDOs = true;
+	const UMovieGraphCameraSettingNode* CameraSettingNode = InConfig->GetSettingForBranch<UMovieGraphCameraSettingNode>(UMovieGraphNode::GlobalsPinName, bIncludeCDOs);
+	if (!ensure(CameraSettingNode))
 	{
-		CameraInfo.ViewInfo = LocalPlayerController->PlayerCameraManager->GetCameraCacheView();
-		CameraInfo.ViewActor = LocalPlayerController->GetViewTarget();
-		
-		UMoviePipelineExecutorShot* CurrentShot = GetOwningGraph()->GetActiveShotList()[GetOwningGraph()->GetCurrentShotIndex()];
+		return CameraInfo;
+	}
+
+	UMoviePipelineExecutorShot* CurrentShot = GetOwningGraph()->GetActiveShotList()[GetOwningGraph()->GetCurrentShotIndex()];
 
-		// INDEX_NONE returns the primary camera from the rendered shot until we have
-		// multi-camera support.
+	// When not using multi-camera the index is "-1" (but we store the data in the 0th array) so we remap,
+	// but the GetCameraName function is -1 aware.
+	int32 LocalArrayIndex = InCameraIndex;
+	if (InCameraIndex < 0)
+	{
+		LocalArrayIndex = 0;
 		const int32 CameraIndex = INDEX_NONE;
 		CameraInfo.CameraName = CurrentShot->GetCameraName(CameraIndex);
 	}
 	else
 	{
-		UE_LOG(LogMovieRenderPipeline, Error, TEXT("Failed to find Local Player Controller/Camera Manager to get viewpoint!"));
+		CameraInfo.CameraName = CurrentShot->GetCameraName(InCameraIndex);
 	}
-	
+
+	APlayerController* LocalPlayerController = GetWorld()->GetFirstPlayerController();
+	if (LocalPlayerController && LocalPlayerController->PlayerCameraManager)
+	{
+		CameraInfo.ViewActor = LocalPlayerController->GetViewTarget();
+	}
+
+	TArray<FMinimalViewInfo> ViewInfos = GetOwningGraph()->GetDataSourceInstance()->GetCameraInformation(GetOwningGraph()->GetActiveShotList()[GetOwningGraph()->GetCurrentShotIndex()], CameraSettingNode->bRenderAllCameras);
+	if (!ensureAlways(ViewInfos.IsValidIndex(LocalArrayIndex)))
+	{
+		return CameraInfo;
+	}
+
+	CameraInfo.ViewInfo = ViewInfos[LocalArrayIndex];
 	return CameraInfo;
 }
 
@@ -471,6 +526,12 @@ TArray<FMovieGraphImagePreviewData> UMovieGraphDefaultRenderer::GetPreviewData()
 { 
 	TArray<FMovieGraphImagePreviewData> Results;
 
+	TSet<FString> CameraNamesInUse;
+	for (const TPair<UE::MovieGraph::DefaultRenderer::FMovieGraphImagePreviewDataPoolParams, TObjectPtr<UTextureRenderTarget2D>>& RenderTarget : PooledViewRenderTargets)
+	{
+		CameraNamesInUse.Add(RenderTarget.Key.Identifier.CameraName);
+	}
+
 	for (const TPair<UE::MovieGraph::DefaultRenderer::FMovieGraphImagePreviewDataPoolParams, TObjectPtr<UTextureRenderTarget2D>>& RenderTarget : PooledViewRenderTargets)
 	{
 		const FString RendererName = RenderTarget.Key.Identifier.RendererName;
@@ -484,6 +545,7 @@ TArray<FMovieGraphImagePreviewData> UMovieGraphDefaultRenderer::GetPreviewData()
 		FMovieGraphImagePreviewData& Data = Results.AddDefaulted_GetRef();
 		Data.Identifier = RenderTarget.Key.Identifier;
 		Data.Texture = RenderTarget.Value.Get();
+		Data.bMultipleCameraNames = CameraNamesInUse.Num() > 1;
 	}
 
 	return Results;
diff --git a/Engine/Plugins/MovieScene/MovieRenderPipeline/Source/MovieRenderPipelineCore/Private/Graph/MovieGraphSequenceDataSource.cpp b/Engine/Plugins/MovieScene/MovieRenderPipeline/Source/MovieRenderPipelineCore/Private/Graph/MovieGraphSequenceDataSource.cpp
--- a/Engine/Plugins/MovieScene/MovieRenderPipeline/Source/MovieRenderPipelineCore/Private/Graph/MovieGraphSequenceDataSource.cpp
+++ b/Engine/Plugins/MovieScene/MovieRenderPipeline/Source/MovieRenderPipelineCore/Private/Graph/MovieGraphSequenceDataSource.cpp
@@ -11,6 +11,7 @@
 #include "LevelSequence.h"
 #include "LevelSequenceActor.h"
 #include "LevelSequencePlayer.h"
+#include "MovieSceneCommonHelpers.h"
 #include "CoreGlobals.h"
 #include "EngineUtils.h"
 #include "Engine/GameViewportClient.h"
@@ -349,6 +350,50 @@ void UMovieGraphSequenceDataSource::JumpDataSource(const FFrameTime& InTimeToJum
 	}
 }
 
+void UMovieGraphSequenceDataSource::InitializeShot(const TObjectPtr<UMoviePipelineExecutorShot>& InShot, const FFrameTime& InEvalTime)
+{
+	Super::InitializeShot(InShot, InEvalTime);
+	
+	SyncDataSourceTime(InEvalTime);
+	JumpDataSource(InEvalTime);
+}
+
+TArray<FMinimalViewInfo> UMovieGraphSequenceDataSource::GetCameraInformation(UMoviePipelineExecutorShot* InShot, bool bIncludeSidecar) const
+{
+	TArray<FMinimalViewInfo> OutInfo;
+
+	if (!bIncludeSidecar)
+	{
+		if (const APlayerController* LocalPlayerController = GetOwningGraph()->GetWorld()->GetFirstPlayerController())
+		{
+			OutInfo.Add(LocalPlayerController->PlayerCameraManager->GetCameraCacheView());
+		}
+	}
+	else
+	{
+		for (int32 Index = 0; Index < InShot->SidecarCameras.Num(); Index++)
+		{
+			const FGuid CameraBindingRef = InShot->SidecarCameras[Index].BindingId;
+			const FMovieSceneSequenceID CameraSequenceRef = InShot->SidecarCameras[Index].SequenceId;
+
+			TArrayView<TWeakObjectPtr<UObject>> BoundCameras = LevelSequenceActor->GetSequencePlayer()->FindBoundObjects(CameraBindingRef, CameraSequenceRef);
+			if (BoundCameras.Num() > 0)
+			{
+				UCameraComponent* BoundCamera = MovieSceneHelpers::CameraComponentFromRuntimeObject(BoundCameras[0].Get());
+				if (BoundCamera)
+				{
+					FMinimalViewInfo ViewInfo;
+					BoundCamera->GetCameraView(GetWorld()->GetDeltaSeconds(), ViewInfo);
+
+					OutInfo.Add(ViewInfo);
+				}
+			}
+		}
+	}
+
+	return OutInfo;
+}
+
 namespace UE::MovieGraph
 {
 	FFrameTime FMovieGraphSequenceTimeController::OnRequestCurrentTime(const FQualifiedFrameTime& InCurrentTime, float InPlayRate)
diff --git a/Engine/Plugins/MovieScene/MovieRenderPipeline/Source/MovieRenderPipelineCore/Private/Graph/Nodes/MovieGraphFileOutputNode.cpp b/Engine/Plugins/MovieScene/MovieRenderPipeline/Source/MovieRenderPipelineCore/Private/Graph/Nodes/MovieGraphFileOutputNode.cpp
--- a/Engine/Plugins/MovieScene/MovieRenderPipeline/Source/MovieRenderPipelineCore/Private/Graph/Nodes/MovieGraphFileOutputNode.cpp
+++ b/Engine/Plugins/MovieScene/MovieRenderPipeline/Source/MovieRenderPipelineCore/Private/Graph/Nodes/MovieGraphFileOutputNode.cpp
@@ -43,8 +43,12 @@ void UMovieGraphFileOutputNode::DisambiguateFilename(FString& InOutFilenameForma
 	{
 		UE::MoviePipeline::ConformOutputFormatStringToken(InOutFilenameFormatString, TEXT("{renderer_sub_name}"), InNodeName, InRenderData.Key.RootBranchName);
 	}
-	
-	// ToDo: Add {camera_name} validation once relevant
+
+	// We only add the camera name token if there are more than one cameras for the active branch
+	if (ValidationInfo.ActiveCameraCount > 1)
+	{
+		UE::MoviePipeline::ConformOutputFormatStringToken(InOutFilenameFormatString, TEXT("{camera_name}"), InNodeName, InRenderData.Key.RootBranchName);
+	}
 }
 
 TArray<FMovieGraphPassData> UMovieGraphFileOutputNode::GetCompositedPasses(UE::MovieGraph::FMovieGraphOutputMergerFrame* InRawFrameData)
diff --git a/Engine/Plugins/MovieScene/MovieRenderPipeline/Source/MovieRenderPipelineCore/Private/Graph/Nodes/MovieGraphWidgetRendererBaseNode.cpp b/Engine/Plugins/MovieScene/MovieRenderPipeline/Source/MovieRenderPipelineCore/Private/Graph/Nodes/MovieGraphWidgetRendererBaseNode.cpp
--- a/Engine/Plugins/MovieScene/MovieRenderPipeline/Source/MovieRenderPipelineCore/Private/Graph/Nodes/MovieGraphWidgetRendererBaseNode.cpp
+++ b/Engine/Plugins/MovieScene/MovieRenderPipeline/Source/MovieRenderPipelineCore/Private/Graph/Nodes/MovieGraphWidgetRendererBaseNode.cpp
@@ -80,8 +80,7 @@ void UMovieGraphWidgetRendererBaseNode::FMovieGraphWidgetPass::Setup(TWeakObject
 
 	// NOTE: Subclasses should specify the SubResourceName
 
-	UE::MovieGraph::DefaultRenderer::FCameraInfo CameraInfo = Renderer->GetCameraInfo(LayerData.CameraIdentifier);
-	RenderDataIdentifier.CameraName = CameraInfo.CameraName;
+	RenderDataIdentifier.CameraName = InLayer.CameraName;
 }
 
 void UMovieGraphWidgetRendererBaseNode::FMovieGraphWidgetPass::Teardown()
diff --git a/Engine/Plugins/MovieScene/MovieRenderPipeline/Source/MovieRenderPipelineCore/Private/MoviePipelineUtils.cpp b/Engine/Plugins/MovieScene/MovieRenderPipeline/Source/MovieRenderPipelineCore/Private/MoviePipelineUtils.cpp
--- a/Engine/Plugins/MovieScene/MovieRenderPipeline/Source/MovieRenderPipelineCore/Private/MoviePipelineUtils.cpp
+++ b/Engine/Plugins/MovieScene/MovieRenderPipeline/Source/MovieRenderPipelineCore/Private/MoviePipelineUtils.cpp
@@ -13,12 +13,14 @@
 #include "AssetRegistry/AssetData.h"
 #include "UObject/UObjectIterator.h"
 #include "HAL/IConsoleManager.h"
+#include "Interfaces/Interface_PostProcessVolume.h"
 #include "MoviePipelineAntiAliasingSetting.h"
 #include "Tracks/MovieSceneSubTrack.h"
 #include "Tracks/MovieSceneCameraCutTrack.h"
 #include "Tracks/MovieSceneCinematicShotTrack.h"
 #include "Sections/MovieSceneCinematicShotSection.h"
 #include "Sections/MovieSceneCameraCutSection.h"
+#include "SceneView.h"
 #include "MovieScene.h"
 #include "MovieSceneSequence.h"
 #include "MovieRenderPipelineCoreModule.h"
@@ -1022,6 +1024,42 @@ namespace UE
 			return FString::Printf(TEXT("%0*d"), InZeroPadCount, InFrameNumber);
 		}
 
+		void DoPostProcessBlend(const FVector& InViewLocation, const UWorld* InWorld, const FMinimalViewInfo& InViewInfo, FSceneView* InOutView)
+		{
+			for (IInterface_PostProcessVolume* PPVolume : InWorld->PostProcessVolumes)
+			{
+				const FPostProcessVolumeProperties VolumeProperties = PPVolume->GetProperties();
+
+				// Skip any volumes which are disabled
+				if (!VolumeProperties.bIsEnabled)
+				{
+					continue;
+				}
+
+				float LocalWeight = FMath::Clamp(VolumeProperties.BlendWeight, 0.0f, 1.0f);
+
+				if (!VolumeProperties.bIsUnbound)
+				{
+					float DistanceToPoint = 0.0f;
+					PPVolume->EncompassesPoint(InViewLocation, 0.0f, &DistanceToPoint);
+
+					if (DistanceToPoint >= 0 && DistanceToPoint < VolumeProperties.BlendRadius)
+					{
+						LocalWeight *= FMath::Clamp(1.0f - DistanceToPoint / VolumeProperties.BlendRadius, 0.0f, 1.0f);
+					}
+					else
+					{
+						LocalWeight = 0.0f;
+					}
+				}
+
+				InOutView->OverridePostProcessSettings(*VolumeProperties.Settings, LocalWeight);
+			}
+
+			// After blending all post processing volumes, blend the camera's post process settings too
+			InOutView->OverridePostProcessSettings(InViewInfo.PostProcessSettings, InViewInfo.PostProcessBlendWeight);
+		}
+
 	}
 }
 namespace UE
diff --git a/Engine/Plugins/MovieScene/MovieRenderPipeline/Source/MovieRenderPipelineCore/Public/Graph/MovieGraphDataTypes.h b/Engine/Plugins/MovieScene/MovieRenderPipeline/Source/MovieRenderPipelineCore/Public/Graph/MovieGraphDataTypes.h
--- a/Engine/Plugins/MovieScene/MovieRenderPipeline/Source/MovieRenderPipelineCore/Public/Graph/MovieGraphDataTypes.h
+++ b/Engine/Plugins/MovieScene/MovieRenderPipeline/Source/MovieRenderPipelineCore/Public/Graph/MovieGraphDataTypes.h
@@ -50,6 +50,10 @@ struct FMovieGraphImagePreviewData
 	/** The identifier for the image, containing the branch name, renderer, etc. */
 	UPROPERTY(BlueprintReadOnly, Category = "Movie Graph")
 	FMovieGraphRenderDataIdentifier Identifier;
+
+	/** If true, then there is more than one camera name being used (ie: multi-camera rendering) */
+	UPROPERTY(BlueprintReadOnly, Category = "Movie Graph")
+	bool bMultipleCameraNames;
 };
 
 USTRUCT(BlueprintType)
@@ -235,13 +239,14 @@ public:
 	virtual void CacheDataPreJob(const FMovieGraphInitConfig& InInitConfig) {}
 	virtual void RestoreCachedDataPostJob() {}
 	virtual void UpdateShotList() {}
-	virtual void InitializeShot(const TObjectPtr<UMoviePipelineExecutorShot>& InShot) {}
+	virtual void InitializeShot(const TObjectPtr<UMoviePipelineExecutorShot>& InShot, const FFrameTime& InEvalTime) {}
 	virtual void CacheHierarchyForShot(const TObjectPtr<UMoviePipelineExecutorShot>& InShot) {}
 	virtual void RestoreHierarchyForShot(const TObjectPtr<UMoviePipelineExecutorShot> &InShot) {}
 	virtual void MuteShot(const TObjectPtr<UMoviePipelineExecutorShot>& InShot) {}
 	virtual void UnmuteShot(const TObjectPtr<UMoviePipelineExecutorShot>& InShot) {}
 	virtual void ExpandShot(const TObjectPtr<UMoviePipelineExecutorShot>& InShot, const int32 InLeftDeltaFrames, const int32 InLeftDeltaFramesUserPoV,
 		const int32 InRightDeltaFrames, const bool bInPrepass) {}
+	virtual TArray<FMinimalViewInfo> GetCameraInformation(UMoviePipelineExecutorShot* InShot, bool bIncludeSidecar) const { return TArray<FMinimalViewInfo>(); }
 
 	UMovieGraphPipeline* GetOwningGraph() const;
 };
@@ -283,7 +288,8 @@ struct FMovieGraphRenderPassLayerData
 {
 	FName BranchName;
 	FString LayerName;
-	FGuid CameraIdentifier;
+	int32 CameraIndex;
+	FString CameraName;
 	TWeakObjectPtr<class UMovieGraphRenderPassNode> RenderPassNode;
 };
 // ToDo: Both of these can probably go into the Default Renderer implementation.
@@ -420,6 +426,7 @@ namespace UE::MovieGraph
 		int32 BranchCount = 0;
 		int32 ActiveBranchRendererCount = 0;
 		int32 ActiveRendererSubresourceCount = 0;
+		int32 ActiveCameraCount = 0;
 	};
 
 	/**
diff --git a/Engine/Plugins/MovieScene/MovieRenderPipeline/Source/MovieRenderPipelineCore/Public/Graph/MovieGraphDefaultRenderer.h b/Engine/Plugins/MovieScene/MovieRenderPipeline/Source/MovieRenderPipelineCore/Public/Graph/MovieGraphDefaultRenderer.h
--- a/Engine/Plugins/MovieScene/MovieRenderPipeline/Source/MovieRenderPipelineCore/Public/Graph/MovieGraphDefaultRenderer.h
+++ b/Engine/Plugins/MovieScene/MovieRenderPipeline/Source/MovieRenderPipelineCore/Public/Graph/MovieGraphDefaultRenderer.h
@@ -50,6 +50,7 @@ namespace UE::MovieGraph::DefaultRenderer
 			, OverscanFraction(0.f)
 			, bAllowCameraAspectRatio(true)
 			, DoFSensorScale(1.0f)
+			, bUseCameraManagerPostProcess(true)
 		{}
 
 		/** Minimal view info describing the position and orientation in the world. */
@@ -77,6 +78,10 @@ namespace UE::MovieGraph::DefaultRenderer
 		// Sub-pixel jitter this camera should use. Only applied when using no AA.
 		FVector2D ProjectionMatrixJitterAmount;
 
+		// If true, we use the Post Process data from the Camera Manager, instead of deriving it from the camera location. This should be
+		// false when rendering multiple cameras.
+		bool bUseCameraManagerPostProcess;
+
 	};
 
 	struct FRenderTargetInitParams
@@ -277,7 +282,8 @@ public:
 	// ~UObject Interface
 
 	void AddOutstandingRenderTask_AnyThread(UE::Tasks::FTask InTask);
-	UE::MovieGraph::DefaultRenderer::FCameraInfo GetCameraInfo(const FGuid& InCameraIdentifier) const;
+	/** Fetches information for the given camera index. Should be "-1" when not using multi-camera rendering, or [0, n] based on the shot's sidecar camera data when using multi-camera rendering. */
+	UE::MovieGraph::DefaultRenderer::FCameraInfo GetCameraInfo(UMovieGraphEvaluatedConfig* InConfig, const int32 InCameraIndex) const;
 	void SetHasRenderedFirstViewThisFrame(bool bInValue) { bHasRenderedFirstViewThisFrame = bInValue; }
 	bool GetHasRenderedFirstViewThisFrame() const { return bHasRenderedFirstViewThisFrame; }
 public:
@@ -303,7 +309,12 @@ public:
 protected:
 	TObjectPtr<UTextureRenderTarget2D> CreateViewRenderTarget(const UE::MovieGraph::DefaultRenderer::FRenderTargetInitParams& InInitParams) const;
 	FMoviePipelineSurfaceQueuePtr CreateSurfaceQueue(const UE::MovieGraph::DefaultRenderer::FRenderTargetInitParams& InInitParams) const;
-	void GetCameraLocationsForFrame(TArray<FVector>& OutLocations) const;
+	
+	/** 
+	* Gets the camera locations for the current frame to inform the grass system of where streaming is required. If bIncludeSidecar is false, only the primary camera location is considered,
+	* otherwise it'll provide all the camera locations for multi-camera streaming (which can increase resource requirements).
+	*/
+	void GetCameraLocationsForFrame(TArray<FVector>& OutLocations, UMoviePipelineExecutorShot* InShot, bool bIncludeSidecar = false) const;
 	void FlushAsyncEngineSystems(const TObjectPtr<UMovieGraphEvaluatedConfig>& InConfig) const;
 
 protected:
diff --git a/Engine/Plugins/MovieScene/MovieRenderPipeline/Source/MovieRenderPipelineCore/Public/Graph/MovieGraphSequenceDataSource.h b/Engine/Plugins/MovieScene/MovieRenderPipeline/Source/MovieRenderPipelineCore/Public/Graph/MovieGraphSequenceDataSource.h
--- a/Engine/Plugins/MovieScene/MovieRenderPipeline/Source/MovieRenderPipelineCore/Public/Graph/MovieGraphSequenceDataSource.h
+++ b/Engine/Plugins/MovieScene/MovieRenderPipeline/Source/MovieRenderPipelineCore/Public/Graph/MovieGraphSequenceDataSource.h
@@ -53,6 +53,8 @@ public:
 	virtual void UnmuteShot(const TObjectPtr<UMoviePipelineExecutorShot>& InShot) override;
 	virtual void ExpandShot(const TObjectPtr<UMoviePipelineExecutorShot>& InShot, const int32 InLeftDeltaFrames, const int32 InLeftDeltaFramesUserPoV,
 		const int32 InRightDeltaFrames, const bool bInPrepass) override;
+	virtual TArray<FMinimalViewInfo> GetCameraInformation(UMoviePipelineExecutorShot* InShot, bool bIncludeSidecar) const override;
+	virtual void InitializeShot(const TObjectPtr<UMoviePipelineExecutorShot>& InShot, const FFrameTime& InEvalTime) override;
 protected:
 	void OverrideSequencePlaybackRangeFromGlobalOutputSettings(ULevelSequence* InSequence);
 	void CacheLevelSequenceData(ULevelSequence* InSequence);
diff --git a/Engine/Plugins/MovieScene/MovieRenderPipeline/Source/MovieRenderPipelineCore/Public/Graph/Nodes/MovieGraphCameraNode.h b/Engine/Plugins/MovieScene/MovieRenderPipeline/Source/MovieRenderPipelineCore/Public/Graph/Nodes/MovieGraphCameraNode.h
--- a/Engine/Plugins/MovieScene/MovieRenderPipeline/Source/MovieRenderPipelineCore/Public/Graph/Nodes/MovieGraphCameraNode.h
+++ b/Engine/Plugins/MovieScene/MovieRenderPipeline/Source/MovieRenderPipelineCore/Public/Graph/Nodes/MovieGraphCameraNode.h
@@ -16,6 +16,7 @@ public:
 	UMovieGraphCameraSettingNode()
 		: ShutterTiming(EMoviePipelineShutterTiming::FrameCenter)
 		, OverscanPercentage(0.f)
+		, bRenderAllCameras(false)
 	{}
 
 	virtual EMovieGraphBranchRestriction GetBranchRestriction() const override { return EMovieGraphBranchRestriction::Globals; }
@@ -35,6 +36,10 @@ public:
 	UPROPERTY(EditAnywhere, BlueprintReadWrite, Category = Overrides, meta = (InlineEditConditionToggle))
 	uint8 bOverride_OverscanPercentage : 1;
 
+
+	UPROPERTY(EditAnywhere, BlueprintReadWrite, Category = Overrides, meta = (InlineEditConditionToggle))
+	uint8 bOverride_bRenderAllCameras : 1;
+
 	/**
 	* Shutter Timing allows you to bias the timing of your shutter angle to either be before, during, or after
 	* a frame. When set to FrameClose, it means that the motion gathered up to produce frame N comes from 
@@ -58,4 +63,15 @@ public:
 	*/
 	UPROPERTY(EditAnywhere, BlueprintReadWrite, meta = (UIMin = "0", UIMax = "100", ClampMin = "0", ClampMax = "100"), Category = "Settings", meta = (EditCondition = "bOverride_OverscanPercentage"))
 	float OverscanPercentage;
+
+	/*
+	* If enabled Movie Render Queue will examine your Level Sequence for additional cameras and create an additional render for each renderer for that camera.
+	* The Camera Cut Track/Section is still used to determine the range of time to render, and then all Camera Actors that are in the level sequence adjacent
+	* to the Camera Cut Track will be considered for rendering. They are expected to exist the entire time and do not support rendering sub-ranges.
+	* 
+	* This increases render duration (100% per camera) and has increased VRAM/RAM requirements. However all cameras are rendered on the same engine tick so they
+	* should all see a consistent view of the world which can be useful for things like particle effects, etc.
+	*/
+	UPROPERTY(EditAnywhere, BlueprintReadWrite, Category = "Settings", meta = (EditCondition = "bOverride_bRenderAllCameras"))
+	bool bRenderAllCameras;
 };
\ No newline at end of file
diff --git a/Engine/Plugins/MovieScene/MovieRenderPipeline/Source/MovieRenderPipelineCore/Public/MoviePipelineUtils.h b/Engine/Plugins/MovieScene/MovieRenderPipeline/Source/MovieRenderPipelineCore/Public/MoviePipelineUtils.h
--- a/Engine/Plugins/MovieScene/MovieRenderPipeline/Source/MovieRenderPipelineCore/Public/MoviePipelineUtils.h
+++ b/Engine/Plugins/MovieScene/MovieRenderPipeline/Source/MovieRenderPipelineCore/Public/MoviePipelineUtils.h
@@ -145,6 +145,7 @@ namespace UE
 		MOVIERENDERPIPELINECORE_API FMoviePipelineRenderPassMetrics GetRenderPassMetrics(UMoviePipelinePrimaryConfig* InPrimaryConfig, UMoviePipelineExecutorShot* InPipelineExecutorShot, const FMoviePipelineRenderPassMetrics& InRenderPassMetrics, const FIntPoint& InEffectiveOutputResolution);
 		MOVIERENDERPIPELINECORE_API bool CanWriteToFile(const TCHAR* InFilename, bool bOverwriteExisting);
 		MOVIERENDERPIPELINECORE_API FString GetPaddingFormatString(int32 InZeroPadCount, const int32 InFrameNumber);
+		MOVIERENDERPIPELINECORE_API void DoPostProcessBlend(const FVector& InViewLocation, const class UWorld* InWorld, const struct FMinimalViewInfo& InViewInfo, class FSceneView* InOutView);
 
 		/** When using spatial/temporal samples without anti-aliasing, get the sub-pixel jitter for the given frame index. FrameIndex is modded by InSamplesPerFrame so that the aa jitter pattern repeats every output frame. */
 		MOVIERENDERPIPELINECORE_API FVector2f GetSubPixelJitter(int32 InFrameIndex, int32 InSamplesPerFrame);
diff --git a/Engine/Plugins/MovieScene/MovieRenderPipeline/Source/MovieRenderPipelineRenderPasses/Private/Graph/Renderers/MovieGraphDeferredPass.cpp b/Engine/Plugins/MovieScene/MovieRenderPipeline/Source/MovieRenderPipelineRenderPasses/Private/Graph/Renderers/MovieGraphDeferredPass.cpp
--- a/Engine/Plugins/MovieScene/MovieRenderPipeline/Source/MovieRenderPipelineRenderPasses/Private/Graph/Renderers/MovieGraphDeferredPass.cpp
+++ b/Engine/Plugins/MovieScene/MovieRenderPipeline/Source/MovieRenderPipelineRenderPasses/Private/Graph/Renderers/MovieGraphDeferredPass.cpp
@@ -43,8 +43,7 @@ void FMovieGraphDeferredPass::Setup(TWeakObjectPtr<UMovieGraphDefaultRenderer> I
 	RenderDataIdentifier.RendererName = InRenderPassNode->GetRendererName();
 	RenderDataIdentifier.SubResourceName = TEXT("beauty");
 	
-	UE::MovieGraph::DefaultRenderer::FCameraInfo CameraInfo = InRenderer->GetCameraInfo(LayerData.CameraIdentifier);
-	RenderDataIdentifier.CameraName =  CameraInfo.CameraName;
+	RenderDataIdentifier.CameraName =  InLayer.CameraName;
 
 	SceneViewState.Allocate(InRenderer->GetWorld()->GetFeatureLevel());
 
@@ -140,10 +139,12 @@ void FMovieGraphDeferredPass::Render(const FMovieGraphTraversalContext& InFrameT
 
 	// Camera nodes are optional
 	const bool bIncludeCDOs = false;
+	bool bRenderAllCameras = false;
 	const UMovieGraphCameraSettingNode* CameraNode = InTimeData.EvaluatedConfig->GetSettingForBranch<UMovieGraphCameraSettingNode>(LayerData.BranchName, bIncludeCDOs);
 	if (CameraNode)
 	{
 		OverscanFraction = FMath::Clamp(CameraNode->OverscanPercentage / 100.f, 0.f, 1.f);
+		bRenderAllCameras = CameraNode->bRenderAllCameras;
 	}
 
 	// ToDo: When tiling is used, this should be the size of the per-tile backbuffer
@@ -175,16 +176,7 @@ void FMovieGraphDeferredPass::Render(const FMovieGraphTraversalContext& InFrameT
 		}
 		 
 		// These are the parameters of our camera 
-		UE::MovieGraph::DefaultRenderer::FCameraInfo CameraInfo;
-
-		// ToDo: Get this from the renderer based on LayerData.CameraIdentifier for eventual multi-camera support
-		APlayerController* LocalPlayerController = GraphRenderer->GetWorld()->GetFirstPlayerController();
-		// CameraAnim override
-		if (LocalPlayerController->PlayerCameraManager)
-		{
-			CameraInfo.ViewInfo = LocalPlayerController->PlayerCameraManager->GetCameraCacheView();
-			CameraInfo.ViewActor = LocalPlayerController->GetViewTarget();
-		}
+		UE::MovieGraph::DefaultRenderer::FCameraInfo CameraInfo = GetRenderer()->GetCameraInfo(InTimeData.EvaluatedConfig, LayerData.CameraIndex);
 
 		CameraInfo.bAllowCameraAspectRatio = true;
 		CameraInfo.TilingParams.TileSize = BackbufferResolution;
@@ -197,6 +189,7 @@ void FMovieGraphDeferredPass::Render(const FMovieGraphTraversalContext& InFrameT
 		CameraInfo.SamplingParams.SpatialSampleCount = NumSpatialSamples;
 		CameraInfo.OverscanFraction = OverscanFraction;
 		CameraInfo.ProjectionMatrixJitterAmount = FVector2D((SpatialShiftAmount.X) * 2.0f / (float)BackbufferResolution.X, SpatialShiftAmount.Y * -2.0f / (float)BackbufferResolution.Y);
+		CameraInfo.bUseCameraManagerPostProcess = !bRenderAllCameras;
 		
 		// ToDo: This math probably needs the per-tile, pre-overlapped size? 
 		FIntPoint OverlappedPad = FIntPoint(FMath::CeilToInt(BackbufferResolution.X * TileOverlapPadRatio), FMath::CeilToInt(BackbufferResolution.Y * TileOverlapPadRatio));
diff --git a/Engine/Plugins/MovieScene/MovieRenderPipeline/Source/MovieRenderPipelineRenderPasses/Private/Graph/Renderers/MovieGraphImagePassBase.cpp b/Engine/Plugins/MovieScene/MovieRenderPipeline/Source/MovieRenderPipelineRenderPasses/Private/Graph/Renderers/MovieGraphImagePassBase.cpp
--- a/Engine/Plugins/MovieScene/MovieRenderPipeline/Source/MovieRenderPipelineRenderPasses/Private/Graph/Renderers/MovieGraphImagePassBase.cpp
+++ b/Engine/Plugins/MovieScene/MovieRenderPipeline/Source/MovieRenderPipelineRenderPasses/Private/Graph/Renderers/MovieGraphImagePassBase.cpp
@@ -69,7 +69,7 @@ FSceneView* FMovieGraphImagePassBase::CreateSceneView(const FSceneViewInitOption
 	InViewFamily->Views.Add(View);
 
 	View->StartFinalPostprocessSettings(InInitOptions.ViewLocation);
-	ApplyCameraManagerPostProcessBlends(View);
+	ApplyCameraManagerPostProcessBlends(View, InCameraInfo.ViewInfo, InCameraInfo.bUseCameraManagerPostProcess);
 
 	// Scaling sensor size inversely with the the projection matrix [0][0] should physically
 	// cause the circle of confusion to be unchanged.
@@ -96,7 +96,7 @@ DefaultRenderer::FRenderTargetInitParams FMovieGraphImagePassBase::GetRenderTarg
 	return InitParams;
 }
 
-void FMovieGraphImagePassBase::ApplyCameraManagerPostProcessBlends(FSceneView* InView) const
+void FMovieGraphImagePassBase::ApplyCameraManagerPostProcessBlends(FSceneView* InView, const FMinimalViewInfo& InViewInfo, bool bUseCameraManagerPostProcess) const
 {
 	check(InView);
 
@@ -106,33 +106,40 @@ void FMovieGraphImagePassBase::ApplyCameraManagerPostProcessBlends(FSceneView* I
 		return;
 	}
 
-	APlayerController* LocalPlayerController = GraphRenderer->GetWorld()->GetFirstPlayerController();
-	// CameraAnim override
-	if (LocalPlayerController->PlayerCameraManager)
+	if (bUseCameraManagerPostProcess)
 	{
-		TArray<FPostProcessSettings> const* CameraAnimPPSettings;
-		TArray<float> const* CameraAnimPPBlendWeights;
-		LocalPlayerController->PlayerCameraManager->GetCachedPostProcessBlends(CameraAnimPPSettings, CameraAnimPPBlendWeights);
-
-		if (LocalPlayerController->PlayerCameraManager->bEnableFading)
+		APlayerController* LocalPlayerController = GraphRenderer->GetWorld()->GetFirstPlayerController();
+		// CameraAnim override
+		if (LocalPlayerController->PlayerCameraManager)
 		{
-			InView->OverlayColor = LocalPlayerController->PlayerCameraManager->FadeColor;
-			InView->OverlayColor.A = FMath::Clamp(LocalPlayerController->PlayerCameraManager->FadeAmount, 0.f, 1.f);
-		}
+			TArray<FPostProcessSettings> const* CameraAnimPPSettings;
+			TArray<float> const* CameraAnimPPBlendWeights;
+			LocalPlayerController->PlayerCameraManager->GetCachedPostProcessBlends(CameraAnimPPSettings, CameraAnimPPBlendWeights);
 
-		if (LocalPlayerController->PlayerCameraManager->bEnableColorScaling)
-		{
-			FVector ColorScale = LocalPlayerController->PlayerCameraManager->ColorScale;
-			InView->ColorScale = FLinearColor(ColorScale.X, ColorScale.Y, ColorScale.Z);
-		}
+			if (LocalPlayerController->PlayerCameraManager->bEnableFading)
+			{
+				InView->OverlayColor = LocalPlayerController->PlayerCameraManager->FadeColor;
+				InView->OverlayColor.A = FMath::Clamp(LocalPlayerController->PlayerCameraManager->FadeAmount, 0.f, 1.f);
+			}
 
-		FMinimalViewInfo ViewInfo = LocalPlayerController->PlayerCameraManager->GetCameraCacheView();
-		for (int32 PPIdx = 0; PPIdx < CameraAnimPPBlendWeights->Num(); ++PPIdx)
-		{
-			InView->OverridePostProcessSettings((*CameraAnimPPSettings)[PPIdx], (*CameraAnimPPBlendWeights)[PPIdx]);
-		}
+			if (LocalPlayerController->PlayerCameraManager->bEnableColorScaling)
+			{
+				FVector ColorScale = LocalPlayerController->PlayerCameraManager->ColorScale;
+				InView->ColorScale = FLinearColor(ColorScale.X, ColorScale.Y, ColorScale.Z);
+			}
+
+			FMinimalViewInfo ViewInfo = LocalPlayerController->PlayerCameraManager->GetCameraCacheView();
+			for (int32 PPIdx = 0; PPIdx < CameraAnimPPBlendWeights->Num(); ++PPIdx)
+			{
+				InView->OverridePostProcessSettings((*CameraAnimPPSettings)[PPIdx], (*CameraAnimPPBlendWeights)[PPIdx]);
+			}
 
-		InView->OverridePostProcessSettings(ViewInfo.PostProcessSettings, ViewInfo.PostProcessBlendWeight);
+			InView->OverridePostProcessSettings(ViewInfo.PostProcessSettings, ViewInfo.PostProcessBlendWeight);
+		}
+	}
+	else
+	{
+		UE::MoviePipeline::DoPostProcessBlend(InViewInfo.Location, GraphRenderer->GetWorld(), InViewInfo, InView);
 	}
 }
 
diff --git a/Engine/Plugins/MovieScene/MovieRenderPipeline/Source/MovieRenderPipelineRenderPasses/Private/MoviePipelineDeferredPasses.cpp b/Engine/Plugins/MovieScene/MovieRenderPipeline/Source/MovieRenderPipelineRenderPasses/Private/MoviePipelineDeferredPasses.cpp
--- a/Engine/Plugins/MovieScene/MovieRenderPipeline/Source/MovieRenderPipelineRenderPasses/Private/MoviePipelineDeferredPasses.cpp
+++ b/Engine/Plugins/MovieScene/MovieRenderPipeline/Source/MovieRenderPipelineRenderPasses/Private/MoviePipelineDeferredPasses.cpp
@@ -859,40 +859,9 @@ void UMoviePipelineDeferredPassBase::BlendPostProcessSettings(FSceneView* InView
 			return;
 		}
 
-		// For sidecar cameras we need to do the blending of PP volumes
+		// For sidecar cameras we need to do the blending of PP volumes and camera PP manually.
 		FVector ViewLocation = OutCamera->GetComponentLocation();
-		for (IInterface_PostProcessVolume* PPVolume : GetWorld()->PostProcessVolumes)
-		{
-			const FPostProcessVolumeProperties VolumeProperties = PPVolume->GetProperties();
-
-			// Skip any volumes which are disabled
-			if (!VolumeProperties.bIsEnabled)
-			{
-				continue;
-			}
-
-			float LocalWeight = FMath::Clamp(VolumeProperties.BlendWeight, 0.0f, 1.0f);
-
-			if (!VolumeProperties.bIsUnbound)
-			{
-				float DistanceToPoint = 0.0f;
-				PPVolume->EncompassesPoint(ViewLocation, 0.0f, &DistanceToPoint);
-
-				if (DistanceToPoint >= 0 && DistanceToPoint < VolumeProperties.BlendRadius)
-				{
-					LocalWeight *= FMath::Clamp(1.0f - DistanceToPoint / VolumeProperties.BlendRadius, 0.0f, 1.0f);
-				}
-				else
-				{
-					LocalWeight = 0.0f;
-				}
-			}
-
-			InView->OverridePostProcessSettings(*VolumeProperties.Settings, LocalWeight);
-		}
-
-		// After blending all post processing volumes, blend the camera's post process settings too
-		InView->OverridePostProcessSettings(OutViewInfo.PostProcessSettings, OutViewInfo.PostProcessBlendWeight);
+		UE::MoviePipeline::DoPostProcessBlend(ViewLocation, GetWorld(), OutViewInfo, InView);
 	}
 }
 
diff --git a/Engine/Plugins/MovieScene/MovieRenderPipeline/Source/MovieRenderPipelineRenderPasses/Public/Graph/Renderers/MovieGraphImagePassBase.h b/Engine/Plugins/MovieScene/MovieRenderPipeline/Source/MovieRenderPipelineRenderPasses/Public/Graph/Renderers/MovieGraphImagePassBase.h
--- a/Engine/Plugins/MovieScene/MovieRenderPipeline/Source/MovieRenderPipelineRenderPasses/Public/Graph/Renderers/MovieGraphImagePassBase.h
+++ b/Engine/Plugins/MovieScene/MovieRenderPipeline/Source/MovieRenderPipelineRenderPasses/Public/Graph/Renderers/MovieGraphImagePassBase.h
@@ -88,7 +88,7 @@ namespace UE::MovieGraph::Rendering
 		/** Gets the render target init params that will be used when creating the render target. */
 		virtual DefaultRenderer::FRenderTargetInitParams GetRenderTargetInitParams(const FMovieGraphTimeStepData& InTimeData, const FIntPoint& InResolution);
 
-		virtual void ApplyCameraManagerPostProcessBlends(FSceneView* InView) const;
+		virtual void ApplyCameraManagerPostProcessBlends(FSceneView* InView, const FMinimalViewInfo& InViewInfo, bool bApplyCameraManagerPostProcessBlends) const;
 		virtual TSharedRef<FSceneViewFamilyContext> CreateSceneViewFamily(const FViewFamilyInitData& InInitData) const;
 		virtual void ApplyMovieGraphOverridesToSceneView(TSharedRef<FSceneViewFamilyContext> InOutFamily, const FViewFamilyInitData& InInitData, const UE::MovieGraph::DefaultRenderer::FCameraInfo& InCameraInfo) const;
 		virtual void ApplyMovieGraphOverridesToViewFamily(TSharedRef<FSceneViewFamilyContext> InOutFamily, const FViewFamilyInitData& InInitData) const;
