Added support for camera overscan to MRQ. MRQ now uses the camera's overscan property to set overscan, and MRQ overscan property has been changed to an override, which when enabled replaces the camera's overscan value with its own.

diff --git a/Engine/Plugins/MovieScene/MovieRenderPipeline/Source/MovieRenderPipelineCore/Private/Graph/MovieGraphBlueprintLibrary.cpp b/Engine/Plugins/MovieScene/MovieRenderPipeline/Source/MovieRenderPipelineCore/Private/Graph/MovieGraphBlueprintLibrary.cpp
--- a/Engine/Plugins/MovieScene/MovieRenderPipeline/Source/MovieRenderPipelineCore/Private/Graph/MovieGraphBlueprintLibrary.cpp
+++ b/Engine/Plugins/MovieScene/MovieRenderPipeline/Source/MovieRenderPipelineCore/Private/Graph/MovieGraphBlueprintLibrary.cpp
@@ -364,7 +364,7 @@ int32 UMovieGraphBlueprintLibrary::GetCurrentVersionNumber(const UMovieGraphPipe
 	return InMovieGraphPipeline->GetActiveShotList()[CurrentShotIndex]->ShotInfo.VersionNumber;
 }
 
-FIntPoint UMovieGraphBlueprintLibrary::GetEffectiveOutputResolution(UMovieGraphEvaluatedConfig* InEvaluatedGraph)
+FIntPoint UMovieGraphBlueprintLibrary::GetEffectiveOutputResolution(UMovieGraphEvaluatedConfig* InEvaluatedGraph, float DefaultOverscan)
 {
 	if (!InEvaluatedGraph)
 	{
@@ -379,11 +379,14 @@ FIntPoint UMovieGraphBlueprintLibrary::GetEffectiveOutputResolution(UMovieGraphE
 		return FIntPoint();
 	}
 
-	float RescaledOverscan = 0.f;
+	float RescaledOverscan = DefaultOverscan;
 	if (UMovieGraphCameraSettingNode* CameraSetting = InEvaluatedGraph->GetSettingForBranch<UMovieGraphCameraSettingNode>(UMovieGraphNode::GlobalsPinName, bIncludeCDOs))
 	{
-		// The old system used [0-1] range for floats, the new system will use [0-100], so we rescale down before calling through.
-		RescaledOverscan = FMath::Clamp(CameraSetting->OverscanPercentage / 100.f, 0.f, 1.f);
+		if (CameraSetting->bOverride_OverscanPercentage)
+		{
+			// The old system used [0-1] range for floats, the new system will use [0-100], so we rescale down before calling through.
+			RescaledOverscan = FMath::Clamp(CameraSetting->OverscanPercentage / 100.f, 0.f, 1.f);
+		}
 	}
 
 	// We need to look at the Project Settings for the latest value for a given profile
diff --git a/Engine/Plugins/MovieScene/MovieRenderPipeline/Source/MovieRenderPipelineCore/Private/Graph/Nodes/MovieGraphWidgetRendererBaseNode.cpp b/Engine/Plugins/MovieScene/MovieRenderPipeline/Source/MovieRenderPipelineCore/Private/Graph/Nodes/MovieGraphWidgetRendererBaseNode.cpp
--- a/Engine/Plugins/MovieScene/MovieRenderPipeline/Source/MovieRenderPipelineCore/Private/Graph/Nodes/MovieGraphWidgetRendererBaseNode.cpp
+++ b/Engine/Plugins/MovieScene/MovieRenderPipeline/Source/MovieRenderPipelineCore/Private/Graph/Nodes/MovieGraphWidgetRendererBaseNode.cpp
@@ -106,7 +106,8 @@ void UMovieGraphWidgetRendererBaseNode::FMovieGraphWidgetPass::Render(const FMov
 	const UMovieGraphPipeline* Pipeline = Renderer->GetOwningGraph();
 	UMovieGraphWidgetRendererBaseNode* ParentNodeThisFrame = GetParentNode(InTimeData.EvaluatedConfig);
 
-	const FIntPoint OutputResolution = UMovieGraphBlueprintLibrary::GetEffectiveOutputResolution(InTimeData.EvaluatedConfig);
+	const float CameraOverscan = Renderer->GetCameraInfo(InTimeData.EvaluatedConfig, LayerData.CameraIndex).ViewInfo.GetOverscan();
+	const FIntPoint OutputResolution = UMovieGraphBlueprintLibrary::GetEffectiveOutputResolution(InTimeData.EvaluatedConfig, CameraOverscan);
 	const int32 MaxResolution = GetMax2DTextureDimension();
 	if ((OutputResolution.X > MaxResolution) || (OutputResolution.Y > MaxResolution))
 	{
diff --git a/Engine/Plugins/MovieScene/MovieRenderPipeline/Source/MovieRenderPipelineCore/Private/MoviePipelineBlueprintLibrary.cpp b/Engine/Plugins/MovieScene/MovieRenderPipeline/Source/MovieRenderPipelineCore/Private/MoviePipelineBlueprintLibrary.cpp
--- a/Engine/Plugins/MovieScene/MovieRenderPipeline/Source/MovieRenderPipelineCore/Private/MoviePipelineBlueprintLibrary.cpp
+++ b/Engine/Plugins/MovieScene/MovieRenderPipeline/Source/MovieRenderPipelineCore/Private/MoviePipelineBlueprintLibrary.cpp
@@ -900,7 +900,7 @@ int32 UMoviePipelineBlueprintLibrary::GetCurrentVersionNumber(const UMoviePipeli
 	return InMoviePipeline->GetActiveShotList()[CurrentShotIndex]->ShotInfo.VersionNumber;
 }
 
-FIntPoint UMoviePipelineBlueprintLibrary::GetEffectiveOutputResolution(UMoviePipelinePrimaryConfig* InPrimaryConfig, UMoviePipelineExecutorShot* InPipelineExecutorShot)
+FIntPoint UMoviePipelineBlueprintLibrary::GetEffectiveOutputResolution(UMoviePipelinePrimaryConfig* InPrimaryConfig, UMoviePipelineExecutorShot* InPipelineExecutorShot, float DefaultOverscan)
 {
 	if (InPrimaryConfig && InPipelineExecutorShot)
 	{
@@ -908,7 +908,7 @@ FIntPoint UMoviePipelineBlueprintLibrary::GetEffectiveOutputResolution(UMoviePip
 		const UMoviePipelineCameraSetting* CameraSetting = Cast<const UMoviePipelineCameraSetting>(UMoviePipelineBlueprintLibrary::FindOrGetDefaultSettingForShot(UMoviePipelineCameraSetting::StaticClass(), InPrimaryConfig, InPipelineExecutorShot));
 		check(OutputSetting);
 		check(CameraSetting);
-		return Utility_GetEffectiveOutputResolution(CameraSetting->OverscanPercentage, OutputSetting->OutputResolution);
+		return Utility_GetEffectiveOutputResolution(CameraSetting->bOverrideCameraOverscan ? CameraSetting->OverscanPercentage : DefaultOverscan, OutputSetting->OutputResolution);
 	}
 
 	return FIntPoint();
diff --git a/Engine/Plugins/MovieScene/MovieRenderPipeline/Source/MovieRenderPipelineCore/Private/MoviePipelineRendering.cpp b/Engine/Plugins/MovieScene/MovieRenderPipeline/Source/MovieRenderPipelineCore/Private/MoviePipelineRendering.cpp
--- a/Engine/Plugins/MovieScene/MovieRenderPipeline/Source/MovieRenderPipelineCore/Private/MoviePipelineRendering.cpp
+++ b/Engine/Plugins/MovieScene/MovieRenderPipeline/Source/MovieRenderPipelineCore/Private/MoviePipelineRendering.cpp
@@ -75,9 +75,11 @@ void UMoviePipeline::SetupRenderingPipelineForShot(UMoviePipelineExecutorShot* I
 	UMoviePipelineOutputSetting* OutputSettings = GetPipelinePrimaryConfig()->FindSetting<UMoviePipelineOutputSetting>();
 	check(OutputSettings);
 
-
+	// TODO: Not much support here for multi-camera, so simply get the player controller camera and use its overscan value
+	FMinimalViewInfo CameraViewInfo = GetWorld()->GetFirstPlayerController()->PlayerCameraManager->GetCameraCacheView();
+	
 	FIntPoint BackbufferTileCount = FIntPoint(HighResSettings->TileCount, HighResSettings->TileCount);
-	FIntPoint OutputResolution = UMoviePipelineBlueprintLibrary::GetEffectiveOutputResolution(GetPipelinePrimaryConfig(), InShot);
+	FIntPoint OutputResolution = UMoviePipelineBlueprintLibrary::GetEffectiveOutputResolution(GetPipelinePrimaryConfig(), InShot, CameraViewInfo.GetOverscan());
 
 	// Figure out how big each sub-region (tile) is.
 	FIntPoint BackbufferResolution = FIntPoint(
@@ -192,9 +194,12 @@ void UMoviePipeline::RenderFrame()
 	check(HighResSettings);
 	check(OutputSettings);
 
+	// TODO: Not much support here for multi-camera, so simply get the player controller camera and use its overscan value
+	FMinimalViewInfo CameraViewInfo = GetWorld()->GetFirstPlayerController()->PlayerCameraManager->GetCameraCacheView();
+	
 	FIntPoint TileCount = FIntPoint(HighResSettings->TileCount, HighResSettings->TileCount);
 	FIntPoint OriginalTileCount = TileCount;
-	FIntPoint OutputResolution = UMoviePipelineBlueprintLibrary::GetEffectiveOutputResolution(GetPipelinePrimaryConfig(), ActiveShotList[CurrentShotIndex]);
+	FIntPoint OutputResolution = UMoviePipelineBlueprintLibrary::GetEffectiveOutputResolution(GetPipelinePrimaryConfig(), ActiveShotList[CurrentShotIndex], CameraViewInfo.GetOverscan());
 
 	int32 NumSpatialSamples = AntiAliasingSettings->SpatialSampleCount;
 	int32 NumTemporalSamples = AntiAliasingSettings->TemporalSampleCount;
@@ -400,6 +405,7 @@ void UMoviePipeline::RenderFrame()
 				SampleState.TextureSharpnessBias = HighResSettings->TextureSharpnessBias;
 				SampleState.OCIOConfiguration = ColorSettings ? &ColorSettings->OCIOConfiguration : nullptr;
 				SampleState.GlobalScreenPercentageFraction = FLegacyScreenPercentageDriver::GetCVarResolutionFraction();
+				SampleState.bOverrideCameraOverscan = CameraSettings->bOverrideCameraOverscan;
 				SampleState.OverscanPercentage = FMath::Clamp(CameraSettings->OverscanPercentage, 0.0f, 1.0f);
 
 				if (FrameThrotteCount > 0)
diff --git a/Engine/Plugins/MovieScene/MovieRenderPipeline/Source/MovieRenderPipelineCore/Public/Graph/MovieGraphBlueprintLibrary.h b/Engine/Plugins/MovieScene/MovieRenderPipeline/Source/MovieRenderPipelineCore/Public/Graph/MovieGraphBlueprintLibrary.h
--- a/Engine/Plugins/MovieScene/MovieRenderPipeline/Source/MovieRenderPipelineCore/Public/Graph/MovieGraphBlueprintLibrary.h
+++ b/Engine/Plugins/MovieScene/MovieRenderPipeline/Source/MovieRenderPipelineCore/Public/Graph/MovieGraphBlueprintLibrary.h
@@ -64,10 +64,11 @@ public:
 	/**
 	* In case of overscan percentage being higher than 0, additional pixels are rendered. This function returns the resolution with overscan taken into account.
 	* @param	InEvaluatedGraph	- The evaluated graph that will provide context for resolving the resolution
+	* @param	DefaultOverscan		- The default overscan to use if there are no camera settings that provide an overscan override value, from 0.0 to 1.0
 	* @return						- The output resolution, taking into account overscan
 	*/
 	UFUNCTION(BlueprintCallable, Category = "Movie Graph")
-	static FIntPoint GetEffectiveOutputResolution(UMovieGraphEvaluatedConfig* InEvaluatedGraph);
+	static FIntPoint GetEffectiveOutputResolution(UMovieGraphEvaluatedConfig* InEvaluatedGraph, float DefaultOverscan = 0.0f);
 
 	/**
 	* Gets the name of the current job.
diff --git a/Engine/Plugins/MovieScene/MovieRenderPipeline/Source/MovieRenderPipelineCore/Public/Graph/MovieGraphDefaultRenderer.h b/Engine/Plugins/MovieScene/MovieRenderPipeline/Source/MovieRenderPipelineCore/Public/Graph/MovieGraphDefaultRenderer.h
--- a/Engine/Plugins/MovieScene/MovieRenderPipeline/Source/MovieRenderPipelineCore/Public/Graph/MovieGraphDefaultRenderer.h
+++ b/Engine/Plugins/MovieScene/MovieRenderPipeline/Source/MovieRenderPipelineCore/Public/Graph/MovieGraphDefaultRenderer.h
@@ -47,7 +47,6 @@ namespace UE::MovieGraph::DefaultRenderer
 	{
 		FCameraInfo()
 			: ViewActor(nullptr)
-			, OverscanFraction(0.f)
 			, bAllowCameraAspectRatio(true)
 			, DoFSensorScale(1.0f)
 			, bUseCameraManagerPostProcess(true)
@@ -62,9 +61,6 @@ namespace UE::MovieGraph::DefaultRenderer
 		/** The name to use for the {camera_name} token. Filled out by the renderer where possible. */
 		FString CameraName;
 
-		/** In the [0-1] Range. */
-		float OverscanFraction;
-
 		/** should we respect the camera's aspect ratio settings. */
 		bool bAllowCameraAspectRatio;
 
diff --git a/Engine/Plugins/MovieScene/MovieRenderPipeline/Source/MovieRenderPipelineCore/Public/Graph/Nodes/MovieGraphCameraNode.h b/Engine/Plugins/MovieScene/MovieRenderPipeline/Source/MovieRenderPipelineCore/Public/Graph/Nodes/MovieGraphCameraNode.h
--- a/Engine/Plugins/MovieScene/MovieRenderPipeline/Source/MovieRenderPipelineCore/Public/Graph/Nodes/MovieGraphCameraNode.h
+++ b/Engine/Plugins/MovieScene/MovieRenderPipeline/Source/MovieRenderPipelineCore/Public/Graph/Nodes/MovieGraphCameraNode.h
@@ -61,7 +61,7 @@ public:
 	* Note: This uses 0-100 and not 0-1 like the previous system did to bring it in-line with other usages
 	* of overscan in the engine (nDisplay).
 	*/
-	UPROPERTY(EditAnywhere, BlueprintReadWrite, meta = (UIMin = "0", UIMax = "100", ClampMin = "0", ClampMax = "100"), Category = "Settings", meta = (EditCondition = "bOverride_OverscanPercentage"))
+	UPROPERTY(EditAnywhere, BlueprintReadWrite, meta = (DisplayName = "Overscan Percentage Override", UIMin = "0", UIMax = "100", ClampMin = "0", ClampMax = "100"), Category = "Settings", meta = (EditCondition = "bOverride_OverscanPercentage"))
 	float OverscanPercentage;
 
 	/*
diff --git a/Engine/Plugins/MovieScene/MovieRenderPipeline/Source/MovieRenderPipelineCore/Public/MoviePipelineBlueprintLibrary.h b/Engine/Plugins/MovieScene/MovieRenderPipeline/Source/MovieRenderPipelineCore/Public/MoviePipelineBlueprintLibrary.h
--- a/Engine/Plugins/MovieScene/MovieRenderPipeline/Source/MovieRenderPipelineCore/Public/MoviePipelineBlueprintLibrary.h
+++ b/Engine/Plugins/MovieScene/MovieRenderPipeline/Source/MovieRenderPipelineCore/Public/MoviePipelineBlueprintLibrary.h
@@ -149,7 +149,7 @@ public:
 	
 	/** In case of Overscan percentage being higher than 0 we render additional pixels. This function returns the resolution with overscan taken into account. */
 	UFUNCTION(BlueprintCallable, Category = "Movie Render Pipeline")
-	static FIntPoint GetEffectiveOutputResolution(UMoviePipelinePrimaryConfig* InPrimaryConfig, UMoviePipelineExecutorShot* InPipelineExecutorShot);
+	static FIntPoint GetEffectiveOutputResolution(UMoviePipelinePrimaryConfig* InPrimaryConfig, UMoviePipelineExecutorShot* InPipelineExecutorShot, float DefaultOverscan = 0.0f);
 
 	/** Allows access to a setting of provided type for specific shot. */
 	UFUNCTION(BlueprintCallable, Category = "Movie Render Pipeline", meta = (DeterminesOutputType = "InSettingType"))
diff --git a/Engine/Plugins/MovieScene/MovieRenderPipeline/Source/MovieRenderPipelineCore/Public/MoviePipelineCameraSetting.h b/Engine/Plugins/MovieScene/MovieRenderPipeline/Source/MovieRenderPipelineCore/Public/MoviePipelineCameraSetting.h
--- a/Engine/Plugins/MovieScene/MovieRenderPipeline/Source/MovieRenderPipelineCore/Public/MoviePipelineCameraSetting.h
+++ b/Engine/Plugins/MovieScene/MovieRenderPipeline/Source/MovieRenderPipelineCore/Public/MoviePipelineCameraSetting.h
@@ -37,6 +37,13 @@ public:
 	UPROPERTY(EditAnywhere, BlueprintReadWrite, Category = "Camera Settings")
 	EMoviePipelineShutterTiming ShutterTiming;
 
+	/**
+	 * If true, the camera settings overscan value will override any overscan on the cameras when rendering;
+	 * otherwise, the overscan value on the cameras will be used.
+	 */
+	UPROPERTY(EditAnywhere, BlueprintReadWrite, Category = "Camera Settings", meta = (InlineEditConditionToggle))
+	bool bOverrideCameraOverscan;
+	
 	/**
 	* Overscan percent allows to render additional pixels beyond the set resolution and can be used in conjunction 
 	* with EXR file output to add post-processing effects such as lens distortion.
@@ -45,7 +52,7 @@ public:
 	* in post production. For all other formats this will increase the final resolution and no pixels will be hidden 
 	* (ie: 1080p /w 0.1 overscan will make a 2112x1188 jpg, but a 1080p exr /w 96/54 pixels hidden on each side)
 	*/
-	UPROPERTY(EditAnywhere, BlueprintReadWrite, meta = (UIMin = "0", UIMax = "1", ClampMin = "0", ClampMax = "1"), Category = "Camera Settings")
+	UPROPERTY(EditAnywhere, BlueprintReadWrite, meta = (DisplayName = "Overscan Percentage Override", UIMin = "0", UIMax = "1", ClampMin = "0", ClampMax = "1", EditCondition = "bOverrideCameraOverscan"), Category = "Camera Settings")
 	float OverscanPercentage;
 	
 	/**
diff --git a/Engine/Plugins/MovieScene/MovieRenderPipeline/Source/MovieRenderPipelineCore/Public/MovieRenderPipelineDataTypes.h b/Engine/Plugins/MovieScene/MovieRenderPipeline/Source/MovieRenderPipelineCore/Public/MovieRenderPipelineDataTypes.h
--- a/Engine/Plugins/MovieScene/MovieRenderPipeline/Source/MovieRenderPipelineCore/Public/MovieRenderPipelineDataTypes.h
+++ b/Engine/Plugins/MovieScene/MovieRenderPipeline/Source/MovieRenderPipelineCore/Public/MovieRenderPipelineDataTypes.h
@@ -999,6 +999,9 @@ public:
 	/** Use overscan percentage to extend render region beyond the set resolution.  */
 	float OverscanPercentage;
 
+	/** Whether to override the camera's overscan value when rendering */
+	bool bOverrideCameraOverscan;
+	
 	/** 
 	* The gamma space to apply accumulation in. During accumulation, pow(x,AccumulationGamma) is applied
 	* and pow(x,1/AccumulationGamma) is applied after accumulation is finished. 1.0 means no change."
diff --git a/Engine/Plugins/MovieScene/MovieRenderPipeline/Source/MovieRenderPipelineRenderPasses/Private/Graph/Renderers/MovieGraphDeferredPass.cpp b/Engine/Plugins/MovieScene/MovieRenderPipeline/Source/MovieRenderPipelineRenderPasses/Private/Graph/Renderers/MovieGraphDeferredPass.cpp
--- a/Engine/Plugins/MovieScene/MovieRenderPipeline/Source/MovieRenderPipelineRenderPasses/Private/Graph/Renderers/MovieGraphDeferredPass.cpp
+++ b/Engine/Plugins/MovieScene/MovieRenderPipeline/Source/MovieRenderPipelineRenderPasses/Private/Graph/Renderers/MovieGraphDeferredPass.cpp
@@ -135,6 +135,7 @@ void FMovieGraphDeferredPass::Render(const FMovieGraphTraversalContext& InFrameT
 	const ESceneCaptureSource SceneCaptureSource = ParentNodeThisFrame->GetDisableToneCurve() ? ESceneCaptureSource::SCS_FinalColorHDR : ESceneCaptureSource::SCS_FinalToneCurveHDR;
 	const EAntiAliasingMethod AntiAliasingMethod = ParentNodeThisFrame->GetAntiAliasingMethod();
 	float OverscanFraction = 0.f;
+	bool bOverrideCameraOverscan = false;
 	const float TileOverlapPadRatio = 0.0f; // No tiling support right now
 
 	// Camera nodes are optional
@@ -145,10 +146,12 @@ void FMovieGraphDeferredPass::Render(const FMovieGraphTraversalContext& InFrameT
 	{
 		OverscanFraction = FMath::Clamp(CameraNode->OverscanPercentage / 100.f, 0.f, 1.f);
 		bRenderAllCameras = CameraNode->bRenderAllCameras;
+		bOverrideCameraOverscan = CameraNode->bOverride_OverscanPercentage;
 	}
 
 	// ToDo: When tiling is used, this should be the size of the per-tile backbuffer
-	FIntPoint AccumulatorResolution = UMovieGraphBlueprintLibrary::GetEffectiveOutputResolution(InTimeData.EvaluatedConfig);
+	const float CameraOverscan = GetRenderer()->GetCameraInfo(InTimeData.EvaluatedConfig, LayerData.CameraIndex).ViewInfo.GetOverscan();
+	FIntPoint AccumulatorResolution = UMovieGraphBlueprintLibrary::GetEffectiveOutputResolution(InTimeData.EvaluatedConfig, CameraOverscan);
 	FIntPoint BackbufferResolution = AccumulatorResolution;
 	
 	DefaultRenderer::FRenderTargetInitParams RenderTargetInitParams = GetRenderTargetInitParams(InTimeData, AccumulatorResolution);
@@ -187,10 +190,15 @@ void FMovieGraphDeferredPass::Render(const FMovieGraphTraversalContext& InFrameT
 		CameraInfo.SamplingParams.TemporalSampleCount = InTimeData.TemporalSampleCount;
 		CameraInfo.SamplingParams.SpatialSampleIndex = SpatialIndex;
 		CameraInfo.SamplingParams.SpatialSampleCount = NumSpatialSamples;
-		CameraInfo.OverscanFraction = OverscanFraction;
 		CameraInfo.ProjectionMatrixJitterAmount = FVector2D((SpatialShiftAmount.X) * 2.0f / (float)BackbufferResolution.X, SpatialShiftAmount.Y * -2.0f / (float)BackbufferResolution.Y);
 		CameraInfo.bUseCameraManagerPostProcess = !bRenderAllCameras;
 		
+		if (bOverrideCameraOverscan)
+		{
+			CameraInfo.ViewInfo.ClearOverscan();
+			CameraInfo.ViewInfo.ApplyOverscan(OverscanFraction);
+		}
+		
 		// ToDo: This math probably needs the per-tile, pre-overlapped size? 
 		FIntPoint OverlappedPad = FIntPoint(FMath::CeilToInt(BackbufferResolution.X * TileOverlapPadRatio), FMath::CeilToInt(BackbufferResolution.Y * TileOverlapPadRatio));
 		
@@ -226,7 +234,7 @@ void FMovieGraphDeferredPass::Render(const FMovieGraphTraversalContext& InFrameT
 		
 		// Scale the DoF sensor scale to counteract overscan, otherwise the size of Bokeh changes when you have Overscan enabled.
 		// This needs to come after we modify it for Tiling.
-		CameraInfo.DoFSensorScale *= 1.0 + CameraInfo.OverscanFraction;
+		CameraInfo.DoFSensorScale *= 1.0 + CameraInfo.ViewInfo.GetOverscan();
 		 
 		// Construct a View to go within this family.
 		FSceneView* NewView = CreateSceneView(SceneViewInitOptions, ViewFamily, CameraInfo);
@@ -261,7 +269,7 @@ void FMovieGraphDeferredPass::Render(const FMovieGraphTraversalContext& InFrameT
 			SampleState.OverlappedPad = OverlappedPad;
 			SampleState.OverlappedOffset = OverlappedOffset;
 			SampleState.OverlappedSubpixelShift = OverlappedSubpixelShift;
-			SampleState.OverscanFraction = OverscanFraction;
+			SampleState.OverscanFraction = CameraInfo.ViewInfo.GetOverscan();
 			SampleState.bAllowOCIO = ParentNodeThisFrame->GetAllowOCIO();
 			SampleState.SceneCaptureSource = SceneCaptureSource;
 			SampleState.CompositingSortOrder = 10;
diff --git a/Engine/Plugins/MovieScene/MovieRenderPipeline/Source/MovieRenderPipelineRenderPasses/Private/Graph/Renderers/MovieGraphImagePassBase.cpp b/Engine/Plugins/MovieScene/MovieRenderPipeline/Source/MovieRenderPipelineRenderPasses/Private/Graph/Renderers/MovieGraphImagePassBase.cpp
--- a/Engine/Plugins/MovieScene/MovieRenderPipeline/Source/MovieRenderPipelineRenderPasses/Private/Graph/Renderers/MovieGraphImagePassBase.cpp
+++ b/Engine/Plugins/MovieScene/MovieRenderPipeline/Source/MovieRenderPipelineRenderPasses/Private/Graph/Renderers/MovieGraphImagePassBase.cpp
@@ -46,16 +46,10 @@ FSceneViewInitOptions FMovieGraphImagePassBase::CreateViewInitOptions(const UE::
 		FPlane(1, 0, 0, 0),
 		FPlane(0, 1, 0, 0),
 		FPlane(0, 0, 0, 1));
-	
-	float ViewFOV = InCameraInfo.ViewInfo.FOV;
-
-	// Inflate our FOV to support the overscan 
-	// ToDo: This is a duplicate of the logic in CalculateProjectionMatrix, should combine.
-	ViewFOV = 2.0f * FMath::RadiansToDegrees(FMath::Atan((1.0f + InCameraInfo.OverscanFraction) * FMath::Tan(FMath::DegreesToRadians( ViewFOV * 0.5f ))));
 
 	ViewInitOptions.SceneViewStateInterface = InViewStateRef.GetReference();
-	ViewInitOptions.FOV = ViewFOV;
-	ViewInitOptions.DesiredFOV = ViewFOV;
+	ViewInitOptions.FOV = InCameraInfo.ViewInfo.FOV;
+	ViewInitOptions.DesiredFOV = InCameraInfo.ViewInfo.FOV;
 	
 	return ViewInitOptions;
 }
@@ -305,17 +299,7 @@ void FMovieGraphImagePassBase::CalculateProjectionMatrix(UE::MovieGraph::Default
 
 	const float DestAspectRatio = ViewRectWidth / ViewRectHeight;
 	const float CameraAspectRatio = InOutCameraInfo.bAllowCameraAspectRatio ? InOutCameraInfo.ViewInfo.AspectRatio : DestAspectRatio;
-
-	float ViewFOV = InOutCameraInfo.ViewInfo.FOV;
-
-	// Inflate our FOV to support the overscan 
-	ViewFOV = 2.0f * FMath::RadiansToDegrees(FMath::Atan((1.0f + InOutCameraInfo.OverscanFraction) * FMath::Tan(FMath::DegreesToRadians(ViewFOV * 0.5f))));
-	InOutCameraInfo.ViewInfo.FOV = ViewFOV;
-	InOutCameraInfo.ViewInfo.DesiredFOV = ViewFOV;
-
-	// Overscan the Orthographic pass too.
-	InOutCameraInfo.ViewInfo.OrthoWidth *= 1.0f + InOutCameraInfo.OverscanFraction;
-
+	
 	const int TotalTileCount = InOutCameraInfo.TilingParams.TileCount.X * InOutCameraInfo.TilingParams.TileCount.Y;
 
 	// If they're using high-resolution tiling we can't support letterboxing (as the blended areas we would render with
diff --git a/Engine/Plugins/MovieScene/MovieRenderPipeline/Source/MovieRenderPipelineRenderPasses/Private/MoviePipelineDeferredPasses.cpp b/Engine/Plugins/MovieScene/MovieRenderPipeline/Source/MovieRenderPipelineRenderPasses/Private/MoviePipelineDeferredPasses.cpp
--- a/Engine/Plugins/MovieScene/MovieRenderPipeline/Source/MovieRenderPipelineRenderPasses/Private/MoviePipelineDeferredPasses.cpp
+++ b/Engine/Plugins/MovieScene/MovieRenderPipeline/Source/MovieRenderPipelineRenderPasses/Private/MoviePipelineDeferredPasses.cpp
@@ -92,7 +92,20 @@ FIntPoint UMoviePipelineDeferredPassBase::GetEffectiveOutputResolutionForCamera(
 	UMoviePipelinePrimaryConfig* PrimaryConfig = GetPipeline()->GetPipelinePrimaryConfig();
 	UMoviePipelineExecutorShot* CurrentShot = GetPipeline()->GetActiveShotList()[GetPipeline()->GetCurrentShotIndex()];
 
-	const FIntPoint OutputResolution = UMoviePipelineBlueprintLibrary::GetEffectiveOutputResolution(PrimaryConfig, CurrentShot);
+	// Get the camera view info to retrieve the camera's overscan, which is used when the settings to not override the overscan
+	FMinimalViewInfo CameraViewInfo;
+
+	if (GetNumCamerasToRender() == 1)
+	{
+		CameraViewInfo = GetPipeline()->GetWorld()->GetFirstPlayerController()->PlayerCameraManager->GetCameraCacheView();
+	}
+	else
+	{
+		UCameraComponent* CameraComponent;
+		GetPipeline()->GetSidecarCameraData(CurrentShot, InCameraIndex, CameraViewInfo, &CameraComponent);
+	}
+	
+	const FIntPoint OutputResolution = UMoviePipelineBlueprintLibrary::GetEffectiveOutputResolution(PrimaryConfig, CurrentShot, CameraViewInfo.GetOverscan());
 
 	return OutputResolution;
 }
diff --git a/Engine/Plugins/MovieScene/MovieRenderPipeline/Source/MovieRenderPipelineRenderPasses/Private/MoviePipelineImagePassBase.cpp b/Engine/Plugins/MovieScene/MovieRenderPipeline/Source/MovieRenderPipelineRenderPasses/Private/MoviePipelineImagePassBase.cpp
--- a/Engine/Plugins/MovieScene/MovieRenderPipeline/Source/MovieRenderPipelineRenderPasses/Private/MoviePipelineImagePassBase.cpp
+++ b/Engine/Plugins/MovieScene/MovieRenderPipeline/Source/MovieRenderPipelineRenderPasses/Private/MoviePipelineImagePassBase.cpp
@@ -469,16 +469,16 @@ FSceneView* UMoviePipelineImagePassBase::GetSceneViewForSampleState(FSceneViewFa
 		FPlane(1, 0, 0, 0),
 		FPlane(0, 1, 0, 0),
 		FPlane(0, 0, 0, 1));
-
-	// Inflate our FOV to support the overscan, we do this before calculating the Projection Matrix so that any downstream consumers of the FOV/DesiredFOV are taken into account (ie: culling)
-	float ViewFOV = 2.0f * FMath::RadiansToDegrees(FMath::Atan((1.0f + InOutSampleState.OverscanPercentage) * FMath::Tan(FMath::DegreesToRadians(CameraInfo.ViewInfo.FOV * 0.5f ))));
-	CameraInfo.ViewInfo.FOV = ViewFOV;
-	CameraInfo.ViewInfo.DesiredFOV = ViewFOV;
+	
+	if (InOutSampleState.bOverrideCameraOverscan)
+	{
+		// If we are overriding the camera's overscan, clear out any overscan the camera added to the view info, and apply the overriding overscan
+		CameraInfo.ViewInfo.ClearOverscan();
+		CameraInfo.ViewInfo.ApplyOverscan(InOutSampleState.OverscanPercentage);
+	}
+	
 	ViewInitOptions.FOV = CameraInfo.ViewInfo.FOV;
-	ViewInitOptions.DesiredFOV = CameraInfo.ViewInfo.DesiredFOV;
-
-	// Overscan the Orthographic pass too.
-	CameraInfo.ViewInfo.OrthoWidth *= 1.0f + InOutSampleState.OverscanPercentage;
+	ViewInitOptions.DesiredFOV = CameraInfo.ViewInfo.FOV;
 
 	float DofSensorScale = 1.0f;
 
diff --git a/Engine/Plugins/MovieScene/MovieRenderPipeline/Source/MovieRenderPipelineRenderPasses/Private/MoviePipelineImageSequenceOutput.cpp b/Engine/Plugins/MovieScene/MovieRenderPipeline/Source/MovieRenderPipelineRenderPasses/Private/MoviePipelineImageSequenceOutput.cpp
--- a/Engine/Plugins/MovieScene/MovieRenderPipeline/Source/MovieRenderPipelineRenderPasses/Private/MoviePipelineImageSequenceOutput.cpp
+++ b/Engine/Plugins/MovieScene/MovieRenderPipeline/Source/MovieRenderPipelineRenderPasses/Private/MoviePipelineImageSequenceOutput.cpp
@@ -242,6 +242,22 @@ void UMoviePipelineImageSequenceOutputBase::OnReceiveImageDataImpl(FMoviePipelin
 					continue;
 				}
 
+				// Check that the composite resolution matches the camera resolution to ensure the composite pass doesn't fail.
+				// This can happen if multiple cameras with different amounts of overscan are rendered, since composite passes
+				// don't support rendering at multiple resolutions
+				const FIntPoint CompositeResolution = CompositePass.PixelData->GetSize();
+				const FIntPoint CameraResolution = RenderPassData.Value->GetSize();
+				if (CompositeResolution != CameraResolution)
+				{
+					UE_LOG(LogMovieRenderPipeline, Warning, TEXT("Composite resolution %dx%d does not match camera resolution %dx%d, skipping composite for %s on camera %s"),
+						CompositeResolution.X, CompositeResolution.Y,
+						CameraResolution.X, CameraResolution.Y,
+						*CompositePass.PassIdentifier.Name,
+						*RenderPassData.Key.CameraName);
+
+					continue;
+				}
+
 				// If there's more than one render pass, we need to copy the composite passes for the first render pass then move for the remaining ones
 				const bool bShouldCopyImageData = RenderPassCount > 1 && RenderPassIteration == 0;
 				TUniquePtr<FImagePixelData> PixelData = bShouldCopyImageData ? CompositePass.PixelData->CopyImageData() : CompositePass.PixelData->MoveImageDataToNew();
diff --git a/Engine/Plugins/MovieScene/MovieRenderPipeline/Source/MovieRenderPipelineSettings/Private/MoviePipelineBurnInSetting.cpp b/Engine/Plugins/MovieScene/MovieRenderPipeline/Source/MovieRenderPipelineSettings/Private/MoviePipelineBurnInSetting.cpp
--- a/Engine/Plugins/MovieScene/MovieRenderPipeline/Source/MovieRenderPipelineSettings/Private/MoviePipelineBurnInSetting.cpp
+++ b/Engine/Plugins/MovieScene/MovieRenderPipeline/Source/MovieRenderPipelineSettings/Private/MoviePipelineBurnInSetting.cpp
@@ -126,7 +126,10 @@ void UMoviePipelineBurnInSetting::SetupImpl(const MoviePipeline::FMoviePipelineR
 		return;
 	}
 
-	OutputResolution = UMoviePipelineBlueprintLibrary::GetEffectiveOutputResolution(GetPipeline()->GetPipelinePrimaryConfig(), GetPipeline()->GetActiveShotList()[GetPipeline()->GetCurrentShotIndex()]);
+	// TODO: Add multi-camera support? For now, simply use the default player controller camera overscan
+	FMinimalViewInfo CameraViewInfo = GetPipeline()->GetWorld()->GetFirstPlayerController()->PlayerCameraManager->GetCameraCacheView();
+	
+	OutputResolution = UMoviePipelineBlueprintLibrary::GetEffectiveOutputResolution(GetPipeline()->GetPipelinePrimaryConfig(), GetPipeline()->GetActiveShotList()[GetPipeline()->GetCurrentShotIndex()], CameraViewInfo.GetOverscan());
 	int32 MaxResolution = GetMax2DTextureDimension();
 	if (OutputResolution.X > MaxResolution || OutputResolution.Y > MaxResolution)
 	{
diff --git a/Engine/Plugins/MovieScene/MovieRenderPipeline/Source/MovieRenderPipelineSettings/Private/MoviePipelineWidgetRenderSetting.cpp b/Engine/Plugins/MovieScene/MovieRenderPipeline/Source/MovieRenderPipelineSettings/Private/MoviePipelineWidgetRenderSetting.cpp
--- a/Engine/Plugins/MovieScene/MovieRenderPipeline/Source/MovieRenderPipelineSettings/Private/MoviePipelineWidgetRenderSetting.cpp
+++ b/Engine/Plugins/MovieScene/MovieRenderPipeline/Source/MovieRenderPipelineSettings/Private/MoviePipelineWidgetRenderSetting.cpp
@@ -116,7 +116,9 @@ void UMoviePipelineWidgetRenderer::SetupImpl(const MoviePipeline::FMoviePipeline
 
 	bool bInForceLinearGamma = false;
 
-	FIntPoint OutputResolution = UMoviePipelineBlueprintLibrary::GetEffectiveOutputResolution(GetPipeline()->GetPipelinePrimaryConfig(), GetPipeline()->GetActiveShotList()[GetPipeline()->GetCurrentShotIndex()]);
+	// TODO: Add multi-camera support? For now, simply use the default player controller camera overscan
+	FMinimalViewInfo CameraViewInfo = GetPipeline()->GetWorld()->GetFirstPlayerController()->PlayerCameraManager->GetCameraCacheView();
+	FIntPoint OutputResolution = UMoviePipelineBlueprintLibrary::GetEffectiveOutputResolution(GetPipeline()->GetPipelinePrimaryConfig(), GetPipeline()->GetActiveShotList()[GetPipeline()->GetCurrentShotIndex()], CameraViewInfo.GetOverscan());
 
 	int32 MaxResolution = GetMax2DTextureDimension();
 	if (OutputResolution.X > MaxResolution || OutputResolution.Y > MaxResolution)
diff --git a/Engine/Source/Runtime/Engine/Classes/Camera/CameraTypes.h b/Engine/Source/Runtime/Engine/Classes/Camera/CameraTypes.h
--- a/Engine/Source/Runtime/Engine/Classes/Camera/CameraTypes.h
+++ b/Engine/Source/Runtime/Engine/Classes/Camera/CameraTypes.h
@@ -140,6 +140,9 @@ struct FMinimalViewInfo
 	float CropFraction;
 	
 private:
+	/** The amount of overscan that has been applied to the view's frustum, with 0.0 meaning no overscan and 1.0 meaning 100% overscan */
+	float Overscan;
+	
 	// Only used for Ortho camera auto plane calculations, tells the Near plane of the extra distance that needs to be added.
 	FVector CameraToViewTarget;
 
@@ -169,6 +172,7 @@ public:
 		, OffCenterProjectionOffset(ForceInitToZero)
 		, OverscanResolutionFraction(1.0f)
 		, CropFraction(1.0f)
+		, Overscan(0.0f)
 		, CameraToViewTarget(FVector::ZeroVector)
 	{
 	}
@@ -227,6 +231,18 @@ public:
 
 	/**
 	 * Apply overscan to the view info, which scales the field of view and ortho width to simulate expanding the view frustum.
+	 * 
+	 * @param InOverscan - The amount of overscan to apply, from 0.0 meaning no overscan to 1.0 meaning 100% overscan
+	 * @param bScaleResolutionWithOverscan - Indicates that the view's resolution should be scaled with the amount of overscan, so that the original frustum remains the same resolution
+	 * @param bCropOverscan - Indicates that the view should be cropped during the final post process pass to remove the overscanned pixels
 	 */
 	ENGINE_API void ApplyOverscan(float InOverscan, bool bScaleResolutionWithOverscan = false, bool bCropOverscan = false);
+
+	/** Gets the total amount of overscan that has been applied to the view's frustum, with 0.0 meaning no overscan and 1.0 meaning 100% overscan */
+	ENGINE_API float GetOverscan() const { return Overscan; }
+	
+	/**
+	 * Removes all overscan from the view info.
+	 */
+	ENGINE_API void ClearOverscan();
 };
diff --git a/Engine/Source/Runtime/Engine/Private/Camera/CameraStackTypes.cpp b/Engine/Source/Runtime/Engine/Private/Camera/CameraStackTypes.cpp
--- a/Engine/Source/Runtime/Engine/Private/Camera/CameraStackTypes.cpp
+++ b/Engine/Source/Runtime/Engine/Private/Camera/CameraStackTypes.cpp
@@ -495,10 +495,17 @@ float FMinimalViewInfo::CalculateFirstPersonFOVCorrectionFactor() const
 
 void FMinimalViewInfo::ApplyOverscan(float InOverscan, bool bScaleResolutionWithOverscan, bool bCropOverscan)
 {
-	if (InOverscan > 0.0)
+	if (!FMath::IsNearlyZero(InOverscan))
 	{
+		// Clamp the incoming overscan so that the new total overscan can never be less than zero
+		const float ClampedOverscan = FMath::Max(InOverscan, -Overscan / (1 + Overscan));
+		
+		// Keep track of the total amount of overscan that has been applied to the view. Mathematically,
+		// this formula is derived from 1 + TotalOverscan = (1 + Overscan) * (1 + InOverscan)
+		Overscan = Overscan * (1.0f + ClampedOverscan) + ClampedOverscan;
+		
 		// By convention, 0.0 means no overscan, so add 1 to compute the scalar needed for altering projection values
-		const float OverscanScalar = 1.0f + InOverscan;
+		const float OverscanScalar = 1.0f + ClampedOverscan;
 		
 		// Overscan directly scales the view frustum, but can be accomplished by scaling the FOV.
 		// However, must scale the tangent of the half-FOV to accomplish the same mathematical transform.
@@ -512,3 +519,17 @@ void FMinimalViewInfo::ApplyOverscan(float InOverscan, bool bScaleResolutionWith
 		CropFraction *= bCropOverscan ? 1.0f / OverscanScalar : 1.0f;
 	}
 }
+
+void FMinimalViewInfo::ClearOverscan()
+{
+	if (Overscan > 0.0f)
+	{
+		// Apply the inverse overscan to the view frustum to obtain the original frustum values (field of view, ortho width, etc)
+		// Inverse overscan derived from (1 + Overscan) * (1 + InverseOverscan) = 1
+		const float InverseOverscan = - Overscan / (1.0f + Overscan);
+		ApplyOverscan(InverseOverscan);
+
+		OverscanResolutionFraction = 1.0f;
+		CropFraction = 1.0f;
+	}
+}
