MoviePipeline: Updated legacy MRQ to take advantage of the new "Name" field for post-process materials. This was added for MRG previously, but MRQ and MRG both share the same data structure that contains this Name property.

#rnx

diff --git a/Engine/Plugins/MovieScene/MovieRenderPipeline/Source/MovieRenderPipelineRenderPasses/Private/MoviePipelineDeferredPasses.cpp b/Engine/Plugins/MovieScene/MovieRenderPipeline/Source/MovieRenderPipelineRenderPasses/Private/MoviePipelineDeferredPasses.cpp
--- a/Engine/Plugins/MovieScene/MovieRenderPipeline/Source/MovieRenderPipelineRenderPasses/Private/MoviePipelineDeferredPasses.cpp
+++ b/Engine/Plugins/MovieScene/MovieRenderPipeline/Source/MovieRenderPipelineRenderPasses/Private/MoviePipelineDeferredPasses.cpp
@@ -460,7 +460,7 @@ void UMoviePipelineDeferredPassBase::GatherOutputPassesImpl(TArray<FMoviePipelin
 		{
 			if (Material)
 			{
-				RenderPasses.Add(Material->GetName());
+				RenderPasses.Add(GetNameForPostProcessMaterial(Material));
 			}
 		}
 
@@ -554,7 +554,8 @@ void UMoviePipelineDeferredPassBase::RenderSample_GameThreadImpl(const FMoviePip
 				{
 					continue;
 				}
-				FMoviePipelinePassIdentifier LayerPassIdentifier = FMoviePipelinePassIdentifier(PassIdentifier.Name + VisMaterial->GetName(), PassIdentifierForCurrentCamera.CameraName);
+
+				FMoviePipelinePassIdentifier LayerPassIdentifier = FMoviePipelinePassIdentifier(PassIdentifier.Name + GetNameForPostProcessMaterial(VisMaterial), PassIdentifierForCurrentCamera.CameraName);
 
 				auto BufferPipe = MakeShared<FImagePixelPipe, ESPMode::ThreadSafe>();
 				BufferPipe->bIsExpecting32BitPixelData = ActiveHighPrecisionPostProcessMaterials.Contains(VisMaterial);
@@ -1203,3 +1204,17 @@ void UMoviePipelineDeferredPassBase::UpdateTelemetry(FMoviePipelineShotRenderTel
 	InTelemetry->bUsesDeferred = true;
 	InTelemetry->bUsesPPMs |= Algo::AnyOf(AdditionalPostProcessMaterials, [](const FMoviePipelinePostProcessPass& Pass) { return Pass.bEnabled; });
 }
+
+FString UMoviePipelineDeferredPassBase::GetNameForPostProcessMaterial(const UMaterialInterface* InMaterial)
+{
+	FString MaterialName = InMaterial->GetName();
+
+	// Use the name specified in the post process pass if it's not empty. Otherwise fall back to the material's name.
+	const FMoviePipelinePostProcessPass* MatchingPass = Algo::FindByPredicate(AdditionalPostProcessMaterials, [InMaterial](const FMoviePipelinePostProcessPass& InPass) { return InPass.Material == InMaterial; });
+	if (MatchingPass && !MatchingPass->Name.IsEmpty())
+	{
+		MaterialName = MatchingPass->Name;
+	}
+
+	return MaterialName;
+}
diff --git a/Engine/Plugins/MovieScene/MovieRenderPipeline/Source/MovieRenderPipelineRenderPasses/Public/MoviePipelineDeferredPasses.h b/Engine/Plugins/MovieScene/MovieRenderPipeline/Source/MovieRenderPipelineRenderPasses/Public/MoviePipelineDeferredPasses.h
--- a/Engine/Plugins/MovieScene/MovieRenderPipeline/Source/MovieRenderPipelineRenderPasses/Public/MoviePipelineDeferredPasses.h
+++ b/Engine/Plugins/MovieScene/MovieRenderPipeline/Source/MovieRenderPipelineRenderPasses/Public/MoviePipelineDeferredPasses.h
@@ -23,8 +23,9 @@ public:
 	bool bEnabled = false;
 
 	/**
-	 * An optional name which which will be used as the {renderer_sub_name} file name format token for this material. If a name is not specified here,
-	 * the name of the material will be used for {renderer_sub_name}.
+	 * An optional name which which will identify this material in the file name. For MRQ, the material name will be included in the {render_pass}
+	 * token. For Movie Render Graph, the material name will be used in {renderer_sub_name}. If a name is not specified here, the full name of the material
+	 * will be used in these tokens instead.
 	 */
 	UPROPERTY(EditAnywhere, BlueprintReadWrite, Category = "Settings")
 	FString Name;
@@ -100,6 +101,10 @@ protected:
 
 	virtual void UpdateTelemetry(FMoviePipelineShotRenderTelemetry* InTelemetry) const override;
 
+private:
+	/** Gets the name for a post-process material (either the custom-specified name, or the material name). */
+	FString GetNameForPostProcessMaterial(const UMaterialInterface* InMaterial);
+
 public:
 	/**
 	* Should multiple temporal/spatial samples accumulate the alpha channel? This requires r.PostProcessing.PropagateAlpha
