MoviePipeline: Fixed the detail lighting mode (via View Mode Index) not properly working on the Deferred Renderer node.

diff --git a/Engine/Plugins/MovieScene/MovieRenderPipeline/Source/MovieRenderPipelineCore/Private/MoviePipelineUtils.cpp b/Engine/Plugins/MovieScene/MovieRenderPipeline/Source/MovieRenderPipelineCore/Private/MoviePipelineUtils.cpp
--- a/Engine/Plugins/MovieScene/MovieRenderPipeline/Source/MovieRenderPipelineCore/Private/MoviePipelineUtils.cpp
+++ b/Engine/Plugins/MovieScene/MovieRenderPipeline/Source/MovieRenderPipelineCore/Private/MoviePipelineUtils.cpp
@@ -150,6 +150,45 @@ namespace UE
 			return GFrameCounter;
 		}
 
+		void UpdateSceneViewForShowFlags(FSceneView* View)
+		{
+			if (View->Family->EngineShowFlags.Wireframe)
+			{
+				// Wireframe color is emissive-only, and mesh-modifying materials do not use material substitution, hence...
+				View->DiffuseOverrideParameter = FVector4f(0.f, 0.f, 0.f, 0.f);
+				View->SpecularOverrideParameter = FVector4f(0.f, 0.f, 0.f, 0.f);
+			}
+			else if (View->Family->EngineShowFlags.OverrideDiffuseAndSpecular)
+			{
+				View->DiffuseOverrideParameter = FVector4f(GEngine->LightingOnlyBrightness.R, GEngine->LightingOnlyBrightness.G, GEngine->LightingOnlyBrightness.B, 0.0f);
+				View->SpecularOverrideParameter = FVector4f(.1f, .1f, .1f, 0.0f);
+			}
+			else if (View->Family->EngineShowFlags.LightingOnlyOverride)
+			{
+				View->DiffuseOverrideParameter = FVector4f(GEngine->LightingOnlyBrightness.R, GEngine->LightingOnlyBrightness.G, GEngine->LightingOnlyBrightness.B, 0.0f);
+				View->SpecularOverrideParameter = FVector4f(0.f, 0.f, 0.f, 0.f);
+			}
+			else if (View->Family->EngineShowFlags.ReflectionOverride)
+			{
+				View->DiffuseOverrideParameter = FVector4f(0.f, 0.f, 0.f, 0.f);
+				View->SpecularOverrideParameter = FVector4f(1, 1, 1, 0.0f);
+				View->NormalOverrideParameter = FVector4f(0, 0, 1, 0.0f);
+				View->RoughnessOverrideParameter = FVector2D(0.0f, 0.0f);
+			}
+	
+			if (!View->Family->EngineShowFlags.Diffuse)
+			{
+				View->DiffuseOverrideParameter = FVector4f(0.f, 0.f, 0.f, 0.f);
+			}
+	
+			if (!View->Family->EngineShowFlags.Specular)
+			{
+				View->SpecularOverrideParameter = FVector4f(0.f, 0.f, 0.f, 0.f);
+			}
+	
+			static const FName BufferVisualizationMode(TEXT("WorldNormal"));
+			View->CurrentBufferVisualizationMode = BufferVisualizationMode;
+		}
 	}
 }
 
diff --git a/Engine/Plugins/MovieScene/MovieRenderPipeline/Source/MovieRenderPipelineCore/Public/MoviePipelineUtils.h b/Engine/Plugins/MovieScene/MovieRenderPipeline/Source/MovieRenderPipelineCore/Public/MoviePipelineUtils.h
--- a/Engine/Plugins/MovieScene/MovieRenderPipeline/Source/MovieRenderPipelineCore/Public/MoviePipelineUtils.h
+++ b/Engine/Plugins/MovieScene/MovieRenderPipeline/Source/MovieRenderPipelineCore/Public/MoviePipelineUtils.h
@@ -127,6 +127,13 @@ namespace UE
 		
 		UE_DEPRECATED(5.3, "Do not use, this is here as a temporary workaround for another issue.")
 		MOVIERENDERPIPELINECORE_API uint64 GetRendererFrameCount();
+		
+		/**
+		 * Updates the provided scene view in accordance with the show flags that are currently set (eg, changes the specular override color if
+		 * the OverrideDiffuseAndSpecular show flag is set). The show flags that are looked at here are usually the ones that are set when the
+		 * view mode index is changed.
+		 */
+		MOVIERENDERPIPELINECORE_API void UpdateSceneViewForShowFlags(FSceneView* View);
 	}
 
 	namespace MoviePipeline
diff --git a/Engine/Plugins/MovieScene/MovieRenderPipeline/Source/MovieRenderPipelineRenderPasses/Private/Graph/Renderers/MovieGraphImagePassBase.cpp b/Engine/Plugins/MovieScene/MovieRenderPipeline/Source/MovieRenderPipelineRenderPasses/Private/Graph/Renderers/MovieGraphImagePassBase.cpp
--- a/Engine/Plugins/MovieScene/MovieRenderPipeline/Source/MovieRenderPipelineRenderPasses/Private/Graph/Renderers/MovieGraphImagePassBase.cpp
+++ b/Engine/Plugins/MovieScene/MovieRenderPipeline/Source/MovieRenderPipelineRenderPasses/Private/Graph/Renderers/MovieGraphImagePassBase.cpp
@@ -72,6 +72,8 @@ FSceneView* FMovieGraphImagePassBase::CreateSceneView(const FSceneViewInitOption
 	View->LensPrincipalPointOffsetScale = CalculatePrinciplePointOffsetForTiling(InCameraInfo.TilingParams);
 	View->EndFinalPostprocessSettings(InInitOptions);
 
+	UE::MovieRenderPipeline::UpdateSceneViewForShowFlags(View);
+
 	return View;
 }
 
@@ -139,7 +141,6 @@ void FMovieGraphImagePassBase::ApplyCameraManagerPostProcessBlends(FSceneView* I
 
 TSharedRef<FSceneViewFamilyContext> FMovieGraphImagePassBase::CreateSceneViewFamily(const FViewFamilyInitData& InInitData) const
 {
-
 	EViewModeIndex ViewModeIndex = InInitData.ViewModeIndex;
 	FEngineShowFlags ShowFlags = InInitData.ShowFlags;
 
@@ -172,7 +173,6 @@ void FMovieGraphImagePassBase::ApplyMovieGraphOverridesToViewFamily(TSharedRef<F
 	// Used to specify if the Tone Curve is being applied or not to our Linear Output data
 	InOutFamily->SceneCaptureSource = InInitData.SceneCaptureSource;
 	InOutFamily->bWorldIsPaused = InInitData.bWorldIsPaused;
-	// InOutFamily->ViewMode = ViewModeIndex;
 	InOutFamily->bOverrideVirtualTextureThrottle = true;
 	
 	// We need to check if this is the first FSceneView being submitted to the renderer module, and set some flags on the ViewFamily for ensuring some
diff --git a/Engine/Plugins/MovieScene/MovieRenderPipeline/Source/MovieRenderPipelineRenderPasses/Private/MoviePipelineImagePassBase.cpp b/Engine/Plugins/MovieScene/MovieRenderPipeline/Source/MovieRenderPipelineRenderPasses/Private/MoviePipelineImagePassBase.cpp
--- a/Engine/Plugins/MovieScene/MovieRenderPipeline/Source/MovieRenderPipelineRenderPasses/Private/MoviePipelineImagePassBase.cpp
+++ b/Engine/Plugins/MovieScene/MovieRenderPipeline/Source/MovieRenderPipelineRenderPasses/Private/MoviePipelineImagePassBase.cpp
@@ -399,41 +399,7 @@ TSharedPtr<FSceneViewFamilyContext> UMoviePipelineImagePassBase::CalculateViewFa
 
 void UMoviePipelineImagePassBase::SetupViewForViewModeOverride(FSceneView* View)
 {
-	if (View->Family->EngineShowFlags.Wireframe)
-	{
-		// Wireframe color is emissive-only, and mesh-modifying materials do not use material substitution, hence...
-		View->DiffuseOverrideParameter = FVector4f(0.f, 0.f, 0.f, 0.f);
-		View->SpecularOverrideParameter = FVector4f(0.f, 0.f, 0.f, 0.f);
-	}
-	else if (View->Family->EngineShowFlags.OverrideDiffuseAndSpecular)
-	{
-		View->DiffuseOverrideParameter = FVector4f(GEngine->LightingOnlyBrightness.R, GEngine->LightingOnlyBrightness.G, GEngine->LightingOnlyBrightness.B, 0.0f);
-		View->SpecularOverrideParameter = FVector4f(.1f, .1f, .1f, 0.0f);
-	}
-	else if (View->Family->EngineShowFlags.LightingOnlyOverride)
-	{
-		View->DiffuseOverrideParameter = FVector4f(GEngine->LightingOnlyBrightness.R, GEngine->LightingOnlyBrightness.G, GEngine->LightingOnlyBrightness.B, 0.0f);
-		View->SpecularOverrideParameter = FVector4f(0.f, 0.f, 0.f, 0.f);
-	}
-	else if (View->Family->EngineShowFlags.ReflectionOverride)
-	{
-		View->DiffuseOverrideParameter = FVector4f(0.f, 0.f, 0.f, 0.f);
-		View->SpecularOverrideParameter = FVector4f(1, 1, 1, 0.0f);
-		View->NormalOverrideParameter = FVector4f(0, 0, 1, 0.0f);
-		View->RoughnessOverrideParameter = FVector2D(0.0f, 0.0f);
-	}
-
-	if (!View->Family->EngineShowFlags.Diffuse)
-	{
-		View->DiffuseOverrideParameter = FVector4f(0.f, 0.f, 0.f, 0.f);
-	}
-
-	if (!View->Family->EngineShowFlags.Specular)
-	{
-		View->SpecularOverrideParameter = FVector4f(0.f, 0.f, 0.f, 0.f);
-	}
-	FName BufferVisualizationMode = "WorldNormal";
-	View->CurrentBufferVisualizationMode = BufferVisualizationMode;
+	UE::MovieRenderPipeline::UpdateSceneViewForShowFlags(View);
 }
 
 void UMoviePipelineImagePassBase::GatherOutputPassesImpl(TArray<FMoviePipelinePassIdentifier>& ExpectedRenderPasses)
