MoviePipeline: Fixed an issue that could cause a crash or stale editor when SaveAs'ing over an existing graph that's already open.

diff --git a/Engine/Plugins/MovieScene/MovieRenderPipeline/Source/MovieRenderPipelineEditor/Private/Graph/MovieGraphAssetToolkit.cpp b/Engine/Plugins/MovieScene/MovieRenderPipeline/Source/MovieRenderPipelineEditor/Private/Graph/MovieGraphAssetToolkit.cpp
--- a/Engine/Plugins/MovieScene/MovieRenderPipeline/Source/MovieRenderPipelineEditor/Private/Graph/MovieGraphAssetToolkit.cpp
+++ b/Engine/Plugins/MovieScene/MovieRenderPipeline/Source/MovieRenderPipelineEditor/Private/Graph/MovieGraphAssetToolkit.cpp
@@ -715,6 +715,25 @@ void FMovieGraphAssetToolkit::SaveAsset_Execute()
 	// TODO: Any custom save logic here
 }
 
+void FMovieGraphAssetToolkit::OnAssetsSavedAs(const TArray<UObject*>& SavedObjects)
+{
+	FAssetEditorToolkit::OnAssetsSavedAs(SavedObjects);
+
+	UAssetEditorSubsystem* AssetEditorSubsystem = GEditor->GetEditorSubsystem<UAssetEditorSubsystem>();
+
+	// The default behavior for SaveAs in the toolkit doesn't properly re-open the assets that were saved during a SaveAs, it only closes the assets
+	// that were the source of the SaveAs. After a SaveAs, the graph potentially goes through a complete data change, and re-opening is the most
+	// reliable way to make sure that the graph and editor are properly in sync. Without this, connections, delegates, etc can get badly out-of-sync
+	// after a SaveAs. Generally a crash won't result, but the graph will be in a nearly unusable state.
+	for (UObject* SavedObject : SavedObjects)
+	{
+		AssetEditorSubsystem->CloseAllEditorsForAsset(SavedObject);
+		AssetEditorSubsystem->NotifyAssetClosed(SavedObject, this);
+	}
+
+	AssetEditorSubsystem->OpenEditorForAssets_Advanced(SavedObjects, ToolkitMode, ToolkitHost.Pin());
+}
+
 void FMovieGraphAssetToolkit::OnClose()
 {
 	// Editor-only nodes are copied to the underlying runtime graph on save/close
diff --git a/Engine/Plugins/MovieScene/MovieRenderPipeline/Source/MovieRenderPipelineEditor/Private/Graph/MovieGraphAssetToolkit.h b/Engine/Plugins/MovieScene/MovieRenderPipeline/Source/MovieRenderPipelineEditor/Private/Graph/MovieGraphAssetToolkit.h
--- a/Engine/Plugins/MovieScene/MovieRenderPipeline/Source/MovieRenderPipelineEditor/Private/Graph/MovieGraphAssetToolkit.h
+++ b/Engine/Plugins/MovieScene/MovieRenderPipeline/Source/MovieRenderPipelineEditor/Private/Graph/MovieGraphAssetToolkit.h
@@ -38,6 +38,7 @@ public:
 	virtual FText GetBaseToolkitName() const override;
 	virtual FString GetWorldCentricTabPrefix() const override;
 	virtual void SaveAsset_Execute() override;
+	virtual void OnAssetsSavedAs(const TArray<UObject*>& SavedObjects) override;
 	virtual void OnClose() override;
 	//~ End FAssetEditorToolkit interface
 
