MoviePipeline: Fixed a branch connection error that could result when connecting the Globals output on a subgraph to a Globals-only node before that node has been connected to the graph's Globals branch.

diff --git a/Engine/Plugins/MovieScene/MovieRenderPipeline/Source/MovieRenderPipelineCore/Private/Graph/MovieGraphPin.cpp b/Engine/Plugins/MovieScene/MovieRenderPipeline/Source/MovieRenderPipelineCore/Private/Graph/MovieGraphPin.cpp
--- a/Engine/Plugins/MovieScene/MovieRenderPipeline/Source/MovieRenderPipelineCore/Private/Graph/MovieGraphPin.cpp
+++ b/Engine/Plugins/MovieScene/MovieRenderPipeline/Source/MovieRenderPipelineCore/Private/Graph/MovieGraphPin.cpp
@@ -339,8 +339,11 @@ bool UMovieGraphPin::IsConnectionToBranchAllowed(const UMovieGraphPin* OtherPin,
 		? TArray{OutputName.ToString()}
 		: GraphConfig->GetUpstreamBranchNames(FromNode, OutputPin, bStopAtSubgraph);
 	
-	const bool bGlobalsIsDownstream = bInputIsGlobals || DownstreamBranchNames.Contains(UMovieGraphNode::GlobalsPinNameString);
-	const bool bGlobalsIsUpstream = bOutputIsGlobals || UpstreamBranchNames.Contains(UMovieGraphNode::GlobalsPinNameString);
+	// Consider Globals to be up/downstream if the connection is directly to the Globals branch, to any node already connected to Globals, or to a
+	// node that has a Globals-only branch restriction (this last check is important when a Globals-only node is not yet connected to the Globals branch).
+	const bool bGlobalsIsDownstream = bInputIsGlobals || DownstreamBranchNames.Contains(UMovieGraphNode::GlobalsPinNameString) || (ToNode->GetBranchRestriction() == EMovieGraphBranchRestriction::Globals);
+	const bool bGlobalsIsUpstream = bOutputIsGlobals || UpstreamBranchNames.Contains(UMovieGraphNode::GlobalsPinNameString) || (FromNode->GetBranchRestriction() == EMovieGraphBranchRestriction::Globals);
+
 	const bool bDownstreamBranchExistsAndIsntOnlyGlobals =
 		!DownstreamBranchNames.IsEmpty() && ((DownstreamBranchNames.Num() != 1) || (DownstreamBranchNames[0] != UMovieGraphNode::GlobalsPinNameString));
 	const bool bUpstreamBranchExistsAndIsntOnlyGlobals =
@@ -350,10 +353,10 @@ bool UMovieGraphPin::IsConnectionToBranchAllowed(const UMovieGraphPin* OtherPin,
 	if (ToNodeIsSubgraph || FromNodeIsSubgraph)
 	{
 		// Only allow Globals -> Globals connections
-		if ((ToNodeIsSubgraph && bInputIsGlobals && !bGlobalsIsUpstream) ||
-			(FromNodeIsSubgraph && bOutputIsGlobals && !bGlobalsIsDownstream))
+		if ((ToNodeIsSubgraph && bInputIsGlobals && bUpstreamBranchExistsAndIsntOnlyGlobals) ||
+			(FromNodeIsSubgraph && bOutputIsGlobals && bDownstreamBranchExistsAndIsntOnlyGlobals))
 		{
-			OutError = NSLOCTEXT("MovieGraph", "SubgraphGlobalsBranchMismatchError", "A subgraph Globals branch can only be connected to another Globals branch.");
+			OutError = NSLOCTEXT("MovieGraph", "SubgraphGlobalsBranchMismatchError", "A subgraph Globals branch can only be connected to another Globals branch or Globals-only nodes.");
 			return false;
 		}
 
@@ -361,7 +364,7 @@ bool UMovieGraphPin::IsConnectionToBranchAllowed(const UMovieGraphPin* OtherPin,
 		if ((ToNodeIsSubgraph && !bInputIsGlobals && bGlobalsIsUpstream) ||
 			(FromNodeIsSubgraph && !bOutputIsGlobals && bGlobalsIsDownstream))
 		{
-			OutError = NSLOCTEXT("MovieGraph", "SubgraphNonGlobalsBranchMismatchError", "A subgraph non-Globals branch can not be connected to the Globals branch.");
+			OutError = NSLOCTEXT("MovieGraph", "SubgraphNonGlobalsBranchMismatchError", "A subgraph non-Globals branch cannot be connected to the Globals branch or Globals-only nodes.");
 			return false;
 		}
 	}
@@ -370,7 +373,7 @@ bool UMovieGraphPin::IsConnectionToBranchAllowed(const UMovieGraphPin* OtherPin,
 		// Globals branches can only be connected to Globals branches
 		if ((bGlobalsIsDownstream && bUpstreamBranchExistsAndIsntOnlyGlobals) || (bGlobalsIsUpstream && bDownstreamBranchExistsAndIsntOnlyGlobals))
 		{
-			OutError = NSLOCTEXT("MovieGraph", "GlobalsBranchMismatchError", "Globals branches can only be connected to other Globals branches.");
+			OutError = NSLOCTEXT("MovieGraph", "GlobalsBranchMismatchError", "Globals branches and Globals-only nodes can only be connected to other Globals branches and Globals-only nodes.");
 			return false;
 		}
 	}
