Movie Render Graph: Initialize bMultipleCameraNames

#rnx

diff --git a/Engine/Plugins/MovieScene/MovieRenderPipeline/Source/MovieRenderPipelineCore/Public/Graph/MovieGraphDataTypes.h b/Engine/Plugins/MovieScene/MovieRenderPipeline/Source/MovieRenderPipelineCore/Public/Graph/MovieGraphDataTypes.h
--- a/Engine/Plugins/MovieScene/MovieRenderPipeline/Source/MovieRenderPipelineCore/Public/Graph/MovieGraphDataTypes.h
+++ b/Engine/Plugins/MovieScene/MovieRenderPipeline/Source/MovieRenderPipelineCore/Public/Graph/MovieGraphDataTypes.h
@@ -53,7 +53,7 @@ struct FMovieGraphImagePreviewData
 
 	/** If true, then there is more than one camera name being used (ie: multi-camera rendering) */
 	UPROPERTY(BlueprintReadOnly, Category = "Movie Graph")
-	bool bMultipleCameraNames;
+	bool bMultipleCameraNames = false;
 };
 
 USTRUCT(BlueprintType)
