MoviePipeline: Fixed a freeze that could occur when using the Command Line Encoder node and the sequence being rendered is playing audio.

diff --git a/Engine/Plugins/MovieScene/MovieRenderPipeline/Source/MovieRenderPipelineCore/Private/Graph/MovieGraphPipeline.cpp b/Engine/Plugins/MovieScene/MovieRenderPipeline/Source/MovieRenderPipelineCore/Private/Graph/MovieGraphPipeline.cpp
--- a/Engine/Plugins/MovieScene/MovieRenderPipeline/Source/MovieRenderPipelineCore/Private/Graph/MovieGraphPipeline.cpp
+++ b/Engine/Plugins/MovieScene/MovieRenderPipeline/Source/MovieRenderPipelineCore/Private/Graph/MovieGraphPipeline.cpp
@@ -1402,11 +1402,8 @@ void UMovieGraphPipeline::TransitionToState(const EMovieRenderPipelineState InNe
 			// This is called once notifying our export step that they can begin the export.
 			PipelineState = EMovieRenderPipelineState::Export;
 
-			// Restore the sequence so that the export processes can operate on the original sequence. 
-			// This is also done in the finished state because it's not guaranteed that the Export state 
-			// will be set when the render is canceled early
-			// LevelSequenceActor->GetSequencePlayer()->Stop();
-			// RestoreTargetSequenceToOriginalState();
+			// Stop playing the data source once export begins. Audio issues may happen if this is not done.
+			GetDataSourceInstance()->StopDataSource();
 
 			// Ensure all of our Futures have been converted to the GeneratedOutputData. This has to happen
 			// after finalize finishes, because the futures won't be available until actually written to disk.
diff --git a/Engine/Plugins/MovieScene/MovieRenderPipeline/Source/MovieRenderPipelineCore/Private/Graph/MovieGraphSequenceDataSource.cpp b/Engine/Plugins/MovieScene/MovieRenderPipeline/Source/MovieRenderPipelineCore/Private/Graph/MovieGraphSequenceDataSource.cpp
--- a/Engine/Plugins/MovieScene/MovieRenderPipeline/Source/MovieRenderPipelineCore/Private/Graph/MovieGraphSequenceDataSource.cpp
+++ b/Engine/Plugins/MovieScene/MovieRenderPipeline/Source/MovieRenderPipelineCore/Private/Graph/MovieGraphSequenceDataSource.cpp
@@ -332,6 +332,14 @@ void UMovieGraphSequenceDataSource::PauseDataSource()
 	}
 }
 
+void UMovieGraphSequenceDataSource::StopDataSource()
+{
+	if (LevelSequenceActor && LevelSequenceActor->GetSequencePlayer())
+	{
+		LevelSequenceActor->GetSequencePlayer()->Stop();
+	}
+}
+
 void UMovieGraphSequenceDataSource::JumpDataSource(const FFrameTime& InTimeToJumpTo) 
 {
 	if(LevelSequenceActor && LevelSequenceActor->GetSequencePlayer())
diff --git a/Engine/Plugins/MovieScene/MovieRenderPipeline/Source/MovieRenderPipelineCore/Public/Graph/MovieGraphDataTypes.h b/Engine/Plugins/MovieScene/MovieRenderPipeline/Source/MovieRenderPipelineCore/Public/Graph/MovieGraphDataTypes.h
--- a/Engine/Plugins/MovieScene/MovieRenderPipeline/Source/MovieRenderPipelineCore/Public/Graph/MovieGraphDataTypes.h
+++ b/Engine/Plugins/MovieScene/MovieRenderPipeline/Source/MovieRenderPipelineCore/Public/Graph/MovieGraphDataTypes.h
@@ -229,6 +229,9 @@ public:
 	/** Called by the Time Step system when the external data source should pause playback. */
 	virtual void PauseDataSource() {}
 
+	/** Called by the Time Step system when the external data source should stop playback (typically at the end of rendering). */
+	virtual void StopDataSource() {}
+
 	/** Called by the Time Step system when the external data source should jump to the given time. Time is in TickResolution scale. */
 	virtual void JumpDataSource(const FFrameTime& InTimeToJumpTo) {}
 
diff --git a/Engine/Plugins/MovieScene/MovieRenderPipeline/Source/MovieRenderPipelineCore/Public/Graph/MovieGraphSequenceDataSource.h b/Engine/Plugins/MovieScene/MovieRenderPipeline/Source/MovieRenderPipelineCore/Public/Graph/MovieGraphSequenceDataSource.h
--- a/Engine/Plugins/MovieScene/MovieRenderPipeline/Source/MovieRenderPipelineCore/Public/Graph/MovieGraphSequenceDataSource.h
+++ b/Engine/Plugins/MovieScene/MovieRenderPipeline/Source/MovieRenderPipelineCore/Public/Graph/MovieGraphSequenceDataSource.h
@@ -46,6 +46,7 @@ public:
 	virtual void SyncDataSourceTime(const FFrameTime& InTime) override;
 	virtual void PlayDataSource() override;
 	virtual void PauseDataSource() override;
+	virtual void StopDataSource() override;
 	virtual void JumpDataSource(const FFrameTime& InTimeToJumpTo) override;
 	virtual void CacheHierarchyForShot(const TObjectPtr<UMoviePipelineExecutorShot>& InShot) override;
 	virtual void RestoreHierarchyForShot(const TObjectPtr<UMoviePipelineExecutorShot>& InShot) override;
