MoviePipeline: Fixed a crash that could occur if the queue has contents removed while a render is occuring. Removing queue contents is now blocked during a render.

diff --git a/Engine/Plugins/MovieScene/MovieRenderPipeline/Source/MovieRenderPipelineEditor/Private/Widgets/SMoviePipelineQueueEditor.cpp b/Engine/Plugins/MovieScene/MovieRenderPipeline/Source/MovieRenderPipelineEditor/Private/Widgets/SMoviePipelineQueueEditor.cpp
--- a/Engine/Plugins/MovieScene/MovieRenderPipeline/Source/MovieRenderPipelineEditor/Private/Widgets/SMoviePipelineQueueEditor.cpp
+++ b/Engine/Plugins/MovieScene/MovieRenderPipeline/Source/MovieRenderPipelineEditor/Private/Widgets/SMoviePipelineQueueEditor.cpp
@@ -1475,6 +1475,13 @@ void SMoviePipelineQueueEditor::Construct(const FArguments& InArgs)
 		.OnGenerateRow(this, &SMoviePipelineQueueEditor::OnGenerateRow)
 		.OnGetChildren(this, &SMoviePipelineQueueEditor::OnGetChildren)
 		.OnContextMenuOpening(this, &SMoviePipelineQueueEditor::GetContextMenuContent)
+		.IsEnabled_Lambda([]()
+		{
+			const UMoviePipelineQueueSubsystem* Subsystem = GEditor->GetEditorSubsystem<UMoviePipelineQueueSubsystem>();
+			check(Subsystem);
+			
+			return !Subsystem->IsRendering();
+		})
 		.HeaderRow
 		(
 			SNew(SHeaderRow)
