MoviePipeline: Fixed some minor issues with the Apply Console Variable Preset node.

* When promoting a pin to a variable, the new variable did not get the value of the associated cvar.
* The promotion context menu did not update to include new cvars if additional cvars were added to the preset while the graph is open.
* Not directly related to the cvar node: Promoting the property to a variable, and promoting a property to a pin, were not undo-able.

diff --git a/Engine/Plugins/MovieScene/MovieRenderPipeline/Source/MovieRenderPipelineCore/Private/Graph/MovieGraphNode.cpp b/Engine/Plugins/MovieScene/MovieRenderPipeline/Source/MovieRenderPipelineCore/Private/Graph/MovieGraphNode.cpp
--- a/Engine/Plugins/MovieScene/MovieRenderPipeline/Source/MovieRenderPipelineCore/Private/Graph/MovieGraphNode.cpp
+++ b/Engine/Plugins/MovieScene/MovieRenderPipeline/Source/MovieRenderPipelineCore/Private/Graph/MovieGraphNode.cpp
@@ -141,6 +141,8 @@ void UMovieGraphNode::TogglePromotePropertyToPin(const FName& PropertyName)
 			return OverrideablePropertyInfo.IsSamePropertyAs(ExposedInfo);
 		});
 
+		Modify();
+
 		if (!bFoundExposedPropertyInfo)
 		{
 			ExposedPropertyInfo.Add(MoveTemp(OverrideablePropertyInfo));
diff --git a/Engine/Plugins/MovieScene/MovieRenderPipeline/Source/MovieRenderPipelineCore/Private/Graph/Nodes/MovieGraphApplyCVarPresetNode.cpp b/Engine/Plugins/MovieScene/MovieRenderPipeline/Source/MovieRenderPipelineCore/Private/Graph/Nodes/MovieGraphApplyCVarPresetNode.cpp
--- a/Engine/Plugins/MovieScene/MovieRenderPipeline/Source/MovieRenderPipelineCore/Private/Graph/Nodes/MovieGraphApplyCVarPresetNode.cpp
+++ b/Engine/Plugins/MovieScene/MovieRenderPipeline/Source/MovieRenderPipelineCore/Private/Graph/Nodes/MovieGraphApplyCVarPresetNode.cpp
@@ -77,6 +77,27 @@ TArray<FPropertyBagPropertyDesc> UMovieGraphApplyCVarPresetNode::GetDynamicPrope
 	return CvarDescs;
 }
 
+bool UMovieGraphApplyCVarPresetNode::GetDynamicPropertyValue(const FName PropertyName, FString& OutValue)
+{
+	// This node uses dynamic properties to enable cvars within the preset to be promoted to pins. However, the DynamicProperties property bag does
+	// not track the value of the cvars within the preset (that would be a syncing nightmare). Instead, resolve the value of the cvar on-demand here.
+
+	constexpr bool bOnlyIncludeChecked = true;
+	TArray<TTuple<FString, FString>> ConsoleVariables;
+	ConsoleVariablePreset->GetConsoleVariablesForTrack(bOnlyIncludeChecked, ConsoleVariables);
+
+	for (const TPair<FString, FString>& CVarPair : ConsoleVariables)
+	{
+		if (CVarPair.Key == PropertyName)
+		{
+			OutValue = CVarPair.Value;
+			return true;
+		}
+	}
+
+	return false;
+}
+
 void UMovieGraphApplyCVarPresetNode::TogglePromotePropertyToPin(const FName& PropertyName)
 {
 	Super::TogglePromotePropertyToPin(PropertyName);
diff --git a/Engine/Plugins/MovieScene/MovieRenderPipeline/Source/MovieRenderPipelineCore/Public/Graph/MovieGraphNode.h b/Engine/Plugins/MovieScene/MovieRenderPipeline/Source/MovieRenderPipelineCore/Public/Graph/MovieGraphNode.h
--- a/Engine/Plugins/MovieScene/MovieRenderPipeline/Source/MovieRenderPipelineCore/Public/Graph/MovieGraphNode.h
+++ b/Engine/Plugins/MovieScene/MovieRenderPipeline/Source/MovieRenderPipelineCore/Public/Graph/MovieGraphNode.h
@@ -185,7 +185,7 @@ public:
 	 * Gets the value of the dynamic property with the specified name. Provides the serialized value of the property in
 	 * "OutValue". Returns true if "OutValue" was set and there were no errors, else returns false.
 	 */
-	bool GetDynamicPropertyValue(const FName PropertyName, FString& OutValue);
+	virtual bool GetDynamicPropertyValue(const FName PropertyName, FString& OutValue);
 
 	/** Gets the override property for the specified dynamic property. If one does not exist, returns nullptr. */
 	const FBoolProperty* FindOverridePropertyForDynamicProperty(const FName& InPropertyName) const;
diff --git a/Engine/Plugins/MovieScene/MovieRenderPipeline/Source/MovieRenderPipelineCore/Public/Graph/Nodes/MovieGraphApplyCVarPresetNode.h b/Engine/Plugins/MovieScene/MovieRenderPipeline/Source/MovieRenderPipelineCore/Public/Graph/Nodes/MovieGraphApplyCVarPresetNode.h
--- a/Engine/Plugins/MovieScene/MovieRenderPipeline/Source/MovieRenderPipelineCore/Public/Graph/Nodes/MovieGraphApplyCVarPresetNode.h
+++ b/Engine/Plugins/MovieScene/MovieRenderPipeline/Source/MovieRenderPipelineCore/Public/Graph/Nodes/MovieGraphApplyCVarPresetNode.h
@@ -21,6 +21,7 @@ public:
 	virtual EMovieGraphBranchRestriction GetBranchRestriction() const override;
 	virtual TArray<FMovieGraphPropertyInfo> GetOverrideablePropertyInfo() const override;
 	virtual TArray<FPropertyBagPropertyDesc> GetDynamicPropertyDescriptions() const override;
+	virtual bool GetDynamicPropertyValue(const FName PropertyName, FString& OutValue) override;
 	virtual void TogglePromotePropertyToPin(const FName& PropertyName) override;
 	virtual void PrepareForFlattening(const UMovieGraphSettingNode* InSourceNode) override;
 
diff --git a/Engine/Plugins/MovieScene/MovieRenderPipeline/Source/MovieRenderPipelineEditor/Private/Graph/MovieEdGraphNode.cpp b/Engine/Plugins/MovieScene/MovieRenderPipeline/Source/MovieRenderPipelineEditor/Private/Graph/MovieEdGraphNode.cpp
--- a/Engine/Plugins/MovieScene/MovieRenderPipeline/Source/MovieRenderPipelineEditor/Private/Graph/MovieEdGraphNode.cpp
+++ b/Engine/Plugins/MovieScene/MovieRenderPipeline/Source/MovieRenderPipelineEditor/Private/Graph/MovieEdGraphNode.cpp
@@ -10,6 +10,7 @@
 #include "Misc/TransactionObjectEvent.h"
 #include "MovieEdGraph.h"
 #include "MovieGraphSchema.h"
+#include "ScopedTransaction.h"
 #include "ToolMenu.h"
 #include "EdGraph/EdGraphSchema.h"
 #include "GraphEditorActions.h"
@@ -274,6 +275,9 @@ void UMoviePipelineEdGraphNode::GetNodeContextMenuActions(UToolMenu* Menu, UGrap
 
 void UMoviePipelineEdGraphNode::GetPropertyPromotionContextMenuActions(UToolMenu* Menu, const UGraphNodeContextMenuContext* Context) const
 {
+	// Before fetching the overrideable properties, update dynamic properties (dynamic properties may be included in overrideable properties)
+	RuntimeNode->UpdateDynamicProperties();
+	
 	const TArray<FMovieGraphPropertyInfo>& OverrideablePropertyInfo = RuntimeNode->GetOverrideablePropertyInfo();
 	
 	FToolMenuSection& PinActionsSection = Menu->FindOrAddSection("EdGraphSchemaPinActions");
@@ -348,6 +352,8 @@ void UMoviePipelineEdGraphNode::GetPropertyPromotionContextMenuActions(UToolMenu
 
 void UMoviePipelineEdGraphNode::PromotePropertyToVariable(const FMovieGraphPropertyInfo& TargetProperty) const
 {
+	FScopedTransaction ScopedTransaction(LOCTEXT("PromotePropertyToVariable_Transaction", "Promote Property to Variable"));
+	
 	const FName PromotedVariableName = TargetProperty.PromotionName.IsNone() ? TargetProperty.Name : TargetProperty.PromotionName;
 	
 	// Note: AddVariable() will take care of determining a unique name if there is already a variable with the property's name
@@ -389,6 +395,8 @@ void UMoviePipelineEdGraphNode::PromotePropertyToVariable(const FMovieGraphPrope
 
 void UMoviePipelineEdGraphNode::TogglePromotePropertyToPin(const FName PropertyName) const
 {
+	FScopedTransaction ScopedTransaction(LOCTEXT("PromotePropertyToPin_Transaction", "Promote Property to Pin"));
+	
 	RuntimeNode->TogglePromotePropertyToPin(PropertyName);
 }
 
