MoviePipeline: Fixed Game Overrides incorrectly restoring some cvars (ShadowDistanceScale, ViewDistanceScale) due to storing them as an int and not a float. This does not apply to the Movie Graph as it does not cache these cvar values.

diff --git a/Engine/Plugins/MovieScene/MovieRenderPipeline/Source/MovieRenderPipelineCore/Private/MoviePipelineGameOverrideSetting.cpp b/Engine/Plugins/MovieScene/MovieRenderPipeline/Source/MovieRenderPipelineCore/Private/MoviePipelineGameOverrideSetting.cpp
--- a/Engine/Plugins/MovieScene/MovieRenderPipeline/Source/MovieRenderPipelineCore/Private/MoviePipelineGameOverrideSetting.cpp
+++ b/Engine/Plugins/MovieScene/MovieRenderPipeline/Source/MovieRenderPipelineCore/Private/MoviePipelineGameOverrideSetting.cpp
@@ -95,14 +95,14 @@ void UMoviePipelineGameOverrideSetting::ApplyCVarSettings(const bool bOverrideVa
 
 	if (bUseHighQualityShadows)
 	{
-		MOVIEPIPELINE_STORE_AND_OVERRIDE_CVAR_INT(PreviousShadowDistanceScale, TEXT("r.Shadow.DistanceScale"), ShadowDistanceScale, bOverrideValues);
+		MOVIEPIPELINE_STORE_AND_OVERRIDE_CVAR_FLOAT(PreviousShadowDistanceScale, TEXT("r.Shadow.DistanceScale"), ShadowDistanceScale, bOverrideValues);
 		MOVIEPIPELINE_STORE_AND_OVERRIDE_CVAR_INT(PreviousShadowQuality, TEXT("r.ShadowQuality"), 5, bOverrideValues);
 		MOVIEPIPELINE_STORE_AND_OVERRIDE_CVAR_FLOAT(PreviousShadowRadiusThreshold, TEXT("r.Shadow.RadiusThreshold"), ShadowRadiusThreshold, bOverrideValues);
 	}
 
 	if (bOverrideViewDistanceScale)
 	{
-		MOVIEPIPELINE_STORE_AND_OVERRIDE_CVAR_INT(PreviousViewDistanceScale, TEXT("r.ViewDistanceScale"), ViewDistanceScale, bOverrideValues);
+		MOVIEPIPELINE_STORE_AND_OVERRIDE_CVAR_FLOAT(PreviousViewDistanceScale, TEXT("r.ViewDistanceScale"), ViewDistanceScale, bOverrideValues);
 	}
 
 	if (bDisableGPUTimeout)
diff --git a/Engine/Plugins/MovieScene/MovieRenderPipeline/Source/MovieRenderPipelineCore/Public/MoviePipelineGameOverrideSetting.h b/Engine/Plugins/MovieScene/MovieRenderPipeline/Source/MovieRenderPipelineCore/Public/MoviePipelineGameOverrideSetting.h
--- a/Engine/Plugins/MovieScene/MovieRenderPipeline/Source/MovieRenderPipelineCore/Public/MoviePipelineGameOverrideSetting.h
+++ b/Engine/Plugins/MovieScene/MovieRenderPipeline/Source/MovieRenderPipelineCore/Public/MoviePipelineGameOverrideSetting.h
@@ -139,10 +139,10 @@ private:
 	int32 PreviousForceLOD;
 	int32 PreviousSkeletalMeshBias;
 	int32 PreviousParticleLODBias;
-	int32 PreviousShadowDistanceScale;
+	float PreviousShadowDistanceScale;
 	int32 PreviousShadowQuality;
 	float PreviousShadowRadiusThreshold;
-	int32 PreviousViewDistanceScale;
+	float PreviousViewDistanceScale;
 	int32 PreviousGPUTimeout;
 	int32 PreviousAnimationUROEnabled;
 	int32 PreviousFoliageDitheredLOD;
