MoviePipeline: Changed where the RHI transitions were happening. They originally were added right before the texture was read back by the CPU, but the rhi validation was actually related to the User Interface drawing the textures as part of the UI preview in UMG. This adds the transitions to immediately after requesting the render, which is needed because the warm-up frames render but don't actually read anything back, thus the transition was not being scheduled.

This is a backout of CL 35881548 + edits.

Tested renders of the burn-in, UI renderer, panoramics, and deferred renderer using "-rdgimmediate -rdgdebug -rhivalidation -norhithread -forcerhibypass -onethread" to ensure there were no validation errors.

diff --git a/Engine/Plugins/MovieScene/MoviePipelineMaskRenderPass/Source/MoviePipelineMaskRenderPass/Private/MoviePipelineObjectIdPass.cpp b/Engine/Plugins/MovieScene/MoviePipelineMaskRenderPass/Source/MoviePipelineMaskRenderPass/Private/MoviePipelineObjectIdPass.cpp
--- a/Engine/Plugins/MovieScene/MoviePipelineMaskRenderPass/Source/MoviePipelineMaskRenderPass/Private/MoviePipelineObjectIdPass.cpp
+++ b/Engine/Plugins/MovieScene/MoviePipelineMaskRenderPass/Source/MoviePipelineMaskRenderPass/Private/MoviePipelineObjectIdPass.cpp
@@ -177,6 +177,14 @@ void UMoviePipelineObjectIdRenderPass::RenderSample_GameThreadImpl(const FMovieP
 
 	FCanvas Canvas = FCanvas(RenderTarget, nullptr, GetPipeline()->GetWorld(), ViewFamily->GetFeatureLevel(), FCanvas::CDM_DeferDrawing, 1.0f);
 	GetRendererModule().BeginRenderingViewFamily(&Canvas, ViewFamily.Get());
+	
+	ENQUEUE_RENDER_COMMAND(TransitionTextureSRVState)(
+	[RenderTarget](FRHICommandListImmediate& RHICmdList) mutable
+	{
+		// Transition our render target from a render target view to a shader resource view to allow the UMG preview material to read from this Render Target.
+		RHICmdList.Transition(FRHITransitionInfo(RenderTarget->GetRenderTargetTexture(), ERHIAccess::RTV, ERHIAccess::SRVGraphicsPixel));
+	});
+	
 
 	// Update the data in place, no need to copy back to the annotation.
 	UE::MoviePipeline::FObjectIdAccelerationData AccelData = ManifestAnnotation.GetAnnotation(this);
@@ -236,9 +244,6 @@ void UMoviePipelineObjectIdRenderPass::RenderSample_GameThreadImpl(const FMovieP
 		ENQUEUE_RENDER_COMMAND(CanvasRenderTargetResolveCommand)(
 			[LocalSurfaceQueue, FramePayload, Callback, RenderTarget](FRHICommandListImmediate& RHICmdList) mutable
 			{
-				// Transition our render target from a render target view to a shader resource view to allow a shader to read from this Render Target.
-				RHICmdList.Transition(FRHITransitionInfo(RenderTarget->GetRenderTargetTexture(), ERHIAccess::RTV, ERHIAccess::SRVGraphicsPixel));
-
 				// Enqueue a encode for this frame onto our worker thread.
 				LocalSurfaceQueue->OnRenderTargetReady_RenderThread(RenderTarget->GetRenderTargetTexture(), FramePayload, MoveTemp(Callback));
 			});
diff --git a/Engine/Plugins/MovieScene/MoviePipelineMaskRenderPass/Source/MoviePipelineMaskRenderPass/Private/MoviePipelinePanoramicPass.cpp b/Engine/Plugins/MovieScene/MoviePipelineMaskRenderPass/Source/MoviePipelineMaskRenderPass/Private/MoviePipelinePanoramicPass.cpp
--- a/Engine/Plugins/MovieScene/MoviePipelineMaskRenderPass/Source/MoviePipelineMaskRenderPass/Private/MoviePipelinePanoramicPass.cpp
+++ b/Engine/Plugins/MovieScene/MoviePipelineMaskRenderPass/Source/MoviePipelineMaskRenderPass/Private/MoviePipelinePanoramicPass.cpp
@@ -465,6 +465,13 @@ void UMoviePipelinePanoramicPass::RenderSample_GameThreadImpl(const FMoviePipeli
 
 				FCanvas Canvas = FCanvas(RenderTarget, nullptr, GetPipeline()->GetWorld(), ViewFamily->GetFeatureLevel(), FCanvas::CDM_DeferDrawing, 1.0f);
 				GetRendererModule().BeginRenderingViewFamily(&Canvas, ViewFamily.Get());
+				
+				ENQUEUE_RENDER_COMMAND(TransitionTextureSRVState)(
+				[RenderTarget](FRHICommandListImmediate& RHICmdList) mutable
+				{
+					// Transition our render target from a render target view to a shader resource view to allow the UMG preview material to read from this Render Target.
+					RHICmdList.Transition(FRHITransitionInfo(RenderTarget->GetRenderTargetTexture(), ERHIAccess::RTV, ERHIAccess::SRVGraphicsPixel));
+				});
 
 				// Schedule a readback and then high-res accumulation.
 				ScheduleReadbackAndAccumulation(InOutSampleState, Pane, Canvas);
@@ -554,9 +561,6 @@ void UMoviePipelinePanoramicPass::ScheduleReadbackAndAccumulation(const FMoviePi
 	ENQUEUE_RENDER_COMMAND(CanvasRenderTargetResolveCommand)(
 		[LocalSurfaceQueue, FramePayload, Callback, RenderTarget](FRHICommandListImmediate& RHICmdList) mutable
 		{
-			// Transition our render target from a render target view to a shader resource view to allow a shader to read from this Render Target.
-			RHICmdList.Transition(FRHITransitionInfo(RenderTarget->GetRenderTargetTexture(), ERHIAccess::RTV, ERHIAccess::SRVGraphicsPixel));
-
 			// Enqueue a encode for this frame onto our worker thread.
 			LocalSurfaceQueue->OnRenderTargetReady_RenderThread(RenderTarget->GetRenderTargetTexture(), FramePayload, MoveTemp(Callback));
 		});
diff --git a/Engine/Plugins/MovieScene/MovieRenderPipeline/Source/MovieRenderPipelineCore/Private/Graph/Nodes/MovieGraphWidgetRendererBaseNode.cpp b/Engine/Plugins/MovieScene/MovieRenderPipeline/Source/MovieRenderPipelineCore/Private/Graph/Nodes/MovieGraphWidgetRendererBaseNode.cpp
--- a/Engine/Plugins/MovieScene/MovieRenderPipeline/Source/MovieRenderPipelineCore/Private/Graph/Nodes/MovieGraphWidgetRendererBaseNode.cpp
+++ b/Engine/Plugins/MovieScene/MovieRenderPipeline/Source/MovieRenderPipelineCore/Private/Graph/Nodes/MovieGraphWidgetRendererBaseNode.cpp
@@ -146,7 +146,7 @@ void UMovieGraphWidgetRendererBaseNode::FMovieGraphWidgetPass::Render(const FMov
 		// Put the widget in our window
 		VirtualWindow->SetContent(WidgetToRender.ToSharedRef());
 
-		// Draw the widget to the render target
+		// Draw the widget to the render target.  This leaves the texture in SRV state so no transition is needed.
 		NodeCDO->WidgetRenderer->DrawWindow(RenderTarget, VirtualWindow->GetHittestGrid(), VirtualWindow.ToSharedRef(),
 			1.f, OutputResolution, InTimeData.FrameDeltaTime);
 
@@ -155,9 +155,6 @@ void UMovieGraphWidgetRendererBaseNode::FMovieGraphWidgetPass::Render(const FMov
 			bComposite = ParentNodeThisFrame->bCompositeOntoFinalImage, CompositingSortOrder = GetCompositingSortOrder(),
 			BackbufferRenderTarget, OutputMerger, OutputResolution](FRHICommandListImmediate& RHICmdList)
 			{
-				// Transition our render target from a render target view to a shader resource view to allow a shader to read from this Render Target.
-				RHICmdList.Transition(FRHITransitionInfo(BackbufferRenderTarget->GetRenderTargetTexture(), ERHIAccess::RTV, ERHIAccess::SRVGraphicsPixel));
-
 				const FIntRect SourceRect = FIntRect(0, 0, BackbufferRenderTarget->GetSizeXY().X, BackbufferRenderTarget->GetSizeXY().Y);
 
 				// Read the data back to the CPU
diff --git a/Engine/Plugins/MovieScene/MovieRenderPipeline/Source/MovieRenderPipelineRenderPasses/Private/Graph/Renderers/MovieGraphImagePassBase.cpp b/Engine/Plugins/MovieScene/MovieRenderPipeline/Source/MovieRenderPipelineRenderPasses/Private/Graph/Renderers/MovieGraphImagePassBase.cpp
--- a/Engine/Plugins/MovieScene/MovieRenderPipeline/Source/MovieRenderPipelineRenderPasses/Private/Graph/Renderers/MovieGraphImagePassBase.cpp
+++ b/Engine/Plugins/MovieScene/MovieRenderPipeline/Source/MovieRenderPipelineRenderPasses/Private/Graph/Renderers/MovieGraphImagePassBase.cpp
@@ -475,9 +475,6 @@ void FMovieGraphImagePassBase::PostRendererSubmission(
 	ENQUEUE_RENDER_COMMAND(CanvasRenderTargetResolveCommand)(
 		[LocalSurfaceQueue, OnSurfaceReadbackFinished, RenderTarget](FRHICommandListImmediate& RHICmdList) mutable
 		{
-			// Transition our render target from a render target view to a shader resource view to allow a shader to read from this Render Target.
-			RHICmdList.Transition(FRHITransitionInfo(RenderTarget->GetRenderTargetTexture(), ERHIAccess::RTV, ERHIAccess::SRVGraphicsPixel));
-
 			// The legacy surface reader takes the payload just so it can shuffle it into our callback, but we can just include the data
 			// directly in the callback, so this is just a dummy payload.
 			TSharedRef<FImagePixelDataPayload, ESPMode::ThreadSafe> FramePayload = MakeShared<FImagePixelDataPayload, ESPMode::ThreadSafe>();
diff --git a/Engine/Plugins/MovieScene/MovieRenderPipeline/Source/MovieRenderPipelineRenderPasses/Private/MoviePipelineDeferredPasses.cpp b/Engine/Plugins/MovieScene/MovieRenderPipeline/Source/MovieRenderPipelineRenderPasses/Private/MoviePipelineDeferredPasses.cpp
--- a/Engine/Plugins/MovieScene/MovieRenderPipeline/Source/MovieRenderPipelineRenderPasses/Private/MoviePipelineDeferredPasses.cpp
+++ b/Engine/Plugins/MovieScene/MovieRenderPipeline/Source/MovieRenderPipelineRenderPasses/Private/MoviePipelineDeferredPasses.cpp
@@ -577,6 +577,13 @@ void UMoviePipelineDeferredPassBase::RenderSample_GameThreadImpl(const FMoviePip
 
 			FCanvas Canvas = FCanvas(RenderTarget, nullptr, GetPipeline()->GetWorld(), View->GetFeatureLevel(), FCanvas::CDM_DeferDrawing, 1.0f);
 			GetRendererModule().BeginRenderingViewFamily(&Canvas, ViewFamily.Get());
+			
+			ENQUEUE_RENDER_COMMAND(TransitionTextureSRVState)(
+			[RenderTarget](FRHICommandListImmediate& RHICmdList) mutable
+			{
+				// Transition our render target from a render target view to a shader resource view to allow the UMG preview material to read from this Render Target.
+				RHICmdList.Transition(FRHITransitionInfo(RenderTarget->GetRenderTargetTexture(), ERHIAccess::RTV, ERHIAccess::SRVGraphicsPixel));
+			});
 
 			// Readback + Accumulate.
 			PostRendererSubmission(InOutSampleState, PassIdentifierForCurrentCamera, GetOutputFileSortingOrder(), Canvas);
@@ -704,6 +711,13 @@ void UMoviePipelineDeferredPassBase::RenderSample_GameThreadImpl(const FMoviePip
 
 						FCanvas Canvas = FCanvas(RenderTarget, nullptr, GetPipeline()->GetWorld(), View->GetFeatureLevel(), FCanvas::CDM_DeferDrawing, 1.0f);
 						GetRendererModule().BeginRenderingViewFamily(&Canvas, ViewFamily.Get());
+						
+						ENQUEUE_RENDER_COMMAND(TransitionTextureSRVState)(
+						[RenderTarget](FRHICommandListImmediate& RHICmdList) mutable
+						{
+							// Transition our render target from a render target view to a shader resource view to allow the UMG preview material to read from this Render Target.
+							RHICmdList.Transition(FRHITransitionInfo(RenderTarget->GetRenderTargetTexture(), ERHIAccess::RTV, ERHIAccess::SRVGraphicsPixel));
+						});
 
 						// Readback + Accumulate.
 						PostRendererSubmission(InOutSampleState, LayerPassIdentifier, GetOutputFileSortingOrder() + 1, Canvas);
@@ -1003,9 +1017,6 @@ void UMoviePipelineDeferredPassBase::PostRendererSubmission(const FMoviePipeline
 	ENQUEUE_RENDER_COMMAND(CanvasRenderTargetResolveCommand)(
 		[LocalSurfaceQueue, FramePayload, Callback, RenderTarget](FRHICommandListImmediate& RHICmdList) mutable
 		{
-			// Transition our render target from a render target view to a shader resource view to allow a shader to read from this Render Target.
-			RHICmdList.Transition(FRHITransitionInfo(RenderTarget->GetRenderTargetTexture(), ERHIAccess::RTV, ERHIAccess::SRVGraphicsPixel));
-
 			// Enqueue a encode for this frame onto our worker thread.
 			LocalSurfaceQueue->OnRenderTargetReady_RenderThread(RenderTarget->GetRenderTargetTexture(), FramePayload, MoveTemp(Callback));
 		});
diff --git a/Engine/Plugins/MovieScene/MovieRenderPipeline/Source/MovieRenderPipelineSettings/Private/MoviePipelineBurnInSetting.cpp b/Engine/Plugins/MovieScene/MovieRenderPipeline/Source/MovieRenderPipelineSettings/Private/MoviePipelineBurnInSetting.cpp
--- a/Engine/Plugins/MovieScene/MovieRenderPipeline/Source/MovieRenderPipelineSettings/Private/MoviePipelineBurnInSetting.cpp
+++ b/Engine/Plugins/MovieScene/MovieRenderPipeline/Source/MovieRenderPipelineSettings/Private/MoviePipelineBurnInSetting.cpp
@@ -81,7 +81,7 @@ void UMoviePipelineBurnInSetting::RenderSample_GameThreadImpl(const FMoviePipeli
 			// Update the Widget with the latest frame information
 			CurrentWidget->OnOutputFrameStarted(GetPipeline());
 
-			// Draw the widget to the render target
+			// Draw the widget to the render target. This leaves the texture in SRV state so no transition is needed.
 			WidgetRenderer->DrawWindow(RenderTarget, VirtualWindow->GetHittestGrid(), VirtualWindow.ToSharedRef(), 1.f, OutputResolution, InSampleState.OutputState.TimeData.FrameDeltaTime);
 
 			FRenderTarget* BackbufferRenderTarget = RenderTarget->GameThread_GetRenderTargetResource();
@@ -90,9 +90,6 @@ void UMoviePipelineBurnInSetting::RenderSample_GameThreadImpl(const FMoviePipeli
 			ENQUEUE_RENDER_COMMAND(BurnInRenderTargetResolveCommand)(
 				[InSampleState, PassIdentifierForCurrentCamera, bComposite = bCompositeOntoFinalImage, BackbufferRenderTarget, OutputBuilder](FRHICommandListImmediate& RHICmdList)
 				{
-					// Transition our render target from a render target view to a shader resource view to allow a shader to read from this Render Target.
-					RHICmdList.Transition(FRHITransitionInfo(BackbufferRenderTarget->GetRenderTargetTexture(), ERHIAccess::RTV, ERHIAccess::SRVGraphicsPixel));
-
 					FIntRect SourceRect = FIntRect(0, 0, BackbufferRenderTarget->GetSizeXY().X, BackbufferRenderTarget->GetSizeXY().Y);
 
 					// Read the data back to the CPU
diff --git a/Engine/Plugins/MovieScene/MovieRenderPipeline/Source/MovieRenderPipelineSettings/Private/MoviePipelineWidgetRenderSetting.cpp b/Engine/Plugins/MovieScene/MovieRenderPipeline/Source/MovieRenderPipelineSettings/Private/MoviePipelineWidgetRenderSetting.cpp
--- a/Engine/Plugins/MovieScene/MovieRenderPipeline/Source/MovieRenderPipelineSettings/Private/MoviePipelineWidgetRenderSetting.cpp
+++ b/Engine/Plugins/MovieScene/MovieRenderPipeline/Source/MovieRenderPipelineSettings/Private/MoviePipelineWidgetRenderSetting.cpp
@@ -73,6 +73,7 @@ void UMoviePipelineWidgetRenderer::RenderSample_GameThreadImpl(const FMoviePipel
 			// Cast the interface to a widget is a little yucky but the implementation is unlikely to change.
 			TSharedPtr<SGameLayerManager> GameLayerManager = StaticCastSharedPtr<SGameLayerManager>(LocalPlayer->ViewportClient->GetGameLayerManager());
 
+			//  This leaves the texture in SRV state so no transition is needed.
 			WidgetRenderer->DrawWidget(BackbufferRenderTarget, GameLayerManager.ToSharedRef(), 1.f, FVector2D(RenderTarget->SizeX, RenderTarget->SizeY), (float)InSampleState.OutputState.TimeData.FrameDeltaTime);
 
 			TSharedPtr<FMoviePipelineOutputMerger, ESPMode::ThreadSafe> OutputBuilder = GetPipeline()->OutputBuilder;
@@ -80,9 +81,6 @@ void UMoviePipelineWidgetRenderer::RenderSample_GameThreadImpl(const FMoviePipel
 			ENQUEUE_RENDER_COMMAND(BurnInRenderTargetResolveCommand)(
 				[InSampleState, PassIdentifierForCurrentCamera, bComposite = bCompositeOntoFinalImage, BackbufferRenderTarget, OutputBuilder](FRHICommandListImmediate& RHICmdList)
 				{
-					// Transition our render target from a render target view to a shader resource view to allow a shader to read from this Render Target.
-					RHICmdList.Transition(FRHITransitionInfo(BackbufferRenderTarget->GetRenderTargetTexture(), ERHIAccess::RTV, ERHIAccess::SRVGraphicsPixel));
-
 					FIntRect SourceRect = FIntRect(0, 0, BackbufferRenderTarget->GetSizeXY().X, BackbufferRenderTarget->GetSizeXY().Y);
 
 					// Read the data back to the CPU
