MoviePipeline: Made several fixes to the Command Line Encoder node.

* The executable specified in Project Settings is no longer required to be part of the project if the path isn't absolute. The path resolution for the executable now works the same way it does for legacy MRQ (PATH will be searched, etc).
* Incorrect warnings about a missing {shot_name} or {sequence_name} token in the file name format have been removed.
* Sequence-level encodes (ie, the file name format excludes {shot_name}) now work correctly.

diff --git a/Engine/Plugins/MovieScene/MovieRenderPipeline/Source/MovieRenderPipelineCore/Private/Graph/Nodes/MovieGraphCommandLineEncoderNode.cpp b/Engine/Plugins/MovieScene/MovieRenderPipeline/Source/MovieRenderPipelineCore/Private/Graph/Nodes/MovieGraphCommandLineEncoderNode.cpp
--- a/Engine/Plugins/MovieScene/MovieRenderPipeline/Source/MovieRenderPipelineCore/Private/Graph/Nodes/MovieGraphCommandLineEncoderNode.cpp
+++ b/Engine/Plugins/MovieScene/MovieRenderPipeline/Source/MovieRenderPipelineCore/Private/Graph/Nodes/MovieGraphCommandLineEncoderNode.cpp
@@ -102,8 +102,8 @@ void UMovieGraphCommandLineEncoderNode::StartEncodingProcess(TArray<FMovieGraphR
 	UE::MoviePipeline::ValidateOutputFormatString(FilePathFormatString, bTestRenderPass, bTestFrameNumber);
 	UE::MoviePipeline::RemoveFrameNumberFormatStrings(FilePathFormatString, true);
 
-	// Launch the encoder, once per render layer
-	for (TTuple<FMovieGraphRenderDataIdentifier, FEncoderParams>& RenderLayerEncoderParams : GenerateRenderLayerEncoderParams(InGeneratedData))
+	// Launch the encoder, once per encoder param (one for a sequence encode, or multiple for shot encodes)
+	for (TTuple<FMovieGraphRenderDataIdentifier, FEncoderParams>& RenderLayerEncoderParams : GenerateRenderLayerEncoderParams(InGeneratedData, bInIsShotEncode))
 	{
 		FMovieGraphRenderDataIdentifier& RenderIdentifier = RenderLayerEncoderParams.Key;
 		FEncoderParams& EncoderParams = RenderLayerEncoderParams.Value;
@@ -117,14 +117,24 @@ void UMovieGraphCommandLineEncoderNode::StartEncodingProcess(TArray<FMovieGraphR
 		CommandLineEncoderIdentifier.RootBranchName = GlobalsPinName;
 		CommandLineEncoderIdentifier.LayerName = GlobalsPinName.ToString();
 		CommandLineEncoderIdentifier.RendererName = FString(TEXT("CommandLineEncoder"));
-		for (FMovieGraphRenderOutputData& OutputData : InGeneratedData)
+
+		// Add the Command Line Encoder output data to the last shot if this is a sequence encode, or to the matching set of output
+		// data if this is a shot encode.
+		if (!bInIsShotEncode)
+		{
+			InGeneratedData.Last().RenderLayerData.FindOrAdd(CommandLineEncoderIdentifier).FilePaths.Add(FinalFilePath);
+		}
+		else
 		{
-			if (OutputData.Shot != EncoderParams.Shot)
+			for (FMovieGraphRenderOutputData& OutputData : InGeneratedData)
 			{
-				continue;
-			}
+				if (OutputData.Shot != EncoderParams.Shot)
+				{
+					continue;
+				}
 
-			OutputData.RenderLayerData.FindOrAdd(CommandLineEncoderIdentifier).FilePaths.Add(FinalFilePath);
+				OutputData.RenderLayerData.FindOrAdd(CommandLineEncoderIdentifier).FilePaths.Add(FinalFilePath);
+			}
 		}
 		
 		EncoderParams.NamedArguments.Add(TEXT("OutputPath"), FinalFilePath);
@@ -164,7 +174,8 @@ void UMovieGraphCommandLineEncoderNode::BeginShotExport(UMovieGraphPipeline* InM
 	
 	if (!NeedsPerShotFlushing())
 	{
-		UE_LOG(LogMovieRenderPipeline, Warning, TEXT("A shot-level Command Line Encoder node requested a non-shot export and will be ignored. Is it missing the {shot_name} token in the File Name Format?"));
+		// If the CLE node didn't request per-shot flushing, skip encoding at this point. When the job-level export is done, the CLE will
+		// run at that point. No need to warn about this.
 		return;
 	}
 	
@@ -205,25 +216,11 @@ bool UMovieGraphCommandLineEncoderNode::AreSettingsValid(TArray<FText>& OutError
 
 	// Validate project settings
 	{
-		FString EncoderPath = EncoderSettings->ExecutablePath;
+		const FString EncoderPath = EncoderSettings->ExecutablePath;
 		if (EncoderPath.IsEmpty())
 		{
 			OutErrors.Add(LOCTEXT("CommandLineEncode_MissingExecutable", "No encoder executable has been specified in the Project Settings. Please set an encoder executable in Project Settings > Movie Pipeline CLI Encoder."));
 		}
-		else
-		{
-			if (FPaths::IsRelative(EncoderPath))
-			{
-				EncoderPath = FPaths::ConvertRelativePathToFull(EncoderPath);
-			}
-		
-			if (!FPaths::FileExists(EncoderPath))
-			{
-				OutErrors.Add(FText::Format(
-					LOCTEXT("CommandLineEncode_InvalidExecutable", "Invalid encoder executable path [{0}] was specified in the Project Settings. Please set a valid encoder executable in Project Settings > Movie Pipeline CLI Encoder."),
-					FText::FromString(EncoderPath)));
-			}
-		}
 	}
 
 	// Validate node settings
@@ -307,7 +304,7 @@ bool UMovieGraphCommandLineEncoderNode::NeedsPerShotFlushing() const
 	return false;
 }
 
-TMap<FMovieGraphRenderDataIdentifier, UMovieGraphCommandLineEncoderNode::FEncoderParams> UMovieGraphCommandLineEncoderNode::GenerateRenderLayerEncoderParams(TArray<FMovieGraphRenderOutputData>& InGeneratedData) const
+TMap<FMovieGraphRenderDataIdentifier, UMovieGraphCommandLineEncoderNode::FEncoderParams> UMovieGraphCommandLineEncoderNode::GenerateRenderLayerEncoderParams(TArray<FMovieGraphRenderOutputData>& InGeneratedData, const bool bInIsShotEncode) const
 {
 	// We produce one set of encoder parameters (ie, one file) per render layer
 	TMap<FMovieGraphRenderDataIdentifier, FEncoderParams> RenderLayerEncoderParams;
@@ -357,10 +354,21 @@ TMap<FMovieGraphRenderDataIdentifier, UMovieGraphCommandLineEncoderNode::FEncode
 			{
 				continue;
 			}
+
+			// If this is a sequence encode, use the first render identifier that's found as the encoder params key. This key will be re-used for all
+			// shots that are included in the sequence encode.
+			FMovieGraphRenderDataIdentifier EncoderParamsKey = RenderIdentifier;
+			if (!bInIsShotEncode)
+			{
+				TArray<FMovieGraphRenderDataIdentifier> Identifiers;
+				InGeneratedData[0].RenderLayerData.GetKeys(Identifiers);
+
+				EncoderParamsKey = Identifiers[0];
+			}
 			
-			FEncoderParams& EncoderParams = RenderLayerEncoderParams.FindOrAdd(RenderIdentifier);
+			FEncoderParams& EncoderParams = RenderLayerEncoderParams.FindOrAdd(EncoderParamsKey);
 			EncoderParams.Shot = GeneratedRenderData.Shot;
-			EncoderParams.RenderDataIdentifier = RenderIdentifier;
+			EncoderParams.RenderDataIdentifier = EncoderParamsKey;
 			EncoderParams.NamedArguments = SharedArguments;
 			
 			for (const FString& FilePath : RenderOutputData.FilePaths)
@@ -369,8 +377,25 @@ TMap<FMovieGraphRenderDataIdentifier, UMovieGraphCommandLineEncoderNode::FEncode
 				EncoderParams.FilesByExtensionType.FindOrAdd(Extension).Add(FilePath);
 			}
 
-			// Attach audio
-			EncoderParams.FilesByExtensionType.Append(AudioPathsByExtension);
+			// Attach audio.
+			// If this is a sequence encode, only add the audio file if it doesn't exist already (because there's only one encoder params for the
+			// sequence, vs multiple for shot encodes). In most cases, the audio provided will be the audio for the entire sequence, so there's no
+			// need to add it again.
+			if (!bInIsShotEncode)
+			{
+				for (const TPair<FString, TArray<FString>>& AudioPair : AudioPathsByExtension)
+				{
+					TArray<FString>& ExistingAudioFilesForExtension = EncoderParams.FilesByExtensionType.FindOrAdd(AudioPair.Key);
+					for (const FString& NewAudioFile : AudioPair.Value)
+					{
+						ExistingAudioFilesForExtension.AddUnique(NewAudioFile);
+					}
+				}
+			}
+			else
+			{
+				EncoderParams.FilesByExtensionType.Append(AudioPathsByExtension);
+			}
 
 			RenderDataToRemove.Add(RenderIdentifier);
 
diff --git a/Engine/Plugins/MovieScene/MovieRenderPipeline/Source/MovieRenderPipelineCore/Public/Graph/Nodes/MovieGraphCommandLineEncoderNode.h b/Engine/Plugins/MovieScene/MovieRenderPipeline/Source/MovieRenderPipelineCore/Public/Graph/Nodes/MovieGraphCommandLineEncoderNode.h
--- a/Engine/Plugins/MovieScene/MovieRenderPipeline/Source/MovieRenderPipelineCore/Public/Graph/Nodes/MovieGraphCommandLineEncoderNode.h
+++ b/Engine/Plugins/MovieScene/MovieRenderPipeline/Source/MovieRenderPipelineCore/Public/Graph/Nodes/MovieGraphCommandLineEncoderNode.h
@@ -159,8 +159,11 @@ private:
 	/** Returns true if there should be one export per shot, else false. */
 	bool NeedsPerShotFlushing() const;
 
-	/** Generates encoder parameters for each render layer. */
-	TMap<FMovieGraphRenderDataIdentifier, FEncoderParams> GenerateRenderLayerEncoderParams(TArray<FMovieGraphRenderOutputData>& InGeneratedData) const;
+	/**
+	 * Generates encoder parameters for each shot (using only the first render layer that's found within the shot). Will only generate one set of
+	 * params if bInIsShotEncode is set to false (ie, this is a sequence encode).
+	 */
+	TMap<FMovieGraphRenderDataIdentifier, FEncoderParams> GenerateRenderLayerEncoderParams(TArray<FMovieGraphRenderOutputData>& InGeneratedData, const bool bInIsShotEncode) const;
 
 	/** Gets the path to the file that the encode process should write to. */
 	FString GetResolvedOutputFilename(const FMovieGraphRenderDataIdentifier& RenderIdentifier, const TWeakObjectPtr<UMoviePipelineExecutorShot>& Shot, const FString& FileNameFormatString) const;
