MoviePipeline: Fixed an issue that caused the sublevel picker in the Collection node to load all sublevel assets, causing lengthy freezes in large projects. The picker will no longer trigger asset loads.

diff --git a/Engine/Plugins/MovieScene/MovieRenderPipeline/Source/MovieRenderPipelineCore/Private/Graph/MovieGraphRenderLayerSubsystem.cpp b/Engine/Plugins/MovieScene/MovieRenderPipeline/Source/MovieRenderPipelineCore/Private/Graph/MovieGraphRenderLayerSubsystem.cpp
--- a/Engine/Plugins/MovieScene/MovieRenderPipeline/Source/MovieRenderPipelineCore/Private/Graph/MovieGraphRenderLayerSubsystem.cpp
+++ b/Engine/Plugins/MovieScene/MovieRenderPipeline/Source/MovieRenderPipelineCore/Private/Graph/MovieGraphRenderLayerSubsystem.cpp
@@ -1987,8 +1987,15 @@ TSharedRef<SWidget> UMovieGraphConditionGroupQuery_Sublevel::GetAddMenuContents(
 		SublevelPickerConfig.OnShouldFilterAsset = FOnShouldFilterAsset::CreateLambda([this](const FAssetData& InLevelAsset)
 		{
 			// Don't show sublevels which have already been picked
-			UWorld* Sublevel = Cast<UWorld>(InLevelAsset.GetAsset());
-			return !Sublevel || Sublevels.Contains(Sublevel);
+			for (const TSoftObjectPtr<UWorld>& Sublevel : Sublevels)
+			{
+				if (Sublevel.ToSoftObjectPath() == InLevelAsset.GetSoftObjectPath())
+				{
+					return true;
+				}
+			}
+
+			return false;
 		});
 	}
 
