MovieGraph: Fixed {camera_name} token not being forcibly added to a filename when you had multiple cameras and were writing to a file format that did not support multiple layers.

Tested ~35 combinations of multi-layer exr vs jpg with 1/2 cameras, 1/2 layers, 1/2 renderers per layer, burn-ins comped and non-comped.

diff --git a/Engine/Plugins/MovieScene/MovieRenderPipeline/Source/MovieRenderPipelineCore/Private/Graph/MovieGraphDataTypes.cpp b/Engine/Plugins/MovieScene/MovieRenderPipeline/Source/MovieRenderPipelineCore/Private/Graph/MovieGraphDataTypes.cpp
--- a/Engine/Plugins/MovieScene/MovieRenderPipeline/Source/MovieRenderPipelineCore/Private/Graph/MovieGraphDataTypes.cpp
+++ b/Engine/Plugins/MovieScene/MovieRenderPipeline/Source/MovieRenderPipelineCore/Private/Graph/MovieGraphDataTypes.cpp
@@ -78,7 +78,7 @@ UE::MovieGraph::FMovieGraphRenderDataValidationInfo UE::MovieGraph::FMovieGraphO
 		// Figure out how many cameras exist for this branch
 		if (PassIdentifier.RootBranchName == ActiveBranchName)
 		{
-			if (CameraNameCounts.Contains(PassIdentifier.CameraName))
+			if (!CameraNameCounts.Contains(PassIdentifier.CameraName))
 			{
 				ValidationInfo.ActiveCameraCount++;
 			}
