MoviePipeline: Fixed a crash that could occur when using the Burn-In and Object ID nodes together in the same branch in the graph.

#rnx

diff --git a/Engine/Plugins/MovieScene/MoviePipelineMaskRenderPass/Source/MoviePipelineMaskRenderPass/Private/Graph/MovieGraphObjectIdNode.cpp b/Engine/Plugins/MovieScene/MoviePipelineMaskRenderPass/Source/MoviePipelineMaskRenderPass/Private/Graph/MovieGraphObjectIdNode.cpp
--- a/Engine/Plugins/MovieScene/MoviePipelineMaskRenderPass/Source/MoviePipelineMaskRenderPass/Private/Graph/MovieGraphObjectIdNode.cpp
+++ b/Engine/Plugins/MovieScene/MoviePipelineMaskRenderPass/Source/MoviePipelineMaskRenderPass/Private/Graph/MovieGraphObjectIdNode.cpp
@@ -40,6 +40,12 @@ bool UMovieGraphObjectIdNode::GetAllowsShowFlagsCustomization() const
 	return false;
 }
 
+bool UMovieGraphObjectIdNode::GetAllowsCompositing() const
+{
+	// Having anything composited into an Object ID pass would corrupt the data in it
+	return false;
+}
+
 #if WITH_EDITOR
 FText UMovieGraphObjectIdNode::GetNodeTitle(const bool bGetDescriptive) const
 {
diff --git a/Engine/Plugins/MovieScene/MoviePipelineMaskRenderPass/Source/MoviePipelineMaskRenderPass/Private/Graph/MovieGraphObjectIdNode.h b/Engine/Plugins/MovieScene/MoviePipelineMaskRenderPass/Source/MoviePipelineMaskRenderPass/Private/Graph/MovieGraphObjectIdNode.h
--- a/Engine/Plugins/MovieScene/MoviePipelineMaskRenderPass/Source/MoviePipelineMaskRenderPass/Private/Graph/MovieGraphObjectIdNode.h
+++ b/Engine/Plugins/MovieScene/MoviePipelineMaskRenderPass/Source/MoviePipelineMaskRenderPass/Private/Graph/MovieGraphObjectIdNode.h
@@ -25,6 +25,7 @@ public:
 	virtual FEngineShowFlags GetShowFlags() const override;
 	virtual EViewModeIndex GetViewModeIndex() const override;
 	virtual bool GetAllowsShowFlagsCustomization() const override;
+	virtual bool GetAllowsCompositing() const override;
 	// ~UMovieGraphImagePassBaseNode Interface
 
 	// UMovieGraphNode Interface
diff --git a/Engine/Plugins/MovieScene/MovieRenderPipeline/Source/MovieRenderPipelineCore/Private/Graph/Nodes/MovieGraphVideoOutputNode.cpp b/Engine/Plugins/MovieScene/MovieRenderPipeline/Source/MovieRenderPipelineCore/Private/Graph/Nodes/MovieGraphVideoOutputNode.cpp
--- a/Engine/Plugins/MovieScene/MovieRenderPipeline/Source/MovieRenderPipelineCore/Private/Graph/Nodes/MovieGraphVideoOutputNode.cpp
+++ b/Engine/Plugins/MovieScene/MovieRenderPipeline/Source/MovieRenderPipelineCore/Private/Graph/Nodes/MovieGraphVideoOutputNode.cpp
@@ -83,6 +83,7 @@ void UMovieGraphVideoOutputNode::OnReceiveImageDataImpl(UMovieGraphPipeline* InP
 		}
 
 		FImagePixelData* RawRenderPassData = RenderPassData.Value.Get();
+		const UE::MovieGraph::FMovieGraphSampleState* Payload = RawRenderPassData->GetPayload<UE::MovieGraph::FMovieGraphSampleState>();
 
 		// Write this frame's beauty pass, and pass along other passes so the encoder can composite them
 		if (RenderPassData.Key.SubResourceName == TEXT("beauty"))
@@ -90,6 +91,12 @@ void UMovieGraphVideoOutputNode::OnReceiveImageDataImpl(UMovieGraphPipeline* InP
 			TArray<FMovieGraphPassData> CompositesForThisCamera;
 			for (FMovieGraphPassData& CompositePass : CompositedPasses)
 			{
+				// This pass may not allow other passes to be composited on it 
+				if (!Payload->bAllowsCompositing)
+				{
+					continue;
+				}
+				
 				// Match them up by camera name so multiple passes intended for different camera names work.
 				if (RenderPassData.Key.CameraName == CompositePass.Key.CameraName)
 				{
diff --git a/Engine/Plugins/MovieScene/MovieRenderPipeline/Source/MovieRenderPipelineCore/Public/Graph/MovieGraphDataTypes.h b/Engine/Plugins/MovieScene/MovieRenderPipeline/Source/MovieRenderPipelineCore/Public/Graph/MovieGraphDataTypes.h
--- a/Engine/Plugins/MovieScene/MovieRenderPipeline/Source/MovieRenderPipelineCore/Public/Graph/MovieGraphDataTypes.h
+++ b/Engine/Plugins/MovieScene/MovieRenderPipeline/Source/MovieRenderPipelineCore/Public/Graph/MovieGraphDataTypes.h
@@ -348,6 +348,7 @@ namespace UE::MovieGraph
 			, bRequiresAccumulator(false)
 			, bFetchFromAccumulator(false)
 			, bCompositeOnOtherRenders(false)
+			, bAllowsCompositing(true)
 			, OverscanFraction(0.f)
 			, CompositingSortOrder(0)
 			, bAllowOCIO(true)
@@ -382,6 +383,9 @@ namespace UE::MovieGraph
 		/** Set this to true if this pass should be composited on top of other renders. */
 		bool bCompositeOnOtherRenders;
 
+		/** Set this to false if this pass should not allow other passes to be composited onto it. */
+		bool bAllowsCompositing;
+
 		/** When using high-res tiling, how many pixels does each tile overlap the adjacent tiles (on each side)? This should be zero if not using tiling. */
 		FIntPoint OverlappedPad;
 
diff --git a/Engine/Plugins/MovieScene/MovieRenderPipeline/Source/MovieRenderPipelineRenderPasses/Private/Graph/Renderers/MovieGraphDeferredPass.cpp b/Engine/Plugins/MovieScene/MovieRenderPipeline/Source/MovieRenderPipelineRenderPasses/Private/Graph/Renderers/MovieGraphDeferredPass.cpp
--- a/Engine/Plugins/MovieScene/MovieRenderPipeline/Source/MovieRenderPipelineRenderPasses/Private/Graph/Renderers/MovieGraphDeferredPass.cpp
+++ b/Engine/Plugins/MovieScene/MovieRenderPipeline/Source/MovieRenderPipelineRenderPasses/Private/Graph/Renderers/MovieGraphDeferredPass.cpp
@@ -271,6 +271,7 @@ void FMovieGraphDeferredPass::Render(const FMovieGraphTraversalContext& InFrameT
 			SampleState.OverlappedSubpixelShift = OverlappedSubpixelShift;
 			SampleState.OverscanFraction = CameraInfo.ViewInfo.GetOverscan();
 			SampleState.bAllowOCIO = ParentNodeThisFrame->GetAllowOCIO();
+			SampleState.bAllowsCompositing = ParentNodeThisFrame->GetAllowsCompositing();
 			SampleState.SceneCaptureSource = SceneCaptureSource;
 			SampleState.CompositingSortOrder = 10;
 		}
diff --git a/Engine/Plugins/MovieScene/MovieRenderPipeline/Source/MovieRenderPipelineRenderPasses/Private/MovieGraphImageSequenceOutputNode.cpp b/Engine/Plugins/MovieScene/MovieRenderPipeline/Source/MovieRenderPipelineRenderPasses/Private/MovieGraphImageSequenceOutputNode.cpp
--- a/Engine/Plugins/MovieScene/MovieRenderPipeline/Source/MovieRenderPipelineRenderPasses/Private/MovieGraphImageSequenceOutputNode.cpp
+++ b/Engine/Plugins/MovieScene/MovieRenderPipeline/Source/MovieRenderPipelineRenderPasses/Private/MovieGraphImageSequenceOutputNode.cpp
@@ -179,6 +179,12 @@ void UMovieGraphImageSequenceOutputNode::OnReceiveImageDataImpl(UMovieGraphPipel
 		// Perform compositing if any composited passes were found earlier
 		for (TPair<FMovieGraphRenderDataIdentifier, TUniquePtr<FImagePixelData>>& CompositedPass : CompositedPasses)
 		{
+			// This pass may not allow other passes to be composited on it
+			if (!Payload->bAllowsCompositing)
+			{
+				continue;
+			}
+			
 			// This composited pass will only composite on top of renders w/ the same branch and camera
 			if (!CompositedPass.Key.IsBranchAndCameraEqual(RenderData.Key))
 			{
@@ -369,6 +375,12 @@ void UMovieGraphImageSequenceOutputNode_EXR::OnReceiveImageDataImpl(UMovieGraphP
 		// Perform compositing if any composited passes were found earlier
 		for (TPair<FMovieGraphRenderDataIdentifier, TUniquePtr<FImagePixelData>>& CompositedPass : CompositedPasses)
 		{
+			// This pass may not allow other passes to be composited on it
+			if (!Payload->bAllowsCompositing)
+			{
+				continue;
+			}
+			
 			// This composited pass will only composite on top of renders w/ the same branch and camera
 			if (CompositedPass.Key.IsBranchAndCameraEqual(RenderData.Key))
 			{
diff --git a/Engine/Plugins/MovieScene/MovieRenderPipeline/Source/MovieRenderPipelineRenderPasses/Public/Graph/Nodes/MovieGraphImagePassBaseNode.h b/Engine/Plugins/MovieScene/MovieRenderPipeline/Source/MovieRenderPipelineRenderPasses/Public/Graph/Nodes/MovieGraphImagePassBaseNode.h
--- a/Engine/Plugins/MovieScene/MovieRenderPipeline/Source/MovieRenderPipelineRenderPasses/Public/Graph/Nodes/MovieGraphImagePassBaseNode.h
+++ b/Engine/Plugins/MovieScene/MovieRenderPipeline/Source/MovieRenderPipelineRenderPasses/Public/Graph/Nodes/MovieGraphImagePassBaseNode.h
@@ -65,6 +65,9 @@ public:
 	/** Whether this node allows changing the Show Flags in the details panel. */
 	virtual bool GetAllowsShowFlagsCustomization() const { return true; }
 
+	/** Whether this pass allows other passes to be composited on it. */
+	virtual bool GetAllowsCompositing() const { return true; }
+
 protected:
 	// Note: Since *individual* show flags are overridden instead of the entire ShowFlags property, manually set to
 	// overridden so the traversal picks the changes up (otherwise they will be ignored).
