MoviePipeline: Use AddUnique instead of Add when building command line argument list to avoid creating many duplicates of the same cvar when there are multiple jobs in the Queue list. Because the CVar values are part of the format string, unique values will still be preserved (which is considered undefined behavior anyways as the device profile cvars can only be applied once when the engine is started). CVars will be applied again when the particular job is started which covers most CVars.

diff --git a/Engine/Plugins/MovieScene/MovieRenderPipeline/Source/MovieRenderPipelineCore/Private/Graph/Nodes/MovieGraphGlobalGameOverrides.cpp b/Engine/Plugins/MovieScene/MovieRenderPipeline/Source/MovieRenderPipelineCore/Private/Graph/Nodes/MovieGraphGlobalGameOverrides.cpp
--- a/Engine/Plugins/MovieScene/MovieRenderPipeline/Source/MovieRenderPipelineCore/Private/Graph/Nodes/MovieGraphGlobalGameOverrides.cpp
+++ b/Engine/Plugins/MovieScene/MovieRenderPipeline/Source/MovieRenderPipelineCore/Private/Graph/Nodes/MovieGraphGlobalGameOverrides.cpp
@@ -33,61 +33,61 @@ void UMovieGraphGlobalGameOverridesNode::BuildNewProcessCommandLineArgsImpl(TArr
 		Scalability::FQualityLevels QualityLevels;
 		QualityLevels.SetFromSingleQualityLevel(static_cast<int32>(ScalabilityQualityLevel));
 
-		InOutDeviceProfileCvars.Add(FString::Format(TEXT("sg.ViewDistanceQuality={0}"), {QualityLevels.ViewDistanceQuality}));
-		InOutDeviceProfileCvars.Add(FString::Format(TEXT("sg.AntiAliasingQuality={0}"), {QualityLevels.AntiAliasingQuality}));
-		InOutDeviceProfileCvars.Add(FString::Format(TEXT("sg.ShadowQuality={0}"), {QualityLevels.ShadowQuality}));
-		InOutDeviceProfileCvars.Add(FString::Format(TEXT("sg.GlobalIlluminationQuality={0}"), {QualityLevels.GlobalIlluminationQuality}));
-		InOutDeviceProfileCvars.Add(FString::Format(TEXT("sg.ReflectionQuality={0}"), {QualityLevels.ReflectionQuality}));
-		InOutDeviceProfileCvars.Add(FString::Format(TEXT("sg.PostProcessQuality={0}"), {QualityLevels.PostProcessQuality}));
-		InOutDeviceProfileCvars.Add(FString::Format(TEXT("sg.TextureQuality={0}"), {QualityLevels.TextureQuality}));
-		InOutDeviceProfileCvars.Add(FString::Format(TEXT("sg.EffectsQuality={0}"), {QualityLevels.EffectsQuality}));
-		InOutDeviceProfileCvars.Add(FString::Format(TEXT("sg.FoliageQuality={0}"), {QualityLevels.FoliageQuality}));
-		InOutDeviceProfileCvars.Add(FString::Format(TEXT("sg.ShadingQuality={0}"), {QualityLevels.ShadingQuality}));
-		InOutDeviceProfileCvars.Add(FString::Format(TEXT("sg.LandscapeQuality={0}"), {QualityLevels.LandscapeQuality}));
+		InOutDeviceProfileCvars.AddUnique(FString::Format(TEXT("sg.ViewDistanceQuality={0}"), {QualityLevels.ViewDistanceQuality}));
+		InOutDeviceProfileCvars.AddUnique(FString::Format(TEXT("sg.AntiAliasingQuality={0}"), {QualityLevels.AntiAliasingQuality}));
+		InOutDeviceProfileCvars.AddUnique(FString::Format(TEXT("sg.ShadowQuality={0}"), {QualityLevels.ShadowQuality}));
+		InOutDeviceProfileCvars.AddUnique(FString::Format(TEXT("sg.GlobalIlluminationQuality={0}"), {QualityLevels.GlobalIlluminationQuality}));
+		InOutDeviceProfileCvars.AddUnique(FString::Format(TEXT("sg.ReflectionQuality={0}"), {QualityLevels.ReflectionQuality}));
+		InOutDeviceProfileCvars.AddUnique(FString::Format(TEXT("sg.PostProcessQuality={0}"), {QualityLevels.PostProcessQuality}));
+		InOutDeviceProfileCvars.AddUnique(FString::Format(TEXT("sg.TextureQuality={0}"), {QualityLevels.TextureQuality}));
+		InOutDeviceProfileCvars.AddUnique(FString::Format(TEXT("sg.EffectsQuality={0}"), {QualityLevels.EffectsQuality}));
+		InOutDeviceProfileCvars.AddUnique(FString::Format(TEXT("sg.FoliageQuality={0}"), {QualityLevels.FoliageQuality}));
+		InOutDeviceProfileCvars.AddUnique(FString::Format(TEXT("sg.ShadingQuality={0}"), {QualityLevels.ShadingQuality}));
+		InOutDeviceProfileCvars.AddUnique(FString::Format(TEXT("sg.LandscapeQuality={0}"), {QualityLevels.LandscapeQuality}));
 	}
 
 	if (bDisableTextureStreaming)
 	{
-		InOutDeviceProfileCvars.Add(TEXT("r.TextureStreaming=0"));
+		InOutDeviceProfileCvars.AddUnique(TEXT("r.TextureStreaming=0"));
 	}
 
 	if (bDisableLODs)
 	{
-		InOutDeviceProfileCvars.Add(TEXT("r.ForceLOD=0"));
-		InOutDeviceProfileCvars.Add(TEXT("r.SkeletalMeshLODBias=-10"));
-		InOutDeviceProfileCvars.Add(TEXT("r.ParticleLODBias=-10"));
-		InOutDeviceProfileCvars.Add(TEXT("foliage.DitheredLOD=0"));
-		InOutDeviceProfileCvars.Add(TEXT("foliage.ForceLOD=0"));
+		InOutDeviceProfileCvars.AddUnique(TEXT("r.ForceLOD=0"));
+		InOutDeviceProfileCvars.AddUnique(TEXT("r.SkeletalMeshLODBias=-10"));
+		InOutDeviceProfileCvars.AddUnique(TEXT("r.ParticleLODBias=-10"));
+		InOutDeviceProfileCvars.AddUnique(TEXT("foliage.DitheredLOD=0"));
+		InOutDeviceProfileCvars.AddUnique(TEXT("foliage.ForceLOD=0"));
 	}
 
 	if (bDisableHLODs)
 	{
 		// It's a command and not an integer cvar (despite taking 1/0)
-		InOutExecCmds.Add(TEXT("r.HLOD 0"));
+		InOutExecCmds.AddUnique(TEXT("r.HLOD 0"));
 	}
 
 	if (bFlushStreamingManagers)
 	{
-		InOutDeviceProfileCvars.Add(TEXT("r.Streaming.SyncStatesWhenBlocking=1"));
+		InOutDeviceProfileCvars.AddUnique(TEXT("r.Streaming.SyncStatesWhenBlocking=1"));
 	}
 
 	// Like the extra cvars applied in ApplySettings(), the below are applied to allow MRQ to function correctly.
 
 #if WITH_EDITOR
 	{
-		InOutDeviceProfileCvars.Add(TEXT("GeometryCache.Streamer.BlockTillFinishStreaming=1"));
-		InOutDeviceProfileCvars.Add(TEXT("GeometryCache.Streamer.ShowNotification=0"));
+		InOutDeviceProfileCvars.AddUnique(TEXT("GeometryCache.Streamer.BlockTillFinishStreaming=1"));
+		InOutDeviceProfileCvars.AddUnique(TEXT("GeometryCache.Streamer.ShowNotification=0"));
 	}
 #endif
 
-	InOutDeviceProfileCvars.Add(FString::Printf(TEXT("a.URO.Enable=%d"), 0));
-	InOutDeviceProfileCvars.Add(FString::Printf(TEXT("r.SkyLight.RealTimeReflectionCapture.TimeSlice=%d"), 0));
-	InOutDeviceProfileCvars.Add(FString::Printf(TEXT("r.VolumetricRenderTarget=%d"), 1));
-	InOutDeviceProfileCvars.Add(FString::Printf(TEXT("r.VolumetricRenderTarget.Mode=%d"), 3));
-	InOutDeviceProfileCvars.Add(FString::Printf(TEXT("wp.Runtime.BlockOnSlowStreaming=%d"), 0));
-	InOutDeviceProfileCvars.Add(FString::Printf(TEXT("p.Chaos.ImmPhys.MinStepTime=%d"), 0));
-	InOutDeviceProfileCvars.Add(FString::Printf(TEXT("r.SkipRedundantTransformUpdate=%d"), 0));
-	InOutDeviceProfileCvars.Add(FString::Printf(TEXT("p.ChaosCloth.UseTimeStepSmoothing=%d"), 0));
+	InOutDeviceProfileCvars.AddUnique(FString::Printf(TEXT("a.URO.Enable=%d"), 0));
+	InOutDeviceProfileCvars.AddUnique(FString::Printf(TEXT("r.SkyLight.RealTimeReflectionCapture.TimeSlice=%d"), 0));
+	InOutDeviceProfileCvars.AddUnique(FString::Printf(TEXT("r.VolumetricRenderTarget=%d"), 1));
+	InOutDeviceProfileCvars.AddUnique(FString::Printf(TEXT("r.VolumetricRenderTarget.Mode=%d"), 3));
+	InOutDeviceProfileCvars.AddUnique(FString::Printf(TEXT("wp.Runtime.BlockOnSlowStreaming=%d"), 0));
+	InOutDeviceProfileCvars.AddUnique(FString::Printf(TEXT("p.Chaos.ImmPhys.MinStepTime=%d"), 0));
+	InOutDeviceProfileCvars.AddUnique(FString::Printf(TEXT("r.SkipRedundantTransformUpdate=%d"), 0));
+	InOutDeviceProfileCvars.AddUnique(FString::Printf(TEXT("p.ChaosCloth.UseTimeStepSmoothing=%d"), 0));
 }
 
 void UMovieGraphGlobalGameOverridesNode::PostLoad()
diff --git a/Engine/Plugins/MovieScene/MovieRenderPipeline/Source/MovieRenderPipelineCore/Private/MoviePipelineGameOverrideSetting.cpp b/Engine/Plugins/MovieScene/MovieRenderPipeline/Source/MovieRenderPipelineCore/Private/MoviePipelineGameOverrideSetting.cpp
--- a/Engine/Plugins/MovieScene/MovieRenderPipeline/Source/MovieRenderPipelineCore/Private/MoviePipelineGameOverrideSetting.cpp
+++ b/Engine/Plugins/MovieScene/MovieRenderPipeline/Source/MovieRenderPipelineCore/Private/MoviePipelineGameOverrideSetting.cpp
@@ -175,26 +175,26 @@ void UMoviePipelineGameOverrideSetting::BuildNewProcessCommandLineArgsImpl(TArra
 	// transition into the correct map which will then use the GameModeOverride setting.
 	if (bCinematicQualitySettings)
 	{
-		InOutDeviceProfileCvars.Add(TEXT("sg.ViewDistanceQuality=4"));
-		InOutDeviceProfileCvars.Add(TEXT("sg.AntiAliasingQuality=4"));
-		InOutDeviceProfileCvars.Add(TEXT("sg.ShadowQuality=4"));
-		InOutDeviceProfileCvars.Add(TEXT("sg.GlobalIlluminationQuality=4"));
-		InOutDeviceProfileCvars.Add(TEXT("sg.ReflectionQuality=4"));
-		InOutDeviceProfileCvars.Add(TEXT("sg.PostProcessQuality=4"));
-		InOutDeviceProfileCvars.Add(TEXT("sg.TextureQuality=4"));
-		InOutDeviceProfileCvars.Add(TEXT("sg.EffectsQuality=4"));
-		InOutDeviceProfileCvars.Add(TEXT("sg.FoliageQuality=4"));
-		InOutDeviceProfileCvars.Add(TEXT("sg.ShadingQuality=4"));
+		InOutDeviceProfileCvars.AddUnique(TEXT("sg.ViewDistanceQuality=4"));
+		InOutDeviceProfileCvars.AddUnique(TEXT("sg.AntiAliasingQuality=4"));
+		InOutDeviceProfileCvars.AddUnique(TEXT("sg.ShadowQuality=4"));
+		InOutDeviceProfileCvars.AddUnique(TEXT("sg.GlobalIlluminationQuality=4"));
+		InOutDeviceProfileCvars.AddUnique(TEXT("sg.ReflectionQuality=4"));
+		InOutDeviceProfileCvars.AddUnique(TEXT("sg.PostProcessQuality=4"));
+		InOutDeviceProfileCvars.AddUnique(TEXT("sg.TextureQuality=4"));
+		InOutDeviceProfileCvars.AddUnique(TEXT("sg.EffectsQuality=4"));
+		InOutDeviceProfileCvars.AddUnique(TEXT("sg.FoliageQuality=4"));
+		InOutDeviceProfileCvars.AddUnique(TEXT("sg.ShadingQuality=4"));
 	}
 
 	switch (TextureStreaming)
 	{
 	case EMoviePipelineTextureStreamingMethod::FullyLoad:
-		InOutDeviceProfileCvars.Add(TEXT("r.Streaming.FramesForFullUpdate=0"));
-		InOutDeviceProfileCvars.Add(TEXT("r.Streaming.FullyLoadUsedTextures=1"));
+		InOutDeviceProfileCvars.AddUnique(TEXT("r.Streaming.FramesForFullUpdate=0"));
+		InOutDeviceProfileCvars.AddUnique(TEXT("r.Streaming.FullyLoadUsedTextures=1"));
 		break;
 	case EMoviePipelineTextureStreamingMethod::Disabled:
-		InOutDeviceProfileCvars.Add(TEXT("r.TextureStreaming=0"));
+		InOutDeviceProfileCvars.AddUnique(TEXT("r.TextureStreaming=0"));
 		break;
 	default:
 		// We don't change their texture streaming settings.
@@ -203,58 +203,58 @@ void UMoviePipelineGameOverrideSetting::BuildNewProcessCommandLineArgsImpl(TArra
 
 	if (bUseLODZero)
 	{
-		InOutDeviceProfileCvars.Add(TEXT("r.ForceLOD=0"));
-		InOutDeviceProfileCvars.Add(TEXT("r.SkeletalMeshLODBias=-10"));
-		InOutDeviceProfileCvars.Add(TEXT("r.ParticleLODBias=-10"));
-		InOutDeviceProfileCvars.Add(TEXT("foliage.DitheredLOD=0"));
-		InOutDeviceProfileCvars.Add(TEXT("foliage.ForceLOD=0"));
+		InOutDeviceProfileCvars.AddUnique(TEXT("r.ForceLOD=0"));
+		InOutDeviceProfileCvars.AddUnique(TEXT("r.SkeletalMeshLODBias=-10"));
+		InOutDeviceProfileCvars.AddUnique(TEXT("r.ParticleLODBias=-10"));
+		InOutDeviceProfileCvars.AddUnique(TEXT("foliage.DitheredLOD=0"));
+		InOutDeviceProfileCvars.AddUnique(TEXT("foliage.ForceLOD=0"));
 	}
 
 	if (bDisableHLODs)
 	{
 		// It's a command and not an integer cvar (despite taking 1/0)
-		InOutExecCmds.Add(TEXT("r.HLOD 0"));
+		InOutExecCmds.AddUnique(TEXT("r.HLOD 0"));
 	}
 
 	if (bUseHighQualityShadows)
 	{
-		InOutDeviceProfileCvars.Add(FString::Printf(TEXT("r.Shadow.DistanceScale=%d"), ShadowDistanceScale));
-		InOutDeviceProfileCvars.Add(FString::Printf(TEXT("r.Shadow.RadiusThreshold=%f"), ShadowRadiusThreshold));
-		InOutDeviceProfileCvars.Add(TEXT("r.ShadowQuality=5"));
+		InOutDeviceProfileCvars.AddUnique(FString::Printf(TEXT("r.Shadow.DistanceScale=%d"), ShadowDistanceScale));
+		InOutDeviceProfileCvars.AddUnique(FString::Printf(TEXT("r.Shadow.RadiusThreshold=%f"), ShadowRadiusThreshold));
+		InOutDeviceProfileCvars.AddUnique(TEXT("r.ShadowQuality=5"));
 	}
 
 	if (bOverrideViewDistanceScale)
 	{
-		InOutDeviceProfileCvars.Add(FString::Printf(TEXT("r.ViewDistanceScale=%d"), ViewDistanceScale));
+		InOutDeviceProfileCvars.AddUnique(FString::Printf(TEXT("r.ViewDistanceScale=%d"), ViewDistanceScale));
 	}
 
 	if (bDisableGPUTimeout)
 	{
-		InOutDeviceProfileCvars.Add(TEXT("r.D3D12.GPUTimeout=0"));
+		InOutDeviceProfileCvars.AddUnique(TEXT("r.D3D12.GPUTimeout=0"));
 	}
 
 	if (bFlushStreamingManagers)
 	{
-		InOutDeviceProfileCvars.Add(TEXT("r.Streaming.SyncStatesWhenBlocking=1"));
+		InOutDeviceProfileCvars.AddUnique(TEXT("r.Streaming.SyncStatesWhenBlocking=1"));
 	}
 	
 #if WITH_EDITOR
 	{
-		InOutDeviceProfileCvars.Add(TEXT("GeometryCache.Streamer.BlockTillFinishStreaming=1"));
-		InOutDeviceProfileCvars.Add(TEXT("GeometryCache.Streamer.ShowNotification=0"));
+		InOutDeviceProfileCvars.AddUnique(TEXT("GeometryCache.Streamer.BlockTillFinishStreaming=1"));
+		InOutDeviceProfileCvars.AddUnique(TEXT("GeometryCache.Streamer.ShowNotification=0"));
 	}
 #endif
 
 	{
-		InOutDeviceProfileCvars.Add(FString::Printf(TEXT("a.URO.Enable=%d"), 0));
+		InOutDeviceProfileCvars.AddUnique(FString::Printf(TEXT("a.URO.Enable=%d"), 0));
 	}
 
-	InOutDeviceProfileCvars.Add(FString::Printf(TEXT("au.NeverMuteNonRealtimeAudioDevices=%d"), 1));
-	InOutDeviceProfileCvars.Add(FString::Printf(TEXT("r.SkyLight.RealTimeReflectionCapture.TimeSlice=%d"), 0));
-	InOutDeviceProfileCvars.Add(FString::Printf(TEXT("r.VolumetricRenderTarget=%d"), 1));
-	InOutDeviceProfileCvars.Add(FString::Printf(TEXT("r.VolumetricRenderTarget.Mode=%d"), 3));
-	InOutDeviceProfileCvars.Add(FString::Printf(TEXT("wp.Runtime.BlockOnSlowStreaming=%d"), 0));
-	InOutDeviceProfileCvars.Add(FString::Printf(TEXT("p.Chaos.ImmPhys.MinStepTime=%d"), 0));
-	InOutDeviceProfileCvars.Add(FString::Printf(TEXT("r.SkipRedundantTransformUpdate=%d"), 0));
-	InOutDeviceProfileCvars.Add(FString::Printf(TEXT("p.ChaosCloth.UseTimeStepSmoothing=%d"), 0));
+	InOutDeviceProfileCvars.AddUnique(FString::Printf(TEXT("au.NeverMuteNonRealtimeAudioDevices=%d"), 1));
+	InOutDeviceProfileCvars.AddUnique(FString::Printf(TEXT("r.SkyLight.RealTimeReflectionCapture.TimeSlice=%d"), 0));
+	InOutDeviceProfileCvars.AddUnique(FString::Printf(TEXT("r.VolumetricRenderTarget=%d"), 1));
+	InOutDeviceProfileCvars.AddUnique(FString::Printf(TEXT("r.VolumetricRenderTarget.Mode=%d"), 3));
+	InOutDeviceProfileCvars.AddUnique(FString::Printf(TEXT("wp.Runtime.BlockOnSlowStreaming=%d"), 0));
+	InOutDeviceProfileCvars.AddUnique(FString::Printf(TEXT("p.Chaos.ImmPhys.MinStepTime=%d"), 0));
+	InOutDeviceProfileCvars.AddUnique(FString::Printf(TEXT("r.SkipRedundantTransformUpdate=%d"), 0));
+	InOutDeviceProfileCvars.AddUnique(FString::Printf(TEXT("p.ChaosCloth.UseTimeStepSmoothing=%d"), 0));
 }
diff --git a/Engine/Plugins/MovieScene/MovieRenderPipeline/Source/MovieRenderPipelineEditor/Private/MoviePipelineNewProcessExecutor.cpp b/Engine/Plugins/MovieScene/MovieRenderPipeline/Source/MovieRenderPipelineEditor/Private/MoviePipelineNewProcessExecutor.cpp
--- a/Engine/Plugins/MovieScene/MovieRenderPipeline/Source/MovieRenderPipelineEditor/Private/MoviePipelineNewProcessExecutor.cpp
+++ b/Engine/Plugins/MovieScene/MovieRenderPipeline/Source/MovieRenderPipelineEditor/Private/MoviePipelineNewProcessExecutor.cpp
@@ -95,6 +95,13 @@ void UMoviePipelineNewProcessExecutor::Execute_Implementation(UMoviePipelineQueu
 	// we go through all jobs and all settings and hope the user doesn't have conflicting settings.
 	for (UMoviePipelineExecutorJob* Job : DuplicatedQueue->GetJobs())
 	{
+		// Only jobs that are enabled should be able to contribute to the command line being built,
+		// as non-enabled jobs won't render.
+		if (!Job || !Job->IsEnabled())
+		{
+			continue;
+		}
+
 		if (Job->IsUsingGraphConfiguration())
 		{
 			UMovieGraphConfig* GraphConfig = Job->GetGraphPreset();
