MovieGraph: Changed the name of non-multilayer exrs to include "Single-Layer" to avoid confusion.

diff --git a/Engine/Plugins/MovieScene/MovieRenderPipeline/Source/MovieRenderPipelineRenderPasses/Public/MovieGraphImageSequenceOutputNode.h b/Engine/Plugins/MovieScene/MovieRenderPipeline/Source/MovieRenderPipelineRenderPasses/Public/MovieGraphImageSequenceOutputNode.h
--- a/Engine/Plugins/MovieScene/MovieRenderPipeline/Source/MovieRenderPipelineRenderPasses/Public/MovieGraphImageSequenceOutputNode.h
+++ b/Engine/Plugins/MovieScene/MovieRenderPipeline/Source/MovieRenderPipelineRenderPasses/Public/MovieGraphImageSequenceOutputNode.h
@@ -103,13 +103,13 @@ public:
 #if WITH_EDITOR
 	virtual FText GetNodeTitle(const bool bGetDescriptive = false) const override 
 	{ 
-		static const FText EXRSequenceNodeName = NSLOCTEXT("MovieGraph", "NodeName_EXRSequence", ".exr Sequence");
+		static const FText EXRSequenceNodeName = NSLOCTEXT("MovieGraph", "NodeName_EXRSequence", ".exr Sequence (Single-Layer)");
 		return EXRSequenceNodeName;
 	}
 
 	virtual FText GetKeywords() const override
 	{
-		static const FText Keywords = NSLOCTEXT("MovieGraph", "ImageSequenceOutputNode_EXR_Keywords", "exr image");
+		static const FText Keywords = NSLOCTEXT("MovieGraph", "ImageSequenceOutputNode_EXR_Keywords", "exr image single layer");
 		return Keywords;
 	}
 	
@@ -196,7 +196,7 @@ public:
 
 	virtual FText GetKeywords() const override
 	{
-		static const FText Keywords = NSLOCTEXT("MovieGraph", "ImageSequenceOutputNode_EXRMultilayer_Keywords", ".exr image (Multilayer)");
+		static const FText Keywords = NSLOCTEXT("MovieGraph", "ImageSequenceOutputNode_EXRMultilayer_Keywords", ".exr image multi layer (Multilayer)");
 		return Keywords;
 	}
 		
