MoviePipeline: Fixed revert issues with the File Name Format property autocomplete widget.

#rnx

diff --git a/Engine/Plugins/MovieScene/MovieRenderPipeline/Source/MovieRenderPipelineEditor/Private/Customizations/Graph/MovieGraphFormatTokenCustomization.cpp b/Engine/Plugins/MovieScene/MovieRenderPipeline/Source/MovieRenderPipelineEditor/Private/Customizations/Graph/MovieGraphFormatTokenCustomization.cpp
--- a/Engine/Plugins/MovieScene/MovieRenderPipeline/Source/MovieRenderPipelineEditor/Private/Customizations/Graph/MovieGraphFormatTokenCustomization.cpp
+++ b/Engine/Plugins/MovieScene/MovieRenderPipeline/Source/MovieRenderPipelineEditor/Private/Customizations/Graph/MovieGraphFormatTokenCustomization.cpp
@@ -28,6 +28,7 @@ void FMovieGraphFormatTokenCustomization::CustomizeDetails(IDetailLayoutBuilder&
 	}
 
 	DetailBuilder.EditDefaultProperty(OutputFormatPropertyHandle)->CustomWidget()
+		.PropertyHandleList({OutputFormatPropertyHandle})
 		.NameContent()
 		[
 			OutputFormatPropertyHandle->CreatePropertyNameWidget()
diff --git a/Engine/Plugins/MovieScene/MovieRenderPipeline/Source/MovieRenderPipelineEditor/Private/Widgets/SMoviePipelineFormatTokenAutoCompleteBox.cpp b/Engine/Plugins/MovieScene/MovieRenderPipeline/Source/MovieRenderPipelineEditor/Private/Widgets/SMoviePipelineFormatTokenAutoCompleteBox.cpp
--- a/Engine/Plugins/MovieScene/MovieRenderPipeline/Source/MovieRenderPipelineEditor/Private/Widgets/SMoviePipelineFormatTokenAutoCompleteBox.cpp
+++ b/Engine/Plugins/MovieScene/MovieRenderPipeline/Source/MovieRenderPipelineEditor/Private/Widgets/SMoviePipelineFormatTokenAutoCompleteBox.cpp
@@ -28,7 +28,6 @@ void SMoviePipelineFormatTokenAutoCompleteBox::Construct(const FArguments& InArg
 			.Font(IDetailLayoutBuilder::GetDetailFont())
 			.OnKeyDownHandler(this, &SMoviePipelineFormatTokenAutoCompleteBox::OnKeyDown)
 			.OnTextChanged(this, &SMoviePipelineFormatTokenAutoCompleteBox::HandleTextBoxTextChanged)
-			.OnTextCommitted(this, &SMoviePipelineFormatTokenAutoCompleteBox::HandleTextBoxTextCommitted)
 			.SelectWordOnMouseDoubleClick(true)
 			.AllowMultiLine(false)
 			.IsEnabled(InArgs._IsEnabled)
@@ -59,11 +58,25 @@ void SMoviePipelineFormatTokenAutoCompleteBox::Construct(const FArguments& InArg
 
 void SMoviePipelineFormatTokenAutoCompleteBox::OnFocusChanging(const FWeakWidgetPath& PreviousFocusPath, const FWidgetPath& NewWidgetPath, const FFocusEvent& InFocusEvent)
 {
+	// If the suggestion list view gets focus when the autocomplete was previously focused, do not process the focus change. This will cause a
+	// commit, which we don't want yet.
+	if (NewWidgetPath.ContainsWidget(SuggestionListView.Get()))
+	{
+		return;
+	}
+
 	// Close the suggestion list if the text box has lost focus to anything other than the suggestion list
-	if (PreviousFocusPath.ContainsWidget(TextBox.Get()) && !NewWidgetPath.ContainsWidget(SuggestionListView.Get()))
+	if (PreviousFocusPath.ContainsWidget(TextBox.Get()) && !NewWidgetPath.ContainsWidget(SuggestionListView.Get()) && MenuAnchor->IsOpen())
 	{
 		CloseMenuAndReset();
 	}
+
+	// If the autocomplete loses focus, commit so an undo entry is created. This is generally only important if the user is typing in something
+	// manually and not choosing an entry from the autocomplete (choosing an entry will cause a commit).
+	FString TextValue;
+	TextHandle->GetValue(TextValue);
+	HandleTextBoxTextCommitted(FText::FromString(TextValue), ETextCommit::Default);
+	
 	SWidget::OnFocusChanging(PreviousFocusPath, NewWidgetPath, InFocusEvent);
 }
 
@@ -161,7 +174,7 @@ void SMoviePipelineFormatTokenAutoCompleteBox::FindAutoCompletableTextAtPos(cons
 	bShowAutoComplete = StartingBracePos >= 0 && OutStr.Len() == 0;
 }
 
-void SMoviePipelineFormatTokenAutoCompleteBox::ReplaceRelevantTextWithSuggestion(const FString& InSuggestionText) const
+void SMoviePipelineFormatTokenAutoCompleteBox::ReplaceRelevantTextWithSuggestion(const FString& InSuggestionText)
 {
 	FString TextBoxText = TextBox->GetText().ToString();
 	int32 CursorPos = TextBoxText.Len();
@@ -200,7 +213,13 @@ void SMoviePipelineFormatTokenAutoCompleteBox::ReplaceRelevantTextWithSuggestion
 		NewCursorPos = TextBoxText.Len() - (Right.Len() - 1);
 	}
 
-	TextBox->SetText(FText::FromString(TextBoxText));
+	// Once the text replacement has been made, commit it so an undo entry is made
+	HandleTextBoxTextCommitted(FText::FromString(TextBoxText), ETextCommit::Default);
+
+	// Seemingly due to some focus event oddities, we need to manually refresh the text box after it regains focus from the suggestion list. Text will
+	// be committed at this point, but the text box may not show the update.
+	TextBox->Refresh();
+	
 	TextBox->GoTo(FTextLocation(0, NewCursorPos));
 }
 
@@ -222,7 +241,7 @@ void SMoviePipelineFormatTokenAutoCompleteBox::HandleTextBoxTextChanged(const FT
 	FilterVisibleSuggestions(OutStr, bShowAutoComplete);
 }
 
-void SMoviePipelineFormatTokenAutoCompleteBox::HandleTextBoxTextCommitted(const FText& InText, ETextCommit::Type CommitInfo)
+void SMoviePipelineFormatTokenAutoCompleteBox::HandleTextBoxTextCommitted(const FText& InText, ETextCommit::Type CommitInfo) const
 {
 	TextHandle->SetValue(InText.ToString(), EPropertyValueSetFlags::DefaultFlags);
 }
diff --git a/Engine/Plugins/MovieScene/MovieRenderPipeline/Source/MovieRenderPipelineEditor/Private/Widgets/SMoviePipelineFormatTokenAutoCompleteBox.h b/Engine/Plugins/MovieScene/MovieRenderPipeline/Source/MovieRenderPipelineEditor/Private/Widgets/SMoviePipelineFormatTokenAutoCompleteBox.h
--- a/Engine/Plugins/MovieScene/MovieRenderPipeline/Source/MovieRenderPipelineEditor/Private/Widgets/SMoviePipelineFormatTokenAutoCompleteBox.h
+++ b/Engine/Plugins/MovieScene/MovieRenderPipeline/Source/MovieRenderPipelineEditor/Private/Widgets/SMoviePipelineFormatTokenAutoCompleteBox.h
@@ -36,11 +36,11 @@ public:
 
 	static void FindAutoCompletableTextAtPos(const FString& InWholeString, int32 InCursorPos, FString& OutStr, bool& bShowAutoComplete);
 
-	void ReplaceRelevantTextWithSuggestion(const FString& InSuggestionText) const;
+	void ReplaceRelevantTextWithSuggestion(const FString& InSuggestionText);
 
 	void HandleTextBoxTextChanged(const FText& InText);
 
-	void HandleTextBoxTextCommitted(const FText& InText, ETextCommit::Type CommitInfo);
+	void HandleTextBoxTextCommitted(const FText& InText, ETextCommit::Type CommitInfo) const;
 
 	void FilterVisibleSuggestions(const FString& StrToMatch, const bool bForceShowAll);
 
