MoviePipeline: Fixed a bug which caused renaming inputs/outputs to not undo properly.

#rnx

diff --git a/Engine/Plugins/MovieScene/MovieRenderPipeline/Source/MovieRenderPipelineCore/Private/Graph/MovieGraphConfig.cpp b/Engine/Plugins/MovieScene/MovieRenderPipeline/Source/MovieRenderPipelineCore/Private/Graph/MovieGraphConfig.cpp
--- a/Engine/Plugins/MovieScene/MovieRenderPipeline/Source/MovieRenderPipelineCore/Private/Graph/MovieGraphConfig.cpp
+++ b/Engine/Plugins/MovieScene/MovieRenderPipeline/Source/MovieRenderPipelineCore/Private/Graph/MovieGraphConfig.cpp
@@ -32,7 +32,9 @@ bool UMovieGraphMember::SetMemberName(const FString& InNewName)
 	FText UnusedError;
 	if (CanRename(FText::FromString(InNewName), UnusedError))
 	{
+		Modify();
 		Name = InNewName;
+		
 		return true;
 	}
 
diff --git a/Engine/Plugins/MovieScene/MovieRenderPipeline/Source/MovieRenderPipelineEditor/Private/Customizations/Graph/MovieGraphMemberCustomization.h b/Engine/Plugins/MovieScene/MovieRenderPipeline/Source/MovieRenderPipelineEditor/Private/Customizations/Graph/MovieGraphMemberCustomization.h
--- a/Engine/Plugins/MovieScene/MovieRenderPipeline/Source/MovieRenderPipelineEditor/Private/Customizations/Graph/MovieGraphMemberCustomization.h
+++ b/Engine/Plugins/MovieScene/MovieRenderPipeline/Source/MovieRenderPipelineEditor/Private/Customizations/Graph/MovieGraphMemberCustomization.h
@@ -47,6 +47,8 @@ protected:
 	
 	void OnNameCommitted(const FText& InText, ETextCommit::Type Arg, UMovieGraphMember* MovieGraphMember) const
 	{
+		const FScopedTransaction Transaction(LOCTEXT("SetMemberName_Transaction", "Set Graph Member Name"));
+		
 		if (MovieGraphMember->SetMemberName(InText.ToString()))
 		{
 			NameEditableTextBox->SetError(FText::GetEmpty());
