MoviePipeline: Packaged builds now automatically mark the Sequence as Volatile so that features that change the size of sections now work as expected in packaged builds.

diff --git a/Engine/Plugins/MovieScene/MovieRenderPipeline/Source/MovieRenderPipelineCore/Private/MoviePipeline.cpp b/Engine/Plugins/MovieScene/MovieRenderPipeline/Source/MovieRenderPipelineCore/Private/MoviePipeline.cpp
--- a/Engine/Plugins/MovieScene/MovieRenderPipeline/Source/MovieRenderPipelineCore/Private/MoviePipeline.cpp
+++ b/Engine/Plugins/MovieScene/MovieRenderPipeline/Source/MovieRenderPipelineCore/Private/MoviePipeline.cpp
@@ -53,6 +53,7 @@
 #include "Components/SkeletalMeshComponent.h"
 #include "MoviePipelineTelemetry.h"
 #include "MovieSceneCommonHelpers.h"
+#include "Compilation/MovieSceneCompiledDataManager.h"
 
 #if WITH_EDITOR
 #include "MovieSceneExportMetadata.h"
@@ -912,7 +913,27 @@ void UMoviePipeline::InitializeLevelSequenceActor()
 	LevelSequenceActor->PlaybackSettings.bPauseAtEnd = true;
 	LevelSequenceActor->PlaybackSettings.FinishCompletionStateOverride = EMovieSceneCompletionModeOverride::ForceRestoreState;
 
-	// Use our duplicated sequence
+#if !WITH_EDITOR
+	// In shipping builds the Level Sequence data is assumed to be non-volatile. This prevents MRQ from automatically
+	// expanding sections (for temporal sub-samples, etc.). To solve this, we want to convert the user provided level
+	// sequence to be volatile, but unfortunately there's a couple of gotchas in the Sequencer API for doing that at runtime.
+	ULevelSequence* DummySequence = NewObject<ULevelSequence>(this, NAME_None, RF_Transient);
+	DummySequence->Initialize();
+	
+	// We need to create a dummy level sequence and specifically tell the LevelSequenceActor to switch to it. This is required
+	// because once we reset the precompiled data for the level sequence, if the LevelSequenceActor was already pointed to this
+	// sequence, then calling SetSequence again doesn't recompile the necessary data.
+	LevelSequenceActor->SetSequence(DummySequence);
+	
+	// Ensure the Volatile flag is set on the user provided level sequence.
+	TargetSequence->SetSequenceFlags(TargetSequence->GetFlags() | EMovieSceneSequenceFlags::Volatile);
+
+	// Reset the data in the Precompiled Data Manager so that it fetches the new Volatile value from the 
+	// sequence flags (otherwise it will used the cached data in the compiled hierarchy, which doesn't
+	// have the flag set).	
+	UMovieSceneCompiledDataManager::GetPrecompiledData()->Reset(TargetSequence);
+#endif
+
 	LevelSequenceActor->SetSequence(TargetSequence);
 
 	LevelSequenceActor->GetSequencePlayer()->SetTimeController(CustomSequenceTimeController);
diff --git a/Engine/Plugins/MovieScene/MovieRenderPipeline/Source/MovieRenderPipelineCore/Private/MoviePipelineUtils.cpp b/Engine/Plugins/MovieScene/MovieRenderPipeline/Source/MovieRenderPipelineCore/Private/MoviePipelineUtils.cpp
--- a/Engine/Plugins/MovieScene/MovieRenderPipeline/Source/MovieRenderPipelineCore/Private/MoviePipelineUtils.cpp
+++ b/Engine/Plugins/MovieScene/MovieRenderPipeline/Source/MovieRenderPipelineCore/Private/MoviePipelineUtils.cpp
@@ -23,6 +23,7 @@
 #include "SceneView.h"
 #include "MovieScene.h"
 #include "MovieSceneSequence.h"
+#include "LevelSequence.h"
 #include "MovieRenderPipelineCoreModule.h"
 #include "Math/Halton.h"
 #include "Misc/Paths.h"
@@ -325,6 +326,11 @@ namespace MoviePipeline
 					}
 
 					Node->EvaluationType = Node->MovieScene->GetEvaluationType();
+
+					if (ULevelSequence* OwningSequence = Node->MovieScene->GetTypedOuter<ULevelSequence>())
+					{
+						Node->OriginalSequenceFlags = OwningSequence->GetFlags();
+					}
 				}
 
 				// Unlock the movie scene so we can make changes to sections below, it'll get re-locked later if needed.
@@ -389,6 +395,10 @@ namespace MoviePipeline
 				{
 					Node->MovieScene->SetPlaybackRange(Node->OriginalMovieScenePlaybackRange);
 					Node->MovieScene->SetEvaluationType(Node->EvaluationType);
+					if (ULevelSequence* OwningSequence = Node->MovieScene->GetTypedOuter<ULevelSequence>())
+					{
+						OwningSequence->SetSequenceFlags(Node->OriginalSequenceFlags);
+					}
 #if WITH_EDITOR
 					Node->MovieScene->SetReadOnly(Node->bOriginalMovieSceneReadOnly);
 					Node->MovieScene->SetPlaybackRangeLocked(Node->bOriginalMovieScenePlaybackRangeLocked);
diff --git a/Engine/Plugins/MovieScene/MovieRenderPipeline/Source/MovieRenderPipelineCore/Public/MovieRenderPipelineDataTypes.h b/Engine/Plugins/MovieScene/MovieRenderPipeline/Source/MovieRenderPipelineCore/Public/MovieRenderPipelineDataTypes.h
--- a/Engine/Plugins/MovieScene/MovieRenderPipeline/Source/MovieRenderPipelineCore/Public/MovieRenderPipelineDataTypes.h
+++ b/Engine/Plugins/MovieScene/MovieRenderPipeline/Source/MovieRenderPipelineCore/Public/MovieRenderPipelineDataTypes.h
@@ -340,6 +340,7 @@ namespace MoviePipeline
 			, bOriginalShotSectionIsLocked(false)
 			, bOriginalCameraCutIsActive(false)
 			, bOriginalShotSectionIsActive(false)
+			, OriginalSequenceFlags(EMovieSceneSequenceFlags::None)
 			, EvaluationType(EMovieSceneEvaluationType::WithSubFrames)
 			, NodeID(MovieSceneSequenceID::Invalid)
 		{
@@ -364,6 +365,7 @@ namespace MoviePipeline
 
 		bool bOriginalCameraCutIsActive;
 		bool bOriginalShotSectionIsActive;
+		EMovieSceneSequenceFlags OriginalSequenceFlags;
 
 		/** An array of sections that we should expand, as well as their original range for restoration later. */
 		TArray<TTuple<UMovieSceneSection*, TRange<FFrameNumber>>> AdditionalSectionsToExpand;
