Fixed issue where animated overscan or Sequencer-spawned cameras with overscan could cause the editor to freeze or crash when rendering an MRQ/MRG shot.

diff --git a/Engine/Plugins/MovieScene/MovieRenderPipeline/Source/MovieRenderPipelineCore/Private/Graph/MovieGraphDefaultRenderer.cpp b/Engine/Plugins/MovieScene/MovieRenderPipeline/Source/MovieRenderPipelineCore/Private/Graph/MovieGraphDefaultRenderer.cpp
--- a/Engine/Plugins/MovieScene/MovieRenderPipeline/Source/MovieRenderPipelineCore/Private/Graph/MovieGraphDefaultRenderer.cpp
+++ b/Engine/Plugins/MovieScene/MovieRenderPipeline/Source/MovieRenderPipelineCore/Private/Graph/MovieGraphDefaultRenderer.cpp
@@ -72,6 +72,9 @@ void UMovieGraphDefaultRenderer::SetupRenderingPipelineForShot(UMoviePipelineExe
 		CameraIndexes.Add(-1);
 	}
 
+	CameraOverscanCache.Empty();
+	bHasWarnedAboutAnimatedOverscan = false;
+	
 	for (const FName& Branch : GraphBranches)
 	{
 		// We follow each branch looking for Render Layer nodes to figure out what render layer this should be. We assume a render layer is named
@@ -466,6 +469,36 @@ UE::MovieGraph::DefaultRenderer::FCameraInfo UMovieGraphDefaultRenderer::GetCame
 	return CameraInfo;
 }
 
+float UMovieGraphDefaultRenderer::GetCameraOverscan(UMovieGraphEvaluatedConfig* InConfig, int32 InCameraIndex)
+{
+	if (CameraOverscanCache.Contains(InCameraIndex))
+	{
+		return CameraOverscanCache[InCameraIndex];
+	}
+
+	const float CameraOverscan = GetCameraInfo(InConfig, InCameraIndex).ViewInfo.GetOverscan();
+	CameraOverscanCache.Add(InCameraIndex, CameraOverscan);
+	return CameraOverscan;
+}
+
+void UMovieGraphDefaultRenderer::WarnAboutAnimatedOverscan(float InInitialOverscan)
+{
+	if (!bHasWarnedAboutAnimatedOverscan)
+	{
+		UMoviePipelineExecutorShot* CurrentShot = GetOwningGraph()->GetActiveShotList()[GetOwningGraph()->GetCurrentShotIndex()];
+		
+		UE_LOG(
+			LogMovieRenderPipeline,
+			Warning,
+			TEXT("Detected animated Camera Overscan value on shot %s for camera %s. MRG does not support changing resolution between frames, and a resolution computed from the initial overscan (%f) will be used instead. Overscan can be affected by both the Overscan camera property or distortion parameters on the Lens component"),
+			*CurrentShot->OuterName,
+			*CurrentShot->InnerName,
+			InInitialOverscan);
+
+		bHasWarnedAboutAnimatedOverscan = true;
+	}
+}
+
 UTextureRenderTarget2D* UMovieGraphDefaultRenderer::GetOrCreateViewRenderTarget(const UE::MovieGraph::DefaultRenderer::FRenderTargetInitParams& InInitParams, const FMovieGraphRenderDataIdentifier& InIdentifier)
 {
 	UE::MovieGraph::DefaultRenderer::FMovieGraphImagePreviewDataPoolParams CombinedParams;
diff --git a/Engine/Plugins/MovieScene/MovieRenderPipeline/Source/MovieRenderPipelineCore/Private/Graph/Nodes/MovieGraphWidgetRendererBaseNode.cpp b/Engine/Plugins/MovieScene/MovieRenderPipeline/Source/MovieRenderPipelineCore/Private/Graph/Nodes/MovieGraphWidgetRendererBaseNode.cpp
--- a/Engine/Plugins/MovieScene/MovieRenderPipeline/Source/MovieRenderPipelineCore/Private/Graph/Nodes/MovieGraphWidgetRendererBaseNode.cpp
+++ b/Engine/Plugins/MovieScene/MovieRenderPipeline/Source/MovieRenderPipelineCore/Private/Graph/Nodes/MovieGraphWidgetRendererBaseNode.cpp
@@ -106,7 +106,7 @@ void UMovieGraphWidgetRendererBaseNode::FMovieGraphWidgetPass::Render(const FMov
 	const UMovieGraphPipeline* Pipeline = Renderer->GetOwningGraph();
 	UMovieGraphWidgetRendererBaseNode* ParentNodeThisFrame = GetParentNode(InTimeData.EvaluatedConfig);
 
-	const float CameraOverscan = Renderer->GetCameraInfo(InTimeData.EvaluatedConfig, LayerData.CameraIndex).ViewInfo.GetOverscan();
+	const float CameraOverscan = Renderer->GetCameraOverscan(InTimeData.EvaluatedConfig, LayerData.CameraIndex);
 	const FIntPoint OutputResolution = UMovieGraphBlueprintLibrary::GetEffectiveOutputResolution(InTimeData.EvaluatedConfig, CameraOverscan);
 	const int32 MaxResolution = GetMax2DTextureDimension();
 	if ((OutputResolution.X > MaxResolution) || (OutputResolution.Y > MaxResolution))
diff --git a/Engine/Plugins/MovieScene/MovieRenderPipeline/Source/MovieRenderPipelineCore/Private/MoviePipeline.cpp b/Engine/Plugins/MovieScene/MovieRenderPipeline/Source/MovieRenderPipelineCore/Private/MoviePipeline.cpp
--- a/Engine/Plugins/MovieScene/MovieRenderPipeline/Source/MovieRenderPipelineCore/Private/MoviePipeline.cpp
+++ b/Engine/Plugins/MovieScene/MovieRenderPipeline/Source/MovieRenderPipelineCore/Private/MoviePipeline.cpp
@@ -1033,9 +1033,6 @@ void UMoviePipeline::InitializeShot(UMoviePipelineExecutorShot* InShot)
 		}
 	}
 
-	// Setup required rendering architecture for all passes in this shot.
-	SetupRenderingPipelineForShot(InShot);
-
 	FMoviePipelineTelemetry::SendBeginShotRenderTelemetry(InShot);
 }
 
@@ -1917,5 +1914,55 @@ bool UMoviePipeline::GetSidecarCameraViewPoints(UMoviePipelineExecutorShot* InSh
 	return true;
 }
 
+float UMoviePipeline::GetCachedCameraOverscan(int32 InCameraIndex) const
+{
+	if (CameraOverscanCache.Contains(InCameraIndex))
+	{
+		return CameraOverscanCache[InCameraIndex];
+	}
+
+	// If no camera under the specified index is found, check if a cached value exists for the "default" camera
+	if (CameraOverscanCache.Contains(INDEX_NONE))
+	{
+		return CameraOverscanCache[INDEX_NONE];
+	}
+
+	return 0.0f;
+}
+
+bool UMoviePipeline::HasCachedCameraOverscan(int32 InCameraIndex) const
+{
+	return CameraOverscanCache.Contains(InCameraIndex);
+}
+
+void UMoviePipeline::CacheCameraOverscan(int32 InCameraIndex, float InCameraOverscan)
+{
+	if (!CameraOverscanCache.Contains(InCameraIndex))
+	{
+		CameraOverscanCache.Add(InCameraIndex, InCameraOverscan);
+	}
+	else
+	{
+		CameraOverscanCache[InCameraIndex] = InCameraOverscan;
+	}
+}
+
+void UMoviePipeline::WarnAboutAnimatedOverscan(float InInitialOverscan)
+{
+	if (!bHasWarnedAboutAnimatedOverscan)
+	{
+		UMoviePipelineExecutorShot* CurrentShot = ActiveShotList[CurrentShotIndex];
+		UE_LOG(
+			LogMovieRenderPipeline,
+			Warning,
+			TEXT("Detected animated Camera Overscan value on shot %s for camera %s. MRQ does not support changing resolution between frames, and a resolution computed from the initial overscan (%f) will be used instead. Overscan can be affected by both the Overscan camera property or distortion parameters on the Lens component"),
+			*CurrentShot->OuterName,
+			*CurrentShot->InnerName,
+			InInitialOverscan);
+
+		bHasWarnedAboutAnimatedOverscan = true;
+	}
+}
+
 #undef LOCTEXT_NAMESPACE // "MoviePipeline"
 
diff --git a/Engine/Plugins/MovieScene/MovieRenderPipeline/Source/MovieRenderPipelineCore/Private/MoviePipelineRendering.cpp b/Engine/Plugins/MovieScene/MovieRenderPipeline/Source/MovieRenderPipelineCore/Private/MoviePipelineRendering.cpp
--- a/Engine/Plugins/MovieScene/MovieRenderPipeline/Source/MovieRenderPipelineCore/Private/MoviePipelineRendering.cpp
+++ b/Engine/Plugins/MovieScene/MovieRenderPipeline/Source/MovieRenderPipelineCore/Private/MoviePipelineRendering.cpp
@@ -75,11 +75,24 @@ void UMoviePipeline::SetupRenderingPipelineForShot(UMoviePipelineExecutorShot* I
 	UMoviePipelineOutputSetting* OutputSettings = GetPipelinePrimaryConfig()->FindSetting<UMoviePipelineOutputSetting>();
 	check(OutputSettings);
 
+	// Reset cached camera overscan
+	CameraOverscanCache.Empty();
+	bHasWarnedAboutAnimatedOverscan = false;
+	
 	// TODO: Not much support here for multi-camera, so simply get the player controller camera and use its overscan value
-	FMinimalViewInfo CameraViewInfo = GetWorld()->GetFirstPlayerController()->PlayerCameraManager->GetCameraCacheView();
+	float CameraOverscan = 0.0f;
+	if (UCameraComponent* BoundCamera = MovieSceneHelpers::CameraComponentFromRuntimeObject(GetWorld()->GetFirstPlayerController()->PlayerCameraManager->GetViewTarget()))
+	{
+		FMinimalViewInfo CameraViewInfo;
+		BoundCamera->GetCameraView(GetWorld()->GetDeltaSeconds(), CameraViewInfo);
+		CameraOverscan = CameraViewInfo.GetOverscan();
+	}
+
+	// Cache the default camera overscan at INDEX_NONE to ensure anything that doesn't have multi-camera support still has an overscan value to utilize
+	CameraOverscanCache.Add(INDEX_NONE, CameraOverscan);
 	
 	FIntPoint BackbufferTileCount = FIntPoint(HighResSettings->TileCount, HighResSettings->TileCount);
-	FIntPoint OutputResolution = UMoviePipelineBlueprintLibrary::GetEffectiveOutputResolution(GetPipelinePrimaryConfig(), InShot, CameraViewInfo.GetOverscan());
+	FIntPoint OutputResolution = UMoviePipelineBlueprintLibrary::GetEffectiveOutputResolution(GetPipelinePrimaryConfig(), InShot, CameraOverscan);
 
 	// Figure out how big each sub-region (tile) is.
 	FIntPoint BackbufferResolution = FIntPoint(
@@ -195,11 +208,12 @@ void UMoviePipeline::RenderFrame()
 	check(OutputSettings);
 
 	// TODO: Not much support here for multi-camera, so simply get the player controller camera and use its overscan value
-	FMinimalViewInfo CameraViewInfo = GetWorld()->GetFirstPlayerController()->PlayerCameraManager->GetCameraCacheView();
+	// Use the cache to get the overscan value for resolution scaling so that it doesn't vary between frames
+	const float CameraOverscan  = CameraOverscanCache[INDEX_NONE];
 	
 	FIntPoint TileCount = FIntPoint(HighResSettings->TileCount, HighResSettings->TileCount);
 	FIntPoint OriginalTileCount = TileCount;
-	FIntPoint OutputResolution = UMoviePipelineBlueprintLibrary::GetEffectiveOutputResolution(GetPipelinePrimaryConfig(), ActiveShotList[CurrentShotIndex], CameraViewInfo.GetOverscan());
+	FIntPoint OutputResolution = UMoviePipelineBlueprintLibrary::GetEffectiveOutputResolution(GetPipelinePrimaryConfig(), ActiveShotList[CurrentShotIndex], CameraOverscan);
 
 	int32 NumSpatialSamples = AntiAliasingSettings->SpatialSampleCount;
 	int32 NumTemporalSamples = AntiAliasingSettings->TemporalSampleCount;
diff --git a/Engine/Plugins/MovieScene/MovieRenderPipeline/Source/MovieRenderPipelineCore/Private/MoviePipelineTiming.cpp b/Engine/Plugins/MovieScene/MovieRenderPipeline/Source/MovieRenderPipelineCore/Private/MoviePipelineTiming.cpp
--- a/Engine/Plugins/MovieScene/MovieRenderPipeline/Source/MovieRenderPipelineCore/Private/MoviePipelineTiming.cpp
+++ b/Engine/Plugins/MovieScene/MovieRenderPipeline/Source/MovieRenderPipelineCore/Private/MoviePipelineTiming.cpp
@@ -238,6 +238,10 @@ void UMoviePipeline::TickProducingFrames()
 			}
 		}
 
+		// We delay setting up the rendering pipeline for the camera cut until the end of setup. This is because we want the Sequence to have evaluated at least once
+		// (via the Jump command earlier) so that spawnable cameras exist during setup.
+		SetupRenderingPipelineForShot(CurrentCameraCut);
+		
 		// We can safely fall through to the below states as they're OK to process the same frame we set up.
 		UE_LOG(LogMovieRenderPipeline, Log, TEXT("[%d] Finished initializing Camera Cut [%d/%d] in [%s] %s."), GFrameCounter,
 			CurrentShotIndex + 1, ActiveShotList.Num(), *CurrentCameraCut->OuterName, *CurrentCameraCut->InnerName);
diff --git a/Engine/Plugins/MovieScene/MovieRenderPipeline/Source/MovieRenderPipelineCore/Public/Graph/MovieGraphDefaultRenderer.h b/Engine/Plugins/MovieScene/MovieRenderPipeline/Source/MovieRenderPipelineCore/Public/Graph/MovieGraphDefaultRenderer.h
--- a/Engine/Plugins/MovieScene/MovieRenderPipeline/Source/MovieRenderPipelineCore/Public/Graph/MovieGraphDefaultRenderer.h
+++ b/Engine/Plugins/MovieScene/MovieRenderPipeline/Source/MovieRenderPipelineCore/Public/Graph/MovieGraphDefaultRenderer.h
@@ -282,6 +282,16 @@ public:
 	UE::MovieGraph::DefaultRenderer::FCameraInfo GetCameraInfo(UMovieGraphEvaluatedConfig* InConfig, const int32 InCameraIndex) const;
 	void SetHasRenderedFirstViewThisFrame(bool bInValue) { bHasRenderedFirstViewThisFrame = bInValue; }
 	bool GetHasRenderedFirstViewThisFrame() const { return bHasRenderedFirstViewThisFrame; }
+
+	/**
+	 * Gets the overscan value for the specified camera. First checks to see if any cached overscan value exists and returning it, and if
+	 * no cached value was found, gets the camera's live overscan value and caches it.
+	 */
+	float GetCameraOverscan(UMovieGraphEvaluatedConfig* InConfig, int32 InCameraIndex);
+
+	/** Outputs a warning message regarding animated overscan to the MRQ log if one has not already been output  */
+	void WarnAboutAnimatedOverscan(float InInitialOverscan);
+	
 public:
 	UTextureRenderTarget2D* GetOrCreateViewRenderTarget(const UE::MovieGraph::DefaultRenderer::FRenderTargetInitParams& InInitParams, const FMovieGraphRenderDataIdentifier& InIdentifier);
 	FMoviePipelineSurfaceQueuePtr GetOrCreateSurfaceQueue(const UE::MovieGraph::DefaultRenderer::FRenderTargetInitParams& InInitParams);
@@ -349,4 +359,10 @@ protected:
 	FCriticalSection OutstandingTasksMutex;
 	/** Array of outstanding accumulation / blending tasks that are currently being worked on. */
 	TArray<UE::Tasks::FTask> OutstandingTasks;
+	
+	/** Caches the camera overscan used during setup to ensure that overscan-scaled resolution stays constant for every frame during a render */
+	TMap<int32, float> CameraOverscanCache;
+
+	/** Indicates if the user has already been warned about animate overscan if it is detected so that logs aren't flooded with warning messages */
+	bool bHasWarnedAboutAnimatedOverscan = false;
 };
\ No newline at end of file
diff --git a/Engine/Plugins/MovieScene/MovieRenderPipeline/Source/MovieRenderPipelineCore/Public/MoviePipeline.h b/Engine/Plugins/MovieScene/MovieRenderPipeline/Source/MovieRenderPipelineCore/Public/MoviePipeline.h
--- a/Engine/Plugins/MovieScene/MovieRenderPipeline/Source/MovieRenderPipelineCore/Public/MoviePipeline.h
+++ b/Engine/Plugins/MovieScene/MovieRenderPipeline/Source/MovieRenderPipelineCore/Public/MoviePipeline.h
@@ -116,6 +116,18 @@ public:
 	void GetSidecarCameraData(UMoviePipelineExecutorShot* InShot, int32 InCameraIndex, FMinimalViewInfo& OutViewInfo, class UCameraComponent** OutCameraComponent) const;
 	bool GetSidecarCameraViewPoints(UMoviePipelineExecutorShot* InShot, TArray<FVector>& OutSidecarViewLocations, TArray<FRotator>& OutSidecarViewRotations) const;
 
+	/** Gets any cached overscan for the specified camera, or 0 if no cached overscan was found */
+	float GetCachedCameraOverscan(int32 InCameraIndex) const;
+
+	/** Gets whether there is a cached overscan value for the specified camera */
+	bool HasCachedCameraOverscan(int32 InCameraIndex) const;
+
+	/** Caches the provided overscan value for the specified camera */
+	void CacheCameraOverscan(int32 InCameraIndex, float InCameraOverscan);
+
+	/** Outputs a warning message regarding animated overscan to the MRQ log if one has not already been output  */
+	void WarnAboutAnimatedOverscan(float InInitialOverscan);
+	
 #if WITH_EDITOR
 	const FMovieSceneExportMetadata& GetOutputMetadata() const { return OutputMetadata; }
 #endif
@@ -432,6 +444,12 @@ private:
 
 	TMap<int32, FRenderTimeStatistics> RenderTimeFrameStatistics;
 
+	/** Caches the camera overscan used during setup to ensure that overscan-scaled resolution stays constant for every frame during a render */
+	TMap<int32, float> CameraOverscanCache;
+
+	/** Indicates if the user has already been warned about animate overscan if it is detected so that logs aren't flooded with warning messages */
+	bool bHasWarnedAboutAnimatedOverscan = false;
+	
 public:
 	static FString DefaultDebugWidgetAsset;
 };
diff --git a/Engine/Plugins/MovieScene/MovieRenderPipeline/Source/MovieRenderPipelineRenderPasses/Private/Graph/Renderers/MovieGraphDeferredPass.cpp b/Engine/Plugins/MovieScene/MovieRenderPipeline/Source/MovieRenderPipelineRenderPasses/Private/Graph/Renderers/MovieGraphDeferredPass.cpp
--- a/Engine/Plugins/MovieScene/MovieRenderPipeline/Source/MovieRenderPipelineRenderPasses/Private/Graph/Renderers/MovieGraphDeferredPass.cpp
+++ b/Engine/Plugins/MovieScene/MovieRenderPipeline/Source/MovieRenderPipelineRenderPasses/Private/Graph/Renderers/MovieGraphDeferredPass.cpp
@@ -150,7 +150,7 @@ void FMovieGraphDeferredPass::Render(const FMovieGraphTraversalContext& InFrameT
 	}
 
 	// ToDo: When tiling is used, this should be the size of the per-tile backbuffer
-	const float CameraOverscan = GetRenderer()->GetCameraInfo(InTimeData.EvaluatedConfig, LayerData.CameraIndex).ViewInfo.GetOverscan();
+	const float CameraOverscan = GetRenderer()->GetCameraOverscan(InTimeData.EvaluatedConfig, LayerData.CameraIndex);
 	FIntPoint AccumulatorResolution = UMovieGraphBlueprintLibrary::GetEffectiveOutputResolution(InTimeData.EvaluatedConfig, CameraOverscan);
 	FIntPoint BackbufferResolution = AccumulatorResolution;
 	
@@ -198,7 +198,15 @@ void FMovieGraphDeferredPass::Render(const FMovieGraphTraversalContext& InFrameT
 			CameraInfo.ViewInfo.ClearOverscan();
 			CameraInfo.ViewInfo.ApplyOverscan(OverscanFraction);
 		}
-		
+		else
+		{
+			// Current overscan is different from originally cached value, indicating an animated overscan value, so output a warning message
+			if (CameraInfo.ViewInfo.GetOverscan() != CameraOverscan)
+			{
+				GetRenderer()->WarnAboutAnimatedOverscan(CameraOverscan);
+			}
+		}
+
 		// ToDo: This math probably needs the per-tile, pre-overlapped size? 
 		FIntPoint OverlappedPad = FIntPoint(FMath::CeilToInt(BackbufferResolution.X * TileOverlapPadRatio), FMath::CeilToInt(BackbufferResolution.Y * TileOverlapPadRatio));
 		
diff --git a/Engine/Plugins/MovieScene/MovieRenderPipeline/Source/MovieRenderPipelineRenderPasses/Private/MoviePipelineDeferredPasses.cpp b/Engine/Plugins/MovieScene/MovieRenderPipeline/Source/MovieRenderPipelineRenderPasses/Private/MoviePipelineDeferredPasses.cpp
--- a/Engine/Plugins/MovieScene/MovieRenderPipeline/Source/MovieRenderPipelineRenderPasses/Private/MoviePipelineDeferredPasses.cpp
+++ b/Engine/Plugins/MovieScene/MovieRenderPipeline/Source/MovieRenderPipelineRenderPasses/Private/MoviePipelineDeferredPasses.cpp
@@ -92,20 +92,32 @@ FIntPoint UMoviePipelineDeferredPassBase::GetEffectiveOutputResolutionForCamera(
 	UMoviePipelinePrimaryConfig* PrimaryConfig = GetPipeline()->GetPipelinePrimaryConfig();
 	UMoviePipelineExecutorShot* CurrentShot = GetPipeline()->GetActiveShotList()[GetPipeline()->GetCurrentShotIndex()];
 
-	// Get the camera view info to retrieve the camera's overscan, which is used when the settings to not override the overscan
-	FMinimalViewInfo CameraViewInfo;
-
-	if (GetNumCamerasToRender() == 1)
+	// Get any cached overscan value for this camera. If there is none, query the live overscan value and cache it
+	float CameraOverscan = 0.0f;
+	if (GetPipeline()->HasCachedCameraOverscan(InCameraIndex))
 	{
-		CameraViewInfo = GetPipeline()->GetWorld()->GetFirstPlayerController()->PlayerCameraManager->GetCameraCacheView();
+		CameraOverscan = GetPipeline()->GetCachedCameraOverscan(InCameraIndex);
 	}
 	else
 	{
-		UCameraComponent* CameraComponent;
-		GetPipeline()->GetSidecarCameraData(CurrentShot, InCameraIndex, CameraViewInfo, &CameraComponent);
+		// Get the camera view info to retrieve the camera's overscan, which is used when the settings to not override the overscan
+		FMinimalViewInfo CameraViewInfo;
+
+		if (GetNumCamerasToRender() == 1)
+		{
+			CameraViewInfo = GetPipeline()->GetWorld()->GetFirstPlayerController()->PlayerCameraManager->GetCameraCacheView();
+		}
+		else
+		{
+			UCameraComponent* CameraComponent;
+			GetPipeline()->GetSidecarCameraData(CurrentShot, InCameraIndex, CameraViewInfo, &CameraComponent);
+		}
+
+		CameraOverscan = CameraViewInfo.GetOverscan();
+		GetPipeline()->CacheCameraOverscan(InCameraIndex, CameraOverscan);
 	}
 	
-	const FIntPoint OutputResolution = UMoviePipelineBlueprintLibrary::GetEffectiveOutputResolution(PrimaryConfig, CurrentShot, CameraViewInfo.GetOverscan());
+	const FIntPoint OutputResolution = UMoviePipelineBlueprintLibrary::GetEffectiveOutputResolution(PrimaryConfig, CurrentShot, CameraOverscan);
 
 	return OutputResolution;
 }
diff --git a/Engine/Plugins/MovieScene/MovieRenderPipeline/Source/MovieRenderPipelineRenderPasses/Private/MoviePipelineImagePassBase.cpp b/Engine/Plugins/MovieScene/MovieRenderPipeline/Source/MovieRenderPipelineRenderPasses/Private/MoviePipelineImagePassBase.cpp
--- a/Engine/Plugins/MovieScene/MovieRenderPipeline/Source/MovieRenderPipelineRenderPasses/Private/MoviePipelineImagePassBase.cpp
+++ b/Engine/Plugins/MovieScene/MovieRenderPipeline/Source/MovieRenderPipelineRenderPasses/Private/MoviePipelineImagePassBase.cpp
@@ -442,6 +442,15 @@ FSceneView* UMoviePipelineImagePassBase::GetSceneViewForSampleState(FSceneViewFa
 		CameraInfo.ViewInfo.ClearOverscan();
 		CameraInfo.ViewInfo.ApplyOverscan(InOutSampleState.OverscanPercentage);
 	}
+	else
+	{
+		// Check for animated overscan, and warn if animated overscan was detected
+		const float CachedOverscan = GetPipeline()->GetCachedCameraOverscan(InOutSampleState.OutputState.CameraIndex);
+		if (CameraInfo.ViewInfo.GetOverscan() != CachedOverscan)
+		{
+			GetPipeline()->WarnAboutAnimatedOverscan(CachedOverscan);
+		}
+	}
 	
 	ViewInitOptions.FOV = CameraInfo.ViewInfo.FOV;
 	ViewInitOptions.DesiredFOV = CameraInfo.ViewInfo.FOV;
diff --git a/Engine/Plugins/MovieScene/MovieRenderPipeline/Source/MovieRenderPipelineSettings/Private/MoviePipelineBurnInSetting.cpp b/Engine/Plugins/MovieScene/MovieRenderPipeline/Source/MovieRenderPipelineSettings/Private/MoviePipelineBurnInSetting.cpp
--- a/Engine/Plugins/MovieScene/MovieRenderPipeline/Source/MovieRenderPipelineSettings/Private/MoviePipelineBurnInSetting.cpp
+++ b/Engine/Plugins/MovieScene/MovieRenderPipeline/Source/MovieRenderPipelineSettings/Private/MoviePipelineBurnInSetting.cpp
@@ -108,7 +108,7 @@ void UMoviePipelineBurnInSetting::RenderSample_GameThreadImpl(const FMoviePipeli
 					FrameData->SortingOrder = 4;
 					FrameData->bCompositeToFinalImage = bComposite;
 
-					TUniquePtr<FImagePixelData> PixelData = MakeUnique<TImagePixelData<FColor>>(InSampleState.BackbufferSize, TArray64<FColor>(MoveTemp(RawPixels)), FrameData);
+					TUniquePtr<FImagePixelData> PixelData = MakeUnique<TImagePixelData<FColor>>(SourceRect.Size(), TArray64<FColor>(MoveTemp(RawPixels)), FrameData);
 
 					OutputBuilder->OnCompleteRenderPassDataAvailable_AnyThread(MoveTemp(PixelData));
 				});
@@ -124,9 +124,9 @@ void UMoviePipelineBurnInSetting::SetupImpl(const MoviePipeline::FMoviePipelineR
 	}
 
 	// TODO: Add multi-camera support? For now, simply use the default player controller camera overscan
-	FMinimalViewInfo CameraViewInfo = GetPipeline()->GetWorld()->GetFirstPlayerController()->PlayerCameraManager->GetCameraCacheView();
+	const float CameraOverscan = GetPipeline()->GetCachedCameraOverscan(INDEX_NONE);
 	
-	OutputResolution = UMoviePipelineBlueprintLibrary::GetEffectiveOutputResolution(GetPipeline()->GetPipelinePrimaryConfig(), GetPipeline()->GetActiveShotList()[GetPipeline()->GetCurrentShotIndex()], CameraViewInfo.GetOverscan());
+	OutputResolution = UMoviePipelineBlueprintLibrary::GetEffectiveOutputResolution(GetPipeline()->GetPipelinePrimaryConfig(), GetPipeline()->GetActiveShotList()[GetPipeline()->GetCurrentShotIndex()], CameraOverscan);
 	int32 MaxResolution = GetMax2DTextureDimension();
 	if (OutputResolution.X > MaxResolution || OutputResolution.Y > MaxResolution)
 	{
diff --git a/Engine/Plugins/MovieScene/MovieRenderPipeline/Source/MovieRenderPipelineSettings/Private/MoviePipelineWidgetRenderSetting.cpp b/Engine/Plugins/MovieScene/MovieRenderPipeline/Source/MovieRenderPipelineSettings/Private/MoviePipelineWidgetRenderSetting.cpp
--- a/Engine/Plugins/MovieScene/MovieRenderPipeline/Source/MovieRenderPipelineSettings/Private/MoviePipelineWidgetRenderSetting.cpp
+++ b/Engine/Plugins/MovieScene/MovieRenderPipeline/Source/MovieRenderPipelineSettings/Private/MoviePipelineWidgetRenderSetting.cpp
@@ -99,7 +99,7 @@ void UMoviePipelineWidgetRenderer::RenderSample_GameThreadImpl(const FMoviePipel
 					FrameData->SortingOrder = 3;
 					FrameData->bCompositeToFinalImage = bComposite;
 
-					TUniquePtr<FImagePixelData> PixelData = MakeUnique<TImagePixelData<FColor>>(InSampleState.BackbufferSize, TArray64<FColor>(MoveTemp(RawPixels)), FrameData);
+					TUniquePtr<FImagePixelData> PixelData = MakeUnique<TImagePixelData<FColor>>(SourceRect.Size(), TArray64<FColor>(MoveTemp(RawPixels)), FrameData);
 
 					OutputBuilder->OnCompleteRenderPassDataAvailable_AnyThread(MoveTemp(PixelData));
 				});
@@ -115,8 +115,8 @@ void UMoviePipelineWidgetRenderer::SetupImpl(const MoviePipeline::FMoviePipeline
 	bool bInForceLinearGamma = false;
 
 	// TODO: Add multi-camera support? For now, simply use the default player controller camera overscan
-	FMinimalViewInfo CameraViewInfo = GetPipeline()->GetWorld()->GetFirstPlayerController()->PlayerCameraManager->GetCameraCacheView();
-	FIntPoint OutputResolution = UMoviePipelineBlueprintLibrary::GetEffectiveOutputResolution(GetPipeline()->GetPipelinePrimaryConfig(), GetPipeline()->GetActiveShotList()[GetPipeline()->GetCurrentShotIndex()], CameraViewInfo.GetOverscan());
+	const float CameraOverscan = GetPipeline()->GetCachedCameraOverscan(INDEX_NONE);
+	FIntPoint OutputResolution = UMoviePipelineBlueprintLibrary::GetEffectiveOutputResolution(GetPipeline()->GetPipelinePrimaryConfig(), GetPipeline()->GetActiveShotList()[GetPipeline()->GetCurrentShotIndex()], CameraOverscan);
 
 	int32 MaxResolution = GetMax2DTextureDimension();
 	if (OutputResolution.X > MaxResolution || OutputResolution.Y > MaxResolution)
