MoviePipeline: Fixed an issue with the Command Line Encoder node which could cause a crash if the Delete Source Files option is enabled.

#rnx

diff --git a/Engine/Plugins/MovieScene/MovieRenderPipeline/Source/MovieRenderPipelineCore/Private/Graph/Nodes/MovieGraphCommandLineEncoderNode.cpp b/Engine/Plugins/MovieScene/MovieRenderPipeline/Source/MovieRenderPipelineCore/Private/Graph/Nodes/MovieGraphCommandLineEncoderNode.cpp
--- a/Engine/Plugins/MovieScene/MovieRenderPipeline/Source/MovieRenderPipelineCore/Private/Graph/Nodes/MovieGraphCommandLineEncoderNode.cpp
+++ b/Engine/Plugins/MovieScene/MovieRenderPipeline/Source/MovieRenderPipelineCore/Private/Graph/Nodes/MovieGraphCommandLineEncoderNode.cpp
@@ -326,6 +326,17 @@ TMap<FMovieGraphRenderDataIdentifier, UMovieGraphCommandLineEncoderNode::FEncode
 	SharedArguments.Add(TEXT("VideoCodec"), VideoCodec);
 	SharedArguments.Add(TEXT("FrameRate"), EffectiveFrameRate.AsDecimal());
 	SharedArguments.Add(TEXT("Quality"), SanitizedEncodeSettings);
+
+	// If this is a sequence encode, use the first render identifier that's found as the encoder params key. This key will be re-used for all
+	// shots that are included in the sequence encode.
+	FMovieGraphRenderDataIdentifier SequenceEncoderParamsKey;
+	if (!bInIsShotEncode)
+	{
+		TArray<FMovieGraphRenderDataIdentifier> Identifiers;
+		InGeneratedData[0].RenderLayerData.GetKeys(Identifiers);
+
+		SequenceEncoderParamsKey = Identifiers[0];
+	}
 	
 	for (FMovieGraphRenderOutputData& GeneratedRenderData : InGeneratedData)
 	{
@@ -355,17 +366,7 @@ TMap<FMovieGraphRenderDataIdentifier, UMovieGraphCommandLineEncoderNode::FEncode
 				continue;
 			}
 
-			// If this is a sequence encode, use the first render identifier that's found as the encoder params key. This key will be re-used for all
-			// shots that are included in the sequence encode.
-			FMovieGraphRenderDataIdentifier EncoderParamsKey = RenderIdentifier;
-			if (!bInIsShotEncode)
-			{
-				TArray<FMovieGraphRenderDataIdentifier> Identifiers;
-				InGeneratedData[0].RenderLayerData.GetKeys(Identifiers);
-
-				EncoderParamsKey = Identifiers[0];
-			}
-			
+			FMovieGraphRenderDataIdentifier EncoderParamsKey = !bInIsShotEncode ? SequenceEncoderParamsKey : RenderIdentifier;
 			FEncoderParams& EncoderParams = RenderLayerEncoderParams.FindOrAdd(EncoderParamsKey);
 			EncoderParams.Shot = GeneratedRenderData.Shot;
 			EncoderParams.RenderDataIdentifier = EncoderParamsKey;
