Holdout: Additional fix for toast notification when the config is read-only.

#rnx

diff --git a/Engine/Plugins/Experimental/Compositing/HoldoutComposite/Source/HoldoutComposite/Private/HoldoutCompositeSubsystem.cpp b/Engine/Plugins/Experimental/Compositing/HoldoutComposite/Source/HoldoutComposite/Private/HoldoutCompositeSubsystem.cpp
--- a/Engine/Plugins/Experimental/Compositing/HoldoutComposite/Source/HoldoutComposite/Private/HoldoutCompositeSubsystem.cpp
+++ b/Engine/Plugins/Experimental/Compositing/HoldoutComposite/Source/HoldoutComposite/Private/HoldoutCompositeSubsystem.cpp
@@ -8,6 +8,7 @@
 
 #if WITH_EDITOR
 #include "Framework/Notifications/NotificationManager.h"
+#include "HAL/PlatformFileManager.h"
 #include "ISettingsEditorModule.h"
 #include "Misc/ConfigCacheIni.h"
 #include "Modules/ModuleManager.h"
@@ -16,6 +17,32 @@
 
 #define LOCTEXT_NAMESPACE "HoldoutCompositeSubsystem"
 
+namespace
+{
+#if WITH_EDITOR
+	void UpdateDependentPropertyInConfigFile(URendererSettings* RendererSettings, FProperty* RendererProperty)
+	{
+		FString RelativePath = RendererSettings->GetDefaultConfigFilename();
+		FString FullPath = FPaths::ConvertRelativePathToFull(RelativePath);
+
+		const bool bIsWriteable = !FPlatformFileManager::Get().GetPlatformFile().IsReadOnly(*FullPath);
+
+		if (!bIsWriteable)
+		{
+			FPlatformFileManager::Get().GetPlatformFile().SetReadOnly(*FullPath, false);
+		}
+
+		RendererSettings->UpdateSinglePropertyInConfigFile(RendererProperty, RendererSettings->GetDefaultConfigFilename());
+
+		// Restore original state for source control
+		if (!bIsWriteable)
+		{
+			FPlatformFileManager::Get().GetPlatformFile().SetReadOnly(*FullPath, true);
+		}
+	}
+#endif
+}
+
 UHoldoutCompositeSubsystem::UHoldoutCompositeSubsystem()
 {
 }
@@ -144,7 +171,7 @@ void UHoldoutCompositeSubsystem::PrimitiveHoldoutSettingsNotification(URendererS
 
 					FPropertyChangedEvent PropertyChangedEvent(Property, EPropertyChangeType::ValueSet, { RendererSettings });
 					RendererSettings->PostEditChangeProperty(PropertyChangedEvent);
-					RendererSettings->UpdateSinglePropertyInConfigFile(Property, RendererSettings->GetDefaultConfigFilename());
+					UpdateDependentPropertyInConfigFile(RendererSettings, Property);
 
 					// SupportPrimitiveAlphaHoldout requires shader recompilation, ask for a restart.
 					FModuleManager::GetModuleChecked<ISettingsEditorModule>("SettingsEditor").OnApplicationRestartRequired();
@@ -159,7 +186,7 @@ void UHoldoutCompositeSubsystem::PrimitiveHoldoutSettingsNotification(URendererS
 
 					FPropertyChangedEvent PropertyChangedEvent(Property, EPropertyChangeType::ValueSet, { RendererSettings });
 					RendererSettings->PostEditChangeProperty(PropertyChangedEvent);
-					RendererSettings->UpdateSinglePropertyInConfigFile(Property, RendererSettings->GetDefaultConfigFilename());
+					UpdateDependentPropertyInConfigFile(RendererSettings, Property);
 				}
 			}
 
diff --git a/Engine/Plugins/MovieScene/MovieRenderPipeline/Source/MovieRenderPipelineCore/Private/Graph/MovieGraphUtils.cpp b/Engine/Plugins/MovieScene/MovieRenderPipeline/Source/MovieRenderPipelineCore/Private/Graph/MovieGraphUtils.cpp
--- a/Engine/Plugins/MovieScene/MovieRenderPipeline/Source/MovieRenderPipelineCore/Private/Graph/MovieGraphUtils.cpp
+++ b/Engine/Plugins/MovieScene/MovieRenderPipeline/Source/MovieRenderPipelineCore/Private/Graph/MovieGraphUtils.cpp
@@ -9,6 +9,7 @@
 #if WITH_EDITOR
 #include "Engine/RendererSettings.h"
 #include "Framework/Notifications/NotificationManager.h"
+#include "HAL/PlatformFileManager.h"
 #include "ISettingsEditorModule.h"
 #include "Misc/ConfigCacheIni.h"
 #include "Modules/ModuleManager.h"
@@ -79,6 +80,27 @@ namespace UE::MovieGraph
 	}
 
 #if WITH_EDITOR
+	void UpdateDependentPropertyInConfigFile(URendererSettings* RendererSettings, FProperty* RendererProperty)
+	{
+		FString RelativePath = RendererSettings->GetDefaultConfigFilename();
+		FString FullPath = FPaths::ConvertRelativePathToFull(RelativePath);
+
+		const bool bIsWriteable = !FPlatformFileManager::Get().GetPlatformFile().IsReadOnly(*FullPath);
+
+		if (!bIsWriteable)
+		{
+			FPlatformFileManager::Get().GetPlatformFile().SetReadOnly(*FullPath, false);
+		}
+
+		RendererSettings->UpdateSinglePropertyInConfigFile(RendererProperty, RendererSettings->GetDefaultConfigFilename());
+
+		// Restore original state for source control
+		if (!bIsWriteable)
+		{
+			FPlatformFileManager::Get().GetPlatformFile().SetReadOnly(*FullPath, true);
+		}
+	}
+
 	void ValidateAlphaProjectSettings(const FText& InRequestingFeatureName, bool bMandatePrimitiveAlphaHoldout)
 	{
 		URendererSettings* RendererSettings = GetMutableDefault<URendererSettings>();
@@ -140,7 +162,7 @@ namespace UE::MovieGraph
 
 							FPropertyChangedEvent PropertyChangedEvent(Property, EPropertyChangeType::ValueSet, { RendererSettings });
 							RendererSettings->PostEditChangeProperty(PropertyChangedEvent);
-							RendererSettings->UpdateSinglePropertyInConfigFile(Property, RendererSettings->GetDefaultConfigFilename());
+							UpdateDependentPropertyInConfigFile(RendererSettings, Property);
 						}
 
 						if (bPrimitiveHoldoutMissing)
@@ -152,7 +174,7 @@ namespace UE::MovieGraph
 
 							FPropertyChangedEvent PropertyChangedEvent(Property, EPropertyChangeType::ValueSet, { RendererSettings });
 							RendererSettings->PostEditChangeProperty(PropertyChangedEvent);
-							RendererSettings->UpdateSinglePropertyInConfigFile(Property, RendererSettings->GetDefaultConfigFilename());
+							UpdateDependentPropertyInConfigFile(RendererSettings, Property);
 
 							// SupportPrimitiveAlphaHoldout requires shader recompilation, ask for a restart.
 							FModuleManager::GetModuleChecked<ISettingsEditorModule>("SettingsEditor").OnApplicationRestartRequired();
