MovieGraph: Burn-ins now show the correct camera data for the camera they are being rendered for. This fixes an issue where when using multi-camera support, all burn-ins showed the Primary camera's data (focal length, aperture, etc.)

#virtualized

diff --git a/Engine/Plugins/MovieScene/MovieRenderPipeline/Source/MovieRenderPipelineCore/Private/Graph/MovieGraphBlueprintLibrary.cpp b/Engine/Plugins/MovieScene/MovieRenderPipeline/Source/MovieRenderPipelineCore/Private/Graph/MovieGraphBlueprintLibrary.cpp
--- a/Engine/Plugins/MovieScene/MovieRenderPipeline/Source/MovieRenderPipelineCore/Private/Graph/MovieGraphBlueprintLibrary.cpp
+++ b/Engine/Plugins/MovieScene/MovieRenderPipeline/Source/MovieRenderPipelineCore/Private/Graph/MovieGraphBlueprintLibrary.cpp
@@ -11,6 +11,7 @@
 #include "Graph/Nodes/MovieGraphRenderLayerNode.h"
 #include "MoviePipelineBlueprintLibrary.h"
 #include "MoviePipelineUtils.h"
+#include "CineCameraActor.h"
 
 #include "Misc/Paths.h"
 #include "HAL/FileManager.h"
@@ -576,41 +577,81 @@ FFrameNumber UMovieGraphBlueprintLibrary::GetCurrentShotFrameNumber(const UMovie
 	return FFrameNumber(-1);
 }
 
-float UMovieGraphBlueprintLibrary::GetCurrentFocusDistance(const UMovieGraphPipeline* InMovieGraphPipeline)
+float UMovieGraphBlueprintLibrary::GetCurrentFocusDistance(const UMovieGraphPipeline* InMovieGraphPipeline, int32 InCameraIndex)
 {
-	if (const UCineCameraComponent* CineCameraComponent = UMoviePipelineBlueprintLibrary::Utility_GetCurrentCineCamera(InMovieGraphPipeline->GetWorld()))
+	if (InCameraIndex == INDEX_NONE)
 	{
-		return CineCameraComponent->CurrentFocusDistance;
+		if (const UCineCameraComponent* CineCameraComponent = UMoviePipelineBlueprintLibrary::Utility_GetCurrentCineCamera(InMovieGraphPipeline->GetWorld()))
+		{
+			return CineCameraComponent->CurrentFocusDistance;
+		}
+	}
+	else if (UCineCameraComponent* CurrentCameraComponent = GetCurrentCineCamera(InMovieGraphPipeline, InCameraIndex))
+	{
+		return CurrentCameraComponent->CurrentFocusDistance;
 	}
 
 	return -1.f;
 }
 
-float UMovieGraphBlueprintLibrary::GetCurrentFocalLength(const UMovieGraphPipeline* InMovieGraphPipeline)
+float UMovieGraphBlueprintLibrary::GetCurrentFocalLength(const UMovieGraphPipeline* InMovieGraphPipeline, int32 InCameraIndex)
 {
-	if (const UCineCameraComponent* CineCameraComponent = UMoviePipelineBlueprintLibrary::Utility_GetCurrentCineCamera(InMovieGraphPipeline->GetWorld()))
+	if (InCameraIndex == INDEX_NONE)
 	{
-		return CineCameraComponent->CurrentFocalLength;
+		if (const UCineCameraComponent* CineCameraComponent = UMoviePipelineBlueprintLibrary::Utility_GetCurrentCineCamera(InMovieGraphPipeline->GetWorld()))
+		{
+			return CineCameraComponent->CurrentFocalLength;
+		}
+	}
+	else if (UCineCameraComponent* CurrentCameraComponent = GetCurrentCineCamera(InMovieGraphPipeline, InCameraIndex))
+	{
+		return CurrentCameraComponent->CurrentFocalLength;
 	}
 
 	return -1.f;
 }
 
-float UMovieGraphBlueprintLibrary::GetCurrentAperture(const UMovieGraphPipeline* InMovieGraphPipeline)
+float UMovieGraphBlueprintLibrary::GetCurrentAperture(const UMovieGraphPipeline* InMovieGraphPipeline, int32 InCameraIndex)
 {
-	if (const UCineCameraComponent* CineCameraComponent = UMoviePipelineBlueprintLibrary::Utility_GetCurrentCineCamera(InMovieGraphPipeline->GetWorld()))
+	if (InCameraIndex == INDEX_NONE)
+	{
+		if (const UCineCameraComponent* CineCameraComponent = UMoviePipelineBlueprintLibrary::Utility_GetCurrentCineCamera(InMovieGraphPipeline->GetWorld()))
+		{
+			return CineCameraComponent->CurrentAperture;
+		}
+
+	}
+	else if (UCineCameraComponent* CurrentCameraComponent = GetCurrentCineCamera(InMovieGraphPipeline, InCameraIndex))
 	{
-		return CineCameraComponent->CurrentAperture;
+		return CurrentCameraComponent->CurrentAperture;
 	}
 
 	return 0.f;
 }
 
-UCineCameraComponent* UMovieGraphBlueprintLibrary::GetCurrentCineCamera(const UMovieGraphPipeline* InMovieGraphPipeline)
+UCineCameraComponent* UMovieGraphBlueprintLibrary::GetCurrentCineCamera(const UMovieGraphPipeline* InMovieGraphPipeline, int32 InCameraIndex)
 {
-	if (UCineCameraComponent* CineCameraComponent = UMoviePipelineBlueprintLibrary::Utility_GetCurrentCineCamera(InMovieGraphPipeline->GetWorld()))
+	if (InCameraIndex == INDEX_NONE)
+	{
+		if (UCineCameraComponent* CineCameraComponent = UMoviePipelineBlueprintLibrary::Utility_GetCurrentCineCamera(InMovieGraphPipeline->GetWorld()))
+		{
+			return CineCameraComponent;
+		}
+	}
+	else if (UMoviePipelineExecutorShot* CurrentShot = GetCurrentExecutorShot(InMovieGraphPipeline))
 	{
-		return CineCameraComponent;
+		// If we're not rendering all cameras, InCameraIndex is -1.
+		const bool bRenderAllCameras = InCameraIndex >= 0;
+		TArray<UE::MovieGraph::FMinimalCameraInfo> MinimalCameraInfos = InMovieGraphPipeline->GetDataSourceInstance()->GetCameraInformation(CurrentShot, bRenderAllCameras);
+		if (!ensure(MinimalCameraInfos.IsValidIndex(InCameraIndex)))
+		{
+			return nullptr;
+		}
+
+		if (const ACineCameraActor* CineCameraActor = Cast<ACineCameraActor>(MinimalCameraInfos[InCameraIndex].ViewActor))
+		{
+			return CineCameraActor->GetCineCameraComponent();
+		}
 	}
 
 	return nullptr;
diff --git a/Engine/Plugins/MovieScene/MovieRenderPipeline/Source/MovieRenderPipelineCore/Private/Graph/MovieGraphDefaultRenderer.cpp b/Engine/Plugins/MovieScene/MovieRenderPipeline/Source/MovieRenderPipelineCore/Private/Graph/MovieGraphDefaultRenderer.cpp
--- a/Engine/Plugins/MovieScene/MovieRenderPipeline/Source/MovieRenderPipelineCore/Private/Graph/MovieGraphDefaultRenderer.cpp
+++ b/Engine/Plugins/MovieScene/MovieRenderPipeline/Source/MovieRenderPipelineCore/Private/Graph/MovieGraphDefaultRenderer.cpp
@@ -140,7 +140,7 @@ void UMovieGraphDefaultRenderer::SetupRenderingPipelineForShot(UMoviePipelineExe
 				LayerData.CameraIndex = CameraIndexes[CameraIndex]; 
 				
 				// We provide the name here so that the setup functions of the renderers can create RenderLayerIdentifiers.
-				UE::MovieGraph::DefaultRenderer::FCameraInfo CameraInfo = GetCameraInfo(EvaluatedConfig, LayerData.CameraIndex);
+				UE::MovieGraph::DefaultRenderer::FCameraInfo CameraInfo = GetCameraInfo(LayerData.CameraIndex);
 				LayerData.CameraName = CameraInfo.CameraName;
 			}
 				
@@ -411,10 +411,10 @@ void UMovieGraphDefaultRenderer::FlushAsyncEngineSystems(const TObjectPtr<UMovie
 
 void UMovieGraphDefaultRenderer::GetCameraLocationsForFrame(TArray<FVector>& OutLocations, UMoviePipelineExecutorShot* InShot, bool bIncludeSidecar) const
 {
-	TArray<FMinimalViewInfo> ViewInfos = GetOwningGraph()->GetDataSourceInstance()->GetCameraInformation(InShot, bIncludeSidecar);
-	for (const FMinimalViewInfo& ViewInfo : ViewInfos)
+	TArray<UE::MovieGraph::FMinimalCameraInfo> MinimalCameraInfos = GetOwningGraph()->GetDataSourceInstance()->GetCameraInformation(InShot, bIncludeSidecar);
+	for (const UE::MovieGraph::FMinimalCameraInfo& MinimalCameraInfo : MinimalCameraInfos)
 	{
-		OutLocations.Add(ViewInfo.Location);
+		OutLocations.Add(MinimalCameraInfo.ViewInfo.Location);
 	}
 }
 
@@ -426,16 +426,10 @@ void UMovieGraphDefaultRenderer::AddOutstandingRenderTask_AnyThread(UE::Tasks::F
 	OutstandingTasks.Add(MoveTemp(InTask));
 }
 
-UE::MovieGraph::DefaultRenderer::FCameraInfo UMovieGraphDefaultRenderer::GetCameraInfo(UMovieGraphEvaluatedConfig* InConfig, const int32 InCameraIndex) const
+UE::MovieGraph::DefaultRenderer::FCameraInfo UMovieGraphDefaultRenderer::GetCameraInfo(const int32 InCameraIndex) const
 {
 	UE::MovieGraph::DefaultRenderer::FCameraInfo CameraInfo;
 	
-	const bool bIncludeCDOs = true;
-	const UMovieGraphCameraSettingNode* CameraSettingNode = InConfig->GetSettingForBranch<UMovieGraphCameraSettingNode>(UMovieGraphNode::GlobalsPinName, bIncludeCDOs);
-	if (!ensure(CameraSettingNode))
-	{
-		return CameraInfo;
-	}
 
 	UMoviePipelineExecutorShot* CurrentShot = GetOwningGraph()->GetActiveShotList()[GetOwningGraph()->GetCurrentShotIndex()];
 
@@ -453,30 +447,29 @@ UE::MovieGraph::DefaultRenderer::FCameraInfo UMovieGraphDefaultRenderer::GetCame
 		CameraInfo.CameraName = CurrentShot->GetCameraName(InCameraIndex);
 	}
 
-	APlayerController* LocalPlayerController = GetWorld()->GetFirstPlayerController();
-	if (LocalPlayerController && LocalPlayerController->PlayerCameraManager)
-	{
-		CameraInfo.ViewActor = LocalPlayerController->GetViewTarget();
-	}
 
-	TArray<FMinimalViewInfo> ViewInfos = GetOwningGraph()->GetDataSourceInstance()->GetCameraInformation(GetOwningGraph()->GetActiveShotList()[GetOwningGraph()->GetCurrentShotIndex()], CameraSettingNode->bRenderAllCameras);
-	if (!ensureAlways(ViewInfos.IsValidIndex(LocalArrayIndex)))
+	// If we're not rendering all cameras, InCameraIndex is -1.
+	const bool bRenderAllCameras = InCameraIndex >= 0;
+	TArray<UE::MovieGraph::FMinimalCameraInfo> MinimalCameraInfos = GetOwningGraph()->GetDataSourceInstance()->GetCameraInformation(GetOwningGraph()->GetActiveShotList()[GetOwningGraph()->GetCurrentShotIndex()], bRenderAllCameras);
+	if (!ensureAlways(MinimalCameraInfos.IsValidIndex(LocalArrayIndex)))
 	{
 		return CameraInfo;
 	}
 
-	CameraInfo.ViewInfo = ViewInfos[LocalArrayIndex];
+	CameraInfo.ViewInfo = MinimalCameraInfos[LocalArrayIndex].ViewInfo;
+	CameraInfo.ViewActor = MinimalCameraInfos[LocalArrayIndex].ViewActor.Get();
+
 	return CameraInfo;
 }
 
-float UMovieGraphDefaultRenderer::GetCameraOverscan(UMovieGraphEvaluatedConfig* InConfig, int32 InCameraIndex)
+float UMovieGraphDefaultRenderer::GetCameraOverscan(int32 InCameraIndex)
 {
 	if (CameraOverscanCache.Contains(InCameraIndex))
 	{
 		return CameraOverscanCache[InCameraIndex];
 	}
 
-	const float CameraOverscan = GetCameraInfo(InConfig, InCameraIndex).ViewInfo.GetOverscan();
+	const float CameraOverscan = GetCameraInfo(InCameraIndex).ViewInfo.GetOverscan();
 	CameraOverscanCache.Add(InCameraIndex, CameraOverscan);
 	return CameraOverscan;
 }
diff --git a/Engine/Plugins/MovieScene/MovieRenderPipeline/Source/MovieRenderPipelineCore/Private/Graph/MovieGraphSequenceDataSource.cpp b/Engine/Plugins/MovieScene/MovieRenderPipeline/Source/MovieRenderPipelineCore/Private/Graph/MovieGraphSequenceDataSource.cpp
--- a/Engine/Plugins/MovieScene/MovieRenderPipeline/Source/MovieRenderPipelineCore/Private/Graph/MovieGraphSequenceDataSource.cpp
+++ b/Engine/Plugins/MovieScene/MovieRenderPipeline/Source/MovieRenderPipelineCore/Private/Graph/MovieGraphSequenceDataSource.cpp
@@ -366,15 +366,21 @@ void UMovieGraphSequenceDataSource::InitializeShot(const TObjectPtr<UMoviePipeli
 	JumpDataSource(InEvalTime);
 }
 
-TArray<FMinimalViewInfo> UMovieGraphSequenceDataSource::GetCameraInformation(UMoviePipelineExecutorShot* InShot, bool bIncludeSidecar) const
+TArray<UE::MovieGraph::FMinimalCameraInfo> UMovieGraphSequenceDataSource::GetCameraInformation(UMoviePipelineExecutorShot* InShot, bool bIncludeSidecar) const
 {
-	TArray<FMinimalViewInfo> OutInfo;
+	TArray<UE::MovieGraph::FMinimalCameraInfo> OutInfo;
 
 	if (!bIncludeSidecar)
 	{
 		if (const APlayerController* LocalPlayerController = GetOwningGraph()->GetWorld()->GetFirstPlayerController())
 		{
-			OutInfo.Add(LocalPlayerController->PlayerCameraManager->GetCameraCacheView());
+			FMinimalViewInfo ViewInfo = LocalPlayerController->PlayerCameraManager->GetCameraCacheView();
+
+			UE::MovieGraph::FMinimalCameraInfo MinimalInfo;
+			MinimalInfo.ViewActor = LocalPlayerController->GetViewTarget();
+			MinimalInfo.ViewInfo = ViewInfo;
+
+			OutInfo.Add(MinimalInfo);
 		}
 	}
 	else
@@ -390,10 +396,14 @@ TArray<FMinimalViewInfo> UMovieGraphSequenceDataSource::GetCameraInformation(UMo
 				UCameraComponent* BoundCamera = MovieSceneHelpers::CameraComponentFromRuntimeObject(BoundCameras[0].Get());
 				if (BoundCamera)
 				{
+
 					FMinimalViewInfo ViewInfo;
 					BoundCamera->GetCameraView(GetWorld()->GetDeltaSeconds(), ViewInfo);
 
-					OutInfo.Add(ViewInfo);
+					UE::MovieGraph::FMinimalCameraInfo MinimalInfo;
+					MinimalInfo.ViewActor = BoundCamera->GetTypedOuter<AActor>();
+					MinimalInfo.ViewInfo = ViewInfo;
+					OutInfo.Add(MinimalInfo);
 				}
 			}
 		}
diff --git a/Engine/Plugins/MovieScene/MovieRenderPipeline/Source/MovieRenderPipelineCore/Private/Graph/Nodes/MovieGraphBurnInNode.cpp b/Engine/Plugins/MovieScene/MovieRenderPipeline/Source/MovieRenderPipelineCore/Private/Graph/Nodes/MovieGraphBurnInNode.cpp
--- a/Engine/Plugins/MovieScene/MovieRenderPipeline/Source/MovieRenderPipelineCore/Private/Graph/Nodes/MovieGraphBurnInNode.cpp
+++ b/Engine/Plugins/MovieScene/MovieRenderPipeline/Source/MovieRenderPipelineCore/Private/Graph/Nodes/MovieGraphBurnInNode.cpp
@@ -142,9 +142,9 @@ void UMovieGraphBurnInNode::FMovieGraphBurnInPass::Render(const FMovieGraphTrave
 	if (const TObjectPtr<UMovieGraphBurnInWidget> BurnInWidget = GetBurnInWidget(ParentNodeThisFrame))
 	{
 		UMovieGraphPipeline* Pipeline = Renderer->GetOwningGraph();
-		BurnInWidget->UpdateForGraph(Pipeline, Pipeline->GetTimeStepInstance()->GetCalculatedTimeData().EvaluatedConfig);
+		BurnInWidget->UpdateForGraph(Pipeline, Pipeline->GetTimeStepInstance()->GetCalculatedTimeData().EvaluatedConfig, LayerData.CameraIndex, LayerData.CameraName);
 	}
-	
+	 
 	FMovieGraphWidgetPass::Render(InFrameTraversalContext, InTimeData);
 }
 
diff --git a/Engine/Plugins/MovieScene/MovieRenderPipeline/Source/MovieRenderPipelineCore/Private/Graph/Nodes/MovieGraphWidgetRendererBaseNode.cpp b/Engine/Plugins/MovieScene/MovieRenderPipeline/Source/MovieRenderPipelineCore/Private/Graph/Nodes/MovieGraphWidgetRendererBaseNode.cpp
--- a/Engine/Plugins/MovieScene/MovieRenderPipeline/Source/MovieRenderPipelineCore/Private/Graph/Nodes/MovieGraphWidgetRendererBaseNode.cpp
+++ b/Engine/Plugins/MovieScene/MovieRenderPipeline/Source/MovieRenderPipelineCore/Private/Graph/Nodes/MovieGraphWidgetRendererBaseNode.cpp
@@ -106,7 +106,8 @@ void UMovieGraphWidgetRendererBaseNode::FMovieGraphWidgetPass::Render(const FMov
 	const UMovieGraphPipeline* Pipeline = Renderer->GetOwningGraph();
 	UMovieGraphWidgetRendererBaseNode* ParentNodeThisFrame = GetParentNode(InTimeData.EvaluatedConfig);
 
-	const float CameraOverscan = Renderer->GetCameraOverscan(InTimeData.EvaluatedConfig, LayerData.CameraIndex);
+	const float CameraOverscan = Renderer->GetCameraOverscan(LayerData.CameraIndex);
+
 	const FIntPoint OutputResolution = UMovieGraphBlueprintLibrary::GetEffectiveOutputResolution(InTimeData.EvaluatedConfig, CameraOverscan);
 	const int32 MaxResolution = GetMax2DTextureDimension();
 	if ((OutputResolution.X > MaxResolution) || (OutputResolution.Y > MaxResolution))
diff --git a/Engine/Plugins/MovieScene/MovieRenderPipeline/Source/MovieRenderPipelineCore/Public/Graph/MovieGraphBlueprintLibrary.h b/Engine/Plugins/MovieScene/MovieRenderPipeline/Source/MovieRenderPipelineCore/Public/Graph/MovieGraphBlueprintLibrary.h
--- a/Engine/Plugins/MovieScene/MovieRenderPipeline/Source/MovieRenderPipelineCore/Public/Graph/MovieGraphBlueprintLibrary.h
+++ b/Engine/Plugins/MovieScene/MovieRenderPipeline/Source/MovieRenderPipelineCore/Public/Graph/MovieGraphBlueprintLibrary.h
@@ -184,28 +184,28 @@ public:
 	* @param	InMovieGraphPipeline	- The pipeline to get the camera information from.
 	*/
 	UFUNCTION(BlueprintPure, Category = "Movie Graph")
-	static float GetCurrentFocusDistance(const UMovieGraphPipeline* InMovieGraphPipeline);
+	static float GetCurrentFocusDistance(const UMovieGraphPipeline* InMovieGraphPipeline, int32 InCameraIndex = -1);
 
 	/**
 	* Gets the focal length for the camera currently in use.
 	* @param	InMovieGraphPipeline	- The pipeline to get the camera information from.
 	*/
 	UFUNCTION(BlueprintPure, Category = "Movie Graph")
-	static float GetCurrentFocalLength(const UMovieGraphPipeline* InMovieGraphPipeline);
+	static float GetCurrentFocalLength(const UMovieGraphPipeline* InMovieGraphPipeline, int32 InCameraIndex = -1);
 	
 	/**
 	* Gets the aperture for the camera currently in use.
 	* @param	InMovieGraphPipeline	- The pipeline to get the camera information from.
 	*/
 	UFUNCTION(BlueprintPure, Category = "Movie Graph")
-	static float GetCurrentAperture(const UMovieGraphPipeline* InMovieGraphPipeline);
+	static float GetCurrentAperture(const UMovieGraphPipeline* InMovieGraphPipeline, int32 InCameraIndex = -1);
 
 	/**
 	* Gets the currently active cine camera, or nullptr if one was not found.
 	* @param	InMovieGraphPipeline	- The pipeline to get the camera from.
 	*/
 	UFUNCTION(BlueprintPure, Category = "Movie Graph")
-	static UCineCameraComponent* GetCurrentCineCamera(const UMovieGraphPipeline* InMovieGraphPipeline);
+	static UCineCameraComponent* GetCurrentCineCamera(const UMovieGraphPipeline* InMovieGraphPipeline, int32 InCameraIndex = -1);
 
 
 
diff --git a/Engine/Plugins/MovieScene/MovieRenderPipeline/Source/MovieRenderPipelineCore/Public/Graph/MovieGraphBurnInWidget.h b/Engine/Plugins/MovieScene/MovieRenderPipeline/Source/MovieRenderPipelineCore/Public/Graph/MovieGraphBurnInWidget.h
--- a/Engine/Plugins/MovieScene/MovieRenderPipeline/Source/MovieRenderPipelineCore/Public/Graph/MovieGraphBurnInWidget.h
+++ b/Engine/Plugins/MovieScene/MovieRenderPipeline/Source/MovieRenderPipelineCore/Public/Graph/MovieGraphBurnInWidget.h
@@ -1,4 +1,4 @@
-﻿// Copyright Epic Games, Inc. All Rights Reserved.
+// Copyright Epic Games, Inc. All Rights Reserved.
 
 #pragma once
 
@@ -21,5 +21,5 @@ public:
 	* @param	InEvaluatedConfig	The evaluated graph that was used to generate this output frame.
 	*/
 	UFUNCTION(BlueprintImplementableEvent)
-	void UpdateForGraph(UMovieGraphPipeline* InGraphPipeline, UMovieGraphEvaluatedConfig* InEvaluatedConfig);
+	void UpdateForGraph(UMovieGraphPipeline* InGraphPipeline, UMovieGraphEvaluatedConfig* InEvaluatedConfig, int32 InCameraIndex, const FString& CameraName);
 };
diff --git a/Engine/Plugins/MovieScene/MovieRenderPipeline/Source/MovieRenderPipelineCore/Public/Graph/MovieGraphDataTypes.h b/Engine/Plugins/MovieScene/MovieRenderPipeline/Source/MovieRenderPipelineCore/Public/Graph/MovieGraphDataTypes.h
--- a/Engine/Plugins/MovieScene/MovieRenderPipeline/Source/MovieRenderPipelineCore/Public/Graph/MovieGraphDataTypes.h
+++ b/Engine/Plugins/MovieScene/MovieRenderPipeline/Source/MovieRenderPipelineCore/Public/Graph/MovieGraphDataTypes.h
@@ -184,11 +184,31 @@ public:
 	float PrevMaxUndilatedFrameTime;
 };
 
+namespace UE::MovieGraph
+{
+	/** 
+	* When fetching multiple cameras from the Data Source, we need to know both
+	* the rendering information (ViewInfo), but also who the renderering object
+	* is, which isn't stored in the ViewInfo.
+	*/
+	struct FMinimalCameraInfo
+	{
+		FMinimalCameraInfo()
+		: ViewActor(nullptr)
+		{}
+		
+		TWeakObjectPtr<AActor> ViewActor;
+		FMinimalViewInfo ViewInfo;
+	};
+}
+
 UCLASS(BlueprintType, Abstract)
 class MOVIERENDERPIPELINECORE_API UMovieGraphRendererBase : public UObject
 {
 	GENERATED_BODY()
 public:
+
+	
 	/** Get an array of image previews that are valid this frame. */
 	UFUNCTION(BlueprintCallable, Category = "Movie Graph")
 	virtual TArray<FMovieGraphImagePreviewData> GetPreviewData() const { return TArray<FMovieGraphImagePreviewData>(); }
@@ -198,6 +218,9 @@ public:
 	virtual void TeardownRenderingPipelineForShot(UMoviePipelineExecutorShot* InShot) {}
 	virtual UE::MovieGraph::FRenderTimeStatistics* GetRenderTimeStatistics(const int32 InFrameNumber) { return nullptr; }
 	
+	/** InCameraIndex can be -1 (for primary camera when not using multi-layer rendering) or 0...n for sidecar cameras. */
+	virtual UE::MovieGraph::FMinimalCameraInfo GetMinimalCameraInfo(const int32 InCameraIndex) const { return UE::MovieGraph::FMinimalCameraInfo(); }
+
 	UMovieGraphPipeline* GetOwningGraph() const;
 };
 
@@ -249,7 +272,7 @@ public:
 	virtual void UnmuteShot(const TObjectPtr<UMoviePipelineExecutorShot>& InShot) {}
 	virtual void ExpandShot(const TObjectPtr<UMoviePipelineExecutorShot>& InShot, const int32 InLeftDeltaFrames, const int32 InLeftDeltaFramesUserPoV,
 		const int32 InRightDeltaFrames, const bool bInPrepass) {}
-	virtual TArray<FMinimalViewInfo> GetCameraInformation(UMoviePipelineExecutorShot* InShot, bool bIncludeSidecar) const { return TArray<FMinimalViewInfo>(); }
+	virtual TArray<UE::MovieGraph::FMinimalCameraInfo> GetCameraInformation(UMoviePipelineExecutorShot* InShot, bool bIncludeSidecar) const { return TArray<UE::MovieGraph::FMinimalCameraInfo>(); }
 
 	UMovieGraphPipeline* GetOwningGraph() const;
 };
diff --git a/Engine/Plugins/MovieScene/MovieRenderPipeline/Source/MovieRenderPipelineCore/Public/Graph/MovieGraphDefaultRenderer.h b/Engine/Plugins/MovieScene/MovieRenderPipeline/Source/MovieRenderPipelineCore/Public/Graph/MovieGraphDefaultRenderer.h
--- a/Engine/Plugins/MovieScene/MovieRenderPipeline/Source/MovieRenderPipelineCore/Public/Graph/MovieGraphDefaultRenderer.h
+++ b/Engine/Plugins/MovieScene/MovieRenderPipeline/Source/MovieRenderPipelineCore/Public/Graph/MovieGraphDefaultRenderer.h
@@ -270,7 +270,6 @@ public:
 	virtual void SetupRenderingPipelineForShot(UMoviePipelineExecutorShot* InShot) override;
 	virtual void TeardownRenderingPipelineForShot(UMoviePipelineExecutorShot* InShot) override;
 	virtual UE::MovieGraph::FRenderTimeStatistics* GetRenderTimeStatistics(const int32 InFrameNumber);
-
 	// ~UMovieGraphRendererBase Interface
 
 	// UObject Interface
@@ -279,7 +278,7 @@ public:
 
 	void AddOutstandingRenderTask_AnyThread(UE::Tasks::FTask InTask);
 	/** Fetches information for the given camera index. Should be "-1" when not using multi-camera rendering, or [0, n] based on the shot's sidecar camera data when using multi-camera rendering. */
-	UE::MovieGraph::DefaultRenderer::FCameraInfo GetCameraInfo(UMovieGraphEvaluatedConfig* InConfig, const int32 InCameraIndex) const;
+	UE::MovieGraph::DefaultRenderer::FCameraInfo GetCameraInfo(const int32 InCameraIndex) const;
 	void SetHasRenderedFirstViewThisFrame(bool bInValue) { bHasRenderedFirstViewThisFrame = bInValue; }
 	bool GetHasRenderedFirstViewThisFrame() const { return bHasRenderedFirstViewThisFrame; }
 
@@ -287,7 +286,7 @@ public:
 	 * Gets the overscan value for the specified camera. First checks to see if any cached overscan value exists and returning it, and if
 	 * no cached value was found, gets the camera's live overscan value and caches it.
 	 */
-	float GetCameraOverscan(UMovieGraphEvaluatedConfig* InConfig, int32 InCameraIndex);
+	float GetCameraOverscan(int32 InCameraIndex);
 
 	/** Outputs a warning message regarding animated overscan to the MRQ log if one has not already been output  */
 	void WarnAboutAnimatedOverscan(float InInitialOverscan);
diff --git a/Engine/Plugins/MovieScene/MovieRenderPipeline/Source/MovieRenderPipelineCore/Public/Graph/MovieGraphSequenceDataSource.h b/Engine/Plugins/MovieScene/MovieRenderPipeline/Source/MovieRenderPipelineCore/Public/Graph/MovieGraphSequenceDataSource.h
--- a/Engine/Plugins/MovieScene/MovieRenderPipeline/Source/MovieRenderPipelineCore/Public/Graph/MovieGraphSequenceDataSource.h
+++ b/Engine/Plugins/MovieScene/MovieRenderPipeline/Source/MovieRenderPipelineCore/Public/Graph/MovieGraphSequenceDataSource.h
@@ -54,7 +54,7 @@ public:
 	virtual void UnmuteShot(const TObjectPtr<UMoviePipelineExecutorShot>& InShot) override;
 	virtual void ExpandShot(const TObjectPtr<UMoviePipelineExecutorShot>& InShot, const int32 InLeftDeltaFrames, const int32 InLeftDeltaFramesUserPoV,
 		const int32 InRightDeltaFrames, const bool bInPrepass) override;
-	virtual TArray<FMinimalViewInfo> GetCameraInformation(UMoviePipelineExecutorShot* InShot, bool bIncludeSidecar) const override;
+	virtual TArray<UE::MovieGraph::FMinimalCameraInfo> GetCameraInformation(UMoviePipelineExecutorShot* InShot, bool bIncludeSidecar) const override;
 	virtual void InitializeShot(const TObjectPtr<UMoviePipelineExecutorShot>& InShot, const FFrameTime& InEvalTime) override;
 protected:
 	void OverrideSequencePlaybackRangeFromGlobalOutputSettings(ULevelSequence* InSequence);
diff --git a/Engine/Plugins/MovieScene/MovieRenderPipeline/Source/MovieRenderPipelineRenderPasses/Private/Graph/Renderers/MovieGraphDeferredPass.cpp b/Engine/Plugins/MovieScene/MovieRenderPipeline/Source/MovieRenderPipelineRenderPasses/Private/Graph/Renderers/MovieGraphDeferredPass.cpp
--- a/Engine/Plugins/MovieScene/MovieRenderPipeline/Source/MovieRenderPipelineRenderPasses/Private/Graph/Renderers/MovieGraphDeferredPass.cpp
+++ b/Engine/Plugins/MovieScene/MovieRenderPipeline/Source/MovieRenderPipelineRenderPasses/Private/Graph/Renderers/MovieGraphDeferredPass.cpp
@@ -150,7 +150,7 @@ void FMovieGraphDeferredPass::Render(const FMovieGraphTraversalContext& InFrameT
 	}
 
 	// ToDo: When tiling is used, this should be the size of the per-tile backbuffer
-	const float CameraOverscan = GetRenderer()->GetCameraOverscan(InTimeData.EvaluatedConfig, LayerData.CameraIndex);
+	const float CameraOverscan = GetRenderer()->GetCameraOverscan(LayerData.CameraIndex);
 	FIntPoint AccumulatorResolution = UMovieGraphBlueprintLibrary::GetEffectiveOutputResolution(InTimeData.EvaluatedConfig, CameraOverscan);
 	FIntPoint BackbufferResolution = AccumulatorResolution;
 	
@@ -179,7 +179,7 @@ void FMovieGraphDeferredPass::Render(const FMovieGraphTraversalContext& InFrameT
 		}
 		 
 		// These are the parameters of our camera 
-		UE::MovieGraph::DefaultRenderer::FCameraInfo CameraInfo = GetRenderer()->GetCameraInfo(InTimeData.EvaluatedConfig, LayerData.CameraIndex);
+		UE::MovieGraph::DefaultRenderer::FCameraInfo CameraInfo = GetRenderer()->GetCameraInfo(LayerData.CameraIndex);
 
 		CameraInfo.bAllowCameraAspectRatio = true;
 		CameraInfo.TilingParams.TileSize = BackbufferResolution;
