MoviePipeline: Automatically set "r.Nanite.VSMInvalidateOnLODDelta 1"  during renders to avoid situations where static VSM pages cache the wrong Nanite LOD data when warming up.

This behavior can be disabled by setting "MoviePipeline.EnableVTInvalidateOnNaniteLOD 0". This is intended as temporary until VSMInvalidateOnLODDelta is no longer experimental.

diff --git a/Engine/Plugins/MovieScene/MovieRenderPipeline/Source/MovieRenderPipelineCore/Private/Graph/Nodes/MovieGraphGlobalGameOverrides.cpp b/Engine/Plugins/MovieScene/MovieRenderPipeline/Source/MovieRenderPipelineCore/Private/Graph/Nodes/MovieGraphGlobalGameOverrides.cpp
--- a/Engine/Plugins/MovieScene/MovieRenderPipeline/Source/MovieRenderPipelineCore/Private/Graph/Nodes/MovieGraphGlobalGameOverrides.cpp
+++ b/Engine/Plugins/MovieScene/MovieRenderPipeline/Source/MovieRenderPipelineCore/Private/Graph/Nodes/MovieGraphGlobalGameOverrides.cpp
@@ -88,6 +88,13 @@ void UMovieGraphGlobalGameOverridesNode::BuildNewProcessCommandLineArgsImpl(TArr
 	InOutDeviceProfileCvars.AddUnique(FString::Printf(TEXT("p.Chaos.ImmPhys.MinStepTime=%d"), 0));
 	InOutDeviceProfileCvars.AddUnique(FString::Printf(TEXT("r.SkipRedundantTransformUpdate=%d"), 0));
 	InOutDeviceProfileCvars.AddUnique(FString::Printf(TEXT("p.ChaosCloth.UseTimeStepSmoothing=%d"), 0));
+	InOutDeviceProfileCvars.AddUnique(FString::Printf(TEXT("r.Water.SkipWaterInfoTextureRenderWhenWorldRenderingDisabled=%d"), 0));
+
+	IConsoleVariable* VTInvalidateCvar = IConsoleManager::Get().FindConsoleVariable(TEXT("MoviePipeline.EnableVTInvalidateOnNaniteLOD"));
+	if (VTInvalidateCvar && VTInvalidateCvar->GetInt() != 0)
+	{
+		InOutDeviceProfileCvars.AddUnique(FString::Printf(TEXT("r.Nanite.VSMInvalidateOnLODDelta=%d"), 1));
+	}
 }
 
 void UMovieGraphGlobalGameOverridesNode::PostLoad()
@@ -200,6 +207,14 @@ void UMovieGraphGlobalGameOverridesNode::ApplySettings(const bool bOverrideValue
 		// Water skips water info texture when the world's game viewport rendering is disabled so we need to prevent this from happening.
 		MOVIEPIPELINE_STORE_AND_OVERRIDE_CVAR_INT_IF_EXIST(NodeCDO->PreviousSkipWaterInfoTextureRenderWhenWorldRenderingDisabled, TEXT("r.Water.SkipWaterInfoTextureRenderWhenWorldRenderingDisabled"), 0, bOverrideValues);
 
+		// This is only a temporary cvar while it's experimental so it's not exposed to the UI, but exposed as a cvar
+		// so that users can turn it off in the event that it causes issues.
+		IConsoleVariable* VTInvalidateCvar = IConsoleManager::Get().FindConsoleVariable(TEXT("MoviePipeline.EnableVTInvalidateOnNaniteLOD"));
+		if (VTInvalidateCvar && VTInvalidateCvar->GetInt() != 0)
+		{
+			MOVIEPIPELINE_STORE_AND_OVERRIDE_CVAR_INT(PreviousNaniteVSMInvalidateOnLODDelta, TEXT("r.Nanite.VSMInvalidateOnLODDelta"), 1, bOverrideValues);
+		}
+
 	}
 	
 	// Must come after the above cvars so that if one of those cvars is also specified by the scalability level, then we
diff --git a/Engine/Plugins/MovieScene/MovieRenderPipeline/Source/MovieRenderPipelineCore/Private/MoviePipelineGameOverrideSetting.cpp b/Engine/Plugins/MovieScene/MovieRenderPipeline/Source/MovieRenderPipelineCore/Private/MoviePipelineGameOverrideSetting.cpp
--- a/Engine/Plugins/MovieScene/MovieRenderPipeline/Source/MovieRenderPipelineCore/Private/MoviePipelineGameOverrideSetting.cpp
+++ b/Engine/Plugins/MovieScene/MovieRenderPipeline/Source/MovieRenderPipelineCore/Private/MoviePipelineGameOverrideSetting.cpp
@@ -9,6 +9,12 @@
 
 #include UE_INLINE_GENERATED_CPP_BY_NAME(MoviePipelineGameOverrideSetting)
 
+static TAutoConsoleVariable<int32> CVarMoviePipelineVTNaniteAutoLOD(
+	TEXT("MoviePipeline.EnableVTInvalidateOnNaniteLOD"),
+	1,
+	TEXT("If true, the Movie Pipeline Game Overrides will automatically apply 'r.Nanite.VSMInvalidateOnLODDelta' during renders.\n"),
+	ECVF_Default);
+
 void UMoviePipelineGameOverrideSetting::SetupForPipelineImpl(UMoviePipeline* InPipeline)
 {
 	// Store the cvar values and apply the ones from this setting
@@ -153,6 +159,13 @@ void UMoviePipelineGameOverrideSetting::ApplyCVarSettings(const bool bOverrideVa
 	// Water skips water info texture when the world's game viewport rendering is disabled so we need to prevent this from happening.
 	MOVIEPIPELINE_STORE_AND_OVERRIDE_CVAR_INT_IF_EXIST(PreviousSkipWaterInfoTextureRenderWhenWorldRenderingDisabled, TEXT("r.Water.SkipWaterInfoTextureRenderWhenWorldRenderingDisabled"), 0, bOverrideValues);
 
+	// This is only a temporary cvar while it's experimental so it's not exposed to the UI, but exposed as a cvar
+	// so that users can turn it off in the event that it causes issues.
+	if (CVarMoviePipelineVTNaniteAutoLOD.GetValueOnGameThread())
+	{
+		MOVIEPIPELINE_STORE_AND_OVERRIDE_CVAR_INT(PreviousNaniteVSMInvalidateOnLODDelta, TEXT("r.Nanite.VSMInvalidateOnLODDelta"), 1, bOverrideValues);
+	}
+
 	// Must come after the above cvars so that if one of those cvars is also specified by the Scalability level, then we restore to the value in the original scalability level
 	// not the value we cached in the Cinematic level (if applied).
 	if (bCinematicQualitySettings)
@@ -257,4 +270,10 @@ void UMoviePipelineGameOverrideSetting::BuildNewProcessCommandLineArgsImpl(TArra
 	InOutDeviceProfileCvars.AddUnique(FString::Printf(TEXT("p.Chaos.ImmPhys.MinStepTime=%d"), 0));
 	InOutDeviceProfileCvars.AddUnique(FString::Printf(TEXT("r.SkipRedundantTransformUpdate=%d"), 0));
 	InOutDeviceProfileCvars.AddUnique(FString::Printf(TEXT("p.ChaosCloth.UseTimeStepSmoothing=%d"), 0));
+	InOutDeviceProfileCvars.AddUnique(FString::Printf(TEXT("r.Water.SkipWaterInfoTextureRenderWhenWorldRenderingDisabled=%d"), 0));
+
+	if (CVarMoviePipelineVTNaniteAutoLOD.GetValueOnGameThread())
+	{
+		InOutDeviceProfileCvars.AddUnique(FString::Printf(TEXT("r.Nanite.VSMInvalidateOnLODDelta=%d"), 1));
+	}
 }
diff --git a/Engine/Plugins/MovieScene/MovieRenderPipeline/Source/MovieRenderPipelineCore/Public/Graph/Nodes/MovieGraphGlobalGameOverrides.h b/Engine/Plugins/MovieScene/MovieRenderPipeline/Source/MovieRenderPipelineCore/Public/Graph/Nodes/MovieGraphGlobalGameOverrides.h
--- a/Engine/Plugins/MovieScene/MovieRenderPipeline/Source/MovieRenderPipelineCore/Public/Graph/Nodes/MovieGraphGlobalGameOverrides.h
+++ b/Engine/Plugins/MovieScene/MovieRenderPipeline/Source/MovieRenderPipelineCore/Public/Graph/Nodes/MovieGraphGlobalGameOverrides.h
@@ -219,6 +219,7 @@ private:
 	int32 PreviousSkipRedundantTransformUpdate;
 	int32 PreviousChaosClothUseTimeStepSmoothing;
 	int32 PreviousSkipWaterInfoTextureRenderWhenWorldRenderingDisabled;
+	int32 PreviousNaniteVSMInvalidateOnLODDelta;
 
 #if WITH_EDITOR
 	int32 PreviousGeoCacheStreamerShowNotification;
diff --git a/Engine/Plugins/MovieScene/MovieRenderPipeline/Source/MovieRenderPipelineCore/Public/MoviePipelineGameOverrideSetting.h b/Engine/Plugins/MovieScene/MovieRenderPipeline/Source/MovieRenderPipelineCore/Public/MoviePipelineGameOverrideSetting.h
--- a/Engine/Plugins/MovieScene/MovieRenderPipeline/Source/MovieRenderPipelineCore/Public/MoviePipelineGameOverrideSetting.h
+++ b/Engine/Plugins/MovieScene/MovieRenderPipeline/Source/MovieRenderPipelineCore/Public/MoviePipelineGameOverrideSetting.h
@@ -161,4 +161,5 @@ private:
 	int32 PreviousSkipRedundantTransformUpdate;
 	int32 PreviousChaosClothUseTimeStepSmoothing;
 	int32 PreviousSkipWaterInfoTextureRenderWhenWorldRenderingDisabled;
+	int32 PreviousNaniteVSMInvalidateOnLODDelta;
 };
\ No newline at end of file
