MoviePipeline: Fixed a regression where global nodes could not be created on the Globals branch in some cases by drag-pin creation (ie, dragging a connection off of a pin and dropping, and choosing from the selection of nodes that is presented).

This issue seems to be isolated to cases where a Globals-only node is created directly off of the Inputs or Outputs nodes. The branch that is being connected to is now determined by graph traversal rather than by simply looking at the branch restriction of the target node.

#rnx

diff --git a/Engine/Plugins/MovieScene/MovieRenderPipeline/Source/MovieRenderPipelineEditor/Private/Graph/MovieGraphSchema.cpp b/Engine/Plugins/MovieScene/MovieRenderPipeline/Source/MovieRenderPipelineEditor/Private/Graph/MovieGraphSchema.cpp
--- a/Engine/Plugins/MovieScene/MovieRenderPipeline/Source/MovieRenderPipelineEditor/Private/Graph/MovieGraphSchema.cpp
+++ b/Engine/Plugins/MovieScene/MovieRenderPipeline/Source/MovieRenderPipelineEditor/Private/Graph/MovieGraphSchema.cpp
@@ -286,10 +286,32 @@ void UMovieGraphSchema::GetGraphContextActions(FGraphContextMenuBuilder& Context
 		bool bCanAppearInMenu = true;
 		if (ContextMenuBuilder.FromPin)
 		{
-			if (const UMovieGraphNode* FromNode = UE::MovieGraph::Private::GetGraphNodeFromEdPin(ContextMenuBuilder.FromPin))
+			const UMovieGraphConfig* GraphConfig = UE::MovieGraph::Private::GetGraphFromEdPin(ContextMenuBuilder.FromPin);
+			const UMovieGraphPin* FromGraphPin = UE::MovieGraph::Private::GetGraphPinFromEdPin(ContextMenuBuilder.FromPin);
+			UMovieGraphNode* FromGraphNode = UE::MovieGraph::Private::GetGraphNodeFromEdPin(ContextMenuBuilder.FromPin);
+			
+			if (GraphConfig && FromGraphPin && FromGraphNode)
 			{
-				const bool bBranchRestrictionIsOk = (PipelineNode->GetBranchRestriction() == EMovieGraphBranchRestriction::Any) ||
-					(FromNode->GetBranchRestriction() == PipelineNode->GetBranchRestriction());
+				// Get the branch name that FromPin is on (there should only be one branch name found in this scenario)
+				constexpr bool bStopAtSubgraph = true;
+				TArray<FString> FromBranchNames = (ContextMenuBuilder.FromPin->Direction == EGPD_Input)
+					? GraphConfig->GetDownstreamBranchNames(FromGraphNode, FromGraphPin, bStopAtSubgraph)
+					: GraphConfig->GetUpstreamBranchNames(FromGraphNode, FromGraphPin, bStopAtSubgraph);
+
+				// Determine if a specific node class can be created on this branch given its branch restriction
+				bool bBranchRestrictionIsOk = true;
+				if (PipelineNode->GetBranchRestriction() == EMovieGraphBranchRestriction::Globals)
+				{
+					bBranchRestrictionIsOk = FromBranchNames.Contains(UMovieGraphNode::GlobalsPinNameString);
+				}
+				else if (PipelineNode->GetBranchRestriction() == EMovieGraphBranchRestriction::RenderLayer)
+				{
+					bBranchRestrictionIsOk = !FromBranchNames.Contains(UMovieGraphNode::GlobalsPinNameString);
+				}
+				else
+				{
+					// The branch restriction is "Any", so the node creation should be ok
+				}
 				
 				bCanAppearInMenu = bBranchRestrictionIsOk && (ContextMenuBuilder.FromPin->PinType.PinCategory == PC_Branch);
 			}
