MoviePipeline: Fixed an issue that caused duplicated graphs (which are used in scripting), evaluated graphs, and potentially other pipeline-related objects to stick around after a render finished in some scenarios. This could prevent proper GC of the PIE world, especially if Collections are used within the graph.

Also tidied up flags on a few other objects, even though this is not necessarily related to the original problem.

#rnx

diff --git a/Engine/Plugins/MovieScene/MovieRenderPipeline/Source/MovieRenderPipelineCore/Private/Graph/MovieGraphConfig.cpp b/Engine/Plugins/MovieScene/MovieRenderPipeline/Source/MovieRenderPipelineCore/Private/Graph/MovieGraphConfig.cpp
--- a/Engine/Plugins/MovieScene/MovieRenderPipeline/Source/MovieRenderPipelineCore/Private/Graph/MovieGraphConfig.cpp
+++ b/Engine/Plugins/MovieScene/MovieRenderPipeline/Source/MovieRenderPipelineCore/Private/Graph/MovieGraphConfig.cpp
@@ -19,6 +19,7 @@
 #include "MovieGraphUtils.h"
 #include "MoviePipelineQueue.h"
 #include "MovieRenderPipelineCoreModule.h"
+#include "UObject/Package.h"
 
 #define LOCTEXT_NAMESPACE "MovieGraphConfig"
 
@@ -1935,7 +1936,13 @@ UMovieGraphEvaluatedConfig* UMovieGraphConfig::CreateFlattenedGraph(const FMovie
 
 	OutError.Empty();
 
-	UMovieGraphEvaluatedConfig* NewContext = NewObject<UMovieGraphEvaluatedConfig>(this);
+	// Create the evaluated config with the transient package as the outer. The transient package is used here in order to avoid some potentially
+	// tricky GC-related issues. Evaluated configs are frequently held via TStrongObjectPtr, which means they're added to the root set. This will cause
+	// everything in the evaluated config's outer chain to not be GC'd until the strong object ptr is deleted. If there's an issue in the pipeline
+	// somewhere that causes the strong object ptr to not be cleaned up when expected, it can cause a cascade that causes multiple objects (the outers)
+	// to not be GC'd in a timely manner, leading to issues like the PIE world not being GC'd (since the pipeline has the PIE world as the outer). Avoid
+	// these (unlikely) issues altogether by using the transient package as the outer.
+	UMovieGraphEvaluatedConfig* NewContext = NewObject<UMovieGraphEvaluatedConfig>(GetTransientPackage(), NAME_None, RF_Transient);
 
 	if (OutputNode)
 	{
diff --git a/Engine/Plugins/MovieScene/MovieRenderPipeline/Source/MovieRenderPipelineCore/Private/Graph/MovieGraphCoreTimeStep.cpp b/Engine/Plugins/MovieScene/MovieRenderPipeline/Source/MovieRenderPipelineCore/Private/Graph/MovieGraphCoreTimeStep.cpp
--- a/Engine/Plugins/MovieScene/MovieRenderPipeline/Source/MovieRenderPipelineCore/Private/Graph/MovieGraphCoreTimeStep.cpp
+++ b/Engine/Plugins/MovieScene/MovieRenderPipeline/Source/MovieRenderPipelineCore/Private/Graph/MovieGraphCoreTimeStep.cpp
@@ -23,6 +23,15 @@ UMovieGraphCoreTimeStep::UMovieGraphCoreTimeStep()
 	ResetForEndOfOutputFrame();
 }
 
+void UMovieGraphCoreTimeStep::Shutdown()
+{
+	// Clear out the evaluated config during shutdown. It's a strong object ptr, so references to it should be removed to ensure everything
+	// gets cleaned up properly.
+	CurrentFrameData.EvaluatedConfig.Reset();
+	
+	Super::Shutdown();
+}
+
 void UMovieGraphCoreTimeStep::TickProducingFrames()
 {
 	int32 CurrentShotIndex = GetOwningGraph()->GetCurrentShotIndex();
diff --git a/Engine/Plugins/MovieScene/MovieRenderPipeline/Source/MovieRenderPipelineCore/Private/Graph/MovieGraphDefaultRenderer.cpp b/Engine/Plugins/MovieScene/MovieRenderPipeline/Source/MovieRenderPipelineCore/Private/Graph/MovieGraphDefaultRenderer.cpp
--- a/Engine/Plugins/MovieScene/MovieRenderPipeline/Source/MovieRenderPipelineCore/Private/Graph/MovieGraphDefaultRenderer.cpp
+++ b/Engine/Plugins/MovieScene/MovieRenderPipeline/Source/MovieRenderPipelineCore/Private/Graph/MovieGraphDefaultRenderer.cpp
@@ -250,7 +250,12 @@ void UMovieGraphDefaultRenderer::Render(const FMovieGraphTimeStepData& InTimeSte
 	// Hide the progress widget before we render anything. This allows widget captures to not include the progress bar.
 	GetOwningGraph()->SetPreviewWidgetVisible(false);
 
-	if (InTimeStepData.bIsFirstTemporalSampleForFrame)
+	const FMoviePipelineCameraCutInfo& CurrentCameraCut = GetOwningGraph()->GetActiveShotList()[GetOwningGraph()->GetCurrentShotIndex()]->ShotInfo;
+	const FMovieGraphTraversalContext CurrentTraversalContext = GetOwningGraph()->GetCurrentTraversalContext();
+
+	// Allocate a new output merger frame and determine the render passes in use. However, only do this if the pipeline state is Rendering (ie, don't
+	// do this if warm-ups are happening because they do not output to disk).
+	if (InTimeStepData.bIsFirstTemporalSampleForFrame && (CurrentCameraCut.State == EMovieRenderShotState::Rendering))
 	{
 		// If this is the first sample for this output frame, then we need to 
 		// talk to all of our render passes and ask them for what data they will
@@ -259,7 +264,7 @@ void UMovieGraphDefaultRenderer::Render(const FMovieGraphTimeStepData& InTimeSte
 
 		// Get the Traversal Context (not specific to any render pass) at the first sample. This is so
 		// we can easily fetch things that are shared between all render layers later.
-		NewOutputFrame.TraversalContext = GetOwningGraph()->GetCurrentTraversalContext();
+		NewOutputFrame.TraversalContext = CurrentTraversalContext;
 		NewOutputFrame.EvaluatedConfig = TStrongObjectPtr<UMovieGraphEvaluatedConfig>(InTimeStepData.EvaluatedConfig);
 
 		for (const TObjectPtr<UMovieGraphRenderPassNode>& RenderPass : RenderPassesInUse)
@@ -312,8 +317,7 @@ void UMovieGraphDefaultRenderer::Render(const FMovieGraphTimeStepData& InTimeSte
 	for (const TObjectPtr<UMovieGraphRenderPassNode>& RenderPass : RenderPassesInUse)
 	{
 		// Pass in a copy of the traversal context so the renderer can decide what to do with it.
-		UE::MovieGraph::FMovieGraphOutputMergerFrame& OutputFrame = GetOwningGraph()->GetOutputMerger()->GetOutputFrame_GameThread(InTimeStepData.RenderedFrameNumber);
-		RenderPass->Render(OutputFrame.TraversalContext, InTimeStepData);
+		RenderPass->Render(CurrentTraversalContext, InTimeStepData);
 	}
 
 	if (NumSceneViewsRendered > 1)
diff --git a/Engine/Plugins/MovieScene/MovieRenderPipeline/Source/MovieRenderPipelineCore/Private/Graph/MovieGraphPipeline.cpp b/Engine/Plugins/MovieScene/MovieRenderPipeline/Source/MovieRenderPipelineCore/Private/Graph/MovieGraphPipeline.cpp
--- a/Engine/Plugins/MovieScene/MovieRenderPipeline/Source/MovieRenderPipelineCore/Private/Graph/MovieGraphPipeline.cpp
+++ b/Engine/Plugins/MovieScene/MovieRenderPipeline/Source/MovieRenderPipelineCore/Private/Graph/MovieGraphPipeline.cpp
@@ -166,16 +166,15 @@ void UMovieGraphPipeline::Initialize(UMoviePipelineExecutorJob* InJob, const FMo
 
 void UMovieGraphPipeline::DuplicateJobAndConfiguration()
 {
-	// Contains all duplicated graphs. Maps the original graph (key) to the duplicated graph (value).
-	TMap<UMovieGraphConfig*, UMovieGraphConfig*> DuplicatedGraphs;
-	
 	// Scripting is likely to want to modify both the job (to set variable assignments) and 
 	// the configuration itself (to add nodes, or override an output directory, etc. If scripts
 	// directly modified the job/configuration it would lead to a lot of unintentional mutation
 	// of assets and queues, so we instead choose to duplicate the job and configurations for
 	// the duration of a render, and all of the Graph Pipeline code should look at the duplicates.
-	FObjectDuplicationParameters JobDuplicationParms = FObjectDuplicationParameters(CurrentJob, GetTransientPackage());
+	FObjectDuplicationParameters JobDuplicationParms = FObjectDuplicationParameters(CurrentJob, this);
 	JobDuplicationParms.DestName = FName(FString::Format(TEXT("{0}_Duplicate"), {CurrentJob->GetFName().ToString()}));
+	JobDuplicationParms.FlagMask = RF_AllFlags & ~RF_Transactional;
+	JobDuplicationParms.ApplyFlags = RF_Transient;
 	CurrentJobDuplicate = Cast<UMoviePipelineExecutorJob>(StaticDuplicateObjectEx(JobDuplicationParms));
 
 	// The duplicate job is a mix of duplicated objects and non-duplicated objects. Objects that 
@@ -238,7 +237,7 @@ void UMovieGraphPipeline::DuplicateJobAndConfiguration()
 
 }
 
-UMovieGraphConfig* UMovieGraphPipeline::DuplicateConfigRecursive(UMovieGraphConfig* InGraphToDuplicate, TMap<UMovieGraphConfig*, UMovieGraphConfig*>& OutDuplicatedGraphs)
+UMovieGraphConfig* UMovieGraphPipeline::DuplicateConfigRecursive(UMovieGraphConfig* InGraphToDuplicate, TMap<TObjectPtr<UMovieGraphConfig>, TObjectPtr<UMovieGraphConfig>>& OutDuplicatedGraphs)
 {
 	UMovieGraphConfig* DuplicateConfig;
 
@@ -252,6 +251,8 @@ UMovieGraphConfig* UMovieGraphPipeline::DuplicateConfigRecursive(UMovieGraphConf
 		// The transient package is used because graphs don't belong to the executor job usually (they belong to an asset package)
 		FObjectDuplicationParameters GraphDuplicationParams(InGraphToDuplicate, GetTransientPackage());
 		GraphDuplicationParams.DestName = FName(FString::Format(TEXT("{0}_Duplicate"), {InGraphToDuplicate->GetFName().ToString()}));
+		GraphDuplicationParams.FlagMask = RF_AllFlags & ~(RF_Standalone | RF_Transactional);
+		GraphDuplicationParams.ApplyFlags = RF_Transient;
 		DuplicateConfig = Cast<UMovieGraphConfig>(StaticDuplicateObjectEx(GraphDuplicationParams));
 		
 		OutDuplicatedGraphs.Add(InGraphToDuplicate, DuplicateConfig);
@@ -273,7 +274,7 @@ UMovieGraphConfig* UMovieGraphPipeline::DuplicateConfigRecursive(UMovieGraphConf
 			// to prevent recursion. Checking the key ensures that we only duplicate if this graph has never been encountered. Checking the value
 			// ensures that we don't re-duplicate a graph that has already been duplicated (the subgraph node was already updated).
 			bool bHasBeenDuplicated = false;
-			for (const TPair<UMovieGraphConfig*, UMovieGraphConfig*>& DuplicateMapping : OutDuplicatedGraphs)
+			for (const TPair<TObjectPtr<UMovieGraphConfig>, TObjectPtr<UMovieGraphConfig>>& DuplicateMapping : OutDuplicatedGraphs)
 			{
 				if ((DuplicateMapping.Key == SubgraphConfig) || (DuplicateMapping.Value == SubgraphConfig))
 				{
@@ -289,7 +290,7 @@ UMovieGraphConfig* UMovieGraphPipeline::DuplicateConfigRecursive(UMovieGraphConf
 
 			// Update the subgraph node to use the duplicated graph. This should always be done, even if the graph was already duplicated (since
 			// a graph can be included as a subgraph in multiple locations).
-			if (UMovieGraphConfig** DuplicatedGraph = OutDuplicatedGraphs.Find(SubgraphConfig))
+			if (const TObjectPtr<UMovieGraphConfig>* DuplicatedGraph = OutDuplicatedGraphs.Find(SubgraphConfig))
 			{
 				SubgraphNode->SetSubGraphAsset(*DuplicatedGraph);
 			}
@@ -300,18 +301,18 @@ UMovieGraphConfig* UMovieGraphPipeline::DuplicateConfigRecursive(UMovieGraphConf
 }
 
 template <typename JobType>
-void UMovieGraphPipeline::UpdateVariableAssignmentsHelper(JobType* InTargetJob, TMap<UMovieGraphConfig*, UMovieGraphConfig*>& InOriginalToDuplicateGraphMap)
+void UMovieGraphPipeline::UpdateVariableAssignmentsHelper(JobType* InTargetJob, TMap<TObjectPtr<UMovieGraphConfig>, TObjectPtr<UMovieGraphConfig>>& InOriginalToDuplicateGraphMap)
 {
 	// Remaps the provided variable assignments to point to the the duplicated graphs.
 	auto UpdateVariableAssignments = [&InOriginalToDuplicateGraphMap](TArray<TObjectPtr<UMovieJobVariableAssignmentContainer>>& InVariableAssignments)
 	{
 		for (const TObjectPtr<UMovieJobVariableAssignmentContainer>& VariableAssignment : InVariableAssignments)
 		{
-			for (const TPair<UMovieGraphConfig*, UMovieGraphConfig*>& GraphMapping : InOriginalToDuplicateGraphMap)
+			for (const TPair<TObjectPtr<UMovieGraphConfig>, TObjectPtr<UMovieGraphConfig>>& GraphMapping : InOriginalToDuplicateGraphMap)
 			{
-				if (VariableAssignment->GetGraphConfig() == GraphMapping.Key)
+				if (VariableAssignment->GetGraphConfig().LoadSynchronous() == GraphMapping.Key.Get())
 				{
-					VariableAssignment->SetGraphConfig(GraphMapping.Value);
+					VariableAssignment->SetGraphConfig(MakeSoftObjectPtr(GraphMapping.Value.Get()));
 					break;
 				}
 			}
@@ -1325,6 +1326,10 @@ void UMovieGraphPipeline::ShutdownImpl(bool bIsError)
 		constexpr bool bForceFinish = true;
 		TickPostFinalizeExport(bForceFinish);
 	}
+
+	// Duplicated graphs are part of the transient package, so by default they'll stick around after the pipeline is destroyed. To prevent
+	// the graphs from hanging onto references after the render is finished, remove them during pipeline shutdown (so they are GC'd).
+	DuplicatedGraphs.Empty();
 }
 
 void UMovieGraphPipeline::TransitionToState(const EMovieRenderPipelineState InNewState)
@@ -1759,6 +1764,15 @@ bool UMovieGraphPipeline::IsPostShotCallbackNeeded() const
 	return bAnyScriptNeedsCallbacks;
 }
 
+void UMovieGraphPipeline::BeginDestroy()
+{
+	// Ask the output merger to clean up. It can hold strong obj pointers, and those should not stick around after the pipeline shuts down. Generally
+	// this shouldn't be needed, but could be necessary if pending frames were not flushed for some reason.
+	GetOutputMerger()->AbandonOutstandingWork();
+	
+	Super::BeginDestroy();
+}
+
 void UMovieGraphPipeline::ExecutePreJobScripts()
 {
 	for (UMovieGraphScriptBase* Script : CurrentScriptInstances)
diff --git a/Engine/Plugins/MovieScene/MovieRenderPipeline/Source/MovieRenderPipelineCore/Public/Graph/MovieGraphCoreTimeStep.h b/Engine/Plugins/MovieScene/MovieRenderPipeline/Source/MovieRenderPipelineCore/Public/Graph/MovieGraphCoreTimeStep.h
--- a/Engine/Plugins/MovieScene/MovieRenderPipeline/Source/MovieRenderPipelineCore/Public/Graph/MovieGraphCoreTimeStep.h
+++ b/Engine/Plugins/MovieScene/MovieRenderPipeline/Source/MovieRenderPipelineCore/Public/Graph/MovieGraphCoreTimeStep.h
@@ -27,6 +27,7 @@ public:
 	UMovieGraphCoreTimeStep();
 
 	// UMovieGraphTimeStepBase Interface
+	virtual void Shutdown() override;
 	virtual void TickProducingFrames() override;
 	virtual FMovieGraphTimeStepData GetCalculatedTimeData() const override { return CurrentTimeStepData; }
 	virtual bool IsExpansionForTSRequired(const TObjectPtr<UMovieGraphEvaluatedConfig>& InConfig) const override;
diff --git a/Engine/Plugins/MovieScene/MovieRenderPipeline/Source/MovieRenderPipelineCore/Public/Graph/MovieGraphPipeline.h b/Engine/Plugins/MovieScene/MovieRenderPipeline/Source/MovieRenderPipelineCore/Public/Graph/MovieGraphPipeline.h
--- a/Engine/Plugins/MovieScene/MovieRenderPipeline/Source/MovieRenderPipelineCore/Public/Graph/MovieGraphPipeline.h
+++ b/Engine/Plugins/MovieScene/MovieRenderPipeline/Source/MovieRenderPipelineCore/Public/Graph/MovieGraphPipeline.h
@@ -175,11 +175,11 @@ protected:
 	 * Helps duplicate graph configs. Prevents re-duplications, duplicates sub-graphs (updates subgraph nodes accordingly), and potentially more.
 	 * Returns the duplicated graph.
 	 */
-	UMovieGraphConfig* DuplicateConfigRecursive(UMovieGraphConfig* InGraphToDuplicate, TMap<UMovieGraphConfig*, UMovieGraphConfig*>& OutDuplicatedGraphs);
+	UMovieGraphConfig* DuplicateConfigRecursive(UMovieGraphConfig* InGraphToDuplicate, TMap<TObjectPtr<UMovieGraphConfig>, TObjectPtr<UMovieGraphConfig>>& OutDuplicatedGraphs);
 
 	/** Helps update variable assignments on the provided job to use duplicated graphs (reflected in the original-to-duplicate graph mapping). */
 	template<typename JobType>
-	void UpdateVariableAssignmentsHelper(JobType* InTargetJob, TMap<UMovieGraphConfig*, UMovieGraphConfig*>& InOriginalToDuplicateGraphMap);
+	void UpdateVariableAssignmentsHelper(JobType* InTargetJob, TMap<TObjectPtr<UMovieGraphConfig>, TObjectPtr<UMovieGraphConfig>>& InOriginalToDuplicateGraphMap);
 
 	// Update our data source to isolate the shot we're currently working on, so that expanded shots don't interfere with each other.
 	virtual void SetSoloShot(const TObjectPtr<UMoviePipelineExecutorShot>& InShot);
@@ -201,6 +201,10 @@ protected:
 	virtual EMovieRenderPipelineState GetPipelineStateImpl() const override { return PipelineState; }
 	virtual bool IsPostShotCallbackNeeded() const override;
 	// ~UMoviePipelineBase Interface
+	
+	// UObject Interface
+	virtual void BeginDestroy() override;
+	// ~UObject Interface
 
 protected:
 	/** Time step instances for each shot, where the index into the array corresponds to the shot index. */
@@ -248,7 +252,7 @@ protected:
 	TObjectPtr<UMoviePipelineExecutorJob> CurrentJob;
 
 	/**
-	* This is the duplicated job, parented to the Transient package. The shots inside
+	* This is the duplicated job, parented to the pipeline. The shots inside
 	* have been duplicated as well, and their graph configurations duplicated. Graph
 	* configurations are assets and scripting may want to modify them, or it may want
 	* to modify the variables in a job, so we have to duplicate both to allow a cleanly
@@ -257,6 +261,13 @@ protected:
 	UPROPERTY(Transient)
 	TObjectPtr<UMoviePipelineExecutorJob> CurrentJobDuplicate;
 
+	/**
+	 * Contains all duplicated graphs (for scripting purposes). Maps the original graph (key) to the duplicated graph (value).
+	 * Duplicated graphs are part of the transient package.
+	 */
+	UPROPERTY(Transient)
+	TMap<TObjectPtr<UMovieGraphConfig>, TObjectPtr<UMovieGraphConfig>> DuplicatedGraphs;
+
 	UPROPERTY(Transient)
 	TArray<TObjectPtr<UMoviePipelineExecutorShot>> ActiveShotList;
 
diff --git a/Engine/Plugins/MovieScene/MovieRenderPipeline/Source/MovieRenderPipelineEditor/Private/MoviePipelinePIEExecutor.cpp b/Engine/Plugins/MovieScene/MovieRenderPipeline/Source/MovieRenderPipelineEditor/Private/MoviePipelinePIEExecutor.cpp
--- a/Engine/Plugins/MovieScene/MovieRenderPipeline/Source/MovieRenderPipelineEditor/Private/MoviePipelinePIEExecutor.cpp
+++ b/Engine/Plugins/MovieScene/MovieRenderPipeline/Source/MovieRenderPipelineEditor/Private/MoviePipelinePIEExecutor.cpp
@@ -231,7 +231,7 @@ void UMoviePipelinePIEExecutor::OnPIEStartupFinished(bool)
 	}
 
 	// This Pipeline belongs to the world being created so that they have context for things they execute.
-	ActiveMoviePipeline = NewObject<UMoviePipelineBase>(ExecutingWorld, PipelineClass);
+	ActiveMoviePipeline = NewObject<UMoviePipelineBase>(ExecutingWorld, PipelineClass, NAME_None, RF_Transient);
 
 	UMoviePipeline* PipelineAsLegacy = Cast<UMoviePipeline>(ActiveMoviePipeline);
 
