MoviePipeline: Added a toast notification for when you attempted to convert a Movie Render Queue job into a Graph configuration, but the graph asset specified by the project settings could not be loaded. Without this warning conversion silently fails and the feature looks broken.

Tested:
Setting the Project Settings > Movie Render Graph > Default Graph to null (Clear the entry), then try to convert an MRQ job to a Graph job.

diff --git a/Engine/Plugins/MovieScene/MovieRenderPipeline/Source/MovieRenderPipelineCore/Private/Graph/MovieGraphRenderLayerSubsystem.cpp b/Engine/Plugins/MovieScene/MovieRenderPipeline/Source/MovieRenderPipelineCore/Private/Graph/MovieGraphRenderLayerSubsystem.cpp
--- a/Engine/Plugins/MovieScene/MovieRenderPipeline/Source/MovieRenderPipelineCore/Private/Graph/MovieGraphRenderLayerSubsystem.cpp
+++ b/Engine/Plugins/MovieScene/MovieRenderPipeline/Source/MovieRenderPipelineCore/Private/Graph/MovieGraphRenderLayerSubsystem.cpp
@@ -685,13 +685,11 @@ TArray<TSharedRef<SWidget>> UMovieGraphConditionGroupQuery_Actor::GetWidgets()
 				// that one in particular won't be added.
 				FNotificationInfo Info(LOCTEXT("TransientActorsUnsupported_Notification", "Actor Conditions do not support Spawnable (Transient) actors"));
 				Info.SubText = LOCTEXT("TransientActorsUnsupported_NotificationSubtext", "Use the \"Actor Name\" Condition to add Spawnable actors to Collections.");
-				Info.Image = FAppStyle::GetBrush(TEXT("MessageLog.Warning"));
 				Info.Image = FAppStyle::GetBrush(TEXT("Icons.Warning"));
 
 				//Set a default expire duration
 				Info.ExpireDuration = 5.0f;
 
-				//And call Add Notification, this is pretty much it!
 				FSlateNotificationManager::Get().AddNotification(Info);
 			}
 
diff --git a/Engine/Plugins/MovieScene/MovieRenderPipeline/Source/MovieRenderPipelineEditor/Private/Widgets/SMoviePipelineQueueEditor.cpp b/Engine/Plugins/MovieScene/MovieRenderPipeline/Source/MovieRenderPipelineEditor/Private/Widgets/SMoviePipelineQueueEditor.cpp
--- a/Engine/Plugins/MovieScene/MovieRenderPipeline/Source/MovieRenderPipelineEditor/Private/Widgets/SMoviePipelineQueueEditor.cpp
+++ b/Engine/Plugins/MovieScene/MovieRenderPipeline/Source/MovieRenderPipelineEditor/Private/Widgets/SMoviePipelineQueueEditor.cpp
@@ -21,6 +21,8 @@
 #include "Widgets/Input/SCheckBox.h"
 #include "Widgets/Input/SButton.h"
 #include "Widgets/Notifications/SProgressBar.h"
+#include "Widgets/Notifications/SNotificationList.h"
+#include "Framework/Notifications/NotificationManager.h"
 #include "Widgets/Views/STreeView.h"
 #include "Widgets/Layout/SWidgetSwitcher.h"
 #include "Widgets/Views/SExpanderArrow.h"
@@ -1640,6 +1642,15 @@ void SMoviePipelineQueueEditor::AssignDefaultGraphPresetToJob(UMoviePipelineExec
 	{
 		InJob->SetGraphPreset(DefaultGraph);
 	}
+	else
+	{
+		FNotificationInfo Info(LOCTEXT("ConvertJobToGraphConfig_InvalidGraphNotification", "Unable to Convert Job"));
+		Info.SubText = LOCTEXT("ConvertJobToGraphConfig_InvalidGraphNotificationSubtext", "The Graph Asset specified in Project Settings (Movie Render Pipeline > Default Graph) could not be loaded.");
+		Info.Image = FAppStyle::GetBrush(TEXT("Icons.Warning"));
+		Info.ExpireDuration = 5.0f;
+
+		FSlateNotificationManager::Get().AddNotification(Info);
+	}
 }
 
 UE_ENABLE_OPTIMIZATION_SHIP
