MoviePipeline: Fixed issue that caused some lists within Collection node conditions to not refresh properly when undoing/redoing list operations.

Also fixed a similar issue in the Modifier node.

#rnx

diff --git a/Engine/Plugins/MovieScene/MovieRenderPipeline/Source/MovieRenderPipelineCore/Private/Graph/MovieGraphRenderLayerSubsystem.cpp b/Engine/Plugins/MovieScene/MovieRenderPipeline/Source/MovieRenderPipelineCore/Private/Graph/MovieGraphRenderLayerSubsystem.cpp
--- a/Engine/Plugins/MovieScene/MovieRenderPipeline/Source/MovieRenderPipelineCore/Private/Graph/MovieGraphRenderLayerSubsystem.cpp
+++ b/Engine/Plugins/MovieScene/MovieRenderPipeline/Source/MovieRenderPipelineCore/Private/Graph/MovieGraphRenderLayerSubsystem.cpp
@@ -629,11 +629,7 @@ TArray<TSharedRef<SWidget>> UMovieGraphConditionGroupQuery_Actor::GetWidgets()
 	TArray<TSharedRef<SWidget>> Widgets;
 
 	// Create the data source for the list view
-	ListDataSource.Empty();
-	for (TSoftObjectPtr<AActor>& Actor : ActorsToMatch)
-	{
-		ListDataSource.Add(MakeShared<TSoftObjectPtr<AActor>>(Actor));
-	}
+	RefreshListDataSource();
 
 	auto GetValidActorsFromOperation = [](TSharedPtr<FDragDropOperation> InOperation, TArray<AActor*>& OutActors, bool& bHadTransient) {
 
@@ -718,6 +714,7 @@ TArray<TSharedRef<SWidget>> UMovieGraphConditionGroupQuery_Actor::GetWidgets()
 				
 				RemoveActors(ActorsToRemove);
 			})
+			.OnRefreshDataSourceRequested_UObject(this, &UMovieGraphConditionGroupQuery_Actor::RefreshListDataSource)
 		]
 	);
 
@@ -925,7 +922,16 @@ void UMovieGraphConditionGroupQuery_Actor::RemoveActors(const TArray<TSoftObject
 	
 	ActorsList->Refresh();
 }
-		
+
+void UMovieGraphConditionGroupQuery_Actor::RefreshListDataSource()
+{
+	ListDataSource.Empty();
+	for (TSoftObjectPtr<AActor>& Actor : ActorsToMatch)
+	{
+		ListDataSource.Add(MakeShared<TSoftObjectPtr<AActor>>(Actor));
+	}
+}
+
 FName UMovieGraphConditionGroupQuery_Actor::FActorSelectionColumn::GetID()
 {
 	static const FName ColumnId = FName("ActorSelection");
@@ -1855,11 +1861,7 @@ TArray<TSharedRef<SWidget>> UMovieGraphConditionGroupQuery_Sublevel::GetWidgets(
 	TArray<TSharedRef<SWidget>> Widgets;
 
 	// Create the data source for the list view
-	ListDataSource.Empty();
-	for (TSoftObjectPtr<UWorld>& Sublevel : Sublevels)
-	{
-		ListDataSource.Add(MakeShared<TSoftObjectPtr<UWorld>>(Sublevel));
-	}
+	RefreshListDataSource();
 
 	Widgets.Add(
 		SNew(SDropTarget)
@@ -1942,6 +1944,7 @@ TArray<TSharedRef<SWidget>> UMovieGraphConditionGroupQuery_Sublevel::GetWidgets(
 				constexpr bool bUpdateSources = true;
 				RefreshLevelPicker.ExecuteIfBound(bUpdateSources);
 			})
+			.OnRefreshDataSourceRequested_UObject(this, &UMovieGraphConditionGroupQuery_Sublevel::RefreshListDataSource)
 		]
 	);
 
@@ -2054,6 +2057,15 @@ void UMovieGraphConditionGroupQuery_Sublevel::AddLevels(const TArray<UWorld*>& I
 	constexpr bool bUpdateSources = false;
 	RefreshLevelPicker.ExecuteIfBound(bUpdateSources);
 }
+
+void UMovieGraphConditionGroupQuery_Sublevel::RefreshListDataSource()
+{
+	ListDataSource.Empty();
+	for (TSoftObjectPtr<UWorld>& Sublevel : Sublevels)
+	{
+		ListDataSource.Add(MakeShared<TSoftObjectPtr<UWorld>>(Sublevel));
+	}
+}
 #endif	// WITH_EDITOR
 
 void UMovieGraphConditionGroupQuery_ActorLayer::Evaluate(const TArray<AActor*>& InActorsToQuery, const UWorld* InWorld, TSet<AActor*>& OutMatchingActors) const
@@ -2238,11 +2250,7 @@ TArray<TSharedRef<SWidget>> UMovieGraphConditionGroupQuery_DataLayer::GetWidgets
 	TArray<TSharedRef<SWidget>> Widgets;
 
 	// Create the data source for the list view
-	ListDataSource.Empty();
-	for (TSoftObjectPtr<UDataLayerAsset>& DataLayer : DataLayers)
-	{
-		ListDataSource.Add(MakeShared<TSoftObjectPtr<UDataLayerAsset>>(DataLayer));
-	}
+	RefreshListDataSource();
 
 	Widgets.Add(
 		SNew(SDropTarget)
@@ -2309,6 +2317,7 @@ TArray<TSharedRef<SWidget>> UMovieGraphConditionGroupQuery_DataLayer::GetWidgets
 				constexpr bool bUpdateSources = true;
 				RefreshDataLayerPicker.ExecuteIfBound(bUpdateSources);
 			})
+			.OnRefreshDataSourceRequested_UObject(this, &UMovieGraphConditionGroupQuery_DataLayer::RefreshListDataSource)
 		]
 	);
 
@@ -2410,6 +2419,15 @@ void UMovieGraphConditionGroupQuery_DataLayer::AddDataLayers(const TArray<const
 	constexpr bool bUpdateSources = false;
 	RefreshDataLayerPicker.ExecuteIfBound(bUpdateSources);
 }
+
+void UMovieGraphConditionGroupQuery_DataLayer::RefreshListDataSource()
+{
+	ListDataSource.Empty();
+	for (TSoftObjectPtr<UDataLayerAsset>& DataLayer : DataLayers)
+	{
+		ListDataSource.Add(MakeShared<TSoftObjectPtr<UDataLayerAsset>>(DataLayer));
+	}
+}
 #endif	// WITH_EDITOR
 
 void UMovieGraphConditionGroupQuery_IsSpawnable::Evaluate(const TArray<AActor*>& InActorsToQuery, const UWorld* InWorld, TSet<AActor*>& OutMatchingActors) const
diff --git a/Engine/Plugins/MovieScene/MovieRenderPipeline/Source/MovieRenderPipelineCore/Public/Graph/MovieGraphRenderLayerSubsystem.h b/Engine/Plugins/MovieScene/MovieRenderPipeline/Source/MovieRenderPipelineCore/Public/Graph/MovieGraphRenderLayerSubsystem.h
--- a/Engine/Plugins/MovieScene/MovieRenderPipeline/Source/MovieRenderPipelineCore/Public/Graph/MovieGraphRenderLayerSubsystem.h
+++ b/Engine/Plugins/MovieScene/MovieRenderPipeline/Source/MovieRenderPipelineCore/Public/Graph/MovieGraphRenderLayerSubsystem.h
@@ -171,6 +171,9 @@ private:
 
 	/** Removes the provided actors from the query, updating the UI as needed. */
 	void RemoveActors(const TArray<TSoftObjectPtr<AActor>>& InActors);
+
+	/** Refreshes the list's data source to reflect the data model. */
+	void RefreshListDataSource();
 #endif
 
 public:
@@ -454,6 +457,9 @@ private:
 	/** Adds the provided levels to the query, updating the UI as needed. Calls InOnAddFinished when done. */
 	void AddLevels(const TArray<UWorld*>& InLevels, const FMovieGraphConditionGroupQueryContentsChanged& InOnAddFinished);
 
+	/** Refreshes the list's data source to reflect the data model. */
+	void RefreshListDataSource();
+
 	/** Displays the names of sublevels which have been chosen. */
 	TSharedPtr<SMovieGraphSimpleList<TSharedPtr<TSoftObjectPtr<UWorld>>>> SublevelsList;
 
@@ -535,6 +541,9 @@ private:
 	/** Adds the provided data layers to the query, updating the UI as needed. Calls InOnAddFinished when done. */
 	void AddDataLayers(const TArray<const UDataLayerAsset*>& InDataLayers, const FMovieGraphConditionGroupQueryContentsChanged& InOnAddFinished);
 
+	/** Refreshes the list's data source to reflect the data model. */
+	void RefreshListDataSource();
+
 	/** Displays the layers which have been chosen. */
 	TSharedPtr<SMovieGraphSimpleList<TSharedPtr<TSoftObjectPtr<UDataLayerAsset>>>> DataLayersList;
 	
diff --git a/Engine/Plugins/MovieScene/MovieRenderPipeline/Source/MovieRenderPipelineCore/Public/Graph/MovieGraphSharedWidgets.h b/Engine/Plugins/MovieScene/MovieRenderPipeline/Source/MovieRenderPipelineCore/Public/Graph/MovieGraphSharedWidgets.h
--- a/Engine/Plugins/MovieScene/MovieRenderPipeline/Source/MovieRenderPipelineCore/Public/Graph/MovieGraphSharedWidgets.h
+++ b/Engine/Plugins/MovieScene/MovieRenderPipeline/Source/MovieRenderPipelineCore/Public/Graph/MovieGraphSharedWidgets.h
@@ -2,7 +2,11 @@
 
 #pragma once
 
+// Don't even declare these widgets unless this is an editor build. This allows us to inherit from FSelfRegisteringEditorUndoClient.
+#if WITH_EDITOR
+
 #include "CoreMinimal.h"
+#include "EditorUndoClient.h"
 #include "Framework/Application/SlateApplication.h"
 #include "Framework/Commands/GenericCommands.h"
 #include "Framework/MultiBox/MultiBoxBuilder.h"
@@ -24,12 +28,13 @@
  * list indicating how many items are in the list.
  */
 template<typename ListType>
-class SMovieGraphSimpleList final : public SCompoundWidget
+class SMovieGraphSimpleList final : public SCompoundWidget, public FSelfRegisteringEditorUndoClient
 {
 public:
 	DECLARE_DELEGATE_RetVal_OneParam(const FSlateBrush*, FGetRowIcon, ListType);
 	DECLARE_DELEGATE_RetVal_OneParam(FText, FGetRowText, ListType);
 	DECLARE_DELEGATE_OneParam(FOnDelete, TArray<ListType>);
+	DECLARE_DELEGATE(FOnRefreshDataSourceRequested);
 	DECLARE_DELEGATE_RetVal_OneParam(bool, FGetRowEnableState, ListType);
 	DECLARE_DELEGATE_TwoParams(FSetRowEnableState, ListType, bool);
 	
@@ -41,6 +46,12 @@ public:
 		/** The source of data that the list will display. */
 		SLATE_ATTRIBUTE(TArray<ListType>*, DataSource)
 
+		/**
+		 * If the data source provided to the list is not the data model being modified by the transaction system, then this event
+		 * should be provided a delegate that refreshes the data source. This delegate will be called when the list needs to be refreshed.
+		 */
+		SLATE_EVENT(FOnRefreshDataSourceRequested, OnRefreshDataSourceRequested)
+
 		/** The name of the data type that will be shown in the summary row. */
 		SLATE_ATTRIBUTE(FText, DataType)
 
@@ -72,6 +83,7 @@ public:
 	void Construct(const FArguments& InArgs)
 	{
 		DataSource = InArgs._DataSource.Get();
+		OnRefreshDataSourceRequested = InArgs._OnRefreshDataSourceRequested;
 		DataType = InArgs._DataType.Get();
 		DataTypePlural = InArgs._DataTypePlural.Get();
 		SelectionMode = InArgs._SelectionMode.Get();
@@ -136,6 +148,7 @@ public:
 	{
 		if (ListView)
 		{
+			OnRefreshDataSourceRequested.ExecuteIfBound();
 			ListView->RequestListRefresh();
 		}
 	}
@@ -266,12 +279,25 @@ private:
 			];
 	}
 
+	//~ Begin FEditorUndoClient Interface
+	virtual void PostUndo(bool bSuccess) override
+	{
+		Refresh();
+	}
+	
+	virtual void PostRedo(bool bSuccess) override
+	{
+		Refresh();
+	}
+	//~ End FEditorUndoClient Interface
+
 private:
 	TSharedPtr<SListView<ListType>> ListView;
 	TSharedPtr<FUICommandList> CommandList;
 	FGetRowIcon OnGetRowIcon;
 	FGetRowText OnGetRowText;
 	FOnDelete OnDelete;
+	FOnRefreshDataSourceRequested OnRefreshDataSourceRequested;
 	bool bShowEnableDisable = false;
 	FGetRowEnableState OnGetRowEnableState;
 	FSetRowEnableState OnSetRowEnableState;
@@ -448,4 +474,6 @@ private:
 	FText DataSourceEmptyMessage;
 };
 
-#undef LOCTEXT_NAMESPACE
\ No newline at end of file
+#undef LOCTEXT_NAMESPACE
+
+#endif	// WITH_EDITOR
\ No newline at end of file
diff --git a/Engine/Plugins/MovieScene/MovieRenderPipeline/Source/MovieRenderPipelineEditor/Private/Customizations/Graph/MovieGraphModifiersCustomization.cpp b/Engine/Plugins/MovieScene/MovieRenderPipeline/Source/MovieRenderPipelineEditor/Private/Customizations/Graph/MovieGraphModifiersCustomization.cpp
--- a/Engine/Plugins/MovieScene/MovieRenderPipeline/Source/MovieRenderPipelineEditor/Private/Customizations/Graph/MovieGraphModifiersCustomization.cpp
+++ b/Engine/Plugins/MovieScene/MovieRenderPipeline/Source/MovieRenderPipelineEditor/Private/Customizations/Graph/MovieGraphModifiersCustomization.cpp
@@ -111,22 +111,15 @@ TSharedRef<IDetailCustomization> FMovieGraphModifiersCustomization::MakeInstance
 
 void FMovieGraphModifiersCustomization::CustomizeDetails(IDetailLayoutBuilder& InDetailBuilder)
 {
-	TArray<TWeakObjectPtr<UMovieGraphModifierNode>> ModifierNodes =
-		InDetailBuilder.GetObjectsOfTypeBeingCustomized<UMovieGraphModifierNode>();
-	if (ModifierNodes.Num() != 1)
-	{
-		// Showing more than one modifier node is not supported
-		return;
-	}
-
-	const TWeakObjectPtr<UMovieGraphModifierNode> ModifierNode = ModifierNodes[0];
+	// The customization only supports editing a single Modifier node
+	const TWeakObjectPtr<UMovieGraphModifierNode> ModifierNode = GetSelectedModifierNode();
 	if (!ModifierNode.IsValid())
 	{
 		return;
 	}
 
 	// Update the data source
-	ListDataSource = ModifierNode->GetCollections();
+	RefreshListDataSource();
 	
 	// Generate a (multi-layered) icon for the "Add" menu
 	const TSharedRef<SLayeredImage> AddIcon =
@@ -236,6 +229,7 @@ void FMovieGraphModifiersCustomization::CustomizeDetails(IDetailLayoutBuilder& I
 				ModifierNode.Get()->SetCollectionEnabled(InCollectionName, bNewEnableState);
 			}
 		})
+		.OnRefreshDataSourceRequested(this, &FMovieGraphModifiersCustomization::RefreshListDataSource)
 	];
 
 	// For all modifiers added to the node, add a category for each, and add each modifier's EditAnywhere properties to the category
@@ -280,4 +274,36 @@ FText FMovieGraphModifiersCustomization::GetCollectionRowText(const FName Collec
 	return FText::FromName(CollectionName);
 }
 
+TWeakObjectPtr<UMovieGraphModifierNode> FMovieGraphModifiersCustomization::GetSelectedModifierNode() const
+{
+	if (const TSharedPtr<IDetailLayoutBuilder> DetailBuilderPin = DetailBuilder.Pin())
+	{
+		TArray<TWeakObjectPtr<UMovieGraphModifierNode>> ModifierNodes =
+			DetailBuilderPin->GetObjectsOfTypeBeingCustomized<UMovieGraphModifierNode>();
+		if (ModifierNodes.Num() != 1)
+		{
+			return nullptr;
+		}
+
+		const TWeakObjectPtr<UMovieGraphModifierNode> ModifierNode = ModifierNodes[0];
+		if (ModifierNode.IsValid())
+		{
+			return ModifierNode;
+		}
+	}
+
+	return nullptr;
+}
+
+void FMovieGraphModifiersCustomization::RefreshListDataSource()
+{
+	const TWeakObjectPtr<UMovieGraphModifierNode> ModifierNode = GetSelectedModifierNode();
+	
+	if (const TStrongObjectPtr<UMovieGraphModifierNode> ModifierNodePin = ModifierNode.Pin())
+	{
+		// Update the data source
+		ListDataSource = ModifierNodePin->GetCollections();
+	}
+}
+
 #undef LOCTEXT_NAMESPACE
\ No newline at end of file
diff --git a/Engine/Plugins/MovieScene/MovieRenderPipeline/Source/MovieRenderPipelineEditor/Private/Customizations/Graph/MovieGraphModifiersCustomization.h b/Engine/Plugins/MovieScene/MovieRenderPipeline/Source/MovieRenderPipelineEditor/Private/Customizations/Graph/MovieGraphModifiersCustomization.h
--- a/Engine/Plugins/MovieScene/MovieRenderPipeline/Source/MovieRenderPipelineEditor/Private/Customizations/Graph/MovieGraphModifiersCustomization.h
+++ b/Engine/Plugins/MovieScene/MovieRenderPipeline/Source/MovieRenderPipelineEditor/Private/Customizations/Graph/MovieGraphModifiersCustomization.h
@@ -6,6 +6,7 @@
 #include "Graph/MovieGraphSharedWidgets.h"
 
 class UMovieGraphCollectionNode;
+class UMovieGraphModifierNode;
 
 /** Customize how the Modifier node appears in the details panel. */
 class FMovieGraphModifiersCustomization final : public IDetailCustomization
@@ -23,6 +24,12 @@ private:
 	static const FSlateBrush* GetCollectionRowIcon(const FName CollectionName);
 	static FText GetCollectionRowText(const FName CollectionName);
 
+	/** Gets the selected modifier node that this customization is displaying. */
+	TWeakObjectPtr<UMovieGraphModifierNode> GetSelectedModifierNode() const;
+	
+	/** Refreshes the list's data source to reflect the data model. */
+	void RefreshListDataSource();
+
 private:
 	/** The details builder associated with the customization. */
 	TWeakPtr<IDetailLayoutBuilder> DetailBuilder;
