MoviePipeline: Fixed issue where the Script Tags property would be in the wrong place in the details panel for some nodes, potentially causing data entry errors as it could be confused with other properties easily.

#rnx

diff --git a/Engine/Plugins/MovieScene/MovieRenderPipeline/Source/MovieRenderPipelineEditor/Private/Customizations/Graph/MovieGraphNodeCustomization.h b/Engine/Plugins/MovieScene/MovieRenderPipeline/Source/MovieRenderPipelineEditor/Private/Customizations/Graph/MovieGraphNodeCustomization.h
--- a/Engine/Plugins/MovieScene/MovieRenderPipeline/Source/MovieRenderPipelineEditor/Private/Customizations/Graph/MovieGraphNodeCustomization.h
+++ b/Engine/Plugins/MovieScene/MovieRenderPipeline/Source/MovieRenderPipelineEditor/Private/Customizations/Graph/MovieGraphNodeCustomization.h
@@ -52,8 +52,10 @@ protected:
 			DetailBuilder.HideCategory("Properties");
 		}
 
-		// Collapse the Tags category by default because it's not used by most users.
-		DetailBuilder.EditCategory("Tags").InitiallyCollapsed(true).SetSortOrder(1000 * ECategoryPriority::Uncommon);
+		// Collapse the Tags category by default because it's not used by most users. Set the sort order to make Tags appear last in the list of
+		// categories (1000 * Uncommon == starting index for all categories with "Uncommon" sort order, and the +1000 is to move the category past all
+		// other categories with "Uncommon" sort order).
+		DetailBuilder.EditCategory("Tags").InitiallyCollapsed(true).SetSortOrder(1000 * ECategoryPriority::Uncommon + 1000);
 
 		// The File Output category should be shown first in all output nodes that have this category.
 		DetailBuilder.EditCategory("File Output").SetSortOrder(ECategoryPriority::Important);
