MovieGraph: Add cloth substepping and update Niagara module with temporal sub-sampling data.

Specifically Tested:
Created test content that has a ball with cloth simulation on it and animate the ball moving back and forth in 3d space. Without this fix the cloth jitters every temporal sub-sample and produces a noisy erratic look, with this fix it produces similar results to having only one Temporal Sub-Sample. Compared results against MRQ and found them comparable.

diff --git a/Engine/Plugins/MovieScene/MovieRenderPipeline/Source/MovieRenderPipelineCore/Private/Graph/MovieGraphCoreTimeStep.cpp b/Engine/Plugins/MovieScene/MovieRenderPipeline/Source/MovieRenderPipelineCore/Private/Graph/MovieGraphCoreTimeStep.cpp
--- a/Engine/Plugins/MovieScene/MovieRenderPipeline/Source/MovieRenderPipelineCore/Private/Graph/MovieGraphCoreTimeStep.cpp
+++ b/Engine/Plugins/MovieScene/MovieRenderPipeline/Source/MovieRenderPipelineCore/Private/Graph/MovieGraphCoreTimeStep.cpp
@@ -16,6 +16,7 @@
 #include "Interfaces/Interface_PostProcessVolume.h"
 #include "GameFramework/WorldSettings.h"
 #include "ProfilingDebugging/MiscTrace.h"
+#include "MoviePipelineUtils.h"
 
 UMovieGraphCoreTimeStep::UMovieGraphCoreTimeStep()
 {
@@ -402,7 +403,16 @@ void UMovieGraphCoreTimeStep::TickProducingFrames()
 	{
 		FrameDeltaTime = CurrentFrameData.TemporalRanges[0].Size<FFrameTime>();
 	}
-	// ToDo: Propagate delta time multipliers to cloth
+
+	// Cloth needs to increase the number of iterations when solving during the "long"
+	// frame of MRQ (shutter closed time).
+	{
+		double Ratio = FrameDeltaTime.FloorToFrame().Value / (double)CurrentFrameMetrics.FrameTimePerTemporalSample.FloorToFrame().Value;
+
+		// Slomo can end up trying to do less than one iteration, we don't want that.	
+		int32 DivisionMultiplier = FMath::Max(FMath::FloorToInt(Ratio), 1);
+		UE::MoviePipeline::SetSkeletalMeshClothSubSteps(DivisionMultiplier, GetWorld(), GetOwningGraph()->GetClothSimCache());
+	}
 
 	// Because we know what time range we're supposed to represent, we can just assign the CurrentTimeInRoot absolutely,
 	// instead of accumulating delta times.
@@ -471,6 +481,18 @@ void UMovieGraphCoreTimeStep::TickProducingFrames()
 		// Calculate metrics for the shot as well
 		CurrentTimeStepData.ShotFrameNumber = ShotFrameNumber;
 		CurrentTimeStepData.ShotTimeCode = FTimecode::FromFrameNumber(CurrentTimeStepData.ShotFrameNumber, EffectiveFrameRate, bDropFrame);
+
+		// Update the lightweight tick info in the module for other modules (Niagara) to have information
+		// about the temporal data for better simulation. We do this here since we calculated the effective
+		// frame rate above.
+		{
+			FMoviePipelineLightweightTickInfo TickInfo;
+			TickInfo.bIsActive = true;
+			TickInfo.TemporalSampleCount = CurrentTimeStepData.TemporalSampleCount;
+			TickInfo.TemporalSampleIndex = CurrentTimeStepData.TemporalSampleIndex;
+			TickInfo.SequenceFPS = EffectiveFrameRate.AsDecimal();
+			FMovieRenderPipelineCoreModule::SetTickInfo(TickInfo);
+		}
 	}
 
 	// Set our time step for the next frame. We use the undilated delta time for the Custom Timestep as the engine will
diff --git a/Engine/Plugins/MovieScene/MovieRenderPipeline/Source/MovieRenderPipelineCore/Private/Graph/MovieGraphPipeline.cpp b/Engine/Plugins/MovieScene/MovieRenderPipeline/Source/MovieRenderPipelineCore/Private/Graph/MovieGraphPipeline.cpp
--- a/Engine/Plugins/MovieScene/MovieRenderPipeline/Source/MovieRenderPipelineCore/Private/Graph/MovieGraphPipeline.cpp
+++ b/Engine/Plugins/MovieScene/MovieRenderPipeline/Source/MovieRenderPipelineCore/Private/Graph/MovieGraphPipeline.cpp
@@ -1006,6 +1006,9 @@ void UMovieGraphPipeline::TeardownShot(const TObjectPtr<UMoviePipelineExecutorSh
 	CVarManager->RevertAllCVars();
 	CVarManager->RunEndConsoleCommands();
 
+	UE::MoviePipeline::RestoreSkeletalMeshClothSubSteps(ClothSimCache);
+	ClothSimCache.Reset();
+
 	// Revert cvars set by the global game overrides. Needs to be done after the CVarManager reverts (since the global
 	// game overrides are applied first in SetupShot).
 	bIncludeCDOs = false;
@@ -1415,6 +1418,9 @@ void UMovieGraphPipeline::TransitionToState(const EMovieRenderPipelineState InNe
 			ProcessOutstandingFutures();
 
 			BeginExport();
+
+			// Clear out MRQ tick information for external consumers as we are done
+			FMovieRenderPipelineCoreModule::SetTickInfo(FMoviePipelineLightweightTickInfo());
 		}
 		break;
 	case EMovieRenderPipelineState::Export:
diff --git a/Engine/Plugins/MovieScene/MovieRenderPipeline/Source/MovieRenderPipelineCore/Private/MoviePipeline.cpp b/Engine/Plugins/MovieScene/MovieRenderPipeline/Source/MovieRenderPipelineCore/Private/MoviePipeline.cpp
--- a/Engine/Plugins/MovieScene/MovieRenderPipeline/Source/MovieRenderPipelineCore/Private/MoviePipeline.cpp
+++ b/Engine/Plugins/MovieScene/MovieRenderPipeline/Source/MovieRenderPipelineCore/Private/MoviePipeline.cpp
@@ -45,12 +45,6 @@
 #include "MoviePipelineCommandLineEncoder.h"
 #include "Kismet/KismetSystemLibrary.h"
 #include "ProfilingDebugging/TraceAuxiliary.h"
-#include "EngineUtils.h"
-#include "ClothingSimulationInteractor.h"
-#include "ClothingSimulationInterface.h"
-#include "ChaosClothAsset/ClothComponent.h"
-#include "ChaosClothAsset/ClothAssetInteractor.h"
-#include "Components/SkeletalMeshComponent.h"
 #include "MoviePipelineTelemetry.h"
 #include "MovieSceneCommonHelpers.h"
 #include "Compilation/MovieSceneCompiledDataManager.h"
@@ -77,7 +71,6 @@ static TAutoConsoleVariable<int32> CVarMovieRenderPipelineFrameStepper(
 	ECVF_Default);
 
 FString UMoviePipeline::DefaultDebugWidgetAsset = TEXT("/MovieRenderPipeline/Blueprints/UI_MovieRenderPipelineScreenOverlay.UI_MovieRenderPipelineScreenOverlay_C");
-DECLARE_CYCLE_STAT(TEXT("STAT_MoviePipeline_ClothAdjust"), STAT_ClothSubstepAdjust, STATGROUP_MoviePipeline);
 
 UMoviePipeline::UMoviePipeline()
 	: CustomTimeStep(nullptr)
@@ -1088,7 +1081,7 @@ void UMoviePipeline::TeardownShot(UMoviePipelineExecutorShot* InShot)
 		MoviePipeline::SaveOrRestoreSubSectionHierarchy(Node, bSaveSettings);
 	}
 
-	RestoreSkeletalMeshClothSubSteps();
+	UE::MoviePipeline::RestoreSkeletalMeshClothSubSteps(ClothSimCache);
 	ClothSimCache.Reset();
 
 	constexpr bool bIsGraph = false;
@@ -1725,105 +1718,6 @@ void UMoviePipeline::StopUnrealInsightsCapture()
 	FTraceAuxiliary::Stop();
 }
 
-void UMoviePipeline::SetSkeletalMeshClothSubSteps(const int32 InSubdivisionCount)
-{
-	SCOPE_CYCLE_COUNTER(STAT_ClothSubstepAdjust);
-
-	for (TActorIterator<AActor> ActorIt(GetWorld()); ActorIt; ++ActorIt)
-	{
-		AActor* FoundActor = *ActorIt;
-		if (FoundActor)
-		{
-			TArray<USkeletalMeshComponent*> SkeletalMeshComponents;
-			FoundActor->GetComponents(SkeletalMeshComponents);
-
-			for (USkeletalMeshComponent* Component : SkeletalMeshComponents)
-			{
-				UClothingSimulationInteractor* ClothInteractor = Component->GetClothingSimulationInteractor();
-				if (ClothInteractor)
-				{
-					const int32 LODIndex = 0;  // There is only a NumSubSteps for LOD 0 in the Skeletal Mesh clothing system
-					TWeakObjectPtr<UObject> WeakPtr = TWeakObjectPtr<UClothingSimulationInteractor>(ClothInteractor);
-		
-					TArray<FClothSimSettingsCache>* ExistingCacheEntry = ClothSimCache.Find(WeakPtr);
-					if (!ExistingCacheEntry)
-					{
-						ClothSimCache.Add(WeakPtr);
-						ExistingCacheEntry = ClothSimCache.Find(WeakPtr);
-						ExistingCacheEntry->SetNumZeroed(1);  // Only store LOD 0
-						const int32 NumSubsteps = Component->GetClothingSimulation() ? FMath::Max(Component->GetClothingSimulation()->GetNumSubsteps(), 1)
-							: 1; // If there's no clothing simulation component just fall back to assuming they only had 1.
-						(*ExistingCacheEntry)[LODIndex].NumSubSteps = NumSubsteps;
-					}
-
-					ClothInteractor->SetNumSubsteps((*ExistingCacheEntry)[LODIndex].NumSubSteps * InSubdivisionCount);
-				}
-			}
-
-			TArray<UChaosClothComponent*> ChaosClothComponents;
-			FoundActor->GetComponents(ChaosClothComponents);
-
-			for (UChaosClothComponent* Component : ChaosClothComponents)
-			{
-				UChaosClothAssetInteractor* ClothAssetInteractor = Component->GetClothOutfitInteractor();
-				if (ClothAssetInteractor)
-				{
-					const int32 NumLODs = Component->GetNumLODs();
-					TWeakObjectPtr<UObject> WeakPtr = TWeakObjectPtr<UChaosClothAssetInteractor>(ClothAssetInteractor);
-		
-					TArray<FClothSimSettingsCache>* ExistingCacheEntry = ClothSimCache.Find(WeakPtr);
-					if (!ExistingCacheEntry)
-					{
-						ClothSimCache.Add(WeakPtr);
-						ExistingCacheEntry = ClothSimCache.Find(WeakPtr);
-						ExistingCacheEntry->SetNumUninitialized(NumLODs);
-						for (int32 LODIndex = 0; LODIndex < NumLODs; ++LODIndex)
-						{
-							constexpr int32 MinNumSubsteps = 1;
-							const int32 NumSubsteps = FMath::Max(ClothAssetInteractor->GetIntValue(TEXT("NumSubsteps"), LODIndex, MinNumSubsteps), MinNumSubsteps);
-							const int32 DynamicSubstepDeltaTime = ClothAssetInteractor->GetFloatValue(TEXT("DynamicSubstepDeltaTime"), LODIndex, MinNumSubsteps);
-							(*ExistingCacheEntry)[LODIndex].NumSubSteps = NumSubsteps;
-							(*ExistingCacheEntry)[LODIndex].DynamicSubstepDeltaTime = DynamicSubstepDeltaTime;
-						}
-					}
-
-					for (int32 LODIndex = 0; LODIndex < NumLODs; ++LODIndex)
-					{
-						if (ExistingCacheEntry->IsValidIndex(LODIndex))
-						{
-							ClothAssetInteractor->SetIntValue(TEXT("NumSubsteps"), LODIndex, (*ExistingCacheEntry)[LODIndex].NumSubSteps * InSubdivisionCount);
-							ClothAssetInteractor->SetFloatValue(TEXT("DynamicSubstepDeltaTime"), LODIndex, 0.f);
-						}
-					}
-				}
-			}
-		}
-	}
-}
-
-void UMoviePipeline::RestoreSkeletalMeshClothSubSteps()
-{
-	for (const TPair<TWeakObjectPtr<UObject>, TArray<FClothSimSettingsCache>>& Pair : ClothSimCache)
-	{
-		if (UObject* Object = Pair.Key.Get())
-		{
-			if (UClothingSimulationInteractor* ClothInteractor = Cast<UClothingSimulationInteractor>(Object))
-			{
-				const int32 LODIndex = 0;  // There is only a NumSubSteps for LOD 0 in the Skeletal Mesh clothing system
-				ClothInteractor->SetNumSubsteps(Pair.Value[LODIndex].NumSubSteps);
-			}
-			else if (UChaosClothAssetInteractor* ClothAssetInteractor = Cast<UChaosClothAssetInteractor>(Object))
-			{
-				for (int32 LODIndex = 0; LODIndex < Pair.Value.Num(); ++LODIndex)
-				{
-					ClothAssetInteractor->SetIntValue(TEXT("NumSubsteps"), LODIndex, Pair.Value[LODIndex].NumSubSteps);
-					ClothAssetInteractor->SetIntValue(TEXT("DynamicSubstepDeltaTime"), LODIndex, Pair.Value[LODIndex].DynamicSubstepDeltaTime);
-				}
-			}
-		}
-	}
-}
-
 void UMoviePipeline::GetSidecarCameraData(UMoviePipelineExecutorShot* InShot, int32 InCameraIndex, FMinimalViewInfo& OutViewInfo, UCameraComponent** OutCameraComponent) const
 {
 	UMovieScene* MovieScene = GetTargetSequence()->GetMovieScene();
diff --git a/Engine/Plugins/MovieScene/MovieRenderPipeline/Source/MovieRenderPipelineCore/Private/MoviePipelineTiming.cpp b/Engine/Plugins/MovieScene/MovieRenderPipeline/Source/MovieRenderPipelineCore/Private/MoviePipelineTiming.cpp
--- a/Engine/Plugins/MovieScene/MovieRenderPipeline/Source/MovieRenderPipelineCore/Private/MoviePipelineTiming.cpp
+++ b/Engine/Plugins/MovieScene/MovieRenderPipeline/Source/MovieRenderPipelineCore/Private/MoviePipelineTiming.cpp
@@ -301,7 +301,7 @@ void UMoviePipeline::TickProducingFrames()
 
 			// Slomo can end up trying to do less than one iteration, we don't want that.	
 			int32 DivisionMultiplier = FMath::Max(FMath::FloorToInt(Ratio), 1);
-			SetSkeletalMeshClothSubSteps(DivisionMultiplier);
+			UE::MoviePipeline::SetSkeletalMeshClothSubSteps(DivisionMultiplier, GetWorld(), ClothSimCache);
 		}
 
 		CachedOutputState.TimeData.FrameDeltaTime = FrameDeltaTime;
@@ -611,7 +611,7 @@ void UMoviePipeline::TickProducingFrames()
 				
 				// Slomo can end up trying to do less than one iteration, we don't want that.	
 				int32 DivisionMultiplier = FMath::Max(FMath::FloorToInt(Ratio), 1);
-				SetSkeletalMeshClothSubSteps(DivisionMultiplier);
+				UE::MoviePipeline::SetSkeletalMeshClothSubSteps(DivisionMultiplier, GetWorld(), ClothSimCache);
 			}
 		}
 
diff --git a/Engine/Plugins/MovieScene/MovieRenderPipeline/Source/MovieRenderPipelineCore/Private/MoviePipelineUtils.cpp b/Engine/Plugins/MovieScene/MovieRenderPipeline/Source/MovieRenderPipelineCore/Private/MoviePipelineUtils.cpp
--- a/Engine/Plugins/MovieScene/MovieRenderPipeline/Source/MovieRenderPipelineCore/Private/MoviePipelineUtils.cpp
+++ b/Engine/Plugins/MovieScene/MovieRenderPipeline/Source/MovieRenderPipelineCore/Private/MoviePipelineUtils.cpp
@@ -36,6 +36,12 @@
 #include "MovieRenderPipelineCoreModule.h"
 #include "RHI.h"
 #include "CineCameraComponent.h"
+#include "EngineUtils.h"
+#include "ClothingSimulationInteractor.h"
+#include "ClothingSimulationInterface.h"
+#include "ChaosClothAsset/ClothComponent.h"
+#include "ChaosClothAsset/ClothAssetInteractor.h"
+#include "Components/SkeletalMeshComponent.h"
 
 namespace UE
 {
@@ -712,6 +718,7 @@ namespace UE
 	namespace MoviePipeline
 	{
 		DECLARE_CYCLE_STAT(TEXT("STAT_MoviePipeline_HardwareMetadata"), STAT_HardwareMetadata, STATGROUP_MoviePipeline);
+		DECLARE_CYCLE_STAT(TEXT("STAT_MoviePipeline_ClothAdjust"), STAT_ClothSubstepAdjust, STATGROUP_MoviePipeline);
 
 		void ConformOutputFormatStringToken(FString& InOutFilenameFormatString, const FStringView InToken, const FName& InNodeName, const FName& InBranchName)
 		{
@@ -1109,6 +1116,104 @@ namespace UE
 			InOutView->OverridePostProcessSettings(InViewInfo.PostProcessSettings, InViewInfo.PostProcessBlendWeight);
 		}
 
+		void SetSkeletalMeshClothSubSteps(const int32 InSubdivisionCount, UWorld* InWorld, TMap<TWeakObjectPtr<UObject>, TArray<::MoviePipeline::FClothSimSettingsCache>> InClothSimCache)
+		{
+			SCOPE_CYCLE_COUNTER(STAT_ClothSubstepAdjust);
+			for (TActorIterator<AActor> ActorIt(InWorld); ActorIt; ++ActorIt)
+			{
+				AActor* FoundActor = *ActorIt;
+				if (FoundActor)
+				{
+					TArray<USkeletalMeshComponent*> SkeletalMeshComponents;
+					FoundActor->GetComponents(SkeletalMeshComponents);
+
+					for (USkeletalMeshComponent* Component : SkeletalMeshComponents)
+					{
+						UClothingSimulationInteractor* ClothInteractor = Component->GetClothingSimulationInteractor();
+						if (ClothInteractor)
+						{
+							const int32 LODIndex = 0;  // There is only a NumSubSteps for LOD 0 in the Skeletal Mesh clothing system
+							TWeakObjectPtr<UObject> WeakPtr = TWeakObjectPtr<UClothingSimulationInteractor>(ClothInteractor);
+
+							TArray<::MoviePipeline::FClothSimSettingsCache>* ExistingCacheEntry = InClothSimCache.Find(WeakPtr);
+							if (!ExistingCacheEntry)
+							{
+								InClothSimCache.Add(WeakPtr);
+								ExistingCacheEntry = InClothSimCache.Find(WeakPtr);
+								ExistingCacheEntry->SetNumZeroed(1);  // Only store LOD 0
+								const int32 NumSubsteps = Component->GetClothingSimulation() ? FMath::Max(Component->GetClothingSimulation()->GetNumSubsteps(), 1)
+									: 1; // If there's no clothing simulation component just fall back to assuming they only had 1.
+								(*ExistingCacheEntry)[LODIndex].NumSubSteps = NumSubsteps;
+							}
+
+							ClothInteractor->SetNumSubsteps((*ExistingCacheEntry)[LODIndex].NumSubSteps * InSubdivisionCount);
+						}
+					}
+
+					TArray<UChaosClothComponent*> ChaosClothComponents;
+					FoundActor->GetComponents(ChaosClothComponents);
+
+					for (UChaosClothComponent* Component : ChaosClothComponents)
+					{
+						UChaosClothAssetInteractor* ClothAssetInteractor = Component->GetClothOutfitInteractor();
+						if (ClothAssetInteractor)
+						{
+							const int32 NumLODs = Component->GetNumLODs();
+							TWeakObjectPtr<UObject> WeakPtr = TWeakObjectPtr<UChaosClothAssetInteractor>(ClothAssetInteractor);
+
+							TArray<::MoviePipeline::FClothSimSettingsCache>* ExistingCacheEntry = InClothSimCache.Find(WeakPtr);
+							if (!ExistingCacheEntry)
+							{
+								InClothSimCache.Add(WeakPtr);
+								ExistingCacheEntry = InClothSimCache.Find(WeakPtr);
+								ExistingCacheEntry->SetNumUninitialized(NumLODs);
+								for (int32 LODIndex = 0; LODIndex < NumLODs; ++LODIndex)
+								{
+									constexpr int32 MinNumSubsteps = 1;
+									const int32 NumSubsteps = FMath::Max(ClothAssetInteractor->GetIntValue(TEXT("NumSubsteps"), LODIndex, MinNumSubsteps), MinNumSubsteps);
+									const int32 DynamicSubstepDeltaTime = ClothAssetInteractor->GetFloatValue(TEXT("DynamicSubstepDeltaTime"), LODIndex, MinNumSubsteps);
+									(*ExistingCacheEntry)[LODIndex].NumSubSteps = NumSubsteps;
+									(*ExistingCacheEntry)[LODIndex].DynamicSubstepDeltaTime = DynamicSubstepDeltaTime;
+								}
+							}
+
+							for (int32 LODIndex = 0; LODIndex < NumLODs; ++LODIndex)
+							{
+								if (ExistingCacheEntry->IsValidIndex(LODIndex))
+								{
+									ClothAssetInteractor->SetIntValue(TEXT("NumSubsteps"), LODIndex, (*ExistingCacheEntry)[LODIndex].NumSubSteps * InSubdivisionCount);
+									ClothAssetInteractor->SetFloatValue(TEXT("DynamicSubstepDeltaTime"), LODIndex, 0.f);
+								}
+							}
+						}
+					}
+				}
+			}
+		}
+
+		void RestoreSkeletalMeshClothSubSteps(TMap<TWeakObjectPtr<UObject>, TArray<::MoviePipeline::FClothSimSettingsCache>> InClothSimCache)
+		{
+			for (const TPair<TWeakObjectPtr<UObject>, TArray<::MoviePipeline::FClothSimSettingsCache>>& Pair : InClothSimCache)
+			{
+				if (UObject* Object = Pair.Key.Get())
+				{
+					if (UClothingSimulationInteractor* ClothInteractor = Cast<UClothingSimulationInteractor>(Object))
+					{
+						const int32 LODIndex = 0;  // There is only a NumSubSteps for LOD 0 in the Skeletal Mesh clothing system
+						ClothInteractor->SetNumSubsteps(Pair.Value[LODIndex].NumSubSteps);
+					}
+					else if (UChaosClothAssetInteractor* ClothAssetInteractor = Cast<UChaosClothAssetInteractor>(Object))
+					{
+						for (int32 LODIndex = 0; LODIndex < Pair.Value.Num(); ++LODIndex)
+						{
+							ClothAssetInteractor->SetIntValue(TEXT("NumSubsteps"), LODIndex, Pair.Value[LODIndex].NumSubSteps);
+							ClothAssetInteractor->SetIntValue(TEXT("DynamicSubstepDeltaTime"), LODIndex, Pair.Value[LODIndex].DynamicSubstepDeltaTime);
+						}
+					}
+				}
+			}
+		}
+
 	}
 }
 namespace UE
diff --git a/Engine/Plugins/MovieScene/MovieRenderPipeline/Source/MovieRenderPipelineCore/Public/Graph/MovieGraphDataTypes.h b/Engine/Plugins/MovieScene/MovieRenderPipeline/Source/MovieRenderPipelineCore/Public/Graph/MovieGraphDataTypes.h
--- a/Engine/Plugins/MovieScene/MovieRenderPipeline/Source/MovieRenderPipelineCore/Public/Graph/MovieGraphDataTypes.h
+++ b/Engine/Plugins/MovieScene/MovieRenderPipeline/Source/MovieRenderPipelineCore/Public/Graph/MovieGraphDataTypes.h
@@ -217,6 +217,9 @@ public:
 	virtual void SetupRenderingPipelineForShot(UMoviePipelineExecutorShot* InShot) {}
 	virtual void TeardownRenderingPipelineForShot(UMoviePipelineExecutorShot* InShot) {}
 	virtual UE::MovieGraph::FRenderTimeStatistics* GetRenderTimeStatistics(const int32 InFrameNumber) { return nullptr; }
+
+	UE_DEPRECATED(5.5, "Use the version that doesn't take a config instead. The information from the config is now derived from the index value.")
+	virtual UE::MovieGraph::FMinimalCameraInfo GetMinimalCameraInfo(UMovieGraphEvaluatedConfig* InConfig, const int32 InCameraIndex) const { return GetMinimalCameraInfo(InCameraIndex); }
 	
 	/** InCameraIndex can be -1 (for primary camera when not using multi-layer rendering) or 0...n for sidecar cameras. */
 	virtual UE::MovieGraph::FMinimalCameraInfo GetMinimalCameraInfo(const int32 InCameraIndex) const { return UE::MovieGraph::FMinimalCameraInfo(); }
diff --git a/Engine/Plugins/MovieScene/MovieRenderPipeline/Source/MovieRenderPipelineCore/Public/Graph/MovieGraphPipeline.h b/Engine/Plugins/MovieScene/MovieRenderPipeline/Source/MovieRenderPipelineCore/Public/Graph/MovieGraphPipeline.h
--- a/Engine/Plugins/MovieScene/MovieRenderPipeline/Source/MovieRenderPipelineCore/Public/Graph/MovieGraphPipeline.h
+++ b/Engine/Plugins/MovieScene/MovieRenderPipeline/Source/MovieRenderPipelineCore/Public/Graph/MovieGraphPipeline.h
@@ -110,6 +110,9 @@ public:
 	
 	/** Gets the Output Merger for this Movie Pipeline which is responsible for gathering all of the data coming in for a given output frame, before making it available to the MovieGraphPipeline. */
 	TSharedPtr<UE::MovieGraph::IMovieGraphOutputMerger> GetOutputMerger() const { return OutputMerger; }
+
+	/** Gets cached data for cloth simulations, needed to save/restore cloth sub-stepping. */
+	TMap<TWeakObjectPtr<UObject>, TArray<MoviePipeline::FClothSimSettingsCache>> GetClothSimCache() const { return ClothSimCache; }
 	
 	/** Writing images to disk is an async process. When you start writing, declare a future with the filename you will eventually write to, and complete the future once it is on disk. */
 	void AddOutputFuture(TFuture<bool>&& InOutputFuture, const UE::MovieGraph::FMovieGraphOutputFutureData& InData);
@@ -339,6 +342,9 @@ protected:
 	/** The previous custom timestep the engine was using, if any. */
 	UPROPERTY(Transient)
 	TObjectPtr<UEngineCustomTimeStep> PrevCustomEngineTimeStep;
+
+	/** Simulation settings cache per cloth interactor object. Needs one per LOD, hence the array. */
+	TMap<TWeakObjectPtr<UObject>, TArray<MoviePipeline::FClothSimSettingsCache>> ClothSimCache;
 	
 public:
 	static FString DefaultPreviewWidgetAsset;
diff --git a/Engine/Plugins/MovieScene/MovieRenderPipeline/Source/MovieRenderPipelineCore/Public/MoviePipeline.h b/Engine/Plugins/MovieScene/MovieRenderPipeline/Source/MovieRenderPipelineCore/Public/MoviePipeline.h
--- a/Engine/Plugins/MovieScene/MovieRenderPipeline/Source/MovieRenderPipelineCore/Public/MoviePipeline.h
+++ b/Engine/Plugins/MovieScene/MovieRenderPipeline/Source/MovieRenderPipelineCore/Public/MoviePipeline.h
@@ -315,10 +315,6 @@ private:
 
 	/** Handles transitioning between states, preventing reentrancy. Normal state flow should be respected, does not handle arbitrary x to y transitions. */
 	void TransitionToState(const EMovieRenderPipelineState InNewState);
-
-	void SetSkeletalMeshClothSubSteps(const int32 InSubdivisionCount);
-	void RestoreSkeletalMeshClothSubSteps();
-
 private:
 	/** Custom TimeStep used to drive the engine while rendering. */
 	UPROPERTY(Transient, Instanced)
@@ -427,14 +423,9 @@ private:
 
 	TSharedPtr<MoviePipeline::FCameraCutSubSectionHierarchyNode> CachedSequenceHierarchyRoot;
 
-	struct FClothSimSettingsCache
-	{
-		int32 NumSubSteps;
-		float DynamicSubstepDeltaTime;
-	};
 
 	/** Simulation settings cache per cloth interactor object. Needs one per LOD, hence the array. */
-	TMap<TWeakObjectPtr<UObject>, TArray<FClothSimSettingsCache>> ClothSimCache;
+	TMap<TWeakObjectPtr<UObject>, TArray<MoviePipeline::FClothSimSettingsCache>> ClothSimCache;
 
 	struct FRenderTimeStatistics
 	{
diff --git a/Engine/Plugins/MovieScene/MovieRenderPipeline/Source/MovieRenderPipelineCore/Public/MoviePipelineUtils.h b/Engine/Plugins/MovieScene/MovieRenderPipeline/Source/MovieRenderPipelineCore/Public/MoviePipelineUtils.h
--- a/Engine/Plugins/MovieScene/MovieRenderPipeline/Source/MovieRenderPipelineCore/Public/MoviePipelineUtils.h
+++ b/Engine/Plugins/MovieScene/MovieRenderPipeline/Source/MovieRenderPipelineCore/Public/MoviePipelineUtils.h
@@ -154,6 +154,10 @@ namespace UE
 		MOVIERENDERPIPELINECORE_API FString GetPaddingFormatString(int32 InZeroPadCount, const int32 InFrameNumber);
 		MOVIERENDERPIPELINECORE_API void DoPostProcessBlend(const FVector& InViewLocation, const class UWorld* InWorld, const struct FMinimalViewInfo& InViewInfo, class FSceneView* InOutView);
 
+		MOVIERENDERPIPELINECORE_API void SetSkeletalMeshClothSubSteps(const int32 InSubdivisionCount, UWorld* InWorld, TMap<TWeakObjectPtr<UObject>, TArray<::MoviePipeline::FClothSimSettingsCache>> InClothSimCache);
+		MOVIERENDERPIPELINECORE_API void RestoreSkeletalMeshClothSubSteps(TMap<TWeakObjectPtr<UObject>, TArray<::MoviePipeline::FClothSimSettingsCache>> InClothSimCache);
+
+
 		/** When using spatial/temporal samples without anti-aliasing, get the sub-pixel jitter for the given frame index. FrameIndex is modded by InSamplesPerFrame so that the aa jitter pattern repeats every output frame. */
 		MOVIERENDERPIPELINECORE_API FVector2f GetSubPixelJitter(int32 InFrameIndex, int32 InSamplesPerFrame);
 
diff --git a/Engine/Plugins/MovieScene/MovieRenderPipeline/Source/MovieRenderPipelineCore/Public/MovieRenderPipelineDataTypes.h b/Engine/Plugins/MovieScene/MovieRenderPipeline/Source/MovieRenderPipelineCore/Public/MovieRenderPipelineDataTypes.h
--- a/Engine/Plugins/MovieScene/MovieRenderPipeline/Source/MovieRenderPipelineCore/Public/MovieRenderPipelineDataTypes.h
+++ b/Engine/Plugins/MovieScene/MovieRenderPipeline/Source/MovieRenderPipelineCore/Public/MovieRenderPipelineDataTypes.h
@@ -445,6 +445,12 @@ namespace MoviePipeline
 			return !(*this == InRHS);
 		}
 	};
+
+	struct FClothSimSettingsCache
+	{
+		int32 NumSubSteps;
+		float DynamicSubstepDeltaTime;
+	};
 }
 
 USTRUCT(BlueprintType)
