MoviePipeline: Updated the UI to indicate that Movie Render Graph is "Beta" rather than "Experimental".

#rnx

diff --git a/Engine/Plugins/MovieScene/MovieRenderPipeline/Source/MovieRenderPipelineCore/Public/Graph/MovieGraphConfig.h b/Engine/Plugins/MovieScene/MovieRenderPipeline/Source/MovieRenderPipelineCore/Public/Graph/MovieGraphConfig.h
--- a/Engine/Plugins/MovieScene/MovieRenderPipeline/Source/MovieRenderPipelineCore/Public/Graph/MovieGraphConfig.h
+++ b/Engine/Plugins/MovieScene/MovieRenderPipeline/Source/MovieRenderPipelineCore/Public/Graph/MovieGraphConfig.h
@@ -706,7 +706,7 @@ public:
 	TArray<UMovieGraphVariable*> GetVariables(const bool bIncludeGlobal = false) const;
 
 	/** Updates the values of all global variables. */
-	UFUNCTION(BlueprintCallable, Category="Experimental")
+	UFUNCTION(BlueprintCallable, Category="Movie Graph")
 	void UpdateGlobalVariableValues(const UMovieGraphPipeline* InPipeline);
 
 	/** Gets all inputs that have been defined on the graph. */
@@ -736,7 +736,7 @@ public:
 	 * Given a user-defined evaluation context, evaluate the graph and build a "flattened" list of settings for each branch discovered.
 	 * If there was an error while evaluating the graph, nullptr will be returned and OutError will be populated with a description of the problem.
 	 */
-	UFUNCTION(BlueprintCallable, Category="Experimental")
+	UFUNCTION(BlueprintCallable, Category="Movie Graph")
 	UMovieGraphEvaluatedConfig* CreateFlattenedGraph(const FMovieGraphTraversalContext& InContext, FString& OutError);
 
 	/** Given a class and FProperty that belongs to that class, search for a FBoolProperty that matches the name "bOverride_<name of InRealProperty>. */
diff --git a/Engine/Plugins/MovieScene/MovieRenderPipeline/Source/MovieRenderPipelineCore/Public/MovieJobVariableAssignmentContainer.h b/Engine/Plugins/MovieScene/MovieRenderPipeline/Source/MovieRenderPipelineCore/Public/MovieJobVariableAssignmentContainer.h
--- a/Engine/Plugins/MovieScene/MovieRenderPipeline/Source/MovieRenderPipelineCore/Public/MovieJobVariableAssignmentContainer.h
+++ b/Engine/Plugins/MovieScene/MovieRenderPipeline/Source/MovieRenderPipelineCore/Public/MovieJobVariableAssignmentContainer.h
@@ -28,7 +28,7 @@ public:
 	/**
 	 * Sets the graph config associated with the variable assignments.
 	 */
-	UFUNCTION(BlueprintCallable, Category="Experimental")
+	UFUNCTION(BlueprintCallable, Category="Graph")
 	void SetGraphConfig(const TSoftObjectPtr<UMovieGraphConfig>& InGraphConfig);
 
 	/** Gets the graph that is associated with this container. */
@@ -38,58 +38,58 @@ public:
 	uint32 GetNumAssignments() const;
 
 	/** Gets the bool value of the specified property. Returns true on success, else false. */
-	UFUNCTION(BlueprintCallable, Category="Experimental")
+	UFUNCTION(BlueprintCallable, Category="Value")
 	bool GetValueBool(const UMovieGraphVariable* InGraphVariable, bool& bOutValue) const;
 
 	/** Gets the byte value of the specified property. Returns true on success, else false. */
-	UFUNCTION(BlueprintCallable, Category="Experimental")
+	UFUNCTION(BlueprintCallable, Category="Value")
 	bool GetValueByte(const UMovieGraphVariable* InGraphVariable, uint8& OutValue) const;
 
 	/** Gets the int32 value of the specified property. Returns true on success, else false. */
-	UFUNCTION(BlueprintCallable, Category="Experimental")
+	UFUNCTION(BlueprintCallable, Category="Value")
 	bool GetValueInt32(const UMovieGraphVariable* InGraphVariable, int32& OutValue) const;
 
 	/** Gets the int64 value of the specified property. Returns true on success, else false. */
-	UFUNCTION(BlueprintCallable, Category="Experimental")
+	UFUNCTION(BlueprintCallable, Category="Value")
 	bool GetValueInt64(const UMovieGraphVariable* InGraphVariable, int64& OutValue) const;
 
 	/** Gets the float value of the specified property. Returns true on success, else false. */
-	UFUNCTION(BlueprintCallable, Category="Experimental")
+	UFUNCTION(BlueprintCallable, Category="Value")
 	bool GetValueFloat(const UMovieGraphVariable* InGraphVariable, float& OutValue) const;
 
 	/** Gets the double value of the specified property. Returns true on success, else false. */
-	UFUNCTION(BlueprintCallable, Category="Experimental")
+	UFUNCTION(BlueprintCallable, Category="Value")
 	bool GetValueDouble(const UMovieGraphVariable* InGraphVariable, double& OutValue) const;
 
 	/** Gets the FName value of the specified property. Returns true on success, else false. */
-	UFUNCTION(BlueprintCallable, Category="Experimental")
+	UFUNCTION(BlueprintCallable, Category="Value")
 	bool GetValueName(const UMovieGraphVariable* InGraphVariable, FName& OutValue) const;
 
 	/** Gets the FString value of the specified property. Returns true on success, else false. */
-	UFUNCTION(BlueprintCallable, Category="Experimental")
+	UFUNCTION(BlueprintCallable, Category="Value")
 	bool GetValueString(const UMovieGraphVariable* InGraphVariable, FString& OutValue) const;
 
 	/** Gets the FText value of the specified property. Returns true on success, else false. */
-	UFUNCTION(BlueprintCallable, Category="Experimental")
+	UFUNCTION(BlueprintCallable, Category="Value")
 	bool GetValueText(const UMovieGraphVariable* InGraphVariable, FText& OutValue) const;
 
 	/** Gets the enum value (for a specific enum) of the specified property. Returns true on success, else false. */
-	UFUNCTION(BlueprintCallable, Category="Experimental")
+	UFUNCTION(BlueprintCallable, Category="Value")
 	bool GetValueEnum(const UMovieGraphVariable* InGraphVariable, uint8& OutValue, const UEnum* RequestedEnum = nullptr) const;
 
 	/** Gets the struct value (for a specific struct) of the specified property. Returns true on success, else false. */
 	bool GetValueStruct(const UMovieGraphVariable* InGraphVariable, FStructView& OutValue, const UScriptStruct* RequestedStruct = nullptr) const;
 
 	/** Gets the object value (for a specific class) of the specified property. Returns true on success, else false. */
-	UFUNCTION(BlueprintCallable, Category="Experimental")
+	UFUNCTION(BlueprintCallable, Category="Value")
 	bool GetValueObject(const UMovieGraphVariable* InGraphVariable, UObject* OutValue, const UClass* RequestedClass = nullptr) const;
 
 	/** Gets the UClass value of the specified property. Returns true on success, else false. */
-	UFUNCTION(BlueprintCallable, Category="Experimental")
+	UFUNCTION(BlueprintCallable, Category="Value")
 	bool GetValueClass(const UMovieGraphVariable* InGraphVariable, UClass*& OutValue) const;
 
 	/** Gets the serialized string value of the specified property. */
-	UFUNCTION(BlueprintCallable, Category="Experimental")
+	UFUNCTION(BlueprintCallable, Category="Value")
 	FString GetValueSerializedString(const UMovieGraphVariable* InGraphVariable);
 
 	/** Gets the enum value of the specified property. Returns true on success, else false. */
@@ -117,58 +117,58 @@ public:
 	}
 
 	/** Sets the bool value of the specified property. Returns true on success, else false. */
-	UFUNCTION(BlueprintCallable, Category="Experimental")
+	UFUNCTION(BlueprintCallable, Category="Value")
 	bool SetValueBool(const UMovieGraphVariable* InGraphVariable, const bool bInValue);
 
 	/** Sets the byte value of the specified property. Returns true on success, else false. */
-	UFUNCTION(BlueprintCallable, Category="Experimental")
+	UFUNCTION(BlueprintCallable, Category="Value")
 	bool SetValueByte(const UMovieGraphVariable* InGraphVariable, const uint8 InValue);
 
 	/** Sets the int32 value of the specified property. Returns true on success, else false. */
-	UFUNCTION(BlueprintCallable, Category="Experimental")
+	UFUNCTION(BlueprintCallable, Category="Value")
 	bool SetValueInt32(const UMovieGraphVariable* InGraphVariable, const int32 InValue);
 
 	/** Sets the int64 value of the specified property. Returns true on success, else false. */
-	UFUNCTION(BlueprintCallable, Category="Experimental")
+	UFUNCTION(BlueprintCallable, Category="Value")
 	bool SetValueInt64(const UMovieGraphVariable* InGraphVariable, const int64 InValue);
 
 	/** Sets the float value of the specified property. Returns true on success, else false. */
-	UFUNCTION(BlueprintCallable, Category="Experimental")
+	UFUNCTION(BlueprintCallable, Category="Value")
 	bool SetValueFloat(const UMovieGraphVariable* InGraphVariable, const float InValue);
 
 	/** Sets the double value of the specified property. Returns true on success, else false. */
-	UFUNCTION(BlueprintCallable, Category="Experimental")
+	UFUNCTION(BlueprintCallable, Category="Value")
 	bool SetValueDouble(const UMovieGraphVariable* InGraphVariable, const double InValue);
 
 	/** Sets the FName value of the specified property. Returns true on success, else false. */
-	UFUNCTION(BlueprintCallable, Category="Experimental")
+	UFUNCTION(BlueprintCallable, Category="Value")
 	bool SetValueName(const UMovieGraphVariable* InGraphVariable, const FName InValue);
 
 	/** Sets the FString value of the specified property. Returns true on success, else false. */
-	UFUNCTION(BlueprintCallable, Category="Experimental")
+	UFUNCTION(BlueprintCallable, Category="Value")
 	bool SetValueString(const UMovieGraphVariable* InGraphVariable, const FString& InValue);
 
 	/** Sets the FText value of the specified property. Returns true on success, else false. */
-	UFUNCTION(BlueprintCallable, Category="Experimental")
+	UFUNCTION(BlueprintCallable, Category="Value")
 	bool SetValueText(const UMovieGraphVariable* InGraphVariable, const FText& InValue);
 
 	/** Sets the enum value of the specified property. Returns true on success, else false. */
-	UFUNCTION(BlueprintCallable, Category="Experimental")
+	UFUNCTION(BlueprintCallable, Category="Value")
 	bool SetValueEnum(const UMovieGraphVariable* InGraphVariable, const uint8 InValue, const UEnum* Enum);
 
 	/** Sets the struct value of the specified property. Returns true on success, else false. */
 	bool SetValueStruct(const UMovieGraphVariable* InGraphVariable, FConstStructView InValue);
 
 	/** Sets the object value of the specified property. Returns true on success, else false. */
-	UFUNCTION(BlueprintCallable, Category="Experimental")
+	UFUNCTION(BlueprintCallable, Category="Value")
 	bool SetValueObject(const UMovieGraphVariable* InGraphVariable, UObject* InValue);
 
 	/** Sets the class value of the specified property. Returns true on success, else false. */
-	UFUNCTION(BlueprintCallable, Category="Experimental")
+	UFUNCTION(BlueprintCallable, Category="Value")
 	bool SetValueClass(const UMovieGraphVariable* InGraphVariable, UClass* InValue);
 
 	/** Sets the serialized value of this member. The string should be the serialized representation of the value. Returns true on success, else false.*/
-	UFUNCTION(BlueprintCallable, Category="Experimental")
+	UFUNCTION(BlueprintCallable, Category="Value")
 	bool SetValueSerializedString(const UMovieGraphVariable* InGraphVariable, const FString& NewValue);
 
 	/** Sets the enum value of the specified property. Returns true on success, else false. */
@@ -193,15 +193,15 @@ public:
 	}
 
 	/** Gets the type of the value stored in the specified property. */
-	UFUNCTION(BlueprintCallable, Category="Experimental")
+	UFUNCTION(BlueprintCallable, Category="Value")
 	EMovieGraphValueType GetValueType(const UMovieGraphVariable* InGraphVariable) const;
 
 	/** Gets the object that defines the enum, struct, or class stored in the specified property. */
-	UFUNCTION(BlueprintCallable, Category="Experimental")
+	UFUNCTION(BlueprintCallable, Category="Value")
 	const UObject* GetValueTypeObject(const UMovieGraphVariable* InGraphVariable) const;
 
 	/** Gets the container type of the stored value in the specified property. */
-	UFUNCTION(BlueprintCallable, Category="Experimental")
+	UFUNCTION(BlueprintCallable, Category="Value")
 	EMovieGraphContainerType GetValueContainerType(const UMovieGraphVariable* InGraphVariable) const;
 
 	/**
@@ -214,14 +214,14 @@ public:
 	 * Updates an existing variable assignment for the provided graph variable to a new enable state, or adds a new
 	 * assignment and updates its enable state. Returns true on success, else false.
 	 */
-	UFUNCTION(BlueprintCallable, Category = "Experimental")
+	UFUNCTION(BlueprintCallable, Category = "Graph")
 	bool SetVariableAssignmentEnableState(const UMovieGraphVariable* InGraphVariable, bool bIsEnabled);
 
 	/**
 	 * Gets the enable state of the variable assignment for the provided graph variable. The enable state is provided
 	 * via bOutIsEnabled. Returns true if an enable state was set on the variable and bOutIsEnabled was changed, else false.
 	 */
-	UFUNCTION(BlueprintCallable, Category = "Experimental")
+	UFUNCTION(BlueprintCallable, Category = "Graph")
 	bool GetVariableAssignmentEnableState(const UMovieGraphVariable* InGraphVariable, bool& bOutIsEnabled);
 
 #if WITH_EDITOR
@@ -231,7 +231,7 @@ public:
 	 * corresponding graph variable will be removed, and overrides will be created for graph variables which do not have
 	 * existing overrides.
 	 */
-	UFUNCTION(BlueprintCallable, Category = "Experimental")
+	UFUNCTION(BlueprintCallable, Category = "Graph")
 	void UpdateGraphVariableOverrides();
 #endif
 
diff --git a/Engine/Plugins/MovieScene/MovieRenderPipeline/Source/MovieRenderPipelineEditor/Private/Widgets/SMoviePipelineQueueEditor.cpp b/Engine/Plugins/MovieScene/MovieRenderPipeline/Source/MovieRenderPipelineEditor/Private/Widgets/SMoviePipelineQueueEditor.cpp
--- a/Engine/Plugins/MovieScene/MovieRenderPipeline/Source/MovieRenderPipelineEditor/Private/Widgets/SMoviePipelineQueueEditor.cpp
+++ b/Engine/Plugins/MovieScene/MovieRenderPipeline/Source/MovieRenderPipelineEditor/Private/Widgets/SMoviePipelineQueueEditor.cpp
@@ -638,7 +638,7 @@ public:
 				if (!bIsShot)
 				{
 					MenuBuilder.AddMenuEntry(
-						LOCTEXT("ReplaceWithGraph_Label", "Replace with Graph (Experimental)"),
+						LOCTEXT("ReplaceWithGraph_Label", "Replace with Graph (Beta)"),
 						LOCTEXT("ReplaceWithGraph_Tooltip", "Replaces the current configuration with a new graph representation."),
 						FSlateIcon(),
 						FUIAction(InNewRenderGraph),
