MoviePipeline: Updated the Frame Count property on the Path Tracer node to have an override checkbox so graph evaluation considers value updates on this property (before, the default was always used). This required a customization due to how it interacts with the denoiser setting.

#rnx

diff --git a/Engine/Plugins/MovieScene/MovieRenderPipeline/Source/MovieRenderPipelineEditor/Private/Customizations/Graph/MovieGraphPathTracedRendererNodeCustomization.h b/Engine/Plugins/MovieScene/MovieRenderPipeline/Source/MovieRenderPipelineEditor/Private/Customizations/Graph/MovieGraphPathTracedRendererNodeCustomization.h
new file mode 100644
--- /dev/null
+++ b/Engine/Plugins/MovieScene/MovieRenderPipeline/Source/MovieRenderPipelineEditor/Private/Customizations/Graph/MovieGraphPathTracedRendererNodeCustomization.h
@@ -0,0 +1,77 @@
+// Copyright Epic Games, Inc. All Rights Reserved.
+
+#pragma once
+
+#include "DetailLayoutBuilder.h"
+#include "Graph/Nodes/MovieGraphPathTracerPassNode.h"
+#include "IDetailCustomization.h"
+#include "PropertyHandle.h"
+#include "ScopedTransaction.h"
+
+#define LOCTEXT_NAMESPACE "MoviePipelineEditor"
+
+/** Customize how the Path Traced Renderer node appears in the details panel. */
+class FMovieGraphPathTracedRendererNodeCustomization : public IDetailCustomization
+{
+public:
+	static TSharedRef<IDetailCustomization> MakeInstance()
+	{
+		return MakeShared<FMovieGraphPathTracedRendererNodeCustomization>();
+	}
+
+protected:
+	//~ Begin IDetailCustomization interface
+	virtual void CustomizeDetails(const TSharedPtr<IDetailLayoutBuilder>& DetailBuilder) override
+	{
+		CustomizeDetails(*DetailBuilder);
+	}
+	
+	virtual void CustomizeDetails(IDetailLayoutBuilder& DetailBuilder) override
+	{
+		for (const TWeakObjectPtr<UMovieGraphPathTracerRenderPassNode>& WeakPathTracerNode : DetailBuilder.GetObjectsOfTypeBeingCustomized<UMovieGraphPathTracerRenderPassNode>())
+		{
+			const TAttribute<bool> FrameCountEnabledAttr = TAttribute<bool>::Create([WeakPathTracerNode]()
+			{
+				if (const TStrongObjectPtr<UMovieGraphPathTracerRenderPassNode> PathTracerNodePin = WeakPathTracerNode.Pin())
+				{
+					return PathTracerNodePin->bOverride_FrameCount > 0;
+				}
+				
+				return false;
+			});
+
+			const TAttribute<EVisibility> FrameCountVisibleAttr = TAttribute<EVisibility>::Create([WeakPathTracerNode]()
+			{
+				if (const TStrongObjectPtr<UMovieGraphPathTracerRenderPassNode> PathTracerNodePin = WeakPathTracerNode.Pin())
+				{
+					return (PathTracerNodePin->DenoiserType == EMovieGraphPathTracerDenoiserType::Temporal) ? EVisibility::Visible : EVisibility::Collapsed;
+				}
+				
+				return EVisibility::Collapsed;
+			});
+
+			// Update the "Frame Count" property to only be visible when the Temporal denoiser is active. Note that we need this customization over
+			// simple EditCondition and EditConditionHides metadata on the UPROPERTY. In order to get the edit condition checkbox, the EditCondition
+			// needs to reference ONLY bOverride_FrameCount (and cannot have any logic around the denoiser type). To get around that, we have to
+			// define our own custom conditions here.
+			const TSharedRef<IPropertyHandle> FrameCountHandle = DetailBuilder.GetProperty(GET_MEMBER_NAME_CHECKED(UMovieGraphPathTracerRenderPassNode, FrameCount));
+			DetailBuilder.EditDefaultProperty(FrameCountHandle)
+				->Visibility(FrameCountVisibleAttr)
+				.EditCondition(
+					FrameCountEnabledAttr,
+					FOnBooleanValueChanged::CreateLambda([WeakPathTracerNode](const bool bNewValue)
+					{
+						if (const TStrongObjectPtr<UMovieGraphPathTracerRenderPassNode> PathTracerNodePin = WeakPathTracerNode.Pin())
+						{
+							FScopedTransaction Transaction(LOCTEXT("FrameCountEditConditionChanged", "Edit Condition Changed"));
+							PathTracerNodePin->Modify();
+							PathTracerNodePin->bOverride_FrameCount = bNewValue;
+						}
+					})
+			);
+		}
+	}
+	//~ End IDetailCustomization interface
+};
+
+#undef LOCTEXT_NAMESPACE
\ No newline at end of file
diff --git a/Engine/Plugins/MovieScene/MovieRenderPipeline/Source/MovieRenderPipelineEditor/Private/Graph/MovieGraphAssetToolkit.cpp b/Engine/Plugins/MovieScene/MovieRenderPipeline/Source/MovieRenderPipelineEditor/Private/Graph/MovieGraphAssetToolkit.cpp
--- a/Engine/Plugins/MovieScene/MovieRenderPipeline/Source/MovieRenderPipelineEditor/Private/Graph/MovieGraphAssetToolkit.cpp
+++ b/Engine/Plugins/MovieScene/MovieRenderPipeline/Source/MovieRenderPipelineEditor/Private/Graph/MovieGraphAssetToolkit.cpp
@@ -10,6 +10,7 @@
 #include "Customizations/Graph/MovieGraphModifiersCustomization.h"
 #include "Customizations/Graph/MovieGraphNamedResolutionCustomization.h"
 #include "Customizations/Graph/MovieGraphNodeCustomization.h"
+#include "Customizations/Graph/MovieGraphPathTracedRendererNodeCustomization.h"
 #include "Customizations/Graph/MovieGraphSelectNodeCustomization.h"
 #include "Customizations/Graph/MovieGraphSetCVarValueNodeCustomization.h"
 #include "Customizations/Graph/MovieGraphShowFlagsCustomization.h"
@@ -23,6 +24,7 @@
 #include "Graph/Nodes/MovieGraphDebugNode.h"
 #include "Graph/Nodes/MovieGraphFileOutputNode.h"
 #include "Graph/Nodes/MovieGraphModifierNode.h"
+#include "Graph/Nodes/MovieGraphPathTracerPassNode.h"
 #include "Graph/Nodes/MovieGraphSelectNode.h"
 #include "Graph/Nodes/MovieGraphSetMetadataAttributesNode.h"
 #include "Graph/Nodes/MovieGraphSetCVarValueNode.h"
@@ -434,6 +436,10 @@ TSharedRef<SDockTab> FMovieGraphAssetToolkit::SpawnTab_RenderGraphDetails(const
 		UMovieGraphDebugSettingNode::StaticClass(),
 		FOnGetDetailCustomizationInstance::CreateStatic(&FMovieGraphFormatTokenCustomization::MakeInstance));
 	
+	SelectedGraphObjectsDetailsWidget->RegisterInstancedCustomPropertyLayout(
+		UMovieGraphPathTracerRenderPassNode::StaticClass(),
+		FOnGetDetailCustomizationInstance::CreateStatic(&FMovieGraphPathTracedRendererNodeCustomization::MakeInstance));
+	
 	TSharedRef<SWidget> CustomContent = SAssignNew(NameAreaCustomContent, SHorizontalBox)
 	+ SHorizontalBox::Slot()
 	[
diff --git a/Engine/Plugins/MovieScene/MovieRenderPipeline/Source/MovieRenderPipelineRenderPasses/Public/Graph/Nodes/MovieGraphPathTracerPassNode.h b/Engine/Plugins/MovieScene/MovieRenderPipeline/Source/MovieRenderPipelineRenderPasses/Public/Graph/Nodes/MovieGraphPathTracerPassNode.h
--- a/Engine/Plugins/MovieScene/MovieRenderPipeline/Source/MovieRenderPipelineRenderPasses/Public/Graph/Nodes/MovieGraphPathTracerPassNode.h
+++ b/Engine/Plugins/MovieScene/MovieRenderPipeline/Source/MovieRenderPipelineRenderPasses/Public/Graph/Nodes/MovieGraphPathTracerPassNode.h
@@ -154,7 +154,7 @@ public:
 	* of the current frame (ie: 2 means consider 2 before, and 2 after the currently denoised frame), but other denoiser 
 	* implementations may interpret this value differently.
 	*/
-	UPROPERTY(EditAnywhere, BlueprintReadWrite, meta = (UIMin = 0, ClampMin = 0, UIMax = 3), Category = "Denoiser", meta = (EditCondition = "DenoiserType==EMovieGraphPathTracerDenoiserType::Temporal"))
+	UPROPERTY(EditAnywhere, BlueprintReadWrite, meta = (UIMin = 0, ClampMin = 0, UIMax = 3), Category = "Denoiser")
 	int32 FrameCount;
 
 	/**
