MoviePipeline: Fixed a crash that could occur in the graph if there was nothing to render and the pipeline immediately shuts down.

Also corrected some errors with telemetry.

#rnx

diff --git a/Engine/Plugins/Media/AppleProResMedia/Source/AppleProResMedia/Private/MoviePipelineAppleProResOutput.cpp b/Engine/Plugins/Media/AppleProResMedia/Source/AppleProResMedia/Private/MoviePipelineAppleProResOutput.cpp
--- a/Engine/Plugins/Media/AppleProResMedia/Source/AppleProResMedia/Private/MoviePipelineAppleProResOutput.cpp
+++ b/Engine/Plugins/Media/AppleProResMedia/Source/AppleProResMedia/Private/MoviePipelineAppleProResOutput.cpp
@@ -4,6 +4,7 @@
 #include "MoviePipeline.h"
 #include "MoviePipelineOutputSetting.h"
 #include "MoviePipelinePrimaryConfig.h"
+#include "MoviePipelineTelemetry.h"
 #include "ImagePixelData.h"
 #include "MoviePipelineImageQuantization.h"
 #include "SampleBuffer.h"
@@ -143,3 +144,8 @@ FText UMoviePipelineAppleProResOutput::GetDisplayText() const
 	}
 }
 #endif
+
+void UMoviePipelineAppleProResOutput::UpdateTelemetry(FMoviePipelineShotRenderTelemetry* InTelemetry) const
+{
+	InTelemetry->bUsesProRes = true;
+}
diff --git a/Engine/Plugins/Media/AppleProResMedia/Source/AppleProResMedia/Private/MoviePipelineAppleProResOutput.h b/Engine/Plugins/Media/AppleProResMedia/Source/AppleProResMedia/Private/MoviePipelineAppleProResOutput.h
--- a/Engine/Plugins/Media/AppleProResMedia/Source/AppleProResMedia/Private/MoviePipelineAppleProResOutput.h
+++ b/Engine/Plugins/Media/AppleProResMedia/Source/AppleProResMedia/Private/MoviePipelineAppleProResOutput.h
@@ -7,7 +7,7 @@
 #include "MoviePipelineAppleProResOutput.generated.h"
 
 // Forward Declare
-
+struct FMoviePipelineShotRenderTelemetry;
 
 UCLASS(BlueprintType)
 class UMoviePipelineAppleProResOutput : public UMoviePipelineVideoOutputBase
@@ -40,6 +40,10 @@ protected:
 #endif
 	// ~UMoviePipelineOutputBase Interface
 
+	// UMoviePipelineSetting Interface
+	virtual void UpdateTelemetry(FMoviePipelineShotRenderTelemetry* InTelemetry) const override;
+	// ~UMoviePipelineSetting Interface
+
 public:
 	/** Which Apple ProRes codec should we use? See Apple documentation for more specifics. Uses Rec 709 color primaries. */
 	UPROPERTY(EditAnywhere, BlueprintReadWrite, Category = "Settings")
diff --git a/Engine/Plugins/Media/AvidDNxMedia/Source/Source/Private/MoviePipelineAvidDNxOutput.cpp b/Engine/Plugins/Media/AvidDNxMedia/Source/Source/Private/MoviePipelineAvidDNxOutput.cpp
--- a/Engine/Plugins/Media/AvidDNxMedia/Source/Source/Private/MoviePipelineAvidDNxOutput.cpp
+++ b/Engine/Plugins/Media/AvidDNxMedia/Source/Source/Private/MoviePipelineAvidDNxOutput.cpp
@@ -4,6 +4,7 @@
 #include "MoviePipeline.h"
 #include "MoviePipelineOutputSetting.h"
 #include "MoviePipelinePrimaryConfig.h"
+#include "MoviePipelineTelemetry.h"
 #include "ImagePixelData.h"
 #include "MoviePipelineImageQuantization.h"
 #include "SampleBuffer.h"
@@ -104,4 +105,9 @@ void UMoviePipelineAvidDNxOutput::Finalize_EncodeThread(MovieRenderPipeline::IVi
 	// Commit this to disk.
 	FAvidWriter* CodecWriter = static_cast<FAvidWriter*>(InWriter);
 	CodecWriter->Writer->Finalize();
-}
\ No newline at end of file
+}
+
+void UMoviePipelineAvidDNxOutput::UpdateTelemetry(FMoviePipelineShotRenderTelemetry* InTelemetry) const
+{
+	InTelemetry->bUsesAvid = true;
+}
diff --git a/Engine/Plugins/Media/AvidDNxMedia/Source/Source/Private/MoviePipelineAvidDNxOutput.h b/Engine/Plugins/Media/AvidDNxMedia/Source/Source/Private/MoviePipelineAvidDNxOutput.h
--- a/Engine/Plugins/Media/AvidDNxMedia/Source/Source/Private/MoviePipelineAvidDNxOutput.h
+++ b/Engine/Plugins/Media/AvidDNxMedia/Source/Source/Private/MoviePipelineAvidDNxOutput.h
@@ -7,7 +7,7 @@
 #include "MoviePipelineAvidDNxOutput.generated.h"
 
 // Forward Declare
-
+struct FMoviePipelineShotRenderTelemetry;
 
 UCLASS(BlueprintType)
 class UMoviePipelineAvidDNxOutput : public UMoviePipelineVideoOutputBase
@@ -38,6 +38,10 @@ protected:
 #endif
 	// ~UMoviePipelineOutputBase Interface
 
+	// UMoviePipelineSetting Interface
+	virtual void UpdateTelemetry(FMoviePipelineShotRenderTelemetry* InTelemetry) const override;
+	// ~UMoviePipelineSetting Interface
+
 public:
 	/** Should we use a lossy compression for the output? */
 	UPROPERTY(EditAnywhere, BlueprintReadWrite, Category = "Settings")
diff --git a/Engine/Plugins/MovieScene/MovieRenderPipeline/Source/MovieRenderPipelineCore/Private/Graph/MovieGraphPipeline.cpp b/Engine/Plugins/MovieScene/MovieRenderPipeline/Source/MovieRenderPipelineCore/Private/Graph/MovieGraphPipeline.cpp
--- a/Engine/Plugins/MovieScene/MovieRenderPipeline/Source/MovieRenderPipelineCore/Private/Graph/MovieGraphPipeline.cpp
+++ b/Engine/Plugins/MovieScene/MovieRenderPipeline/Source/MovieRenderPipelineCore/Private/Graph/MovieGraphPipeline.cpp
@@ -39,6 +39,7 @@ UMovieGraphPipeline::UMovieGraphPipeline()
 	, bIsTransitioningState(false)
 	, bIsTearingDownShot(false)
 	, PipelineState(EMovieRenderPipelineState::Uninitialized)
+	, bDidStartInsightsCapture(false)
 {
 	OutputMerger = MakeShared<UE::MovieGraph::FMovieGraphOutputMerger>(this);
 	CustomEngineTimeStep = CreateDefaultSubobject<UMovieGraphEngineTimeStep>("MovieGraphEngineTimeStep");
@@ -865,6 +866,7 @@ void UMovieGraphPipeline::StartUnrealInsightsCapture(UMovieGraphEvaluatedConfig*
 	const bool bTraceStarted = FTraceAuxiliary::Start(FTraceAuxiliary::EConnectionType::File, *FinalFilePath);
 	if (bTraceStarted)
 	{
+		bDidStartInsightsCapture = true;
 		UE_LOG(LogMovieRenderPipeline, Log, TEXT("Started capturing UnrealInsights trace file to %s"), *FinalFilePath);
 	}
 	else
@@ -1456,11 +1458,7 @@ void UMovieGraphPipeline::TransitionToState(const EMovieRenderPipelineState InNe
 			}
 
 			// Stop Insights trace
-			constexpr bool bIncludeCDOs = false;
-			constexpr bool bExactMatch = true;
-			const UMovieGraphDebugSettingNode* DebugSetting =
-				PostRenderEvaluatedGraph->GetSettingForBranch<UMovieGraphDebugSettingNode>(UMovieGraphNode::GlobalsPinName, bIncludeCDOs, bExactMatch);
-			if (DebugSetting && DebugSetting->bCaptureUnrealInsightsTrace)
+			if (bDidStartInsightsCapture)
 			{
 				StopUnrealInsightsCapture();
 			}
diff --git a/Engine/Plugins/MovieScene/MovieRenderPipeline/Source/MovieRenderPipelineCore/Private/MoviePipelineTelemetry.cpp b/Engine/Plugins/MovieScene/MovieRenderPipeline/Source/MovieRenderPipelineCore/Private/MoviePipelineTelemetry.cpp
--- a/Engine/Plugins/MovieScene/MovieRenderPipeline/Source/MovieRenderPipelineCore/Private/MoviePipelineTelemetry.cpp
+++ b/Engine/Plugins/MovieScene/MovieRenderPipeline/Source/MovieRenderPipelineCore/Private/MoviePipelineTelemetry.cpp
@@ -84,6 +84,8 @@ void FMoviePipelineTelemetry::SendBeginShotRenderTelemetry(UMoviePipelineExecuto
 	EventAttributes.Add(FAnalyticsEventAttribute(TEXT("UsesSubgraphs"), ShotRenderTelemetry.bUsesSubgraphs));
 	EventAttributes.Add(FAnalyticsEventAttribute(TEXT("UsesPPMs"), ShotRenderTelemetry.bUsesPPMs));
 	EventAttributes.Add(FAnalyticsEventAttribute(TEXT("UsesAudio"), ShotRenderTelemetry.bUsesAudio));
+	EventAttributes.Add(FAnalyticsEventAttribute(TEXT("UsesAvid"), ShotRenderTelemetry.bUsesAvid));
+	EventAttributes.Add(FAnalyticsEventAttribute(TEXT("UsesProRes"), ShotRenderTelemetry.bUsesProRes));
 	EventAttributes.Add(FAnalyticsEventAttribute(TEXT("ResolutionX"), ShotRenderTelemetry.ResolutionX));
 	EventAttributes.Add(FAnalyticsEventAttribute(TEXT("ResolutionY"), ShotRenderTelemetry.ResolutionY));
 	EventAttributes.Add(FAnalyticsEventAttribute(TEXT("HandleFrameCount"), ShotRenderTelemetry.HandleFrameCount));
diff --git a/Engine/Plugins/MovieScene/MovieRenderPipeline/Source/MovieRenderPipelineCore/Public/Graph/MovieGraphPipeline.h b/Engine/Plugins/MovieScene/MovieRenderPipeline/Source/MovieRenderPipelineCore/Public/Graph/MovieGraphPipeline.h
--- a/Engine/Plugins/MovieScene/MovieRenderPipeline/Source/MovieRenderPipelineCore/Public/Graph/MovieGraphPipeline.h
+++ b/Engine/Plugins/MovieScene/MovieRenderPipeline/Source/MovieRenderPipelineCore/Public/Graph/MovieGraphPipeline.h
@@ -346,6 +346,9 @@ protected:
 	/** Simulation settings cache per cloth interactor object. Needs one per LOD, hence the array. */
 	TMap<TWeakObjectPtr<UObject>, TArray<MoviePipeline::FClothSimSettingsCache>> ClothSimCache;
 	
+	/** Whether an insights capture was started when the pipeline was initialized. */
+	bool bDidStartInsightsCapture;
+	
 public:
 	static FString DefaultPreviewWidgetAsset;
 };
\ No newline at end of file
diff --git a/Engine/Plugins/MovieScene/MovieRenderPipeline/Source/MovieRenderPipelineCore/Public/MoviePipelineTelemetry.h b/Engine/Plugins/MovieScene/MovieRenderPipeline/Source/MovieRenderPipelineCore/Public/MoviePipelineTelemetry.h
--- a/Engine/Plugins/MovieScene/MovieRenderPipeline/Source/MovieRenderPipelineCore/Public/MoviePipelineTelemetry.h
+++ b/Engine/Plugins/MovieScene/MovieRenderPipeline/Source/MovieRenderPipelineCore/Public/MoviePipelineTelemetry.h
@@ -32,6 +32,9 @@ struct MOVIERENDERPIPELINECORE_API FMoviePipelineShotRenderTelemetry
 	int32 TemporalSampleCount = 0;
 	int32 SpatialSampleCount = 0;
 	int32 RenderLayerCount = 0;
+
+	// Note: If adding an entry here, make sure to also update FMoviePipelineTelemetry::SendBeginShotRenderTelemetry()
+	// Also remember to track the telemetry in both the graph and legacy.
 };
 
 /** Responsible for sending out telemetry for both Movie Render Queue and Movie Render Graph. */
