MoviePipeline: Fixed an issue where "light type" show flags set on the Path Traced Renderer node would not affect the render.

Also updated the message shown in the graph during a render from "simulating" to "rendering".

#rnx

diff --git a/Engine/Plugins/MovieScene/MovieRenderPipeline/Source/MovieRenderPipelineEditor/Private/Widgets/Graph/SMovieGraphConfigPanel.cpp b/Engine/Plugins/MovieScene/MovieRenderPipeline/Source/MovieRenderPipelineEditor/Private/Widgets/Graph/SMovieGraphConfigPanel.cpp
--- a/Engine/Plugins/MovieScene/MovieRenderPipeline/Source/MovieRenderPipelineEditor/Private/Widgets/Graph/SMovieGraphConfigPanel.cpp
+++ b/Engine/Plugins/MovieScene/MovieRenderPipeline/Source/MovieRenderPipelineEditor/Private/Widgets/Graph/SMovieGraphConfigPanel.cpp
@@ -94,9 +94,6 @@ void SMoviePipelineGraphPanel::Construct(const FArguments& InArgs)
 		PipelineQueueEditorWidget->SetSelectedJobs(Jobs);
 	}
 
-	FGraphAppearanceInfo AppearanceInfo;
-	AppearanceInfo.CornerText = LOCTEXT("EditorCornerText", "Movie Graph Config");
-
 	SGraphEditor::FGraphEditorEvents InEvents;
 	InEvents.OnSelectionChanged = SGraphEditor::FOnSelectionChanged::CreateSP(this, &SMoviePipelineGraphPanel::OnSelectedNodesChanged);
 	InEvents.OnNodeDoubleClicked = FSingleNodeEvent::CreateSP(this, &SMoviePipelineGraphPanel::OnNodeDoubleClicked);
@@ -138,7 +135,7 @@ void SMoviePipelineGraphPanel::Construct(const FArguments& InArgs)
 		.GraphToEdit(EdGraph)
 		.AdditionalCommands(GraphEditorCommands)
 		.GraphEvents(InEvents)
-		.Appearance(AppearanceInfo)
+		.Appearance(this, &SMoviePipelineGraphPanel::GetGraphAppearanceInfo)
 	];
 }
 
@@ -726,6 +723,15 @@ void SMoviePipelineGraphPanel::OnNodeTitleCommitted(const FText& NewText, ETextC
 	}
 }
 
+FGraphAppearanceInfo SMoviePipelineGraphPanel::GetGraphAppearanceInfo() const
+{
+	FGraphAppearanceInfo AppearanceInfo;
+	AppearanceInfo.CornerText = LOCTEXT("EditorCornerText", "Movie Graph Config");
+	AppearanceInfo.PIENotifyText = LOCTEXT("PIENotifyText", "RENDERING");
+
+	return AppearanceInfo;
+}
+
 FReply SMoviePipelineGraphPanel::OnRenderLocalRequested()
 {
 	UMoviePipelineQueueSubsystem* Subsystem = GEditor->GetEditorSubsystem<UMoviePipelineQueueSubsystem>();
diff --git a/Engine/Plugins/MovieScene/MovieRenderPipeline/Source/MovieRenderPipelineEditor/Private/Widgets/Graph/SMovieGraphConfigPanel.h b/Engine/Plugins/MovieScene/MovieRenderPipeline/Source/MovieRenderPipelineEditor/Private/Widgets/Graph/SMovieGraphConfigPanel.h
--- a/Engine/Plugins/MovieScene/MovieRenderPipeline/Source/MovieRenderPipelineEditor/Private/Widgets/Graph/SMovieGraphConfigPanel.h
+++ b/Engine/Plugins/MovieScene/MovieRenderPipeline/Source/MovieRenderPipelineEditor/Private/Widgets/Graph/SMovieGraphConfigPanel.h
@@ -5,6 +5,7 @@
 #include "Widgets/DeclarativeSyntaxSupport.h"
 #include "Widgets/SCompoundWidget.h"
 #include "EditorUndoClient.h"
+#include "GraphEditor.h"
 
 class UMoviePipelineConfigBase;
 class SGraphEditor;
@@ -79,6 +80,9 @@ private:
 	void OnSelectedNodesChanged(const TSet<class UObject*>& NewSelection);
 	void OnNodeDoubleClicked(class UEdGraphNode* Node);
 	void OnNodeTitleCommitted(const FText& NewText, ETextCommit::Type CommitInfo, UEdGraphNode* NodeBeingChanged);
+
+	/** Provides appearance info to the graph widget. */
+	FGraphAppearanceInfo GetGraphAppearanceInfo() const;
 	
 	TObjectPtr<class UMovieGraphConfig> CurrentGraph;
 private:
diff --git a/Engine/Plugins/MovieScene/MovieRenderPipeline/Source/MovieRenderPipelineRenderPasses/Private/Graph/Nodes/MovieGraphPathTracerPassNode.cpp b/Engine/Plugins/MovieScene/MovieRenderPipeline/Source/MovieRenderPipelineRenderPasses/Private/Graph/Nodes/MovieGraphPathTracerPassNode.cpp
--- a/Engine/Plugins/MovieScene/MovieRenderPipeline/Source/MovieRenderPipelineRenderPasses/Private/Graph/Nodes/MovieGraphPathTracerPassNode.cpp
+++ b/Engine/Plugins/MovieScene/MovieRenderPipeline/Source/MovieRenderPipelineRenderPasses/Private/Graph/Nodes/MovieGraphPathTracerPassNode.cpp
@@ -222,11 +222,11 @@ bool UMovieGraphPathTracerRenderPassNode::GetAllowDenoiser() const
 
 FEngineShowFlags UMovieGraphPathTracerRenderPassNode::GetShowFlags() const
 {
-	FEngineShowFlags OutShowFlag = FEngineShowFlags(EShowFlagInitMode::ESFIM_Game);
-	OutShowFlag.SetPathTracing(true);
-	OutShowFlag.SetMotionBlur(!bEnableReferenceMotionBlur);
+	FEngineShowFlags OutShowFlags = Super::GetShowFlags();
+	OutShowFlags.SetPathTracing(true);
+	OutShowFlags.SetMotionBlur(!bEnableReferenceMotionBlur);
 
-	return OutShowFlag;
+	return OutShowFlags;
 }
 
 void UMovieGraphPathTracerRenderPassNode::UpdateTelemetry(FMoviePipelineShotRenderTelemetry* InTelemetry) const
