MovieGraph: Fixed an number of issues where MRG was submitting extra renders and then holding onto the results for the duration of the render even though the individual data samples were not actually intended to be used. This also fixes a number of minor issues with the Path Tracer Temporal Denoiser - expansion was not happening to allow proper warm-ups, if you used temporal denoising without regular Warm Up Frames added. Fixes an issue where when the Path Tracer had Use Reference Motion Blur disabled, it would read back too many samples at the end of a shot during cooldown. Fixed an issue where the deferred renderer (if used at the same time as the Path Tracer with Temporal Denoising) would read back unnecessary samples during the cool-down phase.

This also changes the default value for the Path Tracer Use Reference Motion Blur to true, as the only reason you want it off is for High Resolution Tiling, which isn't currently supported in MRG.

Specific Testing that was done:
Path Tracer with no Deferred Pass with various combinations of Temporal Sub-Samples, Spatial Sub-Samples, None/Spatial/Temporal Denoiser, MRG Warm Up Frames, MRG Motion Blur Emulation (both for the first frame)
The above tests but also rendering a Deferred Pass at the same time.
The above tests rendering one shot or multiple shots.
The above tests rendering with one frame or multiple frames.
Writing all samples from Deferred renderer.

diff --git a/Engine/Plugins/MovieScene/MovieRenderPipeline/Source/MovieRenderPipelineCore/Private/Graph/MovieGraphCoreTimeStep.cpp b/Engine/Plugins/MovieScene/MovieRenderPipeline/Source/MovieRenderPipelineCore/Private/Graph/MovieGraphCoreTimeStep.cpp
--- a/Engine/Plugins/MovieScene/MovieRenderPipeline/Source/MovieRenderPipelineCore/Private/Graph/MovieGraphCoreTimeStep.cpp
+++ b/Engine/Plugins/MovieScene/MovieRenderPipeline/Source/MovieRenderPipelineCore/Private/Graph/MovieGraphCoreTimeStep.cpp
@@ -17,6 +17,7 @@
 #include "GameFramework/WorldSettings.h"
 #include "ProfilingDebugging/MiscTrace.h"
 #include "MoviePipelineUtils.h"
+#include "Graph/Nodes/MovieGraphSamplingMethodNode.h"
 
 UMovieGraphCoreTimeStep::UMovieGraphCoreTimeStep()
 {
@@ -628,8 +629,7 @@ void UMovieGraphCoreTimeStep::ResetForEndOfOutputFrame()
 
 bool UMovieGraphCoreTimeStep::IsExpansionForTSRequired(const TObjectPtr<UMovieGraphEvaluatedConfig>& InConfig) const
 {
-	// ToDo: This needs to come from the config (once we have TemporalSampleCount there)
-	return false;
+	return GetTemporalSampleCountFromConfig(InConfig) > 1;
 }
 
 void UMovieGraphCoreTimeStep::UpdateFrameMetrics()
@@ -848,3 +848,21 @@ bool UMovieGraphEngineTimeStep::UpdateTimeStep(UEngine* /*InEngine*/)
 	return false;
 }
 
+int32 UMovieGraphCoreTimeStep::GetTemporalSampleCountFromConfig(UMovieGraphEvaluatedConfig* InConfig) const
+{
+	if (!ensure(InConfig))
+	{
+		return 1;
+	}
+
+	constexpr bool bIncludeCDOs = true;
+	const UMovieGraphSamplingMethodNode* SamplingMethod =
+		InConfig->GetSettingForBranch<UMovieGraphSamplingMethodNode>(UMovieGraphNode::GlobalsPinName, bIncludeCDOs);
+	if (SamplingMethod->TemporalSampleCount <= 0)
+	{
+		UE_LOG(LogMovieRenderPipeline, Error, TEXT("Sampling Method > Temporal Sample Count was zero, this is not allowed. Forcing value to 1!"));
+	}
+
+	return FMath::Max(SamplingMethod->TemporalSampleCount, 1);
+}
+
diff --git a/Engine/Plugins/MovieScene/MovieRenderPipeline/Source/MovieRenderPipelineCore/Private/Graph/MovieGraphDefaultRenderer.cpp b/Engine/Plugins/MovieScene/MovieRenderPipeline/Source/MovieRenderPipelineCore/Private/Graph/MovieGraphDefaultRenderer.cpp
--- a/Engine/Plugins/MovieScene/MovieRenderPipeline/Source/MovieRenderPipelineCore/Private/Graph/MovieGraphDefaultRenderer.cpp
+++ b/Engine/Plugins/MovieScene/MovieRenderPipeline/Source/MovieRenderPipelineCore/Private/Graph/MovieGraphDefaultRenderer.cpp
@@ -260,7 +260,7 @@ void UMovieGraphDefaultRenderer::Render(const FMovieGraphTimeStepData& InTimeSte
 		// If this is the first sample for this output frame, then we need to 
 		// talk to all of our render passes and ask them for what data they will
 		// produce, and set the Output Merger up with that knowledge.
-		UE::MovieGraph::FMovieGraphOutputMergerFrame& NewOutputFrame = GetOwningGraph()->GetOutputMerger()->AllocateNewOutputFrame_GameThread(InTimeStepData.RenderedFrameNumber);
+		UE::MovieGraph::FMovieGraphOutputMergerFrame& NewOutputFrame = GetOwningGraph()->GetOutputMerger()->AllocateNewOutputFrame_GameThread(InTimeStepData.OutputFrameNumber);
 
 		// Get the Traversal Context (not specific to any render pass) at the first sample. This is so
 		// we can easily fetch things that are shared between all render layers later.
@@ -274,7 +274,7 @@ void UMovieGraphDefaultRenderer::Render(const FMovieGraphTimeStepData& InTimeSte
 
 		// Register the frame with our render statistics as being worked on
 
-		UE::MovieGraph::FRenderTimeStatistics* TimeStats = GetRenderTimeStatistics(NewOutputFrame.TraversalContext.Time.RenderedFrameNumber);
+		UE::MovieGraph::FRenderTimeStatistics* TimeStats = GetRenderTimeStatistics(NewOutputFrame.TraversalContext.Time.OutputFrameNumber);
 		if (ensure(TimeStats))
 		{
 			TimeStats->StartTime = FDateTime::UtcNow();
diff --git a/Engine/Plugins/MovieScene/MovieRenderPipeline/Source/MovieRenderPipelineCore/Private/Graph/MovieGraphLinearTimeStep.cpp b/Engine/Plugins/MovieScene/MovieRenderPipeline/Source/MovieRenderPipelineCore/Private/Graph/MovieGraphLinearTimeStep.cpp
--- a/Engine/Plugins/MovieScene/MovieRenderPipeline/Source/MovieRenderPipelineCore/Private/Graph/MovieGraphLinearTimeStep.cpp
+++ b/Engine/Plugins/MovieScene/MovieRenderPipeline/Source/MovieRenderPipelineCore/Private/Graph/MovieGraphLinearTimeStep.cpp
@@ -13,13 +13,5 @@ int32 UMovieGraphLinearTimeStep::GetNextTemporalRangeIndex() const
 
 int32 UMovieGraphLinearTimeStep::GetTemporalSampleCount() const
 {
-	constexpr bool bIncludeCDOs = true;
-	const UMovieGraphSamplingMethodNode* SamplingMethod =
-		CurrentFrameData.EvaluatedConfig->GetSettingForBranch<UMovieGraphSamplingMethodNode>(UMovieGraphNode::GlobalsPinName, bIncludeCDOs);
-	if (SamplingMethod->TemporalSampleCount <= 0)
-	{
-		UE_LOG(LogMovieRenderPipeline, Error, TEXT("Sampling Method > Temporal Sample Count was zero, this is not allowed. Forcing value to 1!"));
-	}
-
-	return FMath::Max(SamplingMethod->TemporalSampleCount, 1);
+	return GetTemporalSampleCountFromConfig(CurrentFrameData.EvaluatedConfig.Get());
 }
diff --git a/Engine/Plugins/MovieScene/MovieRenderPipeline/Source/MovieRenderPipelineCore/Private/Graph/MovieGraphOutputMerger.cpp b/Engine/Plugins/MovieScene/MovieRenderPipeline/Source/MovieRenderPipelineCore/Private/Graph/MovieGraphOutputMerger.cpp
--- a/Engine/Plugins/MovieScene/MovieRenderPipeline/Source/MovieRenderPipelineCore/Private/Graph/MovieGraphOutputMerger.cpp
+++ b/Engine/Plugins/MovieScene/MovieRenderPipeline/Source/MovieRenderPipelineCore/Private/Graph/MovieGraphOutputMerger.cpp
@@ -59,14 +59,15 @@ namespace UE::MovieGraph
 			const UE::MovieGraph::FMovieGraphSampleState* Payload = LocalData->GetPayload<UE::MovieGraph::FMovieGraphSampleState>();
 			const FMovieGraphTraversalContext& TraversalContext = Payload->TraversalContext;
 			const TObjectPtr<UMovieGraphEvaluatedConfig> EvaluatedConfig = TraversalContext.Time.EvaluatedConfig;
-
+			
 			// Resolve the file output path
 			FString FinalFilePath;
 			{
 				// TODO: Add Tile X/Y when tiling is available
-				const FString OutputName = FString::Printf(TEXT("%s_%s_%s_%s_SS_%d_TS_%d.%d"),
+				const FString OutputName = FString::Printf(TEXT("%s_%s_%s_%s_%s_SS_%d_TS_%d.%d"),
 					*TraversalContext.Shot->OuterName,
 					*TraversalContext.RenderDataIdentifier.LayerName,
+					*TraversalContext.RenderDataIdentifier.RendererName,
 					*TraversalContext.RenderDataIdentifier.SubResourceName,
 					*TraversalContext.RenderDataIdentifier.CameraName,
 					TraversalContext.Time.SpatialSampleIndex,
@@ -112,12 +113,12 @@ namespace UE::MovieGraph
 		FMovieGraphSampleState* Payload = InData->GetPayload<FMovieGraphSampleState>();
 		check(Payload);
 		
-		const int32 RenderedFrameNumber = Payload->TraversalContext.Time.RenderedFrameNumber;
+		const int32 IndexedFrameNumber = Payload->TraversalContext.Time.OutputFrameNumber;
 
 		// See if we can find the frame this data is for. This should always be valid, if it's not
 		// valid it means they either forgot to declare they were going to produce it, or this is
 		// coming in after the system already thinks it's finished that frame.
-		FMovieGraphOutputMergerFrame* OutputFrame = PendingData.Find(RenderedFrameNumber);
+		FMovieGraphOutputMergerFrame* OutputFrame = PendingData.Find(IndexedFrameNumber);
 		
 		// Make sure we expected this frame number
 		if (!ensureAlwaysMsgf(OutputFrame, TEXT("Received data for unknown frame. Frame was either already processed or not queued yet!")))
@@ -159,7 +160,7 @@ namespace UE::MovieGraph
 			);
 			// Move this frame into our FinishedFrames array so the Game Thread can read it at its leisure
 			FMovieGraphOutputMergerFrame FinalFrame;
-			ensureMsgf(PendingData.RemoveAndCopyValue(RenderedFrameNumber, FinalFrame), TEXT("Could not find frame in pending data, output will be skipped!"));
+			ensureMsgf(PendingData.RemoveAndCopyValue(IndexedFrameNumber, FinalFrame), TEXT("Could not find frame in pending data, output will be skipped!"));
 			
 			// TQueue is thread safe so it's okay to just push the data into it.
 			FinishedFrames.Enqueue(MoveTemp(FinalFrame));
diff --git a/Engine/Plugins/MovieScene/MovieRenderPipeline/Source/MovieRenderPipelineCore/Private/Graph/MovieGraphPipeline.cpp b/Engine/Plugins/MovieScene/MovieRenderPipeline/Source/MovieRenderPipelineCore/Private/Graph/MovieGraphPipeline.cpp
--- a/Engine/Plugins/MovieScene/MovieRenderPipeline/Source/MovieRenderPipelineCore/Private/Graph/MovieGraphPipeline.cpp
+++ b/Engine/Plugins/MovieScene/MovieRenderPipeline/Source/MovieRenderPipelineCore/Private/Graph/MovieGraphPipeline.cpp
@@ -563,12 +563,6 @@ void UMovieGraphPipeline::BuildShotListFromDataSource()
 			Shot->ShotInfo.CachedShotTickResolution = Shot->ShotInfo.SubSectionHierarchy->MovieScene->GetTickResolution();
 		}
 		
-		const bool bPrePass = true;
-		const bool bExpandForTemporalSubSample = GraphTimeStepInstances.Last()->IsExpansionForTSRequired(EvaluatedConfig);
-		ExpandShot(Shot, OutputNode->HandleFrameCount, bExpandForTemporalSubSample, bPrePass, FinalFrameRate, TickResolution, WarmUpNode->NumWarmUpFrames);
-
-		Shot->ShotInfo.CurrentTimeInRoot = Shot->ShotInfo.TotalOutputRangeRoot.GetLowerBoundValue();
-
 		// Query the max frame count.
 		int32 MaxCoolingDownFrameCount = 0;
 		for (const FName& BranchName : EvaluatedConfig->GetBranchNames())
@@ -581,11 +575,19 @@ void UMovieGraphPipeline::BuildShotListFromDataSource()
 				MaxCoolingDownFrameCount = FMath::Max(Render->GetCoolingDownFrameCount(), MaxCoolingDownFrameCount);
 			}
 		}
-		
+
+		// When using cooldown, we need at least that many warm-up frames even if they have otherwise chosen not to do warm-ups.
 		Shot->ShotInfo.NumEngineCoolDownFramesRemaining = MaxCoolingDownFrameCount;
-		// When using cooldown, we need at least that many warm-up frames even if they have otherwise chose not to do warm-ups.
 		Shot->ShotInfo.NumEngineWarmUpFramesRemaining = FMath::Max(WarmUpNode->NumWarmUpFrames, MaxCoolingDownFrameCount);
 		Shot->ShotInfo.bEmulateFirstFrameMotionBlur = WarmUpNode->bEmulateMotionBlur;
+
+		const bool bPrePass = true;
+		const bool bExpandForTemporalSubSample = GraphTimeStepInstances.Last()->IsExpansionForTSRequired(EvaluatedConfig);
+		
+		// We need to wait until the NumEngineWarmUpFramesRemaining has been set by either actual warm up frames, or by cool-down frames before we expand.
+		ExpandShot(Shot, OutputNode->HandleFrameCount, bExpandForTemporalSubSample, bPrePass, FinalFrameRate, TickResolution, Shot->ShotInfo.NumEngineWarmUpFramesRemaining);
+
+		Shot->ShotInfo.CurrentTimeInRoot = Shot->ShotInfo.TotalOutputRangeRoot.GetLowerBoundValue();
 		Shot->ShotInfo.CalculateWorkMetrics();
 		Shot->ShotInfo.VersionNumber = ResolveVersionForShot(Shot, EvaluatedConfig);
 	}
@@ -1557,7 +1559,7 @@ void UMovieGraphPipeline::ProcessOutstandingFinishedFrames()
 		UE::MovieGraph::FMovieGraphOutputMergerFrame OutputFrame;
 		OutputMerger->GetFinishedFrames().Dequeue(OutputFrame);
 
-		UE::MovieGraph::FRenderTimeStatistics* TimeStats = GetRendererInstance()->GetRenderTimeStatistics(OutputFrame.TraversalContext.Time.RenderedFrameNumber);
+		UE::MovieGraph::FRenderTimeStatistics* TimeStats = GetRendererInstance()->GetRenderTimeStatistics(OutputFrame.TraversalContext.Time.OutputFrameNumber);
 		if (ensure(TimeStats))
 		{
 			TimeStats->EndTime = FDateTime::UtcNow();
diff --git a/Engine/Plugins/MovieScene/MovieRenderPipeline/Source/MovieRenderPipelineCore/Public/Graph/MovieGraphCoreTimeStep.h b/Engine/Plugins/MovieScene/MovieRenderPipeline/Source/MovieRenderPipelineCore/Public/Graph/MovieGraphCoreTimeStep.h
--- a/Engine/Plugins/MovieScene/MovieRenderPipeline/Source/MovieRenderPipelineCore/Public/Graph/MovieGraphCoreTimeStep.h
+++ b/Engine/Plugins/MovieScene/MovieRenderPipeline/Source/MovieRenderPipelineCore/Public/Graph/MovieGraphCoreTimeStep.h
@@ -56,6 +56,7 @@ protected:
 	virtual bool IsFirstTemporalSample() const;
 	virtual void ResetForEndOfOutputFrame();
 	virtual float GetBlendedMotionBlurAmount();
+	virtual int32 GetTemporalSampleCountFromConfig(UMovieGraphEvaluatedConfig* InConfig) const;
 
 protected:
 	/** This is the output data needed by the rest of MRQ to produce a frame. */
diff --git a/Engine/Plugins/MovieScene/MovieRenderPipeline/Source/MovieRenderPipelineRenderPasses/Private/Graph/Nodes/MovieGraphPathTracerPassNode.cpp b/Engine/Plugins/MovieScene/MovieRenderPipeline/Source/MovieRenderPipelineRenderPasses/Private/Graph/Nodes/MovieGraphPathTracerPassNode.cpp
--- a/Engine/Plugins/MovieScene/MovieRenderPipeline/Source/MovieRenderPipelineRenderPasses/Private/Graph/Nodes/MovieGraphPathTracerPassNode.cpp
+++ b/Engine/Plugins/MovieScene/MovieRenderPipeline/Source/MovieRenderPipelineRenderPasses/Private/Graph/Nodes/MovieGraphPathTracerPassNode.cpp
@@ -16,7 +16,7 @@ TUniquePtr<UE::MovieGraph::Rendering::FMovieGraphImagePassBase> UMovieGraphPathT
 
 UMovieGraphPathTracerRenderPassNode::UMovieGraphPathTracerRenderPassNode()
 	: SpatialSampleCount(1)
-	, bEnableReferenceMotionBlur(false)
+	, bEnableReferenceMotionBlur(true)
 	, bEnableDenoiser(true)
 	, DenoiserType(EMovieGraphPathTracerDenoiserType::Spatial)
 	, FrameCount(2)
diff --git a/Engine/Plugins/MovieScene/MovieRenderPipeline/Source/MovieRenderPipelineRenderPasses/Private/Graph/Renderers/MovieGraphDeferredPass.cpp b/Engine/Plugins/MovieScene/MovieRenderPipeline/Source/MovieRenderPipelineRenderPasses/Private/Graph/Renderers/MovieGraphDeferredPass.cpp
--- a/Engine/Plugins/MovieScene/MovieRenderPipeline/Source/MovieRenderPipelineRenderPasses/Private/Graph/Renderers/MovieGraphDeferredPass.cpp
+++ b/Engine/Plugins/MovieScene/MovieRenderPipeline/Source/MovieRenderPipelineRenderPasses/Private/Graph/Renderers/MovieGraphDeferredPass.cpp
@@ -51,6 +51,7 @@ void FMovieGraphDeferredPass::Setup(TWeakObjectPtr<UMovieGraphDefaultRenderer> I
 	// initialize the frames to delay for post submission.
 	UMovieGraphRenderPassNode* RenderPassNode = InLayer.RenderPassNode.Get();
 	FramesToDelayPostSubmission = RenderPassNode ? RenderPassNode->GetCoolingDownFrameCount() : 0;
+	RemainingCooldownReadbackFrames = FramesToDelayPostSubmission;
 }
 
 void FMovieGraphDeferredPass::Teardown()
@@ -101,6 +102,24 @@ FName FMovieGraphDeferredPass::GetBranchName() const
 	return LayerData.BranchName;
 }
 
+bool FMovieGraphDeferredPass::ShouldDiscardOutput(const TSharedRef<FSceneViewFamilyContext>& InFamily, const UE::MovieGraph::DefaultRenderer::FCameraInfo& InCameraInfo) const
+{
+	TObjectPtr<UMovieGraphDefaultRenderer> GraphRenderer = GetRenderer().Get();
+	if (!GraphRenderer)
+	{
+		return false;
+	}
+
+	if (UMovieGraphPipeline* Pipeline = GraphRenderer->GetOwningGraph())
+	{
+		// The deferred renderer should attempt to discard anything that isn't from the rendering state, as we don't need any data from
+		// the warm-up or cool-down phases of the shot.
+		return Pipeline->GetActiveShotList()[Pipeline->GetCurrentShotIndex()]->ShotInfo.State != EMovieRenderShotState::Rendering;
+	}
+	return false;
+}
+
+
 UMovieGraphImagePassBaseNode* FMovieGraphDeferredPass::GetParentNode(UMovieGraphEvaluatedConfig* InConfig) const
 {
 	// This is a bit of a workaround for the fact that the pass doesn't have a strong pointer to the node it's supposed to be associated with,
@@ -127,7 +146,9 @@ void FMovieGraphDeferredPass::Render(const FMovieGraphTraversalContext& InFrameT
 	}
 
 	UMovieGraphImagePassBaseNode* ParentNodeThisFrame = GetParentNode(InTimeData.EvaluatedConfig);
-	const bool bWriteAllSamples = ParentNodeThisFrame->GetWriteAllSamples();
+	
+	// We can only write rendered frames to disk right now (warm-up/cool-down indexes aren't propagated so files overwrite each other).
+	const bool bWriteAllSamples = ParentNodeThisFrame->GetWriteAllSamples() && InFrameTraversalContext.Shot->ShotInfo.State == EMovieRenderShotState::Rendering;
 	const bool bIsRenderingState = InFrameTraversalContext.Shot->ShotInfo.State == EMovieRenderShotState::Rendering ||
 									InFrameTraversalContext.Shot->ShotInfo.State == EMovieRenderShotState::CoolingDown;
 	int32 NumSpatialSamples = FMath::Max(1, bIsRenderingState ? ParentNodeThisFrame->GetNumSpatialSamples() : ParentNodeThisFrame->GetNumSpatialSamplesDuringWarmUp());
@@ -347,6 +368,26 @@ void FMovieGraphDeferredPass::Render(const FMovieGraphTraversalContext& InFrameT
 		{
 			continue;
 		}
+
+		// Example Assumptions: 2 frame denoise temporal denoise with 8 temporal sub-samples. 
+		// If you're using Cooldown Frames, we can get into a scenario where due to the particular render pass settings,
+		// you don't need all the cooldown frames. If you're using Path Tracer's "Use Reference Motion Blur" then the 
+		// above Discard is true for everything but the last sample. This means that we needed 8 Cool Down _samples_ to 
+		// produce the two output frames (matching the 2 frame delay in the PT denoiser).
+		// But if Use Reference Motion Blur is off, then the first two samples of the cool-down are enough to finish flushing
+		// the PT denoiser, and the remaining 14 end up confusing the system because it gets data it doesn't think it should.
+		if(InFrameTraversalContext.Shot->ShotInfo.State == EMovieRenderShotState::CoolingDown)
+		{
+			// When we're cooling down, we track the number of times we actually
+			// go to do a readback (ie: pass the above bShouldDiscardOutput check)
+			// and once we reach the number needed to actually flush the PT denoiser
+			// we stop submitting.
+			if (RemainingCooldownReadbackFrames == 0)
+			{
+				continue;
+			}
+			RemainingCooldownReadbackFrames--;
+		}
 		
 		// Post-submission is a little bit complicated to allow supporting temporal denoisers in the Path Tracer.
 		// When using the denoiser with a sample frame count of 2, for a frame to be produced it must look 2 frames
diff --git a/Engine/Plugins/MovieScene/MovieRenderPipeline/Source/MovieRenderPipelineRenderPasses/Private/Graph/Renderers/MovieGraphPathTracerPass.cpp b/Engine/Plugins/MovieScene/MovieRenderPipeline/Source/MovieRenderPipelineRenderPasses/Private/Graph/Renderers/MovieGraphPathTracerPass.cpp
--- a/Engine/Plugins/MovieScene/MovieRenderPipeline/Source/MovieRenderPipelineRenderPasses/Private/Graph/Renderers/MovieGraphPathTracerPass.cpp
+++ b/Engine/Plugins/MovieScene/MovieRenderPipeline/Source/MovieRenderPipelineRenderPasses/Private/Graph/Renderers/MovieGraphPathTracerPass.cpp
@@ -42,7 +42,10 @@ namespace Private
 }
 
 bool FMovieGraphPathTracerPass::ShouldDiscardOutput(const TSharedRef<FSceneViewFamilyContext>& InFamily, const UE::MovieGraph::DefaultRenderer::FCameraInfo& InCameraInfo) const
-{ 
+{
+	// We intentionally skip calling the super because it discards any samples not done during Rendering state,
+	// but we need to not skip the results due to Temporal Denoising cooldown frames.
+
 	int32 SampleCount, SampleIndex;
 	Private::GetSampleData(InFamily, InCameraInfo, SampleCount, SampleIndex);
 
diff --git a/Engine/Plugins/MovieScene/MovieRenderPipeline/Source/MovieRenderPipelineRenderPasses/Public/Graph/Renderers/MovieGraphDeferredPass.h b/Engine/Plugins/MovieScene/MovieRenderPipeline/Source/MovieRenderPipelineRenderPasses/Public/Graph/Renderers/MovieGraphDeferredPass.h
--- a/Engine/Plugins/MovieScene/MovieRenderPipeline/Source/MovieRenderPipelineRenderPasses/Public/Graph/Renderers/MovieGraphDeferredPass.h
+++ b/Engine/Plugins/MovieScene/MovieRenderPipeline/Source/MovieRenderPipelineRenderPasses/Public/Graph/Renderers/MovieGraphDeferredPass.h
@@ -22,6 +22,8 @@ namespace UE::MovieGraph::Rendering
 		virtual void AddReferencedObjects(FReferenceCollector& Collector) override;
 		virtual FName GetBranchName() const override;
 		virtual UMovieGraphImagePassBaseNode* GetParentNode(UMovieGraphEvaluatedConfig* InConfig) const override;
+		bool ShouldDiscardOutput(const TSharedRef<FSceneViewFamilyContext>& InFamily, const UE::MovieGraph::DefaultRenderer::FCameraInfo& InCameraInfo) const override;
+
 		// End FMovieGraphImagePassBase
 			
 	protected:
@@ -37,6 +39,9 @@ namespace UE::MovieGraph::Rendering
 
 		// The number of frames to delay to send frames from SubmissionQueue to post-render submission.
 		int32 FramesToDelayPostSubmission;
+		
+		// If using cooldown, the number of cool-down frames we still need to process.
+		int32 RemainingCooldownReadbackFrames;
 
 		// FIFO queue of rendered frames. It allows frames to be sent to post-render submission with a delay if needed (e.g., when temporal denoising is used with path tracers).  
 		TQueue<FMovieGraphPostRendererSubmissionParams> SubmissionQueue;
