MoviePipeline: Added reroute nodes to the graph (under Utility -> Reroute in the node creation menu).

diff --git a/Engine/Plugins/MovieScene/MovieRenderPipeline/Source/MovieRenderPipelineCore/Private/Graph/MovieGraphConfig.cpp b/Engine/Plugins/MovieScene/MovieRenderPipeline/Source/MovieRenderPipelineCore/Private/Graph/MovieGraphConfig.cpp
--- a/Engine/Plugins/MovieScene/MovieRenderPipeline/Source/MovieRenderPipelineCore/Private/Graph/MovieGraphConfig.cpp
+++ b/Engine/Plugins/MovieScene/MovieRenderPipeline/Source/MovieRenderPipelineCore/Private/Graph/MovieGraphConfig.cpp
@@ -12,6 +12,7 @@
 #include "Graph/Nodes/MovieGraphGlobalOutputSettingNode.h"
 #include "Graph/Nodes/MovieGraphRemoveRenderSettingNode.h"
 #include "Graph/Nodes/MovieGraphRenderLayerNode.h"
+#include "Graph/Nodes/MovieGraphRerouteNode.h"
 #include "Graph/Nodes/MovieGraphSubgraphNode.h"
 #include "Graph/Nodes/MovieGraphVariableNode.h"
 #include "Graph/Nodes/MovieGraphSelectNode.h"
@@ -1752,29 +1753,35 @@ bool UMovieGraphConfig::CreateFlattenedGraph_Recursive(UMovieGraphEvaluatedConfi
 			}
 		}
 #endif
-		const UMovieGraphSettingNode* NodeAsSetting = CastChecked<UMovieGraphSettingNode>(Node);
-		const FString& NodeInstanceName = NodeAsSetting->GetNodeInstanceName();
-		
-		UMovieGraphSettingNode* ExistingNode = Cast<UMovieGraphSettingNode>(OutBranchConfig.GetNodeByClassExactMatch(Node->GetClass(), NodeInstanceName));
-		if (!ExistingNode)
+
+		// Initialize (or copy properties to) the node in the evaluated graph. Reroute nodes aren't included in this because they're
+		// passthrough and have no properties.
+		if (!Node->IsA<UMovieGraphRerouteNode>())
 		{
-			// Create a new instance of this node inside our flattened eval graph
-			ExistingNode = NewObject<UMovieGraphSettingNode>(InOwningConfig, Node->GetClass());
-			OutBranchConfig.NamedNodes.FindOrAdd(NodeInstanceName).NodeInstances.Add(ExistingNode);
+			const UMovieGraphSettingNode* NodeAsSetting = CastChecked<UMovieGraphSettingNode>(Node);
+			const FString& NodeInstanceName = NodeAsSetting->GetNodeInstanceName();
+		
+			UMovieGraphSettingNode* ExistingNode = Cast<UMovieGraphSettingNode>(OutBranchConfig.GetNodeByClassExactMatch(Node->GetClass(), NodeInstanceName));
+			if (!ExistingNode)
+			{
+				// Create a new instance of this node inside our flattened eval graph
+				ExistingNode = NewObject<UMovieGraphSettingNode>(InOwningConfig, Node->GetClass());
+				OutBranchConfig.NamedNodes.FindOrAdd(NodeInstanceName).NodeInstances.Add(ExistingNode);
 
-			// This shouldn't be abused by nodes, but in some very rare cases the flattened node needs to be "primed" before being put through the
-			// evaluation process.
-			ExistingNode->PrepareForFlattening(NodeAsSetting);
+				// This shouldn't be abused by nodes, but in some very rare cases the flattened node needs to be "primed" before being put through the
+				// evaluation process.
+				ExistingNode->PrepareForFlattening(NodeAsSetting);
 
-			// Set all of the boolean edit condition values to false, so we can use "true" to indicate
-			// that the value was overridden already during traversal.
-			InitializeFlattenedNode(ExistingNode);
-		}
+				// Set all of the boolean edit condition values to false, so we can use "true" to indicate
+				// that the value was overridden already during traversal.
+				InitializeFlattenedNode(ExistingNode);
+			}
 
-		// Now do a property-copy from this node onto our flattened one. We don't use the generic property
-		// copy routines in the engine because we have special handling (we want to check if the property
-		// is actually marked for override, and also skip if this has already been overridden).
-		CopyOverriddenProperties(Node, ExistingNode, InEvaluationContext);
+			// Now do a property-copy from this node onto our flattened one. We don't use the generic property
+			// copy routines in the engine because we have special handling (we want to check if the property
+			// is actually marked for override, and also skip if this has already been overridden).
+			CopyOverriddenProperties(Node, ExistingNode, InEvaluationContext);
+		}
 	}
 
 	// If this is a special "removal" node, keep track of the type that should be removed. Since this method is recursive,
@@ -1815,7 +1822,8 @@ bool UMovieGraphConfig::CreateFlattenedGraph_Recursive(UMovieGraphEvaluatedConfi
 	{
 		for (UMovieGraphEdge* Edge : Pin->Edges)
 		{
-			UMovieGraphPin* OtherPin = Edge->GetOtherPin(Pin);
+			constexpr bool bFollowRerouteConnections = true;
+			UMovieGraphPin* OtherPin = Edge->GetOtherPin(Pin, bFollowRerouteConnections);
 			if (!OtherPin)
 			{
 				continue;
@@ -1876,7 +1884,8 @@ void UMovieGraphConfig::VisitUpstreamNodes_Recursive(UMovieGraphNode* FromNode,
 	{
 		for (const UMovieGraphPin* ConnectedPin : Pin->GetAllConnectedPins())
 		{
-			if (ConnectedPin->Properties.bIsBranch)
+			// Branches are followed during traversal, but we also want to traverse reroute nodes as well since they're passthrough
+			if (ConnectedPin->Properties.bIsBranch || (Pin->Node && Pin->Node->IsA<UMovieGraphRerouteNode>()))
 			{
 				bool bContinueVisiting = true;
 				if (VisitCallback.IsBound())
@@ -1912,7 +1921,8 @@ void UMovieGraphConfig::VisitDownstreamNodes_Recursive(UMovieGraphNode* FromNode
 	{
 		for (const UMovieGraphPin* ConnectedPin : Pin->GetAllConnectedPins())
 		{
-			if (ConnectedPin->Properties.bIsBranch)
+			// Branches are followed during traversal, but we also want to traverse reroute nodes as well since they're passthrough
+			if (ConnectedPin->Properties.bIsBranch || (Pin->Node && Pin->Node->IsA<UMovieGraphRerouteNode>()))
 			{
 				bool bContinueVisiting = true;
 				if (VisitCallback.IsBound())
diff --git a/Engine/Plugins/MovieScene/MovieRenderPipeline/Source/MovieRenderPipelineCore/Private/Graph/MovieGraphEdge.cpp b/Engine/Plugins/MovieScene/MovieRenderPipeline/Source/MovieRenderPipelineCore/Private/Graph/MovieGraphEdge.cpp
new file mode 100644
--- /dev/null
+++ b/Engine/Plugins/MovieScene/MovieRenderPipeline/Source/MovieRenderPipelineCore/Private/Graph/MovieGraphEdge.cpp
@@ -0,0 +1,36 @@
+// Copyright Epic Games, Inc. All Rights Reserved.
+
+#include "Graph/MovieGraphEdge.h"
+
+#include "Graph/Nodes/MovieGraphRerouteNode.h"
+
+bool UMovieGraphEdge::IsValid() const
+{
+	return InputPin.Get() && OutputPin.Get();
+}
+
+UMovieGraphPin* UMovieGraphEdge::GetOtherPin(const UMovieGraphPin* InFirstPin, const bool bFollowRerouteConnections)
+{
+	check(InFirstPin == InputPin || InFirstPin == OutputPin);
+	UMovieGraphPin* OtherPin = InFirstPin == InputPin ? OutputPin : InputPin;
+
+	if (!OtherPin)
+	{
+		return nullptr;
+	}
+
+	// If the other pin is connected to the reroute node, recursively follow the chain until a non-reroute node is reached
+	if (bFollowRerouteConnections)
+	{
+		if (const UMovieGraphRerouteNode* RerouteNode = Cast<UMovieGraphRerouteNode>(OtherPin->Node))
+		{
+			UMovieGraphPin* PassThroughPin = RerouteNode->GetPassThroughPin(OtherPin);
+			if (PassThroughPin && !PassThroughPin->Edges.IsEmpty())
+			{
+				OtherPin = PassThroughPin->Edges[0]->GetOtherPin(PassThroughPin, bFollowRerouteConnections);
+			}
+		}
+	}
+		
+	return OtherPin;
+}
\ No newline at end of file
diff --git a/Engine/Plugins/MovieScene/MovieRenderPipeline/Source/MovieRenderPipelineCore/Private/Graph/MovieGraphNode.cpp b/Engine/Plugins/MovieScene/MovieRenderPipeline/Source/MovieRenderPipelineCore/Private/Graph/MovieGraphNode.cpp
--- a/Engine/Plugins/MovieScene/MovieRenderPipeline/Source/MovieRenderPipelineCore/Private/Graph/MovieGraphNode.cpp
+++ b/Engine/Plugins/MovieScene/MovieRenderPipeline/Source/MovieRenderPipelineCore/Private/Graph/MovieGraphNode.cpp
@@ -3,6 +3,7 @@
 #include "Graph/MovieGraphNode.h"
 
 #include "Graph/MovieGraphConfig.h"
+#include "Graph/Nodes/MovieGraphRerouteNode.h"
 #include "UObject/TextProperty.h"
 
 FName UMovieGraphNode::GlobalsPinName = FName("Globals");
@@ -558,14 +559,17 @@ TArray<UMovieGraphPin*> UMovieGraphNode::EvaluatePinsToFollow(FMovieGraphEvaluat
 		return PinsToFollow;
 	}
 
-	// If this is a data (non-branch) connection, just return the first connected pin.
+	// If this is a data (non-branch) connection, just return the first connected pin (following any reroute nodes).
 	if (!InContext.PinBeingFollowed->Properties.bIsBranch)
 	{
+		const TObjectPtr<UMovieGraphNode> OwningNode = InContext.PinBeingFollowed->Node;
+		
 		// Only do this if the pin is an input though. For data output pins, the generic case does not have enough information to determine what
-		// the upstream connected data pin is.
-		if (InContext.PinBeingFollowed->IsInputPin())
+		// the upstream connected data pin is. If this pin is on a reroute node, keep following the connection.
+		if (InContext.PinBeingFollowed->IsInputPin() || (OwningNode && OwningNode->IsA<UMovieGraphRerouteNode>()))
 		{
-			if (UMovieGraphPin* ConnectedInputPin = InContext.PinBeingFollowed->GetFirstConnectedPin())
+			constexpr bool bFollowRerouteConnections = true;
+			if (UMovieGraphPin* ConnectedInputPin = InContext.PinBeingFollowed->GetFirstConnectedPin(bFollowRerouteConnections))
 			{
 				PinsToFollow.Add(ConnectedInputPin);
 			}
diff --git a/Engine/Plugins/MovieScene/MovieRenderPipeline/Source/MovieRenderPipelineCore/Private/Graph/MovieGraphPin.cpp b/Engine/Plugins/MovieScene/MovieRenderPipeline/Source/MovieRenderPipelineCore/Private/Graph/MovieGraphPin.cpp
--- a/Engine/Plugins/MovieScene/MovieRenderPipeline/Source/MovieRenderPipelineCore/Private/Graph/MovieGraphPin.cpp
+++ b/Engine/Plugins/MovieScene/MovieRenderPipeline/Source/MovieRenderPipelineCore/Private/Graph/MovieGraphPin.cpp
@@ -4,11 +4,12 @@
 
 #include "Algo/Find.h"
 #include "EdGraph/EdGraphSchema.h"
-#include "Graph/MovieGraphNode.h"
+#include "Graph/MovieGraphConfig.h"
 #include "Graph/MovieGraphEdge.h"
+#include "Graph/MovieGraphNode.h"
+#include "Graph/Nodes/MovieGraphRerouteNode.h"
 #include "Graph/Nodes/MovieGraphSubgraphNode.h"
 #include "MovieRenderPipelineCoreModule.h"
-#include "Graph/MovieGraphConfig.h"
 
 #if WITH_EDITOR
 bool UMovieGraphPin::Modify(bool bAlwaysMarkDirty)
@@ -52,6 +53,10 @@ bool UMovieGraphPin::AddEdgeTo(UMovieGraphPin* InOtherPin)
 
 	NewEdge->InputPin = bThisPinIsUpstream ? this : InOtherPin;
 	NewEdge->OutputPin = bThisPinIsUpstream ? InOtherPin : this;
+	
+	// When an edge is made, the pin type may need to be propagated (eg, to reroute nodes that could currently be wildcards)
+	PropagatePinType(NewEdge->OutputPin, NewEdge->InputPin);
+	
 	return true;
 }
 
@@ -72,6 +77,11 @@ bool UMovieGraphPin::BreakEdgeTo(UMovieGraphPin* InOtherPin)
 			ensure(InOtherPin->Edges.Remove(Edge));
 
 			Edges.Remove(Edge);
+
+			// After a disconnection, some pins may need to revert to being wildcard (eg, on reroute nodes)
+			MaybeRevertToWildcard(this);
+			MaybeRevertToWildcard(InOtherPin);
+
 			return true;
 		}
 	}
@@ -94,6 +104,10 @@ bool UMovieGraphPin::BreakAllEdges()
 			OtherPin->Modify();
 			ensure(OtherPin->Edges.Remove(Edge));
 			bChanged = true;
+
+			// After a disconnection, some pins may need to revert to being wildcard (eg, on reroute nodes)
+			MaybeRevertToWildcard(this);
+			MaybeRevertToWildcard(OtherPin);
 		}
 	}
 
@@ -115,8 +129,10 @@ FPinConnectionResponse UMovieGraphPin::CanCreateConnection_PinConnectionResponse
 	}
 
 	const bool bBothPinsAreBranch = Properties.bIsBranch && InOtherPin->Properties.bIsBranch;
+	const bool bPinIsWildcard = Properties.bIsWildcard || InOtherPin->Properties.bIsWildcard;
 
-	const bool bBothPinsAreSameType = bBothPinsAreBranch ||				// Both are branches or
+	const bool bBothPinsAreSameType = bPinIsWildcard ||					// Any connection can be made to a wildcard pin or
+		bBothPinsAreBranch ||											// Both are branches or
 		(!Properties.bIsBranch && !InOtherPin->Properties.bIsBranch &&	// Neither is branch and
 		Properties.Type == InOtherPin->Properties.Type &&				// They have the same property type and
 		Properties.TypeObject == InOtherPin->Properties.TypeObject);	// They have the same type object (for enums, structs, objects, classes)
@@ -200,7 +216,7 @@ bool UMovieGraphPin::AllowsMultipleConnections() const
 	return IsOutputPin() || Properties.bAllowMultipleConnections;
 }
 
-UMovieGraphPin* UMovieGraphPin::GetFirstConnectedPin() const
+UMovieGraphPin* UMovieGraphPin::GetFirstConnectedPin(const bool bFollowRerouteConnections) const
 {
 	if (Edges.IsEmpty())
 	{
@@ -209,7 +225,7 @@ UMovieGraphPin* UMovieGraphPin::GetFirstConnectedPin() const
 
 	if (ensureMsgf(Edges[0], TEXT("Null edge found when trying to get connected pin!")))
 	{
-		return Edges[0]->GetOtherPin(this);
+		return Edges[0]->GetOtherPin(this, bFollowRerouteConnections);
 	}
 
 	return nullptr;
@@ -416,3 +432,191 @@ bool UMovieGraphPin::IsPinDirectionCompatibleWith(const UMovieGraphPin* OtherPin
 {
 	return IsInputPin() != OtherPin->IsInputPin();
 }
+
+void UMovieGraphPin::PropagatePinType(const UMovieGraphPin* InInputPin, const UMovieGraphPin* InOutputPin)
+{
+	// If the pin's nodes aren't valid, don't proceed. This shouldn't normally happen, but could theoretically occur if one of the nodes could not
+	// be loaded properly (eg, it comes from a plugin that isn't loaded).
+	if ((InInputPin && !InInputPin->Node) || (InOutputPin && !InOutputPin->Node))
+	{
+		return;
+	}
+
+	// Find the first downstream pin with a non-wildcard type
+	const UMovieGraphPin* DownstreamPinWithType = InInputPin && InInputPin->Properties.bIsWildcard ? nullptr : InInputPin;
+	if (!DownstreamPinWithType && InInputPin)
+	{
+		InInputPin->Node->GetGraph()->VisitDownstreamNodes(InInputPin->Node, UMovieGraphConfig::FVisitNodesCallback::CreateLambda(
+			[&DownstreamPinWithType](const UMovieGraphNode* VisitedNode, const UMovieGraphPin* VisitedPin)
+			{
+				if (VisitedPin && !VisitedPin->Properties.bIsWildcard)
+				{
+					DownstreamPinWithType = VisitedPin;
+					return false;
+				}
+
+				return true;
+			}));
+	}
+
+	// Find the first pin upstream with a non-wildcard type
+	const UMovieGraphPin* UpstreamPinWithType = InOutputPin && InOutputPin->Properties.bIsWildcard ? nullptr : InOutputPin;
+	if (!UpstreamPinWithType && InOutputPin)
+	{
+		InOutputPin->Node->GetGraph()->VisitUpstreamNodes(InOutputPin->Node, UMovieGraphConfig::FVisitNodesCallback::CreateLambda(
+			[&UpstreamPinWithType](const UMovieGraphNode* VisitedNode, const UMovieGraphPin* VisitedPin)
+			{
+				if (VisitedPin && !VisitedPin->Properties.bIsWildcard)
+				{
+					UpstreamPinWithType = VisitedPin;
+					return false;
+				}
+
+				return true;
+			}));
+	}
+
+	// Prioritize pins downstream with a type; if one is found, propagate that pin type upstream
+	if (DownstreamPinWithType)
+	{
+		constexpr bool bPropagateUpstream = true;
+		DownstreamPinWithType->PropagatePinProperties(DownstreamPinWithType->Properties, bPropagateUpstream);
+		return;
+	}
+
+	// Otherwise, propagate the type from upstream (if possible)
+	if (UpstreamPinWithType)
+	{
+		constexpr bool bPropagateUpstream = false;
+		UpstreamPinWithType->PropagatePinProperties(UpstreamPinWithType->Properties, bPropagateUpstream);
+	}
+}
+
+void UMovieGraphPin::PropagatePinProperties(const FMovieGraphPinProperties PinProperties, const bool bPropagateUpstream) const
+{
+	if (bPropagateTypeRecursionGuard)
+	{
+		return;
+	}
+
+	TGuardValue<bool> RecursionGuard(bPropagateTypeRecursionGuard, true);
+	
+	const UMovieGraphRerouteNode* RerouteNode = Cast<UMovieGraphRerouteNode>(Node);
+
+	// Get the pins to propagate to. If propagating upstream, this will only be one pin. For downstream, there may be multiple.
+	TArray<UMovieGraphPin*> NextPins;
+	if (bPropagateUpstream)
+	{
+		NextPins.Add(IsInputPin()
+			? GetFirstConnectedPin()
+			: (RerouteNode ? RerouteNode->GetPassThroughPin(this) : nullptr));
+	}
+	else
+	{
+		// When propagating downstream, there may be multiple connections to propagate to
+		if (IsInputPin())
+		{
+			if (RerouteNode)
+			{
+				NextPins.Add(RerouteNode->GetPassThroughPin(this));
+			}
+		}
+		else
+		{
+			NextPins.Append(GetAllConnectedPins());
+		}
+	}
+
+	// Propagate the type for all connected pins
+	for (const UMovieGraphPin* NextPin : NextPins)
+	{
+		if (NextPin && NextPin->Node && NextPin->Node->IsA<UMovieGraphRerouteNode>())
+		{
+			UMovieGraphRerouteNode* NextRerouteNode = Cast<UMovieGraphRerouteNode>(NextPin->Node);
+
+			// "PinProperties" isn't supplied directly to SetPinProperties() because there are members within the properties (like the label) that
+			// need to remain the same.
+			FMovieGraphPinProperties PropertiesToUpdate = NextRerouteNode->GetPinProperties();
+			PropertiesToUpdate.Type = PinProperties.Type;
+			PropertiesToUpdate.TypeObject = PinProperties.TypeObject;
+			PropertiesToUpdate.bIsBranch = PinProperties.bIsBranch;
+			PropertiesToUpdate.bIsWildcard = PinProperties.bIsWildcard;
+			NextRerouteNode->SetPinProperties(PropertiesToUpdate);
+			
+			NextRerouteNode->UpdatePins();
+
+			// Continue propagating the type through the connection chain
+			NextPin->PropagatePinProperties(PinProperties, bPropagateUpstream);
+
+			// When propagating upstream, there may be a situation where propagation also needs to fork off into a separate downstream propagation.
+			// For example, if we're coming from pin C, into pin B, propagation also needs to continue to A.
+			//
+			//  /``````A
+			// B
+			//  \......C
+			if (bPropagateUpstream)
+			{
+				TArray<UMovieGraphPin*> DownstreamPins = NextPin->GetAllConnectedPins();
+
+				// Remove this pin from the downstream pins (eg, NextPin is B, remove C)
+				DownstreamPins.Remove(const_cast<UMovieGraphPin*>(this));
+				
+				if (!DownstreamPins.IsEmpty())
+				{
+					NextPin->PropagatePinProperties(PinProperties, !bPropagateUpstream);
+				}
+			}
+		}
+	}
+}
+
+void UMovieGraphPin::MaybeRevertToWildcard(const UMovieGraphPin* InPin)
+{
+	check(InPin);
+
+	if (!InPin->Node)
+	{
+		return;
+	}
+
+	// Traverse the graph up/downstream and determine if there are only reroute nodes left in the connection chain (ie, this is an isolated island of
+	// only reroute nodes). If that's the case, then revert all of the reroute nodes to wildcards.
+	bool bFoundNonRerouteNodes = !InPin->Node->IsA<UMovieGraphRerouteNode>();
+
+	// Determine if there are any upstream non-reroute nodes
+	InPin->Node->GetGraph()->VisitUpstreamNodes(InPin->Node, UMovieGraphConfig::FVisitNodesCallback::CreateLambda(
+	[&bFoundNonRerouteNodes](const UMovieGraphNode* VisitedNode, const UMovieGraphPin* VisitedPin)
+	{
+		if (VisitedNode && !VisitedNode->IsA<UMovieGraphRerouteNode>())
+		{
+			bFoundNonRerouteNodes = true;
+			return false;
+		}
+
+		return true;
+	}));
+
+	// Determine if there are any downstream non-reroute nodes
+	InPin->Node->GetGraph()->VisitDownstreamNodes(InPin->Node, UMovieGraphConfig::FVisitNodesCallback::CreateLambda(
+	[&bFoundNonRerouteNodes](const UMovieGraphNode* VisitedNode, const UMovieGraphPin* VisitedPin)
+	{
+		if (VisitedNode && !VisitedNode->IsA<UMovieGraphRerouteNode>())
+		{
+			bFoundNonRerouteNodes = true;
+			return false;
+		}
+
+		return true;
+	}));
+
+	// If no non-reroute nodes were found, revert everything in the connection chain to wildcard
+	if (!bFoundNonRerouteNodes)
+	{
+		const FMovieGraphPinProperties WildcardProperties = FMovieGraphPinProperties::MakeWildcardProperties();
+		
+		bool bPropagateUpstream = true;
+		InPin->PropagatePinProperties(WildcardProperties, bPropagateUpstream);
+		bPropagateUpstream = false;
+		InPin->PropagatePinProperties(WildcardProperties, bPropagateUpstream);
+	}
+}
diff --git a/Engine/Plugins/MovieScene/MovieRenderPipeline/Source/MovieRenderPipelineCore/Private/Graph/Nodes/MovieGraphBranchNode.cpp b/Engine/Plugins/MovieScene/MovieRenderPipeline/Source/MovieRenderPipelineCore/Private/Graph/Nodes/MovieGraphBranchNode.cpp
--- a/Engine/Plugins/MovieScene/MovieRenderPipeline/Source/MovieRenderPipelineCore/Private/Graph/Nodes/MovieGraphBranchNode.cpp
+++ b/Engine/Plugins/MovieScene/MovieRenderPipeline/Source/MovieRenderPipelineCore/Private/Graph/Nodes/MovieGraphBranchNode.cpp
@@ -43,7 +43,8 @@ TArray<UMovieGraphPin*> UMovieGraphBranchNode::EvaluatePinsToFollow(FMovieGraphE
 	UMovieGraphPin* OtherPin = nullptr;
 	for (UMovieGraphEdge* Edge : ConditionalPin->Edges)
 	{
-		OtherPin = Edge->GetOtherPin(ConditionalPin);
+		constexpr bool bFollowRerouteConnections = true;
+		OtherPin = Edge->GetOtherPin(ConditionalPin, bFollowRerouteConnections);
 
 		// We only support a single connection to this pin type anyways.
 		break;
diff --git a/Engine/Plugins/MovieScene/MovieRenderPipeline/Source/MovieRenderPipelineCore/Private/Graph/Nodes/MovieGraphRerouteNode.cpp b/Engine/Plugins/MovieScene/MovieRenderPipeline/Source/MovieRenderPipelineCore/Private/Graph/Nodes/MovieGraphRerouteNode.cpp
new file mode 100644
--- /dev/null
+++ b/Engine/Plugins/MovieScene/MovieRenderPipeline/Source/MovieRenderPipelineCore/Private/Graph/Nodes/MovieGraphRerouteNode.cpp
@@ -0,0 +1,67 @@
+// Copyright Epic Games, Inc. All Rights Reserved.
+
+#include "Graph/Nodes/MovieGraphRerouteNode.h"
+
+#include "Styling/AppStyle.h"
+
+UMovieGraphRerouteNode::UMovieGraphRerouteNode()
+{
+	InputOutputProperties = FMovieGraphPinProperties::MakeWildcardProperties();
+}
+
+TArray<FMovieGraphPinProperties> UMovieGraphRerouteNode::GetInputPinProperties() const
+{
+	return {InputOutputProperties};
+}
+
+TArray<FMovieGraphPinProperties> UMovieGraphRerouteNode::GetOutputPinProperties() const
+{
+	return {InputOutputProperties};
+}
+
+bool UMovieGraphRerouteNode::CanBeDisabled() const
+{
+	return false;
+}
+
+#if WITH_EDITOR
+FText UMovieGraphRerouteNode::GetNodeTitle(const bool bGetDescriptive) const
+{
+	static const FText RerouteNodeName = NSLOCTEXT("MovieGraphNodes", "NodeName_Reroute", "Reroute");
+
+	return RerouteNodeName;
+}
+
+FText UMovieGraphRerouteNode::GetMenuCategory() const
+{
+	return NSLOCTEXT("MovieGraphNodes", "RerouteGraphNode_Category", "Utility");
+}
+#endif // WITH_EDITOR
+
+UMovieGraphPin* UMovieGraphRerouteNode::GetPassThroughPin(const UMovieGraphPin* InFromPin) const
+{
+	if (InFromPin)
+	{
+		if (InputPins.Contains(InFromPin))
+		{
+			return OutputPins[0];
+		}
+
+		if (OutputPins.Contains(InFromPin))
+		{
+			return InputPins[0];
+		}
+	}
+
+	return nullptr;
+}
+
+FMovieGraphPinProperties UMovieGraphRerouteNode::GetPinProperties() const
+{
+	return InputOutputProperties;
+}
+
+void UMovieGraphRerouteNode::SetPinProperties(const FMovieGraphPinProperties& InPinProperties)
+{
+	InputOutputProperties = InPinProperties;
+}
diff --git a/Engine/Plugins/MovieScene/MovieRenderPipeline/Source/MovieRenderPipelineCore/Private/Graph/Nodes/MovieGraphSelectNode.cpp b/Engine/Plugins/MovieScene/MovieRenderPipeline/Source/MovieRenderPipelineCore/Private/Graph/Nodes/MovieGraphSelectNode.cpp
--- a/Engine/Plugins/MovieScene/MovieRenderPipeline/Source/MovieRenderPipelineCore/Private/Graph/Nodes/MovieGraphSelectNode.cpp
+++ b/Engine/Plugins/MovieScene/MovieRenderPipeline/Source/MovieRenderPipelineCore/Private/Graph/Nodes/MovieGraphSelectNode.cpp
@@ -186,7 +186,8 @@ TArray<UMovieGraphPin*> UMovieGraphSelectNode::EvaluatePinsToFollow(FMovieGraphE
 	bool bGotValueFromConnection = false;
 	if (const UMovieGraphPin* SelectPin = GetInputPin(UE::MovieGraph::SelectNode::SelectedOption, EMovieGraphPinQueryRequirement::BuiltIn))
 	{
-		if (const UMovieGraphPin* OtherPin = SelectPin->GetFirstConnectedPin())
+		constexpr bool bFollowRerouteConnections = true;
+		if (const UMovieGraphPin* OtherPin = SelectPin->GetFirstConnectedPin(bFollowRerouteConnections))
 		{
 			if (const UMovieGraphNode* ConnectedNode = OtherPin->Node)
 			{
diff --git a/Engine/Plugins/MovieScene/MovieRenderPipeline/Source/MovieRenderPipelineCore/Private/Graph/Nodes/MovieGraphSubgraphNode.cpp b/Engine/Plugins/MovieScene/MovieRenderPipeline/Source/MovieRenderPipelineCore/Private/Graph/Nodes/MovieGraphSubgraphNode.cpp
--- a/Engine/Plugins/MovieScene/MovieRenderPipeline/Source/MovieRenderPipelineCore/Private/Graph/Nodes/MovieGraphSubgraphNode.cpp
+++ b/Engine/Plugins/MovieScene/MovieRenderPipeline/Source/MovieRenderPipelineCore/Private/Graph/Nodes/MovieGraphSubgraphNode.cpp
@@ -113,7 +113,8 @@ TArray<UMovieGraphPin*> UMovieGraphSubgraphNode::EvaluatePinsToFollow(FMovieGrap
 	// If an input on the node is being followed, just follow the connection on the pin, nothing fancy required.
 	if (InContext.PinBeingFollowed->IsInputPin())
 	{
-		if (UMovieGraphPin* ConnectedInputPin = InContext.PinBeingFollowed->GetFirstConnectedPin())
+		constexpr bool bFollowRerouteConnections = true;
+		if (UMovieGraphPin* ConnectedInputPin = InContext.PinBeingFollowed->GetFirstConnectedPin(bFollowRerouteConnections))
 		{
 			PinsToFollow.Add(ConnectedInputPin);
 		}
diff --git a/Engine/Plugins/MovieScene/MovieRenderPipeline/Source/MovieRenderPipelineCore/Public/Graph/MovieGraphEdge.h b/Engine/Plugins/MovieScene/MovieRenderPipeline/Source/MovieRenderPipelineCore/Public/Graph/MovieGraphEdge.h
--- a/Engine/Plugins/MovieScene/MovieRenderPipeline/Source/MovieRenderPipelineCore/Public/Graph/MovieGraphEdge.h
+++ b/Engine/Plugins/MovieScene/MovieRenderPipeline/Source/MovieRenderPipelineCore/Public/Graph/MovieGraphEdge.h
@@ -1,7 +1,9 @@
 // Copyright Epic Games, Inc. All Rights Reserved.
 
 #pragma once
+
 #include "CoreMinimal.h"
+
 #include "MovieGraphEdge.generated.h"
 
 class UMovieGraphPin;
@@ -18,15 +20,12 @@ public:
 	UPROPERTY()
 	TObjectPtr<UMovieGraphPin> OutputPin;
 
-	bool IsValid() const
-	{
-		return InputPin.Get() && OutputPin.Get();
-	}
-
-	UMovieGraphPin* GetOtherPin(const UMovieGraphPin* InFirstPin)
-	{
-		check(InFirstPin == InputPin || InFirstPin == OutputPin);
-		return InFirstPin == InputPin ? OutputPin : InputPin;
-	}
+	/** Whether the edge is valid or not. Being valid means it contains a non-null input and output pin. */
+	bool IsValid() const;
 
+	/**
+	 * Gets the pin on the other side of the edge. If bFollowRerouteConnections is true, reroute nodes will be passthrough, and this method will
+	 * continue traversing edges until a pin on a non-reroute node is found.
+	 */
+	UMovieGraphPin* GetOtherPin(const UMovieGraphPin* InFirstPin, const bool bFollowRerouteConnections = false);
 };
\ No newline at end of file
diff --git a/Engine/Plugins/MovieScene/MovieRenderPipeline/Source/MovieRenderPipelineCore/Public/Graph/MovieGraphPin.h b/Engine/Plugins/MovieScene/MovieRenderPipeline/Source/MovieRenderPipelineCore/Public/Graph/MovieGraphPin.h
--- a/Engine/Plugins/MovieScene/MovieRenderPipeline/Source/MovieRenderPipelineCore/Public/Graph/MovieGraphPin.h
+++ b/Engine/Plugins/MovieScene/MovieRenderPipeline/Source/MovieRenderPipelineCore/Public/Graph/MovieGraphPin.h
@@ -31,6 +31,13 @@ struct FMovieGraphPinProperties
 		return MoveTemp(Properties);
 	}
 
+	static FMovieGraphPinProperties MakeWildcardProperties(const FName& InLabel = NAME_None)
+	{
+		FMovieGraphPinProperties Properties(InLabel, EMovieGraphValueType::None, nullptr, false);
+		Properties.bIsWildcard = true;
+		return MoveTemp(Properties);
+	}
+
 	/** The name assigned to the pin. */
 	UPROPERTY(BlueprintReadWrite, EditAnywhere, Category = "Settings")
 	FName Label = NAME_None;
@@ -51,6 +58,10 @@ struct FMovieGraphPinProperties
 	UPROPERTY(BlueprintReadWrite, EditAnywhere, Category = "Settings")
 	bool bIsBranch = false;
 
+	/** Whether this pin represents a wildcard, which means it can be any type (including being a branch). */
+	UPROPERTY(BlueprintReadWrite, EditAnywhere, Category = "Settings")
+	bool bIsWildcard = false;
+
 	/**
 	 * Whether this pin is built-in (ie, the pin ships with the node and cannot be removed). Option pins on the Select
 	 * node would be an example of pins which are not built-in (they can be added and removed dynamically).
@@ -124,9 +135,12 @@ public:
 	UFUNCTION(BlueprintCallable, Category = "Movie Graph")
 	bool AllowsMultipleConnections() const;
 
-	/** Gets the first pin connected to this pin. Returns nullptr if no valid connection exists. */
+	/**
+	 * Gets the first pin connected to this pin. Returns nullptr if no valid connection exists. If bFollowRerouteConnections is true,
+	 * this method will continue traversing the graph past reroute nodes until a non-reroute pin is found.
+	 */
 	UFUNCTION(BlueprintCallable, Category = "Movie Graph")
-	UMovieGraphPin* GetFirstConnectedPin() const;
+	UMovieGraphPin* GetFirstConnectedPin(const bool bFollowRerouteConnections = false) const;
 
 	/** Gets all connected pins. */
 	UFUNCTION(BlueprintCallable, Category = "Movie Graph")
@@ -164,4 +178,24 @@ public:
 	*/
 	UPROPERTY(BlueprintReadOnly, TextExportTransient, Category = "Properties")
 	TArray<TObjectPtr<UMovieGraphEdge>> Edges;
+
+private:
+	/**
+	 * Given a pair of input and output pins which a connection was made between, propagate the type either upstream or downstream depending on
+	 * the nodes the connection was made between.
+	 */
+	void PropagatePinType(const UMovieGraphPin* InInputPin, const UMovieGraphPin* InOutputPin);
+
+	/** Propagate the given properties either upstream or downstream from this pin. This pin is not changed. */
+	void PropagatePinProperties(const FMovieGraphPinProperties PinProperties, const bool bPropagateUpstream) const;
+
+	/**
+	 * Potentially reverts InPin (and the reroute nodes connected to/from it) to a wildcard if InPin no longer has any upstream or downstream
+	 * connections that provide a type. In other words, InPin is part of a reroute "island" where the reroute nodes are only connected to other
+	 * reroute nodes.
+	 */
+	void MaybeRevertToWildcard(const UMovieGraphPin* InPin);
+	
+	// Recursion guard boolean to prevent PropagatePinProperties from hanging if there is a cycle of reroute nodes.
+	mutable bool bPropagateTypeRecursionGuard = false;
 };
\ No newline at end of file
diff --git a/Engine/Plugins/MovieScene/MovieRenderPipeline/Source/MovieRenderPipelineCore/Public/Graph/Nodes/MovieGraphRerouteNode.h b/Engine/Plugins/MovieScene/MovieRenderPipeline/Source/MovieRenderPipelineCore/Public/Graph/Nodes/MovieGraphRerouteNode.h
new file mode 100644
--- /dev/null
+++ b/Engine/Plugins/MovieScene/MovieRenderPipeline/Source/MovieRenderPipelineCore/Public/Graph/Nodes/MovieGraphRerouteNode.h
@@ -0,0 +1,45 @@
+// Copyright Epic Games, Inc. All Rights Reserved.
+
+#pragma once
+
+#include "Graph/MovieGraphNode.h"
+#include "MovieGraphRerouteNode.generated.h"
+
+/** A node which is effectively a no-op/passthrough. Allows a connection to be routed untouched through this node to organize the graph. */
+UCLASS()
+class MOVIERENDERPIPELINECORE_API UMovieGraphRerouteNode : public UMovieGraphSettingNode
+{
+	GENERATED_BODY()
+
+public:
+	UMovieGraphRerouteNode();
+
+	//~ UMovieGraphNode Interface
+	virtual TArray<FMovieGraphPinProperties> GetInputPinProperties() const override;
+	virtual TArray<FMovieGraphPinProperties> GetOutputPinProperties() const override;
+	virtual bool CanBeDisabled() const override;
+#if WITH_EDITOR
+	virtual FText GetNodeTitle(const bool bGetDescriptive = false) const override;
+	virtual FText GetMenuCategory() const override;
+#endif
+	//~ End UMovieGraphNode Interface
+
+	/** Gets the pin opposite to the specified InFromPin. */
+	virtual UMovieGraphPin* GetPassThroughPin(const UMovieGraphPin* InFromPin) const;
+
+	/** Sets the pin properties for this reroute node. Note that this applies to both the input and output pin. */
+	UFUNCTION(BlueprintPure, Category = "Movie Graph")
+	FMovieGraphPinProperties GetPinProperties() const;
+
+	/**
+	 * Sets the pin properties for this reroute node (both the input and output pin have the same properties). This generally should not be called
+	 * unless you know what you're doing; normal connection/disconnection should handle setting the properties correctly.
+	 */
+	UFUNCTION(BlueprintCallable, Category = "Movie Graph")
+	void SetPinProperties(const FMovieGraphPinProperties& InPinProperties);
+
+private:
+	/** Pin properties that are shared with both the input and output pins. */
+	UPROPERTY()
+	FMovieGraphPinProperties InputOutputProperties;
+};
\ No newline at end of file
diff --git a/Engine/Plugins/MovieScene/MovieRenderPipeline/Source/MovieRenderPipelineEditor/Private/Customizations/Graph/MovieGraphMemberCustomization.h b/Engine/Plugins/MovieScene/MovieRenderPipeline/Source/MovieRenderPipelineEditor/Private/Customizations/Graph/MovieGraphMemberCustomization.h
--- a/Engine/Plugins/MovieScene/MovieRenderPipeline/Source/MovieRenderPipelineEditor/Private/Customizations/Graph/MovieGraphMemberCustomization.h
+++ b/Engine/Plugins/MovieScene/MovieRenderPipeline/Source/MovieRenderPipelineEditor/Private/Customizations/Graph/MovieGraphMemberCustomization.h
@@ -244,7 +244,8 @@ protected:
 					}
 
 					constexpr bool bIsBranch = false;
-					return UMoviePipelineEdGraphNodeBase::GetPinType(WeakMemberObject->GetValueType(), bIsBranch, WeakMemberObject->GetValueTypeObject());
+					constexpr bool bIsWildcard = false;
+					return UMoviePipelineEdGraphNodeBase::GetPinType(WeakMemberObject->GetValueType(), bIsBranch, bIsWildcard, WeakMemberObject->GetValueTypeObject());
 				})
 				.OnPinTypeChanged_Lambda(PinInfoChanged, WeakMemberObject)
 				.Schema(GetDefault<UPropertyBagSchema>())
diff --git a/Engine/Plugins/MovieScene/MovieRenderPipeline/Source/MovieRenderPipelineEditor/Private/Customizations/Graph/MovieGraphSelectNodeCustomization.h b/Engine/Plugins/MovieScene/MovieRenderPipeline/Source/MovieRenderPipelineEditor/Private/Customizations/Graph/MovieGraphSelectNodeCustomization.h
--- a/Engine/Plugins/MovieScene/MovieRenderPipeline/Source/MovieRenderPipelineEditor/Private/Customizations/Graph/MovieGraphSelectNodeCustomization.h
+++ b/Engine/Plugins/MovieScene/MovieRenderPipeline/Source/MovieRenderPipelineEditor/Private/Customizations/Graph/MovieGraphSelectNodeCustomization.h
@@ -124,7 +124,8 @@ protected:
 					}
 
 					constexpr bool bIsBranch = false;
-					return UMoviePipelineEdGraphNodeBase::GetPinType(SelectNodeObject->GetValueType(), bIsBranch, SelectNodeObject->GetValueTypeObject());
+					constexpr bool bIsWildcard = false;
+					return UMoviePipelineEdGraphNodeBase::GetPinType(SelectNodeObject->GetValueType(), bIsBranch, bIsWildcard, SelectNodeObject->GetValueTypeObject());
 				})
 				.OnPinTypeChanged_Lambda([PropUtils, SelectNodeObject](const FEdGraphPinType& InPinType)
 				{
diff --git a/Engine/Plugins/MovieScene/MovieRenderPipeline/Source/MovieRenderPipelineEditor/Private/Graph/MovieEdGraph.cpp b/Engine/Plugins/MovieScene/MovieRenderPipeline/Source/MovieRenderPipelineEditor/Private/Graph/MovieEdGraph.cpp
--- a/Engine/Plugins/MovieScene/MovieRenderPipeline/Source/MovieRenderPipelineEditor/Private/Graph/MovieEdGraph.cpp
+++ b/Engine/Plugins/MovieScene/MovieRenderPipeline/Source/MovieRenderPipelineEditor/Private/Graph/MovieEdGraph.cpp
@@ -4,11 +4,13 @@
 #include "Graph/MovieGraphConfig.h"
 #include "Graph/MovieGraphEdge.h"
 #include "Graph/MovieGraphPin.h"
+#include "Graph/Nodes/MovieGraphRerouteNode.h"
 #include "Graph/Nodes/MovieGraphSubgraphNode.h"
 #include "Graph/Nodes/MovieGraphVariableNode.h"
 #include "Graph/MovieGraphNode.h"
-#include "MovieEdGraphOutputNode.h"
 #include "MovieEdGraphInputNode.h"
+#include "MovieEdGraphOutputNode.h"
+#include "MovieEdGraphRerouteNode.h"
 #include "MovieEdGraphVariableNode.h"
 #include "MoviePipelineEdGraphSubgraphNode.h"
 #include "MovieRenderPipelineCoreModule.h"
@@ -73,6 +75,10 @@ void UMoviePipelineEdGraph::InitFromRuntimeGraph(UMovieGraphConfig* InGraph)
 		{
 			NodeLookup.Add(RuntimeNode, CreateNodeFromRuntimeNode<UMoviePipelineEdGraphSubgraphNode>(RuntimeNode));
 		}
+		else if (RuntimeNode.IsA<UMovieGraphRerouteNode>())
+		{
+			NodeLookup.Add(RuntimeNode, CreateNodeFromRuntimeNode<UMoviePipelineEdGraphRerouteNode>(RuntimeNode));
+		}
 		else
 		{
 			NodeLookup.Add(RuntimeNode, CreateNodeFromRuntimeNode<UMoviePipelineEdGraphNode>(RuntimeNode));
diff --git a/Engine/Plugins/MovieScene/MovieRenderPipeline/Source/MovieRenderPipelineEditor/Private/Graph/MovieEdGraphNode.cpp b/Engine/Plugins/MovieScene/MovieRenderPipeline/Source/MovieRenderPipelineEditor/Private/Graph/MovieEdGraphNode.cpp
--- a/Engine/Plugins/MovieScene/MovieRenderPipeline/Source/MovieRenderPipelineEditor/Private/Graph/MovieEdGraphNode.cpp
+++ b/Engine/Plugins/MovieScene/MovieRenderPipeline/Source/MovieRenderPipelineEditor/Private/Graph/MovieEdGraphNode.cpp
@@ -57,7 +57,7 @@ void UMoviePipelineEdGraphNodeBase::PostTransacted(const FTransactionObjectEvent
 	}
 }
 
-FEdGraphPinType UMoviePipelineEdGraphNodeBase::GetPinType(EMovieGraphValueType ValueType, bool bIsBranch, const UObject* InValueTypeObject)
+FEdGraphPinType UMoviePipelineEdGraphNodeBase::GetPinType(EMovieGraphValueType ValueType, bool bIsBranch, bool bIsWildcard, const UObject* InValueTypeObject)
 {
 	FEdGraphPinType EdPinType;
 	EdPinType.ResetToDefaults();
@@ -72,6 +72,13 @@ FEdGraphPinType UMoviePipelineEdGraphNodeBase::GetPinType(EMovieGraphValueType V
 		return EdPinType;
 	}
 
+	// Special case for wildcard pins
+	if (bIsWildcard)
+	{
+		EdPinType.PinCategory = UMovieGraphSchema::PC_Wildcard;
+		return EdPinType;
+	}
+
 	switch (ValueType)
 	{
 	case EMovieGraphValueType::Bool:
@@ -137,7 +144,7 @@ FEdGraphPinType UMoviePipelineEdGraphNodeBase::GetPinType(EMovieGraphValueType V
 
 FEdGraphPinType UMoviePipelineEdGraphNodeBase::GetPinType(const UMovieGraphPin* InPin)
 {
-	return GetPinType(InPin->Properties.Type, InPin->Properties.bIsBranch, InPin->Properties.TypeObject);
+	return GetPinType(InPin->Properties.Type, InPin->Properties.bIsBranch, InPin->Properties.bIsWildcard, InPin->Properties.TypeObject);
 }
 
 EMovieGraphValueType UMoviePipelineEdGraphNodeBase::GetValueTypeFromPinType(const FEdGraphPinType& InPinType)
@@ -448,11 +455,22 @@ void UMoviePipelineEdGraphNodeBase::CreatePins(const TArray<UMovieGraphPin*>& In
 FString UMoviePipelineEdGraphNodeBase::GetPinTooltip(const UMovieGraphPin* InPin) const
 {
 	const EMovieGraphValueType PinType = InPin->Properties.Type;
-	const FText TypeText = InPin->Properties.bIsBranch
-		? LOCTEXT("PinTypeTooltip_Branch", "Branch")
-		: StaticEnum<EMovieGraphValueType>()->GetDisplayNameTextByValue(static_cast<int64>(PinType));
 	const FText TypeObjectText = InPin->Properties.TypeObject ? FText::FromString(InPin->Properties.TypeObject.Get()->GetName()) : FText::GetEmpty();
 
+	FText TypeText;
+	if (InPin->Properties.bIsBranch)
+	{
+		TypeText = LOCTEXT("PinTypeTooltip_Branch", "Branch");
+	}
+	else if (InPin->Properties.bIsWildcard)
+	{
+		TypeText = LOCTEXT("PinTypeTooltip_Wildcard", "Any");
+	}
+	else
+	{
+		TypeText = StaticEnum<EMovieGraphValueType>()->GetDisplayNameTextByValue(static_cast<int64>(PinType));
+	}
+
 	const FText PinTooltipFormat = LOCTEXT("PinTypeTooltip_NoValueTypeObject", "Type: {ValueType}");
 	const FText PinTooltipFormatWithTypeObject = LOCTEXT("PinTypeTooltip_WithValueTypeObject", "Type: {ValueType} ({ValueTypeObject})");
 
@@ -484,9 +502,12 @@ void UMoviePipelineEdGraphNodeBase::AutowireNewNode(UEdGraphPin* FromPin)
 		check(OtherPin);
 
 		const FName& OtherPinName = OtherPin->Properties.Label;
-		UEdGraphPin* ToPin = FindPinChecked(OtherPinName, bFromPinIsInput ? EEdGraphPinDirection::EGPD_Output : EEdGraphPinDirection::EGPD_Input);
+		UEdGraphPin* ToPin = FindPinChecked(OtherPinName, bFromPinIsInput ? EGPD_Output : EGPD_Input);
 		if (ToPin && GetSchema()->TryCreateConnection(FromPin, ToPin))
 		{
+			// The pin (or owning node) may have been re-generated after the connection was made (eg, this happens with Reroute nodes)
+			ToPin = FindPinChecked(OtherPinName, bFromPinIsInput ? EGPD_Output : EGPD_Input);
+			
 			// Connection succeeded. Notify our other node that their connections changed.
 			if (ToPin->GetOwningNode())
 			{
diff --git a/Engine/Plugins/MovieScene/MovieRenderPipeline/Source/MovieRenderPipelineEditor/Private/Graph/MovieEdGraphNode.h b/Engine/Plugins/MovieScene/MovieRenderPipeline/Source/MovieRenderPipelineEditor/Private/Graph/MovieEdGraphNode.h
--- a/Engine/Plugins/MovieScene/MovieRenderPipeline/Source/MovieRenderPipelineEditor/Private/Graph/MovieEdGraphNode.h
+++ b/Engine/Plugins/MovieScene/MovieRenderPipeline/Source/MovieRenderPipelineEditor/Private/Graph/MovieEdGraphNode.h
@@ -34,7 +34,7 @@ public:
 	virtual void GetNodeContextMenuActions(UToolMenu* Menu, UGraphNodeContextMenuContext* Context) const override;
 	//~ End UEdGraphNode Interface
 
-	static FEdGraphPinType GetPinType(EMovieGraphValueType ValueType, bool bIsBranch, const UObject* InValueTypeObject = nullptr);
+	static FEdGraphPinType GetPinType(EMovieGraphValueType ValueType, bool bIsBranch, bool bIsWildcard, const UObject* InValueTypeObject = nullptr);
 	static FEdGraphPinType GetPinType(const UMovieGraphPin* InPin);
 	static EMovieGraphValueType GetValueTypeFromPinType(const FEdGraphPinType& InPinType);
 
diff --git a/Engine/Plugins/MovieScene/MovieRenderPipeline/Source/MovieRenderPipelineEditor/Private/Graph/MovieEdGraphRerouteNode.cpp b/Engine/Plugins/MovieScene/MovieRenderPipeline/Source/MovieRenderPipelineEditor/Private/Graph/MovieEdGraphRerouteNode.cpp
new file mode 100644
--- /dev/null
+++ b/Engine/Plugins/MovieScene/MovieRenderPipeline/Source/MovieRenderPipelineEditor/Private/Graph/MovieEdGraphRerouteNode.cpp
@@ -0,0 +1,44 @@
+// Copyright Epic Games, Inc. All Rights Reserved.
+
+#include "Graph/MovieEdGraphRerouteNode.h"
+
+#include "EdGraph/EdGraphPin.h"
+#include "SGraphNodeKnot.h"
+
+#define LOCTEXT_NAMESPACE "MoviePipelineGraph"
+
+FText UMoviePipelineEdGraphRerouteNode::GetNodeTitle(ENodeTitleType::Type TitleType) const
+{
+	static const FText RerouteNodeTitle = LOCTEXT("RerouteNodeTitle", "Reroute");
+	return RerouteNodeTitle;
+}
+
+UEdGraphPin* UMoviePipelineEdGraphRerouteNode::GetPassThroughPin(const UEdGraphPin* FromPin) const
+{
+	if (FromPin && Pins.Contains(FromPin))
+	{
+		return FromPin == Pins[0] ? Pins[1] : Pins[0];
+	}
+
+	return nullptr;
+}
+
+bool UMoviePipelineEdGraphRerouteNode::ShouldDrawNodeAsControlPointOnly(int32& OutInputPinIndex, int32& OutOutputPinIndex) const
+{
+	OutInputPinIndex = 0;
+	OutOutputPinIndex = 1;
+
+	return true;
+}
+
+bool UMoviePipelineEdGraphRerouteNode::CanSplitPin(const UEdGraphPin* Pin) const
+{
+	return false;
+}
+
+TSharedPtr<SGraphNode> UMoviePipelineEdGraphRerouteNode::CreateVisualWidget()
+{
+	return SNew(SGraphNodeKnot, this);
+}
+
+#undef LOCTEXT_NAMESPACE
diff --git a/Engine/Plugins/MovieScene/MovieRenderPipeline/Source/MovieRenderPipelineEditor/Private/Graph/MovieEdGraphRerouteNode.h b/Engine/Plugins/MovieScene/MovieRenderPipeline/Source/MovieRenderPipelineEditor/Private/Graph/MovieEdGraphRerouteNode.h
new file mode 100644
--- /dev/null
+++ b/Engine/Plugins/MovieScene/MovieRenderPipeline/Source/MovieRenderPipelineEditor/Private/Graph/MovieEdGraphRerouteNode.h
@@ -0,0 +1,23 @@
+// Copyright Epic Games, Inc. All Rights Reserved.
+
+#pragma once
+
+#include "Graph/MovieEdGraphNode.h"
+
+#include "MovieEdGraphRerouteNode.generated.h"
+
+/** A node which reroutes connections in order to organize the graph cleanly. */
+UCLASS()
+class UMoviePipelineEdGraphRerouteNode : public UMoviePipelineEdGraphNode
+{
+	GENERATED_BODY()
+
+public:
+	// ~Begin UEdGraphNode interface
+	virtual FText GetNodeTitle(ENodeTitleType::Type TitleType) const override;
+	virtual UEdGraphPin* GetPassThroughPin(const UEdGraphPin* FromPin) const override;
+	virtual bool ShouldDrawNodeAsControlPointOnly(int32& OutInputPinIndex, int32& OutOutputPinIndex) const override;
+	virtual bool CanSplitPin(const UEdGraphPin* Pin) const override;
+	virtual TSharedPtr<SGraphNode> CreateVisualWidget() override;
+	// ~End UEdGraphNode interface
+};
\ No newline at end of file
diff --git a/Engine/Plugins/MovieScene/MovieRenderPipeline/Source/MovieRenderPipelineEditor/Private/Graph/MovieGraphSchema.cpp b/Engine/Plugins/MovieScene/MovieRenderPipeline/Source/MovieRenderPipelineEditor/Private/Graph/MovieGraphSchema.cpp
--- a/Engine/Plugins/MovieScene/MovieRenderPipeline/Source/MovieRenderPipelineEditor/Private/Graph/MovieGraphSchema.cpp
+++ b/Engine/Plugins/MovieScene/MovieRenderPipeline/Source/MovieRenderPipelineEditor/Private/Graph/MovieGraphSchema.cpp
@@ -6,6 +6,7 @@
 #include "Graph/MovieGraphConfig.h"
 #include "Graph/Nodes/MovieGraphInputNode.h"
 #include "Graph/Nodes/MovieGraphOutputNode.h"
+#include "Graph/Nodes/MovieGraphRerouteNode.h"
 #include "Graph/Nodes/MovieGraphSubgraphNode.h"
 #include "Graph/Nodes/MovieGraphVariableNode.h"
 #include "Graph/MovieGraphNode.h"
@@ -18,6 +19,7 @@
 #include "GraphEditor.h"
 #include "MovieEdGraphInputNode.h"
 #include "MovieEdGraphOutputNode.h"
+#include "MovieEdGraphRerouteNode.h"
 #include "MovieEdGraphVariableNode.h"
 #include "MoviePipelineEdGraphSubgraphNode.h"
 
@@ -26,6 +28,7 @@ TArray<UClass*> UMovieGraphSchema::MoviePipelineNodeClasses;
 #define LOCTEXT_NAMESPACE "MoviePipelineGraphSchema"
 
 const FName UMovieGraphSchema::PC_Branch(TEXT("branch"));	// The branch looks like an Exec pin, but isn't the same thing, so we don't use the BP Exec type
+const FName UMovieGraphSchema::PC_Wildcard(UEdGraphSchema_K2::PC_Wildcard);
 const FName UMovieGraphSchema::PC_Boolean(UEdGraphSchema_K2::PC_Boolean);
 const FName UMovieGraphSchema::PC_Byte(UEdGraphSchema_K2::PC_Byte);
 const FName UMovieGraphSchema::PC_Integer(UEdGraphSchema_K2::PC_Int);
@@ -193,6 +196,25 @@ UEdGraphPin* UMovieGraphSchema::DropPinOnNode(UEdGraphNode* InTargetNode, const
 	return NewEdPin; 
 }
 
+bool UMovieGraphSchema::SafeDeleteNodeFromGraph(UEdGraph* Graph, UEdGraphNode* Node) const
+{
+	if (!Graph || !Node)
+	{
+		return false;
+	}
+
+	const UMoviePipelineEdGraphNodeBase* EdGraphNode = CastChecked<UMoviePipelineEdGraphNodeBase>(Node);
+	UMovieGraphNode* RuntimeNode = EdGraphNode->GetRuntimeNode();
+	check(RuntimeNode);
+
+	UMovieGraphConfig* OwningGraph = RuntimeNode->GetGraph();
+	check(OwningGraph);
+
+	OwningGraph->RemoveNode(RuntimeNode);
+
+	return true;
+}
+
 void UMovieGraphSchema::InitMoviePipelineNodeClasses()
 {
 	if (MoviePipelineNodeClasses.Num() > 0)
@@ -312,8 +334,17 @@ void UMovieGraphSchema::GetGraphContextActions(FGraphContextMenuBuilder& Context
 				{
 					// The branch restriction is "Any", so the node creation should be ok
 				}
-				
-				bCanAppearInMenu = bBranchRestrictionIsOk && (ContextMenuBuilder.FromPin->PinType.PinCategory == PC_Branch);
+
+				// Determine if the node can be shown in the menu. An exception to the above rules is the Reroute node -- this can always
+				// be created in any context.
+				if (PipelineNodeClass == UMovieGraphRerouteNode::StaticClass())
+				{
+					bCanAppearInMenu = true;
+				}
+				else
+				{
+					bCanAppearInMenu = bBranchRestrictionIsOk && (ContextMenuBuilder.FromPin->PinType.PinCategory == PC_Branch);
+				}
 			}
 		}
 		
@@ -632,6 +663,10 @@ UEdGraphNode* FMovieGraphSchemaAction_NewNode::PerformAction(UEdGraph* ParentGra
 	{
 		InvokableEdGraphNodeClass = UMoviePipelineEdGraphSubgraphNode::StaticClass();
 	}
+	else if (RuntimeNode->IsA(UMovieGraphRerouteNode::StaticClass()))
+	{
+		InvokableEdGraphNodeClass = UMoviePipelineEdGraphRerouteNode::StaticClass();
+	}
 	
 	UMoviePipelineEdGraphNode* GraphNode = NodeCreator.CreateUserInvokedNode(bSelectNewNode, InvokableEdGraphNodeClass);
 	GraphNode->Construct(RuntimeNode);
diff --git a/Engine/Plugins/MovieScene/MovieRenderPipeline/Source/MovieRenderPipelineEditor/Private/Graph/MovieGraphSchema.h b/Engine/Plugins/MovieScene/MovieRenderPipeline/Source/MovieRenderPipelineEditor/Private/Graph/MovieGraphSchema.h
--- a/Engine/Plugins/MovieScene/MovieRenderPipeline/Source/MovieRenderPipelineEditor/Private/Graph/MovieGraphSchema.h
+++ b/Engine/Plugins/MovieScene/MovieRenderPipeline/Source/MovieRenderPipelineEditor/Private/Graph/MovieGraphSchema.h
@@ -29,6 +29,7 @@ public:
 	virtual bool ShouldHidePinDefaultValue(UEdGraphPin* Pin) const override;
 	virtual bool SupportsDropPinOnNode(UEdGraphNode* InTargetNode, const FEdGraphPinType& InSourcePinType, EEdGraphPinDirection InSourcePinDirection, FText& OutErrorMessage) const override;
 	virtual UEdGraphPin* DropPinOnNode(UEdGraphNode* InTargetNode, const FName& InSourcePinName, const FEdGraphPinType& InSourcePinType, EEdGraphPinDirection InSourcePinDirection) const override;
+	virtual bool SafeDeleteNodeFromGraph(UEdGraph* Graph, UEdGraphNode* Node) const override;
 	//~ End EdGraphSchema Interface
 	
 	static void InitMoviePipelineNodeClasses();
@@ -60,6 +61,7 @@ private:
 public:
 	// Allowed "PinCategory" values for use on EdGraphPin
 	static const FName PC_Branch;
+	static const FName PC_Wildcard;
 	static const FName PC_Boolean;
 	static const FName PC_Byte;
 	static const FName PC_Integer;
diff --git a/Engine/Plugins/MovieScene/MovieRenderPipeline/Source/MovieRenderPipelineEditor/Private/Graph/SMovieGraphMembersTabContent.cpp b/Engine/Plugins/MovieScene/MovieRenderPipeline/Source/MovieRenderPipelineEditor/Private/Graph/SMovieGraphMembersTabContent.cpp
--- a/Engine/Plugins/MovieScene/MovieRenderPipeline/Source/MovieRenderPipelineEditor/Private/Graph/SMovieGraphMembersTabContent.cpp
+++ b/Engine/Plugins/MovieScene/MovieRenderPipeline/Source/MovieRenderPipelineEditor/Private/Graph/SMovieGraphMembersTabContent.cpp
@@ -72,7 +72,8 @@ namespace UE::MovieGraph::Private
 		}
 		
 		constexpr bool bIsBranch = false;
-		const FEdGraphPinType PinType = UMoviePipelineEdGraphNodeBase::GetPinType(InGraphVariable->GetValueType(), bIsBranch, InGraphVariable->GetValueTypeObject());
+		constexpr bool bIsWildcard = false;
+		const FEdGraphPinType PinType = UMoviePipelineEdGraphNodeBase::GetPinType(InGraphVariable->GetValueType(), bIsBranch, bIsWildcard, InGraphVariable->GetValueTypeObject());
 
 		OutPrimaryBrush = FAppStyle::GetBrush("Kismet.AllClasses.VariableIcon");
 		OutIconColor = UMovieGraphSchema::GetTypeColor(PinType.PinCategory, PinType.PinSubCategory);
@@ -164,7 +165,7 @@ TSharedRef<SWidget> SMovieGraphMembersTabContent::CreateActionWidget(FCreateWidg
 			static_cast<FMovieGraphSchemaAction_NewVariableNode*>(InCreateData->Action.Get());
 		const UMovieGraphVariable* Variable = Cast<UMovieGraphVariable>(VariableAction->ActionTarget);
 		const FEdGraphPinType PinType = Variable
-			? UMoviePipelineEdGraphNodeBase::GetPinType(Variable->GetValueType(), false, Variable->GetValueTypeObject())
+			? UMoviePipelineEdGraphNodeBase::GetPinType(Variable->GetValueType(), false, false, Variable->GetValueTypeObject())
 			: FEdGraphPinType();
 		const FLinearColor PinColor = CurrentGraph->PipelineEdGraph->GetSchema()->GetPinTypeColor(PinType);
 		
