MoviePipeline: Exposed more properties and methods on graph variables to Python.

diff --git a/Engine/Plugins/MovieScene/MovieRenderPipeline/Source/MovieRenderPipelineCore/Public/Graph/MovieGraphConfig.h b/Engine/Plugins/MovieScene/MovieRenderPipeline/Source/MovieRenderPipelineCore/Public/Graph/MovieGraphConfig.h
--- a/Engine/Plugins/MovieScene/MovieRenderPipeline/Source/MovieRenderPipelineCore/Public/Graph/MovieGraphConfig.h
+++ b/Engine/Plugins/MovieScene/MovieRenderPipeline/Source/MovieRenderPipelineCore/Public/Graph/MovieGraphConfig.h
@@ -40,13 +40,14 @@ public:
 	FString GetMemberName() const { return Name; }
 
 	/** Sets the name of this member. Returns true if the rename was successful, else false. */
-	UFUNCTION(BlueprintCallable, Category ="Movie Graph")
+	UFUNCTION(BlueprintCallable, Category = "Movie Graph")
 	virtual bool SetMemberName(const FString& InNewName);
 
 	/**
 	 * Determines if this member can be renamed to the specified name. If the rename is not possible, returns false
 	 * and OutError is populated with the reason, else returns true.
 	 */
+	UFUNCTION(BlueprintPure, Category = "Movie Graph")
 	virtual bool CanRename(const FText& InNewName, FText& OutError) const;
 
 	/** Gets the GUID that uniquely identifies this member. */
@@ -57,9 +58,11 @@ public:
 	void SetGuid(const FGuid& InGuid) { Guid = InGuid; }
 
 	/** Determines if this member can be deleted. */
+	UFUNCTION(BlueprintPure, Category = "Movie Graph")
 	virtual bool IsDeletable() const { return true; }
 
 	/** Gets whether this member is editable via the UI. */
+	UFUNCTION(BlueprintPure, Category = "Movie Graph")
 	virtual bool IsEditable() const { return bIsEditable; }
 	
 protected:
@@ -81,7 +84,7 @@ protected:
 
 public:
 	/** The optional description of this member, which is user-facing. */
-	UPROPERTY(EditAnywhere, Category = "General", meta=(EditCondition="bIsEditable", HideEditConditionToggle))
+	UPROPERTY(EditAnywhere, BlueprintReadWrite, Category = "General", meta=(EditCondition="bIsEditable", HideEditConditionToggle))
 	FString Description;
 
 protected:
@@ -117,15 +120,18 @@ public:
 	UMovieGraphVariable() = default;
 
 	/** Returns true if this variable is a global variable. */
+	UFUNCTION(BlueprintPure, Category = "Movie Graph")
 	bool IsGlobal() const;
 
 	/** Gets the category (if any) assigned to this variable. */
+	UFUNCTION(BlueprintPure, Category = "Movie Graph")
 	const FString& GetCategory() const;
 
 	/**
 	 * Sets the variable to the provided category. Be aware that the category provided here may not be the final category set on the
 	 * variable (InNewCategory will be put through FName::NameToDisplayString().
 	 */
+	UFUNCTION(BlueprintCallable, Category = "Movie Graph")
 	void SetCategory(const FString& InNewCategory);
 
 	//~ Begin UMovieGraphMember interface
