MovieGraph: Fixed view extensions not being set up correctly for either the view or the View Family, they were only being initialized, but not set up.

This results in inconsistent handling of the Alpha Composite plugin, as well as any third party plugins which utilize View Extensions to modify the projection matrix for rendering.

It can also create inconsistent handling of Water Rendering, Landscapes, HLOD Resource Residency, and Image Media all of which utilize scene view extensions for their functionality.

diff --git a/Engine/Plugins/MovieScene/MovieRenderPipeline/Source/MovieRenderPipelineRenderPasses/Private/Graph/Renderers/MovieGraphImagePassBase.cpp b/Engine/Plugins/MovieScene/MovieRenderPipeline/Source/MovieRenderPipelineRenderPasses/Private/Graph/Renderers/MovieGraphImagePassBase.cpp
--- a/Engine/Plugins/MovieScene/MovieRenderPipeline/Source/MovieRenderPipelineRenderPasses/Private/Graph/Renderers/MovieGraphImagePassBase.cpp
+++ b/Engine/Plugins/MovieScene/MovieRenderPipeline/Source/MovieRenderPipelineRenderPasses/Private/Graph/Renderers/MovieGraphImagePassBase.cpp
@@ -74,6 +74,11 @@ FSceneView* FMovieGraphImagePassBase::CreateSceneView(const FSceneViewInitOption
 
 	UE::MovieRenderPipeline::UpdateSceneViewForShowFlags(View);
 
+	for (int ViewExt = 0; ViewExt < InViewFamily->ViewExtensions.Num(); ViewExt++)
+	{
+		InViewFamily->ViewExtensions[ViewExt]->SetupView(*InViewFamily, *View);
+	}
+
 	return View;
 }
 
@@ -163,6 +168,11 @@ TSharedRef<FSceneViewFamilyContext> FMovieGraphImagePassBase::CreateSceneViewFam
 	// Need to add the engine-wide view extensions, as rendering code may depend on them (ie: landscapes)
 	OutViewFamily->ViewExtensions.Append(GEngine->ViewExtensions->GatherActiveExtensions(FSceneViewExtensionContext(InInitData.World->Scene)));
 
+	for (FSceneViewExtensionRef& ViewExt : OutViewFamily->ViewExtensions)
+	{
+		ViewExt->SetupViewFamily(*OutViewFamily);
+	}
+
 	return OutViewFamily;
 }
 
