- remove function register_callback_class that constructs the callbackOverrides class
- fix comment in MovieGraphExecuteScriptNode.h with the new python signature that just cals on_job_start

diff --git a/Engine/Plugins/MovieScene/MovieRenderPipeline/Source/MovieRenderPipelineCore/Public/Graph/Nodes/MovieGraphExecuteScriptNode.h b/Engine/Plugins/MovieScene/MovieRenderPipeline/Source/MovieRenderPipelineCore/Public/Graph/Nodes/MovieGraphExecuteScriptNode.h
--- a/Engine/Plugins/MovieScene/MovieRenderPipeline/Source/MovieRenderPipelineCore/Public/Graph/Nodes/MovieGraphExecuteScriptNode.h
+++ b/Engine/Plugins/MovieScene/MovieRenderPipeline/Source/MovieRenderPipelineCore/Public/Graph/Nodes/MovieGraphExecuteScriptNode.h
@@ -26,20 +26,13 @@
 * In "MovieGraphPipelineScriptExample.py" use the following:
 *
 *import unreal
-*
 [REDACTED_EMAIL]()
 *class MovieGraphPipelineScriptExample(unreal.MovieGraphScriptBase):
-
 *    @unreal.ufunction(override=True)
-*    def on_job_start(self, inJobCopy):
-*        super().on_job_start(inJobCopy)
-*        unreal.log_warning("OnJobStart")
-*        for shot in inOutputData.graph_data:
-*            for layerIdentifier in shot.render_layer_data:
-*                unreal.log("render layer: " + layerIdentifier.layer_name)
-*                for file in shot.render_layer_data[layerIdentifier].file_paths:
-*                    unreal.log("file: " + file)
-* 
+*    def on_job_start(self, in_job_copy):
+*        super().on_job_start(in_job_copy)
+*        unreal.log("OnJobStart")
+*
 * Save and restart Unreal Editor, then create a Movie Graph Configuration
 * with an Execute Script Node, and set it to run MovieGraphPipelineScriptExample. 
 * When the job is run, the output console should print "OnJobStart" as a warning.
@@ -47,8 +40,8 @@
 * Use "help(unreal.MovieGraphScriptBase)" in the Python Console to see all available 
 * overridable functions.
 * 
-* See /Engine/Plugins/MovieScene/MovieRenderPipeline/Content/Python for more 
-* python examples.
+* See /Engine/Plugins/MovieScene/MovieRenderPipeline/Content/Python/MovieGraphScriptNodeExample.py 
+* for more python examples.
 */
 UCLASS(BlueprintType, Abstract)
 class MOVIERENDERPIPELINECORE_API UMovieGraphScriptBase : public UObject
