MoviePipeline: Fixed an issue where UI Renderers composited incorrectly when dealing with partially translucent UIs, which resulted in overly bright composites compared to the viewport.

When UMG/Slate does a composite of the UMG render onto the game viewport it currently does it in the incorrect colorspace for premultiplied images. This change adjusts MRQ to match the order used by UMG by shifting around what colorspace the premultiply is done in.

diff --git a/Engine/Plugins/MovieScene/MovieRenderPipeline/Source/MovieRenderPipelineCore/Public/MovieRenderPipelineDataTypes.h b/Engine/Plugins/MovieScene/MovieRenderPipeline/Source/MovieRenderPipelineCore/Public/MovieRenderPipelineDataTypes.h
--- a/Engine/Plugins/MovieScene/MovieRenderPipeline/Source/MovieRenderPipelineCore/Public/MovieRenderPipelineDataTypes.h
+++ b/Engine/Plugins/MovieScene/MovieRenderPipeline/Source/MovieRenderPipelineCore/Public/MovieRenderPipelineDataTypes.h
@@ -1428,7 +1428,8 @@ struct TAsyncCompositeImage<FColor>
 		{
 			return;
 		}
-
+		
+		TRACE_CPUPROFILER_EVENT_SCOPE(MoviePipeline::TAsyncCompositeImageFColor);
 
 		TImagePixelData<FColor>* DestColorData = static_cast<TImagePixelData<FColor>*>(PixelData);
 		TImagePixelData<FColor>* SrcColorData = static_cast<TImagePixelData<FColor>*>(ImageToComposite.Get());
@@ -1443,13 +1444,65 @@ struct TAsyncCompositeImage<FColor>
 					FColor& Dst = DestColorData->Pixels[DstIndex];
 					FColor& Src = SrcColorData->Pixels[DstIndex];
 
-					float SourceAlpha = Src.A / 255.f;
-					FColor Out;
-					Out.A = FMath::Clamp(Src.A + FMath::RoundToInt(Dst.A * (1.f - SourceAlpha)), 0, 255);
-					Out.R = FMath::Clamp(Src.R + FMath::RoundToInt(Dst.R * (1.f - SourceAlpha)), 0, 255);
-					Out.G = FMath::Clamp(Src.G + FMath::RoundToInt(Dst.G * (1.f - SourceAlpha)), 0, 255);
-					Out.B = FMath::Clamp(Src.B + FMath::RoundToInt(Dst.B * (1.f - SourceAlpha)), 0, 255);
-					Dst = Out;
+					if (Src.A == 255)
+					{
+						// If the source is fully opaque then we can just stomp the destination without any calculation
+						Dst = Src;
+					}
+					else if (Src.A == 0)
+					{
+						// If the source is fully transparent, there's no composite needed (which is true for most of our burn-in)
+					}
+					else
+					{
+						// If a pixel is partially transparent then we need to change how we comp.
+						auto Lin2sRGB = [](float InSRGB) -> float
+						{
+							return InSRGB < 0.0031308 ? InSRGB * 12.92 : (1.055f * FMath::Pow(InSRGB, 1 / 2.4f) - 0.055f);
+						};
+
+						// To match the composite done by UMG/Viewport we need to convert from sRGB to linear and
+						// undo the premultiply, then convert back to sRGB and re-premultiply,  before comping.
+						// The operation done by the viewport is done in the incorrect order, so we must match
+						// the incorrect behavior here.
+	
+						// Convert sRGB -> Linear
+						FLinearColor SrcAsLinear = FLinearColor(Src);
+
+						// Unpre-mult. The divide should be safe here as we should have a non-zero alpha to be in this branch.
+						FLinearColor SrcLinearUnPremult;
+						SrcLinearUnPremult.R = SrcAsLinear.R / SrcAsLinear.A;
+						SrcLinearUnPremult.G = SrcAsLinear.G / SrcAsLinear.A;
+						SrcLinearUnPremult.B = SrcAsLinear.B / SrcAsLinear.A;
+						SrcLinearUnPremult.A = SrcAsLinear.A;
+
+						// Convert Linear -> sRGB
+						FLinearColor SrcAsSRGBUnPremult;
+						SrcAsSRGBUnPremult.R = Lin2sRGB(SrcLinearUnPremult.R);
+						SrcAsSRGBUnPremult.G = Lin2sRGB(SrcLinearUnPremult.G);
+						SrcAsSRGBUnPremult.B = Lin2sRGB(SrcLinearUnPremult.B);
+						SrcAsSRGBUnPremult.A = SrcLinearUnPremult.A;
+
+						// Re premult now that it is in sRGB
+						FLinearColor SrcAsSRGBPremult;
+						SrcAsSRGBPremult.R = SrcAsSRGBUnPremult.R * SrcAsSRGBUnPremult.A;
+						SrcAsSRGBPremult.G = SrcAsSRGBUnPremult.G * SrcAsSRGBUnPremult.A;
+						SrcAsSRGBPremult.B = SrcAsSRGBUnPremult.B * SrcAsSRGBUnPremult.A;
+						SrcAsSRGBPremult.A = SrcAsSRGBUnPremult.A;
+
+						// Comp in sRGB (even though they're in floats)
+						FLinearColor SrcL = SrcAsSRGBPremult;
+						FLinearColor DstL = Dst.ReinterpretAsLinear();
+
+						FLinearColor OutL;
+						OutL.R = SrcL.R + (DstL.R * (1.f - SrcL.A));
+						OutL.G = SrcL.G + (DstL.G * (1.f - SrcL.A));
+						OutL.B = SrcL.B + (DstL.B * (1.f - SrcL.A));
+						OutL.A = SrcL.A + (DstL.A * (1.f - SrcL.A));
+
+						// Quantize back to 0-255, without applying sRGB. This also clamps to 0-255.
+						Dst = OutL.QuantizeRound();
+					}
 				}
 			});
 	}
