MovieRenderPipeline: Added functionality to crop overscan from non-exr output formats. In addition, composites will now be offset to the original frustum when overscan is present, as all composites are sized to the user requested resolution instead of the overscanned resolution

diff --git a/Engine/Plugins/Media/AppleProResMedia/Source/AppleProResMedia/Private/MovieGraphAppleProResNode.cpp b/Engine/Plugins/Media/AppleProResMedia/Source/AppleProResMedia/Private/MovieGraphAppleProResNode.cpp
--- a/Engine/Plugins/Media/AppleProResMedia/Source/AppleProResMedia/Private/MovieGraphAppleProResNode.cpp
+++ b/Engine/Plugins/Media/AppleProResMedia/Source/AppleProResMedia/Private/MovieGraphAppleProResNode.cpp
@@ -70,8 +70,8 @@ TUniquePtr<MovieRenderGraph::IVideoCodecWriter> UMovieGraphAppleProResNode::Init
 	
 	FAppleProResEncoderOptions Options;
 	Options.OutputFilename = InInitializationContext.FileName;
-	Options.Width = InInitializationContext.PassData->Value->GetSize().X;
-	Options.Height = InInitializationContext.PassData->Value->GetSize().Y;
+	Options.Width = InInitializationContext.Resolution.X;
+	Options.Height = InInitializationContext.Resolution.Y;
 	Options.FrameRate = EffectiveFrameRate;
 	Options.Codec = EvaluatedNode->Quality;
 	Options.ColorPrimaries = EAppleProResEncoderColorPrimaries::CD_HDREC709; // Force Rec 709 for now
diff --git a/Engine/Plugins/Media/AvidDNxMedia/Source/Source/Private/MovieGraphAvidDNxHRNode.cpp b/Engine/Plugins/Media/AvidDNxMedia/Source/Source/Private/MovieGraphAvidDNxHRNode.cpp
--- a/Engine/Plugins/Media/AvidDNxMedia/Source/Source/Private/MovieGraphAvidDNxHRNode.cpp
+++ b/Engine/Plugins/Media/AvidDNxMedia/Source/Source/Private/MovieGraphAvidDNxHRNode.cpp
@@ -85,8 +85,8 @@ TUniquePtr<MovieRenderGraph::IVideoCodecWriter> UMovieGraphAvidDNxHRNode::Initia
 	
 	FAvidDNxEncoderOptions Options;
 	Options.OutputFilename = InInitializationContext.FileName;
-	Options.Width = InInitializationContext.PassData->Value->GetSize().X;
-	Options.Height = InInitializationContext.PassData->Value->GetSize().Y;
+	Options.Width = InInitializationContext.Resolution.X;
+	Options.Height = InInitializationContext.Resolution.Y;
 	Options.Quality = EvaluatedNode->Quality;
 	Options.FrameRate = EffectiveFrameRate;
 	Options.bCompress = true;
diff --git a/Engine/Plugins/MovieScene/MovieRenderPipeline/Source/MovieRenderPipelineCore/Private/Graph/MovieGraphBlueprintLibrary.cpp b/Engine/Plugins/MovieScene/MovieRenderPipeline/Source/MovieRenderPipelineCore/Private/Graph/MovieGraphBlueprintLibrary.cpp
--- a/Engine/Plugins/MovieScene/MovieRenderPipeline/Source/MovieRenderPipelineCore/Private/Graph/MovieGraphBlueprintLibrary.cpp
+++ b/Engine/Plugins/MovieScene/MovieRenderPipeline/Source/MovieRenderPipelineCore/Private/Graph/MovieGraphBlueprintLibrary.cpp
@@ -405,6 +405,34 @@ FIntPoint UMovieGraphBlueprintLibrary::GetEffectiveOutputResolution(UMovieGraphE
 	return UMoviePipelineBlueprintLibrary::Utility_GetEffectiveOutputResolution(RescaledOverscan, NamedResolution.Resolution);
 }
 
+FIntRect UMovieGraphBlueprintLibrary::GetOverscanCropRectangle(UMovieGraphEvaluatedConfig* InEvaluatedGraph, float DefaultOverscan)
+{
+	if (!InEvaluatedGraph)
+	{
+		FFrame::KismetExecutionMessage(TEXT("InEvaluatedGraph cannot be null."), ELogVerbosity::Error);
+		return FIntRect();
+	}
+	
+	const FIntPoint OverscannedResolution = GetEffectiveOutputResolution(InEvaluatedGraph, DefaultOverscan);
+	
+	constexpr bool bIncludeCDOs = true;
+	const UMovieGraphGlobalOutputSettingNode* OutputSetting = InEvaluatedGraph->GetSettingForBranch<UMovieGraphGlobalOutputSettingNode>(UMovieGraphNode::GlobalsPinName, bIncludeCDOs);
+	if (!ensure(OutputSetting))
+	{
+		return FIntRect();
+	}
+
+	const FIntPoint OutputResolution = OutputSetting->OutputResolution.Resolution;
+	
+	FIntRect CropRect;
+
+	// Overscan is symmetrical, so this difference is always evenly divisible by 2.
+	CropRect.Min = (OverscannedResolution - OutputResolution) / 2;
+	CropRect.Max = CropRect.Min + OutputResolution;
+
+	return CropRect;
+}
+
 FText UMovieGraphBlueprintLibrary::GetJobName(const UMovieGraphPipeline* InMovieGraphPipeline)
 {
 	return InMovieGraphPipeline ? FText::FromString(InMovieGraphPipeline->GetCurrentJob()->JobName) : FText();
diff --git a/Engine/Plugins/MovieScene/MovieRenderPipeline/Source/MovieRenderPipelineCore/Private/Graph/Nodes/MovieGraphVideoOutputNode.cpp b/Engine/Plugins/MovieScene/MovieRenderPipeline/Source/MovieRenderPipelineCore/Private/Graph/Nodes/MovieGraphVideoOutputNode.cpp
--- a/Engine/Plugins/MovieScene/MovieRenderPipeline/Source/MovieRenderPipelineCore/Private/Graph/Nodes/MovieGraphVideoOutputNode.cpp
+++ b/Engine/Plugins/MovieScene/MovieRenderPipeline/Source/MovieRenderPipelineCore/Private/Graph/Nodes/MovieGraphVideoOutputNode.cpp
@@ -85,6 +85,13 @@ void UMovieGraphVideoOutputNode::OnReceiveImageDataImpl(UMovieGraphPipeline* InP
 		FImagePixelData* RawRenderPassData = RenderPassData.Value.Get();
 		const UE::MovieGraph::FMovieGraphSampleState* Payload = RawRenderPassData->GetPayload<UE::MovieGraph::FMovieGraphSampleState>();
 
+		FAsyncCropImage CropImagePreprocessor(Payload->CropRectangle);
+		if (ShouldCropOverscanImpl())
+		{
+			CropImagePreprocessor(RawRenderPassData);
+			RawRenderPassData = CropImagePreprocessor.OutCroppedImage.Get();
+		}
+		
 		// Write this frame's beauty pass, and pass along other passes so the encoder can composite them
 		if (RenderPassData.Key.SubResourceName == TEXT("beauty"))
 		{
@@ -283,11 +290,18 @@ UMovieGraphVideoOutputNode::FMovieGraphCodecWriterWithPromise* UMovieGraphVideoO
 	
 	if (!OutputWriter)
 	{
+		FIntPoint OutputResolution = InRenderPassData.Value->GetSize();
+		if (ShouldCropOverscanImpl())
+		{
+			OutputResolution = Payload->CropRectangle.Size();
+		}
+		
 		FMovieGraphVideoNodeInitializationContext InitializationContext;
 		InitializationContext.Pipeline = InPipeline;
 		InitializationContext.EvaluatedConfig = EvaluatedConfig;
 		InitializationContext.TraversalContext = &InRawFrameData->TraversalContext;
 		InitializationContext.PassData = &InRenderPassData;
+		InitializationContext.Resolution = OutputResolution;
 		InitializationContext.FileName = FinalFilePath;
 		InitializationContext.bAllowOCIO = Payload->bAllowOCIO;
 
diff --git a/Engine/Plugins/MovieScene/MovieRenderPipeline/Source/MovieRenderPipelineCore/Private/Graph/Nodes/MovieGraphWidgetRendererBaseNode.cpp b/Engine/Plugins/MovieScene/MovieRenderPipeline/Source/MovieRenderPipelineCore/Private/Graph/Nodes/MovieGraphWidgetRendererBaseNode.cpp
--- a/Engine/Plugins/MovieScene/MovieRenderPipeline/Source/MovieRenderPipelineCore/Private/Graph/Nodes/MovieGraphWidgetRendererBaseNode.cpp
+++ b/Engine/Plugins/MovieScene/MovieRenderPipeline/Source/MovieRenderPipelineCore/Private/Graph/Nodes/MovieGraphWidgetRendererBaseNode.cpp
@@ -107,8 +107,11 @@ void UMovieGraphWidgetRendererBaseNode::FMovieGraphWidgetPass::Render(const FMov
 	UMovieGraphWidgetRendererBaseNode* ParentNodeThisFrame = GetParentNode(InTimeData.EvaluatedConfig);
 
 	const float CameraOverscan = Renderer->GetCameraOverscan(LayerData.CameraIndex);
-
-	const FIntPoint OutputResolution = UMovieGraphBlueprintLibrary::GetEffectiveOutputResolution(InTimeData.EvaluatedConfig, CameraOverscan);
+	const FIntRect CropRect = UMovieGraphBlueprintLibrary::GetOverscanCropRectangle(InTimeData.EvaluatedConfig, CameraOverscan);
+	
+	// Composited elements should be sized to the original frustum size, as the final image is either cropped to that size, or the composite will be offset
+	// to match the original frustum
+	const FIntPoint OutputResolution = CropRect.Size();
 	const int32 MaxResolution = GetMax2DTextureDimension();
 	if ((OutputResolution.X > MaxResolution) || (OutputResolution.Y > MaxResolution))
 	{
diff --git a/Engine/Plugins/MovieScene/MovieRenderPipeline/Source/MovieRenderPipelineCore/Private/MoviePipelineBlueprintLibrary.cpp b/Engine/Plugins/MovieScene/MovieRenderPipeline/Source/MovieRenderPipelineCore/Private/MoviePipelineBlueprintLibrary.cpp
--- a/Engine/Plugins/MovieScene/MovieRenderPipeline/Source/MovieRenderPipelineCore/Private/MoviePipelineBlueprintLibrary.cpp
+++ b/Engine/Plugins/MovieScene/MovieRenderPipeline/Source/MovieRenderPipelineCore/Private/MoviePipelineBlueprintLibrary.cpp
@@ -914,6 +914,29 @@ FIntPoint UMoviePipelineBlueprintLibrary::GetEffectiveOutputResolution(UMoviePip
 	return FIntPoint();
 }
 
+FIntRect UMoviePipelineBlueprintLibrary::GetOverscanCropRectangle(const UMoviePipelinePrimaryConfig* InPrimaryConfig, const UMoviePipelineExecutorShot* InPipelineExecutorShot, float DefaultOverscan)
+{
+	if (InPrimaryConfig && InPipelineExecutorShot)
+	{
+		const UMoviePipelineOutputSetting* OutputSetting = InPrimaryConfig->FindSetting<UMoviePipelineOutputSetting>();
+		const UMoviePipelineCameraSetting* CameraSetting = Cast<const UMoviePipelineCameraSetting>(UMoviePipelineBlueprintLibrary::FindOrGetDefaultSettingForShot(UMoviePipelineCameraSetting::StaticClass(), InPrimaryConfig, InPipelineExecutorShot));
+		check(OutputSetting);
+		check(CameraSetting);
+		
+		const FIntPoint OverscannedResolution = Utility_GetEffectiveOutputResolution(CameraSetting->bOverrideCameraOverscan ? CameraSetting->OverscanPercentage : DefaultOverscan, OutputSetting->OutputResolution);
+		
+		FIntRect CropRect;
+
+		// Overscan is symmetrical, so this difference is always evenly divisible by 2.
+		CropRect.Min = (OverscannedResolution - OutputSetting->OutputResolution) / 2;
+		CropRect.Max = CropRect.Min + OutputSetting->OutputResolution;
+
+		return CropRect;
+	}
+
+	return FIntRect();
+}
+
 FIntPoint UMoviePipelineBlueprintLibrary::Utility_GetEffectiveOutputResolution(const float OverscanPercentage, const FIntPoint& InOutputResolution)
 {
 	const float ClampedOverscanPercentage = FMath::Clamp(OverscanPercentage, 0.0f, 1.0f);
diff --git a/Engine/Plugins/MovieScene/MovieRenderPipeline/Source/MovieRenderPipelineCore/Private/MoviePipelineUtils.cpp b/Engine/Plugins/MovieScene/MovieRenderPipeline/Source/MovieRenderPipelineCore/Private/MoviePipelineUtils.cpp
--- a/Engine/Plugins/MovieScene/MovieRenderPipeline/Source/MovieRenderPipelineCore/Private/MoviePipelineUtils.cpp
+++ b/Engine/Plugins/MovieScene/MovieRenderPipeline/Source/MovieRenderPipelineCore/Private/MoviePipelineUtils.cpp
@@ -1226,7 +1226,9 @@ namespace UE
 
 			if (InPrimaryConfig && InPipelineExecutorShot)
 			{
+				const UMoviePipelineOutputSetting* OutputSettings = InPrimaryConfig->FindSetting<UMoviePipelineOutputSetting>();
 				const UMoviePipelineHighResSetting* HighResSettings = Cast<const UMoviePipelineHighResSetting>(UMoviePipelineBlueprintLibrary::FindOrGetDefaultSettingForShot(UMoviePipelineHighResSetting::StaticClass(), InPrimaryConfig, InPipelineExecutorShot));
+				check(OutputSettings);
 				check(HighResSettings);
 
 				FIntPoint BackbufferResolution = FIntPoint(FMath::CeilToInt((float)InEffectiveOutputResolution.X / (float)OutRenderPassMetrics.OriginalTileCounts.X), FMath::CeilToInt((float)InEffectiveOutputResolution.Y / (float)OutRenderPassMetrics.OriginalTileCounts.Y));
@@ -1242,6 +1244,15 @@ namespace UE
 				OutRenderPassMetrics.BackbufferSize = BackbufferResolution;
 				OutRenderPassMetrics.EffectiveOutputResolution = InEffectiveOutputResolution;
 
+
+				FIntRect CropRect;
+
+				// Overscan is symmetrical, so this difference is always evenly divisible by 2.
+				CropRect.Min = (InEffectiveOutputResolution - OutputSettings->OutputResolution) / 2;
+				CropRect.Max = CropRect.Min + OutputSettings->OutputResolution;
+
+				OutRenderPassMetrics.CropRectangle = CropRect;
+				
 				{
 					OutRenderPassMetrics.OverlappedPad = FIntPoint(FMath::CeilToInt(TileResolution.X * HighResSettings->OverlapRatio),
 						FMath::CeilToInt(TileResolution.Y * HighResSettings->OverlapRatio));
diff --git a/Engine/Plugins/MovieScene/MovieRenderPipeline/Source/MovieRenderPipelineCore/Private/MoviePipelineVideoOutputBase.cpp b/Engine/Plugins/MovieScene/MovieRenderPipeline/Source/MovieRenderPipelineCore/Private/MoviePipelineVideoOutputBase.cpp
--- a/Engine/Plugins/MovieScene/MovieRenderPipeline/Source/MovieRenderPipelineCore/Private/MoviePipelineVideoOutputBase.cpp
+++ b/Engine/Plugins/MovieScene/MovieRenderPipeline/Source/MovieRenderPipelineCore/Private/MoviePipelineVideoOutputBase.cpp
@@ -183,9 +183,17 @@ void UMoviePipelineVideoOutputBase::OnReceiveImageDataImpl(FMoviePipelineMergerO
 		
 		if(!OutputWriter)
 		{
+			// If overscan is going to be cropped, use the output resolution from the crop rectangle. Otherwise, use the internal resolution from the
+			// raw render pass data
+			FIntPoint OutputResolution = RenderPassData.Value->GetSize();
+			if (ShouldCropOverscanImpl())
+			{
+				OutputResolution = Payload->SampleState.CropRectangle.Size();
+			}
+			
 			// Create a new writer for this file name (and output format settings)
 			TUniquePtr<MovieRenderPipeline::IVideoCodecWriter> NewWriter = Initialize_GameThread(FinalFilePath,
-				RenderPassData.Value->GetSize(), RenderPassData.Value->GetType(), RenderPassData.Value->GetPixelLayout(),
+				OutputResolution, RenderPassData.Value->GetType(), RenderPassData.Value->GetPixelLayout(),
 				RenderPassData.Value->GetBitDepth(), RenderPassData.Value->GetNumChannels());
 
 			if (NewWriter)
@@ -225,6 +233,13 @@ void UMoviePipelineVideoOutputBase::OnReceiveImageDataImpl(FMoviePipelineMergerO
 		FMoviePipelineBackgroundMediaTasks Task;
 		FImagePixelData* RawRenderPassData = RenderPassData.Value.Get();
 
+		FAsyncCropImage CropImagePreprocessor(Payload->SampleState.CropRectangle);
+		if (ShouldCropOverscanImpl())
+		{
+			CropImagePreprocessor(RawRenderPassData);
+			RawRenderPassData = CropImagePreprocessor.OutCroppedImage.Get();
+		}
+		
 		// Making sure that if OCIO is enabled the Quantization won't do additional color conversion.
 		UMoviePipelineColorSetting* ColorSetting = GetPipeline()->GetPipelinePrimaryConfig()->FindSetting<UMoviePipelineColorSetting>();
 		OutputWriter->Get<0>()->bConvertToSrgb = !(ColorSetting && ColorSetting->OCIOConfiguration.bIsEnabled);
diff --git a/Engine/Plugins/MovieScene/MovieRenderPipeline/Source/MovieRenderPipelineCore/Public/Graph/MovieGraphBlueprintLibrary.h b/Engine/Plugins/MovieScene/MovieRenderPipeline/Source/MovieRenderPipelineCore/Public/Graph/MovieGraphBlueprintLibrary.h
--- a/Engine/Plugins/MovieScene/MovieRenderPipeline/Source/MovieRenderPipelineCore/Public/Graph/MovieGraphBlueprintLibrary.h
+++ b/Engine/Plugins/MovieScene/MovieRenderPipeline/Source/MovieRenderPipelineCore/Public/Graph/MovieGraphBlueprintLibrary.h
@@ -70,6 +70,15 @@ public:
 	UFUNCTION(BlueprintCallable, Category = "Movie Graph")
 	static FIntPoint GetEffectiveOutputResolution(UMovieGraphEvaluatedConfig* InEvaluatedGraph, float DefaultOverscan = 0.0f);
 
+	/**
+	 * Gets a rectangle that will crop out any overscan applied to the scene. If there is no overscan, the rectangle returned will match the backbuffer rectangle
+	 * @param InEvaluatedGraph		- The evaluated graph that will provide context for resolving the resolution
+	 * @param DefaultOverscan		- The default overscan to use if there are no camera settings that provide an overscan override value, from 0.0 to 1.0
+	 * @return						- The crop rectangle, which will have a resolution matching the requested output resolution
+	 */
+	UFUNCTION(BlueprintCallable, Category = "Movie Graph")
+	static FIntRect GetOverscanCropRectangle(UMovieGraphEvaluatedConfig* InEvaluatedGraph, float DefaultOverscan = 0.0f);
+	
 	/**
 	* Gets the name of the current job.
 	* @param	InMovieGraphPipeline	 - The pipeline to get the job name from.
diff --git a/Engine/Plugins/MovieScene/MovieRenderPipeline/Source/MovieRenderPipelineCore/Public/Graph/MovieGraphDataTypes.h b/Engine/Plugins/MovieScene/MovieRenderPipeline/Source/MovieRenderPipelineCore/Public/Graph/MovieGraphDataTypes.h
--- a/Engine/Plugins/MovieScene/MovieRenderPipeline/Source/MovieRenderPipelineCore/Public/Graph/MovieGraphDataTypes.h
+++ b/Engine/Plugins/MovieScene/MovieRenderPipeline/Source/MovieRenderPipelineCore/Public/Graph/MovieGraphDataTypes.h
@@ -424,6 +424,13 @@ namespace UE::MovieGraph
 		/** When using camera overscan, what fraction (0-1) did this render use? Needed so that exrs can take a center-out crop of the data. */
 		float OverscanFraction;
 
+		/**
+		 * The rectangle to use to crop the image prior to compositing and output. Used to crop out overscan from appropriate formats, and will
+		 * generally have size equal to the originally requested resolution, whereas internal resolutions such as Backbuffer and Accumulator
+		 * resolutions will be scaled up to include overscan
+		 */
+		FIntRect CropRectangle;
+		
 		/**
 		* If multiple passes are composited on top of a render, the sort order determines the order in which they're composited.
 		* Passes with a low sort order will composite on top of passes with a higher sort order.
diff --git a/Engine/Plugins/MovieScene/MovieRenderPipeline/Source/MovieRenderPipelineCore/Public/Graph/Nodes/MovieGraphFileOutputNode.h b/Engine/Plugins/MovieScene/MovieRenderPipeline/Source/MovieRenderPipelineCore/Public/Graph/Nodes/MovieGraphFileOutputNode.h
--- a/Engine/Plugins/MovieScene/MovieRenderPipeline/Source/MovieRenderPipelineCore/Public/Graph/Nodes/MovieGraphFileOutputNode.h
+++ b/Engine/Plugins/MovieScene/MovieRenderPipeline/Source/MovieRenderPipelineCore/Public/Graph/Nodes/MovieGraphFileOutputNode.h
@@ -83,6 +83,9 @@ protected:
 	virtual void OnAllFramesFinalizedImpl(UMovieGraphPipeline* InPipeline, TObjectPtr<UMovieGraphEvaluatedConfig>& InPrimaryJobEvaluatedGraph) {}
 	virtual void OnAllShotFramesSubmittedImpl(UMovieGraphPipeline* InPipeline, const UMoviePipelineExecutorShot* InShot, TObjectPtr<UMovieGraphEvaluatedConfig>& InShotEvaluatedGraph) {}
 	virtual bool IsFinishedWritingToDiskImpl() const { return true; }
+
+	/** Returns whether any overscan should be automatically cropped from the image before compositing and output */
+	virtual bool ShouldCropOverscanImpl() const { return true; }
 	
 	/** Returns the number of evaluated (active) file nodes on the specified branch. */
 	static int32 GetNumFileOutputNodes(const UMovieGraphEvaluatedConfig& InEvaluatedConfig, const FName& InBranchName);
diff --git a/Engine/Plugins/MovieScene/MovieRenderPipeline/Source/MovieRenderPipelineCore/Public/Graph/Nodes/MovieGraphVideoOutputNode.h b/Engine/Plugins/MovieScene/MovieRenderPipeline/Source/MovieRenderPipelineCore/Public/Graph/Nodes/MovieGraphVideoOutputNode.h
--- a/Engine/Plugins/MovieScene/MovieRenderPipeline/Source/MovieRenderPipelineCore/Public/Graph/Nodes/MovieGraphVideoOutputNode.h
+++ b/Engine/Plugins/MovieScene/MovieRenderPipeline/Source/MovieRenderPipelineCore/Public/Graph/Nodes/MovieGraphVideoOutputNode.h
@@ -42,6 +42,7 @@ protected:
 		TObjectPtr<UMovieGraphEvaluatedConfig> EvaluatedConfig;
 		const FMovieGraphTraversalContext* TraversalContext;
 		const FMovieGraphPassData* PassData;
+		FIntPoint Resolution;
 		FString FileName;
 		bool bAllowOCIO;
 	};
diff --git a/Engine/Plugins/MovieScene/MovieRenderPipeline/Source/MovieRenderPipelineCore/Public/MoviePipelineBlueprintLibrary.h b/Engine/Plugins/MovieScene/MovieRenderPipeline/Source/MovieRenderPipelineCore/Public/MoviePipelineBlueprintLibrary.h
--- a/Engine/Plugins/MovieScene/MovieRenderPipeline/Source/MovieRenderPipelineCore/Public/MoviePipelineBlueprintLibrary.h
+++ b/Engine/Plugins/MovieScene/MovieRenderPipeline/Source/MovieRenderPipelineCore/Public/MoviePipelineBlueprintLibrary.h
@@ -151,6 +151,16 @@ public:
 	UFUNCTION(BlueprintCallable, Category = "Movie Render Pipeline")
 	static FIntPoint GetEffectiveOutputResolution(UMoviePipelinePrimaryConfig* InPrimaryConfig, UMoviePipelineExecutorShot* InPipelineExecutorShot, float DefaultOverscan = 0.0f);
 
+	/**
+	 * Gets a rectangle that will crop out any overscan applied to the scene. If there is no overscan, the rectangle returned will match the backbuffer rectangle
+	 * @param InPrimaryConfig			- The pipeline configuration that will provide context for resolving the resolution
+	 * @param InPipelineExecutorShot	- The shot that is being generated by the pipeline
+	 * @param DefaultOverscan			- The default overscan to use if there are no camera settings that provide an overscan override value, from 0.0 to 1.0
+	 * @return							- The crop rectangle, which will have a resolution matching the requested output resolution
+	 */
+	UFUNCTION(BlueprintCallable, Category = "Movie Render Pipeline")
+	static FIntRect GetOverscanCropRectangle(const UMoviePipelinePrimaryConfig* InPrimaryConfig, const UMoviePipelineExecutorShot* InPipelineExecutorShot, float DefaultOverscan = 0.0f);
+	
 	/** Allows access to a setting of provided type for specific shot. */
 	UFUNCTION(BlueprintCallable, Category = "Movie Render Pipeline", meta = (DeterminesOutputType = "InSettingType"))
 	static UMoviePipelineSetting* FindOrGetDefaultSettingForShot(TSubclassOf<UMoviePipelineSetting> InSettingType, const UMoviePipelinePrimaryConfig* InPrimaryConfig, const UMoviePipelineExecutorShot* InShot);
diff --git a/Engine/Plugins/MovieScene/MovieRenderPipeline/Source/MovieRenderPipelineCore/Public/MoviePipelineOutputBase.h b/Engine/Plugins/MovieScene/MovieRenderPipeline/Source/MovieRenderPipelineCore/Public/MoviePipelineOutputBase.h
--- a/Engine/Plugins/MovieScene/MovieRenderPipeline/Source/MovieRenderPipelineCore/Public/MoviePipelineOutputBase.h
+++ b/Engine/Plugins/MovieScene/MovieRenderPipeline/Source/MovieRenderPipelineCore/Public/MoviePipelineOutputBase.h
@@ -68,6 +68,9 @@ protected:
 	virtual void OnShotFinishedImpl(const UMoviePipelineExecutorShot* InShot, const bool bFlushToDisk = false) {}
 	virtual void OnPipelineFinishedImpl() {}
 	virtual void OnPostTickImpl() {}
+	
+	/** Returns whether any overscan should be automatically cropped from the image before compositing and output */
+	virtual bool ShouldCropOverscanImpl() const { return true; }
 	// ~UMoviePipelineOutputBase
 
 protected:
diff --git a/Engine/Plugins/MovieScene/MovieRenderPipeline/Source/MovieRenderPipelineCore/Public/MovieRenderPipelineDataTypes.h b/Engine/Plugins/MovieScene/MovieRenderPipeline/Source/MovieRenderPipelineCore/Public/MovieRenderPipelineDataTypes.h
--- a/Engine/Plugins/MovieScene/MovieRenderPipeline/Source/MovieRenderPipelineCore/Public/MovieRenderPipelineDataTypes.h
+++ b/Engine/Plugins/MovieScene/MovieRenderPipeline/Source/MovieRenderPipelineCore/Public/MovieRenderPipelineDataTypes.h
@@ -3,6 +3,7 @@
 
 #include "Misc/Timecode.h"
 #include "ImagePixelData.h"
+#include "ImageWriteTask.h"
 #include "Containers/UnrealString.h"
 #include "Engine/Scene.h"
 #include "Engine/EngineTypes.h"
@@ -1009,6 +1010,13 @@ public:
 
 	/** Whether to override the camera's overscan value when rendering */
 	bool bOverrideCameraOverscan;
+
+	/**
+	 * The rectangle to use to crop the image prior to compositing and output. Used to crop out overscan from appropriate formats, and will
+	 * generally have size equal to the originally requested resolution, whereas internal resolutions such as Backbuffer and Accumulator
+	 * resolutions will be scaled up to include overscan
+	 */
+	FIntRect CropRectangle;
 	
 	/** 
 	* The gamma space to apply accumulation in. During accumulation, pow(x,AccumulationGamma) is applied
@@ -1415,16 +1423,32 @@ template<typename PixelType> struct TAsyncCompositeImage;
 template<>
 struct TAsyncCompositeImage<FColor>
 {
-	TAsyncCompositeImage(TUniquePtr<FImagePixelData>&& InPixelData)
+	TAsyncCompositeImage(TUniquePtr<FImagePixelData>&& InPixelData, FIntPoint InOffset = FIntPoint())
 		: ImageToComposite(MoveTemp(InPixelData))
+		, Offset(InOffset)
 	{}
 
 	void operator()(FImagePixelData* PixelData)
 	{
 		check(PixelData->GetType() == EImagePixelType::Color);
 		check(ImageToComposite && ImageToComposite->GetType() == EImagePixelType::Color);
-		if (!ensureMsgf(ImageToComposite->GetSize() == PixelData->GetSize(), TEXT("Cannot composite images of different sizes! Source: (%d,%d) Target: (%d,%d)"),
-			ImageToComposite->GetSize().X, ImageToComposite->GetSize().Y, PixelData->GetSize().X, PixelData->GetSize().Y))
+
+		// Offset within source and destination images to start compositing from and to
+		const FIntPoint SrcOffset = FIntPoint(FMath::Max(-Offset.X, 0), FMath::Max(-Offset.Y, 0));
+		const FIntPoint DestOffset = FIntPoint(FMath::Max(Offset.X, 0), FMath::Max(Offset.Y, 0));
+		
+		const FIntPoint SrcSize = ImageToComposite->GetSize();
+		const FIntPoint DestSize = PixelData->GetSize();
+		
+		// The size of the region being composited
+		const FIntPoint CompositeSize = FIntPoint(
+			FMath::Min(SrcSize.X - SrcOffset.X, DestSize.X - DestOffset.X),
+			FMath::Min(SrcSize.Y - SrcOffset.Y, DestSize.Y - DestOffset.Y));
+
+		// If either dimension of our composite size is less than or equal to zero, it means that one or more axes of our source image is completely outside the bounds of
+		// the destination image when offset to the specified position, and can't be properly composited onto it
+		if (!ensureMsgf(CompositeSize.X > 0 && CompositeSize.Y > 0, TEXT("Source of size (%d,%d) does not overlap destination of size (%d,%d) when offset to (%d,%d)!"),
+			SrcSize.X, SrcSize.Y, DestSize.X, DestSize.Y, Offset.X, Offset.Y))
 		{
 			return;
 		}
@@ -1434,15 +1458,16 @@ struct TAsyncCompositeImage<FColor>
 		TImagePixelData<FColor>* DestColorData = static_cast<TImagePixelData<FColor>*>(PixelData);
 		TImagePixelData<FColor>* SrcColorData = static_cast<TImagePixelData<FColor>*>(ImageToComposite.Get());
 
-		ParallelFor(DestColorData->GetSize().Y,
+		ParallelFor(CompositeSize.Y,
 			[&](int32 ScanlineIndex = 0)
 			{
-				for (int64 ColumnIndex = 0; ColumnIndex < DestColorData->GetSize().X; ColumnIndex++)
+				for (int64 ColumnIndex = 0; ColumnIndex < CompositeSize.X; ColumnIndex++)
 				{
-					int64 DstIndex = int64(ScanlineIndex) * int64(DestColorData->GetSize().X) + int64(ColumnIndex);
-
+					int64 SrcIndex = (int64(ScanlineIndex) + int64(SrcOffset.Y)) * int64(SrcColorData->GetSize().X) + int64(ColumnIndex) + int64(SrcOffset.X);
+					int64 DstIndex = (int64(ScanlineIndex) + int64(DestOffset.Y)) * int64(DestColorData->GetSize().X) + int64(ColumnIndex) + int64(DestOffset.X);
+					
 					FColor& Dst = DestColorData->Pixels[DstIndex];
-					FColor& Src = SrcColorData->Pixels[DstIndex];
+					FColor& Src = SrcColorData->Pixels[SrcIndex];
 
 					if (Src.A == 255)
 					{
@@ -1508,35 +1533,61 @@ struct TAsyncCompositeImage<FColor>
 	}
 
 	TUniquePtr<FImagePixelData> ImageToComposite;
+	
+	/**
+	 * Amount in pixels from the destination's top left to offset the source when compositing onto the destination.
+	 * Can be negative, in which case a portion of the source's bottom right will be composited onto the destination's top left
+	 */
+	FIntPoint Offset;
 };
 
 template<>
 struct TAsyncCompositeImage<FFloat16Color>
 {
-	TAsyncCompositeImage(TUniquePtr<FImagePixelData>&& InPixelData)
+	TAsyncCompositeImage(TUniquePtr<FImagePixelData>&& InPixelData, FIntPoint InOffset = FIntPoint())
 		: ImageToComposite(MoveTemp(InPixelData))
+		, Offset(InOffset)
 	{}
 
 	void operator()(FImagePixelData* PixelData)
 	{
 		check(PixelData->GetType() == EImagePixelType::Float16);
 		check(ImageToComposite && ImageToComposite->GetType() == EImagePixelType::Color);
-		if (!ensureMsgf(ImageToComposite->GetSize() == PixelData->GetSize(), TEXT("Cannot composite images of different sizes! Source: (%d,%d) Target: (%d,%d)"),
-			ImageToComposite->GetSize().X, ImageToComposite->GetSize().Y, PixelData->GetSize().X, PixelData->GetSize().Y))
+		
+		// Offset within source and destination images to start compositing from and to
+        const FIntPoint SrcOffset = FIntPoint(FMath::Max(-Offset.X, 0), FMath::Max(-Offset.Y, 0));
+        const FIntPoint DestOffset = FIntPoint(FMath::Max(Offset.X, 0), FMath::Max(Offset.Y, 0));
+        
+        const FIntPoint SrcSize = ImageToComposite->GetSize();
+        const FIntPoint DestSize = PixelData->GetSize();
+        
+        // The size of the region being composited
+        const FIntPoint CompositeSize = FIntPoint(
+        	FMath::Min(SrcSize.X - SrcOffset.X, DestSize.X - DestOffset.X),
+        	FMath::Min(SrcSize.Y - SrcOffset.Y, DestSize.Y - DestOffset.Y));
+
+		// If either dimension of our composite size is less than or equal to zero, it means that one or more axes of our source image is completely outside the bounds of
+		// the destination image when offset to the specified position, and can't be properly composited onto it
+		if (!ensureMsgf(CompositeSize.X > 0 && CompositeSize.Y > 0, TEXT("Source of size (%d,%d) does not overlap destination of size (%d,%d) when offset to (%d,%d)!"),
+			SrcSize.X, SrcSize.Y, DestSize.X, DestSize.Y, Offset.X, Offset.Y))
 		{
 			return;
 		}
 
+		TRACE_CPUPROFILER_EVENT_SCOPE(MoviePipeline::TAsyncCompositeImageFFloat16Color);
+		
 		TImagePixelData<FFloat16Color>* DestColorData = static_cast<TImagePixelData<FFloat16Color>*>(PixelData);
 		TImagePixelData<FColor>* SrcColorData = static_cast<TImagePixelData<FColor>*>(ImageToComposite.Get());
-		ParallelFor(DestColorData->GetSize().Y,
+		ParallelFor(CompositeSize.Y,
 			[&](int32 ScanlineIndex = 0)
 			{
-				for (int64 ColumnIndex = 0; ColumnIndex < DestColorData->GetSize().X; ColumnIndex++)
+				for (int64 ColumnIndex = 0; ColumnIndex < CompositeSize.X; ColumnIndex++)
 				{
-					int64 DstIndex = int64(ScanlineIndex) * int64(DestColorData->GetSize().X) + int64(ColumnIndex);
+					int64 SrcIndex = (int64(ScanlineIndex) + int64(SrcOffset.Y)) * int64(SrcColorData->GetSize().X) + int64(ColumnIndex) + int64(SrcOffset.X);
+					int64 DstIndex = (int64(ScanlineIndex) + int64(DestOffset.Y)) * int64(DestColorData->GetSize().X) + int64(ColumnIndex) + int64(DestOffset.X);
+					
 					FFloat16Color& Dst = DestColorData->Pixels[DstIndex];
-					FColor& Src = SrcColorData->Pixels[DstIndex];
+					FColor& Src = SrcColorData->Pixels[SrcIndex];
 
 					float SourceAlpha = Src.A / 255.f;
 					FFloat16Color Out;
@@ -1550,36 +1601,62 @@ struct TAsyncCompositeImage<FFloat16Color>
 	}
 
 	TUniquePtr<FImagePixelData> ImageToComposite;
+	
+	/**
+	 * Amount in pixels from the destination's top left to offset the source when compositing onto the destination.
+	 * Can be negative, in which case a portion of the source's bottom right will be composited onto the destination's top left
+	 */
+	FIntPoint Offset;
 };
 
 template<>
 struct TAsyncCompositeImage<FLinearColor>
 {
-	TAsyncCompositeImage(TUniquePtr<FImagePixelData>&& InPixelData)
+	TAsyncCompositeImage(TUniquePtr<FImagePixelData>&& InPixelData, FIntPoint InOffset = FIntPoint())
 		: ImageToComposite(MoveTemp(InPixelData))
+		, Offset(InOffset)
 	{}
 
 	void operator()(FImagePixelData* PixelData)
 	{
 		check(PixelData->GetType() == EImagePixelType::Float32);
 		check(ImageToComposite && ImageToComposite->GetType() == EImagePixelType::Color);
-		if (!ensureMsgf(ImageToComposite->GetSize() == PixelData->GetSize(), TEXT("Cannot composite images of different sizes! Source: (%d,%d) Target: (%d,%d)"),
-			ImageToComposite->GetSize().X, ImageToComposite->GetSize().Y, PixelData->GetSize().X, PixelData->GetSize().Y))
+
+		// Offset within source and destination images to start compositing from and to
+		const FIntPoint SrcOffset = FIntPoint(FMath::Max(-Offset.X, 0), FMath::Max(-Offset.Y, 0));
+		const FIntPoint DestOffset = FIntPoint(FMath::Max(Offset.X, 0), FMath::Max(Offset.Y, 0));
+		
+		const FIntPoint SrcSize = ImageToComposite->GetSize();
+		const FIntPoint DestSize = PixelData->GetSize();
+		
+		// The size of the region being composited
+		const FIntPoint CompositeSize = FIntPoint(
+			FMath::Min(SrcSize.X - SrcOffset.X, DestSize.X - DestOffset.X),
+			FMath::Min(SrcSize.Y - SrcOffset.Y, DestSize.Y - DestOffset.Y));
+
+		// If either dimension of our composite size is less than or equal to zero, it means that one or more axes of our source image is completely outside the bounds of
+		// the destination image when offset to the specified position, and can't be properly composited onto it
+		if (!ensureMsgf(CompositeSize.X > 0 && CompositeSize.Y > 0, TEXT("Source of size (%d,%d) does not overlap destination of size (%d,%d) when offset to (%d,%d)!"),
+			SrcSize.X, SrcSize.Y, DestSize.X, DestSize.Y, Offset.X, Offset.Y))
 		{
 			return;
 		}
 
+		TRACE_CPUPROFILER_EVENT_SCOPE(MoviePipeline::TAsyncCompositeImageFLinearColor);
+		
 		TImagePixelData<FLinearColor>* DestColorData = static_cast<TImagePixelData<FLinearColor>*>(PixelData);
 		TImagePixelData<FColor>* SrcColorData = static_cast<TImagePixelData<FColor>*>(ImageToComposite.Get());
 
-		ParallelFor(DestColorData->GetSize().Y,
+		ParallelFor(CompositeSize.Y,
 			[&](int32 ScanlineIndex = 0)
 			{
-				for (int64 ColumnIndex = 0; ColumnIndex < DestColorData->GetSize().X; ColumnIndex++)
+				for (int64 ColumnIndex = 0; ColumnIndex < CompositeSize.X; ColumnIndex++)
 				{
-					int64 DstIndex = int64(ScanlineIndex) * int64(DestColorData->GetSize().X) + int64(ColumnIndex);
+					int64 SrcIndex = (int64(ScanlineIndex) + int64(SrcOffset.Y)) * int64(SrcColorData->GetSize().X) + int64(ColumnIndex) + int64(SrcOffset.X);
+					int64 DstIndex = (int64(ScanlineIndex) + int64(DestOffset.Y)) * int64(DestColorData->GetSize().X) + int64(ColumnIndex) + int64(DestOffset.X);
+					
 					FLinearColor& Dst = DestColorData->Pixels[DstIndex];
-					FColor& Src = SrcColorData->Pixels[DstIndex];
+					FColor& Src = SrcColorData->Pixels[SrcIndex];
 
 					float SourceAlpha = Src.A / 255.f;
 					FLinearColor Out;
@@ -1593,4 +1670,74 @@ struct TAsyncCompositeImage<FLinearColor>
 	}
 
 	TUniquePtr<FImagePixelData> ImageToComposite;
+	
+	/**
+	 * Amount in pixels from the destination's top left to offset the source when compositing onto the destination.
+	 * Can be negative, in which case a portion of the source's bottom right will be composited onto the destination's top left
+	 */
+	FIntPoint Offset;
 };
+
+
+struct FAsyncCropImage
+{
+	FAsyncCropImage(FImageWriteTask* InParentWriteTask, const FIntRect& InCropRect)
+		: ParentWriteTask(InParentWriteTask)
+		, CropRect(InCropRect)
+	{ }
+
+	FAsyncCropImage(const FIntRect& InCropRect)
+		: ParentWriteTask(nullptr)
+		, CropRect(InCropRect)
+	{ }
+	
+	void operator()(FImagePixelData* PixelData)
+	{
+		const bool bIsSrcBigEnough = PixelData->GetSize().X >= CropRect.Size().X && PixelData->GetSize().Y >= CropRect.Size().Y;
+		if (!ensureMsgf(bIsSrcBigEnough, TEXT("Cannot crop image of size (%d,%d) using crop rectangle of size (%d,%d)"),
+			PixelData->GetSize().X, PixelData->GetSize().Y, CropRect.Size().X, CropRect.Size().Y))
+		{
+			return;
+		}
+
+		TRACE_CPUPROFILER_EVENT_SCOPE(MoviePipeline::FAsyncCropImage)
+		
+		TArray64<FFloat16Color> DefaultPixels;
+		DefaultPixels.AddUninitialized(CropRect.Size().X * CropRect.Size().Y);
+
+		TSharedRef<FImagePixelDataPayload> NewPayload = PixelData->GetPayload<FImagePixelDataPayload>()->Copy();
+		
+		TUniquePtr<FImagePixelData> CroppedImage = MakeUnique<TImagePixelData<FFloat16Color>>(CropRect.Size(), MoveTemp(DefaultPixels), NewPayload);
+		TImagePixelData<FFloat16Color>* DestColorData = static_cast<TImagePixelData<FFloat16Color>*>(CroppedImage.Get());
+		TImagePixelData<FFloat16Color>* SrcColorData = static_cast<TImagePixelData<FFloat16Color>*>(PixelData);
+
+		ParallelFor(DestColorData->GetSize().Y,
+			[&](int32 ScanlineIndex = 0)
+			{
+				int64 DstIndex = int64(ScanlineIndex) * int64(DestColorData->GetSize().X);
+				int64 SrcIndex = (int64(ScanlineIndex) + int64(CropRect.Min.Y)) * int64(SrcColorData->GetSize().X) + int64(CropRect.Min.X);
+				
+				FMemory::Memcpy(&DestColorData->Pixels[DstIndex], &SrcColorData->Pixels[SrcIndex], DestColorData->GetSize().X * sizeof(FFloat16Color));
+			});
+
+		// If we have a parent write task that is holding the image being processed, write to that. Otherwise,
+		// write to a local image that can be used by the caller
+		if (ParentWriteTask)
+		{
+			ParentWriteTask->PixelData = MoveTemp(CroppedImage);
+		}
+		else
+		{
+			OutCroppedImage = MoveTemp(CroppedImage);
+		}
+	}
+
+	/** The parent image write task using this image preprocessor. If not null, the cropped image will be written to the write task's pixel data */
+	FImageWriteTask* ParentWriteTask;
+
+	/** If the parent write task is null, the cropped image will be written to this pixel data instead */
+	TUniquePtr<FImagePixelData> OutCroppedImage;
+
+	/** the rectangle to crop the source image to */
+	FIntRect CropRect;
+};
\ No newline at end of file
diff --git a/Engine/Plugins/MovieScene/MovieRenderPipeline/Source/MovieRenderPipelineRenderPasses/Private/Graph/Renderers/MovieGraphDeferredPass.cpp b/Engine/Plugins/MovieScene/MovieRenderPipeline/Source/MovieRenderPipelineRenderPasses/Private/Graph/Renderers/MovieGraphDeferredPass.cpp
--- a/Engine/Plugins/MovieScene/MovieRenderPipeline/Source/MovieRenderPipelineRenderPasses/Private/Graph/Renderers/MovieGraphDeferredPass.cpp
+++ b/Engine/Plugins/MovieScene/MovieRenderPipeline/Source/MovieRenderPipelineRenderPasses/Private/Graph/Renderers/MovieGraphDeferredPass.cpp
@@ -180,6 +180,8 @@ void FMovieGraphDeferredPass::Render(const FMovieGraphTraversalContext& InFrameT
 	FRenderTarget* RenderTargetResource = RenderTarget->GameThread_GetRenderTargetResource();
 	check(RenderTargetResource);
 
+	const FIntRect CropRect = UMovieGraphBlueprintLibrary::GetOverscanCropRectangle(InTimeData.EvaluatedConfig, CameraOverscan);
+	
 	// Now we can actually construct our ViewFamily, SceneView, and submit it for Rendering + Readback
 	for(int32 SpatialIndex = 0; SpatialIndex < NumSpatialSamples; SpatialIndex++)
 	{
@@ -299,6 +301,7 @@ void FMovieGraphDeferredPass::Render(const FMovieGraphTraversalContext& InFrameT
 			SampleState.OverlappedOffset = OverlappedOffset;
 			SampleState.OverlappedSubpixelShift = OverlappedSubpixelShift;
 			SampleState.OverscanFraction = CameraInfo.ViewInfo.GetOverscan();
+			SampleState.CropRectangle = CropRect;
 			SampleState.bAllowOCIO = ParentNodeThisFrame->GetAllowOCIO();
 			SampleState.bAllowsCompositing = ParentNodeThisFrame->GetAllowsCompositing();
 			SampleState.SceneCaptureSource = SceneCaptureSource;
diff --git a/Engine/Plugins/MovieScene/MovieRenderPipeline/Source/MovieRenderPipelineRenderPasses/Private/MovieGraphImageSequenceOutputNode.cpp b/Engine/Plugins/MovieScene/MovieRenderPipeline/Source/MovieRenderPipelineRenderPasses/Private/MovieGraphImageSequenceOutputNode.cpp
--- a/Engine/Plugins/MovieScene/MovieRenderPipeline/Source/MovieRenderPipelineRenderPasses/Private/MovieGraphImageSequenceOutputNode.cpp
+++ b/Engine/Plugins/MovieScene/MovieRenderPipeline/Source/MovieRenderPipelineRenderPasses/Private/MovieGraphImageSequenceOutputNode.cpp
@@ -158,6 +158,17 @@ void UMovieGraphImageSequenceOutputNode::OnReceiveImageDataImpl(UMovieGraphPipel
 		}
 
 		UE::MovieGraph::FMovieGraphSampleState* Payload = RenderData.Value->GetPayload<UE::MovieGraph::FMovieGraphSampleState>();
+
+		// If the overscan isn't cropped from the final image, offset any composites to the top-left of the original frustum
+		FIntPoint CompositeOffset = FIntPoint::ZeroValue;
+		if (ShouldCropOverscanImpl())
+		{
+			TileImageTask->PixelPreProcessors.Add(FAsyncCropImage(TileImageTask.Get(), Payload->CropRectangle));
+		}
+		else
+		{
+			CompositeOffset = Payload->CropRectangle.Min;
+		}
 		
 		bool bQuantizationEncodeSRGB = true;
 #if WITH_OCIO
@@ -195,13 +206,13 @@ void UMovieGraphImageSequenceOutputNode::OnReceiveImageDataImpl(UMovieGraphPipel
 			switch (PixelType)
 			{
 			case EImagePixelType::Color:
-				TileImageTask->PixelPreProcessors.Add(TAsyncCompositeImage<FColor>(CompositedPass.Value->CopyImageData()));
+				TileImageTask->PixelPreProcessors.Add(TAsyncCompositeImage<FColor>(CompositedPass.Value->CopyImageData(), CompositeOffset));
 				break;
 			case EImagePixelType::Float16:
-				TileImageTask->PixelPreProcessors.Add(TAsyncCompositeImage<FFloat16Color>(CompositedPass.Value->CopyImageData()));
+				TileImageTask->PixelPreProcessors.Add(TAsyncCompositeImage<FFloat16Color>(CompositedPass.Value->CopyImageData(), CompositeOffset));
 				break;
 			case EImagePixelType::Float32:
-				TileImageTask->PixelPreProcessors.Add(TAsyncCompositeImage<FLinearColor>(CompositedPass.Value->CopyImageData()));
+				TileImageTask->PixelPreProcessors.Add(TAsyncCompositeImage<FLinearColor>(CompositedPass.Value->CopyImageData(), CompositeOffset));
 				break;
 			}
 		}
@@ -385,7 +396,7 @@ void UMovieGraphImageSequenceOutputNode_EXR::OnReceiveImageDataImpl(UMovieGraphP
 			if (CompositedPass.Key.IsBranchAndCameraEqual(RenderData.Key))
 			{
 				// There could be multiple renders within this branch using the composited pass, so we have to copy the image data
-				ImageWriteTask->PixelPreprocessors.FindOrAdd(LayerIndex).Add(TAsyncCompositeImage<FFloat16Color>(CompositedPass.Value->CopyImageData()));
+				ImageWriteTask->PixelPreprocessors.FindOrAdd(LayerIndex).Add(TAsyncCompositeImage<FFloat16Color>(CompositedPass.Value->CopyImageData(), Payload->CropRectangle.Min));
 			}
 		}
 
diff --git a/Engine/Plugins/MovieScene/MovieRenderPipeline/Source/MovieRenderPipelineRenderPasses/Private/MoviePipelineImageSequenceOutput.cpp b/Engine/Plugins/MovieScene/MovieRenderPipeline/Source/MovieRenderPipelineRenderPasses/Private/MoviePipelineImageSequenceOutput.cpp
--- a/Engine/Plugins/MovieScene/MovieRenderPipeline/Source/MovieRenderPipelineRenderPasses/Private/MoviePipelineImageSequenceOutput.cpp
+++ b/Engine/Plugins/MovieScene/MovieRenderPipeline/Source/MovieRenderPipelineRenderPasses/Private/MoviePipelineImageSequenceOutput.cpp
@@ -205,6 +205,17 @@ void UMoviePipelineImageSequenceOutputBase::OnReceiveImageDataImpl(FMoviePipelin
 		TileImageTask->CompressionQuality = 100;
 		TileImageTask->Filename = OutputData.FilePath;
 
+		// If the overscan isn't cropped from the final image, offset any composites to the top-left of the original frustum
+		FIntPoint CompositeOffset = FIntPoint::ZeroValue;
+		if (ShouldCropOverscanImpl())
+		{
+			TileImageTask->PixelPreProcessors.Add(FAsyncCropImage(TileImageTask.Get(), Payload->SampleState.CropRectangle));
+		}
+		else
+		{
+			CompositeOffset = Payload->SampleState.CropRectangle.Min;
+		}
+		
 		TUniquePtr<FImagePixelData> QuantizedPixelData = RenderPassData.Value->CopyImageData();
 		EImagePixelType QuantizedPixelType = QuantizedPixelData->GetType();
 
@@ -242,16 +253,16 @@ void UMoviePipelineImageSequenceOutputBase::OnReceiveImageDataImpl(FMoviePipelin
 					continue;
 				}
 
-				// Check that the composite resolution matches the camera resolution to ensure the composite pass doesn't fail.
+				// Check that the composite resolution matches the original frustum resolution to ensure the composite pass doesn't fail.
 				// This can happen if multiple cameras with different amounts of overscan are rendered, since composite passes
 				// don't support rendering at multiple resolutions
 				const FIntPoint CompositeResolution = CompositePass.PixelData->GetSize();
-				const FIntPoint CameraResolution = RenderPassData.Value->GetSize();
-				if (CompositeResolution != CameraResolution)
+				const FIntPoint CameraOutputResolution = Payload->SampleState.CropRectangle.Size();
+				if (CompositeResolution != CameraOutputResolution)
 				{
-					UE_LOG(LogMovieRenderPipeline, Warning, TEXT("Composite resolution %dx%d does not match camera resolution %dx%d, skipping composite for %s on camera %s"),
+					UE_LOG(LogMovieRenderPipeline, Warning, TEXT("Composite resolution %dx%d does not match output resolution %dx%d, skipping composite for %s on camera %s"),
 						CompositeResolution.X, CompositeResolution.Y,
-						CameraResolution.X, CameraResolution.Y,
+						CameraOutputResolution.X, CameraOutputResolution.Y,
 						*CompositePass.PassIdentifier.Name,
 						*RenderPassData.Key.CameraName);
 
@@ -267,13 +278,13 @@ void UMoviePipelineImageSequenceOutputBase::OnReceiveImageDataImpl(FMoviePipelin
 				switch (QuantizedPixelType)
 				{
 					case EImagePixelType::Color:
-						TileImageTask->PixelPreProcessors.Add(TAsyncCompositeImage<FColor>(MoveTemp(PixelData)));
+						TileImageTask->PixelPreProcessors.Add(TAsyncCompositeImage<FColor>(MoveTemp(PixelData), CompositeOffset));
 						break;
 					case EImagePixelType::Float16:
-						TileImageTask->PixelPreProcessors.Add(TAsyncCompositeImage<FFloat16Color>(MoveTemp(PixelData)));
+						TileImageTask->PixelPreProcessors.Add(TAsyncCompositeImage<FFloat16Color>(MoveTemp(PixelData), CompositeOffset));
 						break;
 					case EImagePixelType::Float32:
-						TileImageTask->PixelPreProcessors.Add(TAsyncCompositeImage<FLinearColor>(MoveTemp(PixelData)));
+						TileImageTask->PixelPreProcessors.Add(TAsyncCompositeImage<FLinearColor>(MoveTemp(PixelData), CompositeOffset));
 						break;
 				}
 			}
@@ -299,7 +310,6 @@ void UMoviePipelineImageSequenceOutputBase::OnReceiveImageDataImpl(FMoviePipelin
 	}
 }
 
-
 void UMoviePipelineImageSequenceOutputBase::GetFormatArguments(FMoviePipelineFormatArgs& InOutFormatArgs) const
 {
 	// Stub in a dummy extension (so people know it exists)
diff --git a/Engine/Plugins/MovieScene/MovieRenderPipeline/Source/MovieRenderPipelineRenderPasses/Public/MovieGraphImageSequenceOutputNode.h b/Engine/Plugins/MovieScene/MovieRenderPipeline/Source/MovieRenderPipelineRenderPasses/Public/MovieGraphImageSequenceOutputNode.h
--- a/Engine/Plugins/MovieScene/MovieRenderPipeline/Source/MovieRenderPipelineRenderPasses/Public/MovieGraphImageSequenceOutputNode.h
+++ b/Engine/Plugins/MovieScene/MovieRenderPipeline/Source/MovieRenderPipelineRenderPasses/Public/MovieGraphImageSequenceOutputNode.h
@@ -99,7 +99,8 @@ public:
 	}
 
 	virtual void OnReceiveImageDataImpl(UMovieGraphPipeline* InPipeline, UE::MovieGraph::FMovieGraphOutputMergerFrame* InRawFrameData, const TSet<FMovieGraphRenderDataIdentifier>& InMask) override;
-
+	virtual bool ShouldCropOverscanImpl() const override { return false; }
+	
 #if WITH_EDITOR
 	virtual FText GetNodeTitle(const bool bGetDescriptive = false) const override 
 	{ 
@@ -186,7 +187,8 @@ public:
 	}
 
 	virtual void OnReceiveImageDataImpl(UMovieGraphPipeline* InPipeline, UE::MovieGraph::FMovieGraphOutputMergerFrame* InRawFrameData, const TSet<FMovieGraphRenderDataIdentifier>& InMask) override;
-
+	virtual bool ShouldCropOverscanImpl() const override { return false; }
+	
 #if WITH_EDITOR
 	virtual FText GetNodeTitle(const bool bGetDescriptive = false) const override
 	{
diff --git a/Engine/Plugins/MovieScene/MovieRenderPipeline/Source/MovieRenderPipelineRenderPasses/Public/MoviePipelineEXROutput.h b/Engine/Plugins/MovieScene/MovieRenderPipeline/Source/MovieRenderPipelineRenderPasses/Public/MoviePipelineEXROutput.h
--- a/Engine/Plugins/MovieScene/MovieRenderPipeline/Source/MovieRenderPipelineRenderPasses/Public/MoviePipelineEXROutput.h
+++ b/Engine/Plugins/MovieScene/MovieRenderPipeline/Source/MovieRenderPipelineRenderPasses/Public/MoviePipelineEXROutput.h
@@ -177,6 +177,7 @@ public:
 
 	virtual void OnReceiveImageDataImpl(FMoviePipelineMergerOutputFrame* InMergedOutputFrame) override;
 	virtual bool IsAlphaAllowed() const override { return true; }
+	virtual bool ShouldCropOverscanImpl() const override { return false; }
 
 public:
 	/**
diff --git a/Engine/Plugins/MovieScene/MovieRenderPipeline/Source/MovieRenderPipelineSettings/Private/MoviePipelineBurnInSetting.cpp b/Engine/Plugins/MovieScene/MovieRenderPipeline/Source/MovieRenderPipelineSettings/Private/MoviePipelineBurnInSetting.cpp
--- a/Engine/Plugins/MovieScene/MovieRenderPipeline/Source/MovieRenderPipelineSettings/Private/MoviePipelineBurnInSetting.cpp
+++ b/Engine/Plugins/MovieScene/MovieRenderPipeline/Source/MovieRenderPipelineSettings/Private/MoviePipelineBurnInSetting.cpp
@@ -125,8 +125,11 @@ void UMoviePipelineBurnInSetting::SetupImpl(const MoviePipeline::FMoviePipelineR
 
 	// TODO: Add multi-camera support? For now, simply use the default player controller camera overscan
 	const float CameraOverscan = GetPipeline()->GetCachedCameraOverscan(INDEX_NONE);
+	const FIntRect CropRect = UMoviePipelineBlueprintLibrary::GetOverscanCropRectangle(GetPipeline()->GetPipelinePrimaryConfig(), GetPipeline()->GetActiveShotList()[GetPipeline()->GetCurrentShotIndex()], CameraOverscan);
 	
-	OutputResolution = UMoviePipelineBlueprintLibrary::GetEffectiveOutputResolution(GetPipeline()->GetPipelinePrimaryConfig(), GetPipeline()->GetActiveShotList()[GetPipeline()->GetCurrentShotIndex()], CameraOverscan);
+	// Composited elements should be sized to the original frustum size, as the final image is either cropped to that size, or the composite will be offset
+	// to match the original frustum
+	OutputResolution = CropRect.Size();
 	int32 MaxResolution = GetMax2DTextureDimension();
 	if (OutputResolution.X > MaxResolution || OutputResolution.Y > MaxResolution)
 	{
diff --git a/Engine/Plugins/MovieScene/MovieRenderPipeline/Source/MovieRenderPipelineSettings/Private/MoviePipelineWidgetRenderSetting.cpp b/Engine/Plugins/MovieScene/MovieRenderPipeline/Source/MovieRenderPipelineSettings/Private/MoviePipelineWidgetRenderSetting.cpp
--- a/Engine/Plugins/MovieScene/MovieRenderPipeline/Source/MovieRenderPipelineSettings/Private/MoviePipelineWidgetRenderSetting.cpp
+++ b/Engine/Plugins/MovieScene/MovieRenderPipeline/Source/MovieRenderPipelineSettings/Private/MoviePipelineWidgetRenderSetting.cpp
@@ -116,8 +116,12 @@ void UMoviePipelineWidgetRenderer::SetupImpl(const MoviePipeline::FMoviePipeline
 
 	// TODO: Add multi-camera support? For now, simply use the default player controller camera overscan
 	const float CameraOverscan = GetPipeline()->GetCachedCameraOverscan(INDEX_NONE);
-	FIntPoint OutputResolution = UMoviePipelineBlueprintLibrary::GetEffectiveOutputResolution(GetPipeline()->GetPipelinePrimaryConfig(), GetPipeline()->GetActiveShotList()[GetPipeline()->GetCurrentShotIndex()], CameraOverscan);
+	const FIntRect CropRect = UMoviePipelineBlueprintLibrary::GetOverscanCropRectangle(GetPipeline()->GetPipelinePrimaryConfig(), GetPipeline()->GetActiveShotList()[GetPipeline()->GetCurrentShotIndex()], CameraOverscan);
 
+	// Composited elements should be sized to the original frustum size, as the final image is either cropped to that size, or the composite will be offset
+	// to match the original frustum
+	FIntPoint OutputResolution = CropRect.Size();
+	
 	int32 MaxResolution = GetMax2DTextureDimension();
 	if (OutputResolution.X > MaxResolution || OutputResolution.Y > MaxResolution)
 	{
