MoviePipeline: Added the capability to match components (in addition to actors) to collection/condition group evaluation.

The "Actor" condition group query was updated to optionally allow matching components within an actor.

diff --git a/Engine/Plugins/MovieScene/MovieRenderPipeline/Source/MovieRenderPipelineCore/Private/Graph/MovieGraphRenderLayerSubsystem.cpp b/Engine/Plugins/MovieScene/MovieRenderPipeline/Source/MovieRenderPipelineCore/Private/Graph/MovieGraphRenderLayerSubsystem.cpp
--- a/Engine/Plugins/MovieScene/MovieRenderPipeline/Source/MovieRenderPipelineCore/Private/Graph/MovieGraphRenderLayerSubsystem.cpp
+++ b/Engine/Plugins/MovieScene/MovieRenderPipeline/Source/MovieRenderPipelineCore/Private/Graph/MovieGraphRenderLayerSubsystem.cpp
@@ -32,12 +32,14 @@
 #include "ContentBrowserDataSource.h"
 #include "ContentBrowserModule.h"
 #include "DataLayer/DataLayerDragDropOp.h"
+#include "DetailLayoutBuilder.h"
 #include "DragAndDrop/ActorDragDropGraphEdOp.h"
 #include "DragAndDrop/AssetDragDropOp.h"
 #include "DragAndDrop/FolderDragDropOp.h"
 #include "DragAndDrop/LevelDragDropOp.h"
 #include "Editor.h"
 #include "EditorActorFolders.h"
+#include "Framework/MultiBox/MultiBoxBuilder.h"
 #include "Framework/Notifications/NotificationManager.h"
 #include "Graph/MovieGraphSharedWidgets.h"
 #include "IContentBrowserSingleton.h"
@@ -123,22 +125,17 @@ void UMovieGraphMaterialModifier::ApplyModifier(const UWorld* World)
 			continue;
 		}
 
-		for (const AActor* Actor : Collection->Evaluate(World))
-		{
-			const bool bIncludeFromChildActors = true;
-			TArray<UPrimitiveComponent*> PrimitiveComponents;
-			Actor->GetComponents<UPrimitiveComponent>(PrimitiveComponents, bIncludeFromChildActors);
+		FMovieGraphEvaluationResult EvaluationResult = Collection->EvaluateActorsAndComponents(World);
 
-			for (UPrimitiveComponent* PrimitiveComponent : PrimitiveComponents)
+		for (UPrimitiveComponent* PrimitiveComponent : EvaluationResult.GetAllComponentsOfType<UPrimitiveComponent>())
+		{
+			TArray<FMaterialSlotAssignment>& ModifiedMaterials = ModifiedComponents.FindOrAdd(PrimitiveComponent);
+			
+			for (int32 Index = 0; Index < PrimitiveComponent->GetNumMaterials(); ++Index)
 			{
-				TArray<FMaterialSlotAssignment>& ModifiedMaterials = ModifiedComponents.FindOrAdd(PrimitiveComponent);
-				
-				for (int32 Index = 0; Index < PrimitiveComponent->GetNumMaterials(); ++Index)
-				{
-					ModifiedMaterials.Add(FMaterialSlotAssignment(Index, PrimitiveComponent->GetMaterial(Index)));
-				
-					PrimitiveComponent->SetMaterial(Index, NewMaterial);
-				}
+				ModifiedMaterials.Add(FMaterialSlotAssignment(Index, PrimitiveComponent->GetMaterial(Index)));
+			
+				PrimitiveComponent->SetMaterial(Index, NewMaterial);
 			}
 		}
 	}
@@ -204,12 +201,58 @@ void UMovieGraphRenderPropertyModifier::PostEditChangeProperty(FPropertyChangedE
 
 void UMovieGraphRenderPropertyModifier::ApplyModifier(const UWorld* World)
 {
-	ModifiedActors.Empty();
+	auto CacheComponentState = [this](USceneComponent* InComponent)
+	{
+#if WITH_EDITORONLY_DATA
+		// Don't bother processing editor-only components (editor billboard icons, text, etc)
+		if (InComponent->IsEditorOnly())
+		{
+			return;
+		}
+#endif // WITH_EDITORONLY_DATA
+		
+		FComponentState& CachedState = CachedComponentState.AddDefaulted_GetRef();
+		CachedState.Component = InComponent;
 
-	// SetActorVisibilityState() takes a FActorVisibilityState, but all calls here will contain the same settings (with
-	// a differing actor), so create one copy to prevent constantly re-creating structs.
-	FActorVisibilityState NewVisibilityState;
-	NewVisibilityState.bIsHidden = bIsHidden;
+		// Cache the state
+		if (const UPrimitiveComponent* AsPrimitiveComponent = Cast<UPrimitiveComponent>(InComponent))
+		{
+			CachedState.bCastsShadows = AsPrimitiveComponent->CastShadow;
+			CachedState.bCastShadowWhileHidden = AsPrimitiveComponent->bCastHiddenShadow;
+			CachedState.bAffectIndirectLightingWhileHidden = AsPrimitiveComponent->bAffectIndirectLightingWhileHidden;
+			CachedState.bHoldout = AsPrimitiveComponent->bHoldout;
+			CachedState.bIsHidden = AsPrimitiveComponent->bHiddenInGame;
+		}
+		// Volumetrics are special cases as they don't inherit from UPrimitiveComponent, and don't support all of the flags.
+		else if (const UVolumetricCloudComponent* AsVolumetricCloudComponent = Cast<UVolumetricCloudComponent>(InComponent))
+		{
+			CachedState.bIsHidden = AsVolumetricCloudComponent->bHiddenInGame;
+			CachedState.bHoldout = AsVolumetricCloudComponent->bHoldout;
+			CachedState.bAffectIndirectLightingWhileHidden = !AsVolumetricCloudComponent->bRenderInMainPass;
+		}
+		else if (const USkyAtmosphereComponent* AsSkyAtmosphereComponent = Cast<USkyAtmosphereComponent>(InComponent))
+		{
+			CachedState.bIsHidden = AsSkyAtmosphereComponent->bHiddenInGame;
+			CachedState.bHoldout = AsSkyAtmosphereComponent->bHoldout;
+			CachedState.bAffectIndirectLightingWhileHidden = !AsSkyAtmosphereComponent->bRenderInMainPass;
+		}
+		else if (const UExponentialHeightFogComponent* AsExponentialHeightFogComponent = Cast<UExponentialHeightFogComponent>(InComponent))
+		{
+			CachedState.bIsHidden = AsExponentialHeightFogComponent->bHiddenInGame;
+			CachedState.bHoldout = AsExponentialHeightFogComponent->bHoldout;
+			CachedState.bAffectIndirectLightingWhileHidden = !AsExponentialHeightFogComponent->bRenderInMainPass;
+		}
+	};
+
+	auto CacheActorState = [this](AActor* InActor)
+	{
+		FActorState& CachedState = CachedActorState.AddDefaulted_GetRef();
+		CachedState.Actor = InActor;
+		CachedState.bIsHidden = InActor->IsHidden();
+	};
+	
+	CachedActorState.Empty();
+	CachedComponentState.Empty();
 
 	ValidateProjectSettings();
 	
@@ -220,112 +263,68 @@ void UMovieGraphRenderPropertyModifier::ApplyModifier(const UWorld* World)
 			continue;
 		}
 
-		const TSet<AActor*> MatchingActors = Collection->Evaluate(World);
-		ModifiedActors.Reserve(MatchingActors.Num());
+		FMovieGraphEvaluationResult EvaluationResult = Collection->EvaluateActorsAndComponents(World);
 
-		for (AActor* Actor : MatchingActors)
+		CachedActorState.Reserve(EvaluationResult.MatchingActors.Num());
+		CachedComponentState.Reserve(EvaluationResult.MatchingComponents.Num());	// Can't really accurately reserve the proper number of elements here
+
+		for (AActor* Actor : EvaluationResult.MatchingActors)
 		{
-			NewVisibilityState.Actor = Actor;
-			
-			// Save out visibility state before the modifier is applied
-			FActorVisibilityState OriginalVisibilityState;
-			OriginalVisibilityState.Actor = Actor;
-			OriginalVisibilityState.bIsHidden = Actor->IsHidden();
+			CacheActorState(Actor);
 
 			constexpr bool bIncludeFromChildActors = true;
-			TInlineComponentArray<USceneComponent*> Components;
-			Actor->GetComponents<USceneComponent>(Components, bIncludeFromChildActors);
+			TInlineComponentArray<USceneComponent*> Components(Actor, bIncludeFromChildActors);
 
-			OriginalVisibilityState.Components.Reserve(Components.Num());
-			NewVisibilityState.Components.Empty(Components.Num());
-			
 			for (USceneComponent* SceneComponent : Components)
 			{
-#if WITH_EDITORONLY_DATA
-				// Don't bother processing editor-only components (editor billboard icons, text, etc)
-				if (SceneComponent->IsEditorOnly())
-				{
-					continue;
-				}
-#endif // WITH_EDITORONLY_DATA
-				
-				FActorVisibilityState::FComponentState& CachedComponentState = OriginalVisibilityState.Components.AddDefaulted_GetRef();
-				CachedComponentState.Component = SceneComponent;
-
-				// Cache the state
-				if (const UPrimitiveComponent* AsPrimitiveComponent = Cast<UPrimitiveComponent>(SceneComponent))
-				{
-					CachedComponentState.bCastsShadows = AsPrimitiveComponent->CastShadow;
-					CachedComponentState.bCastShadowWhileHidden = AsPrimitiveComponent->bCastHiddenShadow;
-					CachedComponentState.bAffectIndirectLightingWhileHidden = AsPrimitiveComponent->bAffectIndirectLightingWhileHidden;
-					CachedComponentState.bHoldout = AsPrimitiveComponent->bHoldout;
-				}
-				// Volumetrics are special cases as they don't inherit from UPrimitiveComponent, and don't support all of the flags.
-				else if (const UVolumetricCloudComponent* AsVolumetricCloudComponent = Cast<UVolumetricCloudComponent>(SceneComponent))
-				{
-					CachedComponentState.bHoldout = AsVolumetricCloudComponent->bHoldout;
-					CachedComponentState.bAffectIndirectLightingWhileHidden = !AsVolumetricCloudComponent->bRenderInMainPass;
-				}
-				else if (const USkyAtmosphereComponent* AsSkyAtmosphereComponent = Cast<USkyAtmosphereComponent>(SceneComponent))
-				{
-					CachedComponentState.bHoldout = AsSkyAtmosphereComponent->bHoldout;
-					CachedComponentState.bAffectIndirectLightingWhileHidden = !AsSkyAtmosphereComponent->bRenderInMainPass;
-				}
-				else if (const UExponentialHeightFogComponent* AsExponentialHeightFogComponent = Cast<UExponentialHeightFogComponent>(SceneComponent))
-				{
-					CachedComponentState.bHoldout = AsExponentialHeightFogComponent->bHoldout;
-					CachedComponentState.bAffectIndirectLightingWhileHidden = !AsExponentialHeightFogComponent->bRenderInMainPass;
-				}
+				CacheComponentState(SceneComponent);
+			}
+		}
 
-				// SetActorVisibilityState() relies on inspecting individual components to determine what to do, hence why we need to generate a
-				// separate state for each component. Ideally this would not be needed in order to prevent constantly regenerating these structs.
-				FActorVisibilityState::FComponentState& NewComponentState = NewVisibilityState.Components.AddDefaulted_GetRef();
-				NewComponentState.Component = SceneComponent;
-				NewComponentState.bCastsShadows = bCastsShadows;
-				NewComponentState.bCastShadowWhileHidden = bCastShadowWhileHidden;
-				NewComponentState.bAffectIndirectLightingWhileHidden = bAffectIndirectLightingWhileHidden;
-				NewComponentState.bHoldout = bHoldout;
+		for (UActorComponent* Component : EvaluationResult.MatchingComponents)
+		{
+			if (USceneComponent* SceneComponent = Cast<USceneComponent>(Component))
+			{
+				CacheComponentState(SceneComponent);
 			}
-			
-			SetActorVisibilityState(NewVisibilityState);
-			
-			ModifiedActors.Add(OriginalVisibilityState);
 		}
+
+		// Actor/component state has been cached. Apply the node's state (settings) to all affected (ie, cached) actors/components.
+		constexpr bool bUseStateFromNode = true;
+		SetActorAndComponentState({}, {}, bUseStateFromNode);
 	}
 }
 
 void UMovieGraphRenderPropertyModifier::UndoModifier()
 {
-	for (const FActorVisibilityState& PrevVisibilityState : ModifiedActors)
-	{
-		SetActorVisibilityState(PrevVisibilityState);
-	}
+	constexpr bool bUseStateFromNode = false;
+	SetActorAndComponentState(CachedActorState, CachedComponentState, bUseStateFromNode);
 
-	ModifiedActors.Empty();
+	CachedActorState.Empty();
+	CachedComponentState.Empty();
 }
 
-void UMovieGraphRenderPropertyModifier::SetActorVisibilityState(const FActorVisibilityState& NewVisibilityState)
+void UMovieGraphRenderPropertyModifier::SetActorAndComponentState(const TArray<FActorState>& InActorState, const TArray<FComponentState>& InComponentState, const bool bUseStateFromNode)
 {
-	const TSoftObjectPtr<AActor> Actor = NewVisibilityState.Actor.LoadSynchronous();
-	if (!Actor)
-	{
-		return;
-	}
-
 	// In most cases, if the hidden state is being modified, the hidden state should be set. However, there is an exception for volumetrics.
 	// If volumetrics set the 'Affect Indirect Lighting While Hidden' flag to true, the volumetric component needs to set the 'Render in Main' flag
 	// instead, and the 'Hidden' flag should NOT be set on the *actor*. Setting the Hidden flag on the actor in this case will override the behavior
 	// of 'Render in Main' and volumetrics will not affect indirect lighting.
-	bool bShouldSetActorHiddenState = bOverride_bIsHidden;
+	TSet<AActor*> VolumetricComponentsWithAffectIndirect;
 
 	// Volumetrics are a special case and their visibility properties need to be handled separately
-	auto SetStateForVolumetrics = [this, &bShouldSetActorHiddenState]<typename VolumetricsType>(VolumetricsType& VolumetricsComponent, const FActorVisibilityState::FComponentState& NewComponentState)
+	auto SetStateForVolumetrics = [this, &VolumetricComponentsWithAffectIndirect]<typename VolumetricsType>(VolumetricsType& VolumetricsComponent, const FComponentState& NewComponentState)
 	{
 		if (bOverride_bHoldout)
 		{
 			VolumetricsComponent->SetHoldout(NewComponentState.bHoldout);
 		}
 
+		if (bOverride_bIsHidden)
+		{
+			VolumetricsComponent->SetRenderInMainPass(NewComponentState.bIsHidden);
+		}
+
 		if (bOverride_bAffectIndirectLightingWhileHidden)
 		{
 			// If the component should affect indirect while hidden, then we need to use 'Render in Main' instead.
@@ -335,64 +334,123 @@ void UMovieGraphRenderPropertyModifier::SetActorVisibilityState(const FActorVisi
 			// negate the effects of setting Render In Main Pass.
 			if (NewComponentState.bAffectIndirectLightingWhileHidden)
 			{
-				bShouldSetActorHiddenState = false;
+				VolumetricComponentsWithAffectIndirect.Add(VolumetricsComponent->GetOwner());
 			}
 		}
 	};
 
-	for (const FActorVisibilityState::FComponentState& ComponentState : NewVisibilityState.Components)
+	// Sets the state for a component to the specified component state
+	auto SetStateForComponent = [this, &SetStateForVolumetrics](UActorComponent* InComponent, const FComponentState& InNewComponentState)
 	{
 		// TODO: These could potentially cause a large rendering penalty due to dirtying the render state; investigate potential
 		// ways to optimize this
-		if (!ComponentState.Component)
-		{
-			continue;
-		}
-
-		if (UPrimitiveComponent* AsPrimitiveComponent = Cast<UPrimitiveComponent>(ComponentState.Component.Get()))
+		if (UPrimitiveComponent* AsPrimitiveComponent = Cast<UPrimitiveComponent>(InComponent))
 		{
 			if (bOverride_bCastsShadows)
 			{
-				AsPrimitiveComponent->SetCastShadow(ComponentState.bCastsShadows);
+				AsPrimitiveComponent->SetCastShadow(InNewComponentState.bCastsShadows);
 			}
 
 			if (bOverride_bCastShadowWhileHidden)
 			{
-				AsPrimitiveComponent->SetCastHiddenShadow(ComponentState.bCastShadowWhileHidden);
+				AsPrimitiveComponent->SetCastHiddenShadow(InNewComponentState.bCastShadowWhileHidden);
 			}
 
 			if (bOverride_bAffectIndirectLightingWhileHidden)
 			{
-				AsPrimitiveComponent->SetAffectIndirectLightingWhileHidden(ComponentState.bAffectIndirectLightingWhileHidden);
+				AsPrimitiveComponent->SetAffectIndirectLightingWhileHidden(InNewComponentState.bAffectIndirectLightingWhileHidden);
 			}
 
 			if (bOverride_bHoldout)
 			{
-				AsPrimitiveComponent->SetHoldout(ComponentState.bHoldout);
+				AsPrimitiveComponent->SetHoldout(InNewComponentState.bHoldout);
+			}
+
+			if (bOverride_bIsHidden)
+			{
+				AsPrimitiveComponent->SetHiddenInGame(InNewComponentState.bIsHidden);
+#if WITH_EDITOR
+				AsPrimitiveComponent->SetIsTemporarilyHiddenInEditor(InNewComponentState.bIsHidden);
+#endif
 			}
 		}
 		// Volumetrics are special cases as they don't inherit from UPrimitiveComponent, and don't support all of the flags.
-		else if (UVolumetricCloudComponent* AsVolumetricCloudComponent = Cast<UVolumetricCloudComponent>(ComponentState.Component.Get()))
+		else if (UVolumetricCloudComponent* AsVolumetricCloudComponent = Cast<UVolumetricCloudComponent>(InComponent))
 		{
-			SetStateForVolumetrics(AsVolumetricCloudComponent, ComponentState);
+			SetStateForVolumetrics(AsVolumetricCloudComponent, InNewComponentState);
 		}
-		else if (USkyAtmosphereComponent* AsSkyAtmosphereComponent = Cast<USkyAtmosphereComponent>(ComponentState.Component.Get()))
+		else if (USkyAtmosphereComponent* AsSkyAtmosphereComponent = Cast<USkyAtmosphereComponent>(InComponent))
 		{
-			SetStateForVolumetrics(AsSkyAtmosphereComponent, ComponentState);
+			SetStateForVolumetrics(AsSkyAtmosphereComponent, InNewComponentState);
 		}
-		else if (UExponentialHeightFogComponent* AsExponentialHeightFogComponent = Cast<UExponentialHeightFogComponent>(ComponentState.Component.Get()))
+		else if (UExponentialHeightFogComponent* AsExponentialHeightFogComponent = Cast<UExponentialHeightFogComponent>(InComponent))
 		{
-			SetStateForVolumetrics(AsExponentialHeightFogComponent, ComponentState);
+			SetStateForVolumetrics(AsExponentialHeightFogComponent, InNewComponentState);
 		}
-	}
-	
-	if (bShouldSetActorHiddenState)
+	};
+
+	// Sets the state for an actor to the specified actor state
+	auto SetStateForActor = [this, &VolumetricComponentsWithAffectIndirect](AActor* InActor, const FActorState& InNewActorState)
 	{
-		Actor->SetActorHiddenInGame(NewVisibilityState.bIsHidden);
+		if (VolumetricComponentsWithAffectIndirect.Contains(InActor))
+		{
+			return;
+		}
 
+		InActor->SetActorHiddenInGame(InNewActorState.bIsHidden);
 #if WITH_EDITOR
-		Actor->SetIsTemporarilyHiddenInEditor(NewVisibilityState.bIsHidden);
+		InActor->SetIsTemporarilyHiddenInEditor(InNewActorState.bIsHidden);
 #endif
+	};
+
+	// If the node's state/settings should be applied, generate a singular FComponentState and apply it to all cached components.
+	// Otherwise, apply all of the component states that were provided.
+	if (bUseStateFromNode)
+	{
+		FComponentState NewComponentState;
+		NewComponentState.bIsHidden = bIsHidden;
+		NewComponentState.bHoldout = bHoldout;
+		NewComponentState.bCastsShadows = bCastsShadows;
+		NewComponentState.bCastShadowWhileHidden = bCastShadowWhileHidden;
+		NewComponentState.bAffectIndirectLightingWhileHidden = bAffectIndirectLightingWhileHidden;
+
+		for (const FComponentState& CachedState : CachedComponentState)
+		{
+			SetStateForComponent(CachedState.Component.Get(), NewComponentState);
+		}
+	}
+	else
+	{
+		for (const FComponentState& ComponentState : InComponentState)
+		{
+			if (USceneComponent* Component = ComponentState.Component.Get())
+			{
+				SetStateForComponent(Component, ComponentState);
+			}
+		}
+	}
+
+	// Do a similar operation for actors; if the node's state should be applied, generate a singular FActorState and apply it to all cached actors.
+	// Otherwise, apply all of the actor states that were provided.
+	if (bUseStateFromNode)
+	{
+		FActorState NewActorState;
+		NewActorState.bIsHidden = bIsHidden;
+
+		for (const FActorState& CachedState : CachedActorState)
+		{
+			SetStateForActor(CachedState.Actor.LoadSynchronous(), NewActorState);
+		}
+	}
+	else
+	{
+		for (const FActorState& ActorState : InActorState)
+		{
+			if (AActor* Actor = ActorState.Actor.LoadSynchronous())
+			{
+				SetStateForActor(Actor, ActorState);
+			}
+		}
 	}
 }
 
@@ -465,12 +523,23 @@ void UMovieGraphConditionGroupQueryBase::Evaluate(const TArray<AActor*>& InActor
 	// No implementation
 }
 
+void UMovieGraphConditionGroupQueryBase::EvaluateActorsAndComponents(const TArray<AActor*>& InActorsToQuery, const UWorld* InWorld, TSet<AActor*>& OutMatchingActors, TSet<UActorComponent*>& OutMatchingComponents) const
+{
+	// No implementation
+}
+
 bool UMovieGraphConditionGroupQueryBase::ShouldHidePropertyNames() const
 {
 	// Show property names by default; subclassed queries can opt-out if they want a cleaner UI 
 	return false;
 }
 
+bool UMovieGraphConditionGroupQueryBase::ShouldEvaluateComponents() const
+{
+	// By default, only actors are evaluated
+	return false;
+}
+
 const FSlateIcon& UMovieGraphConditionGroupQueryBase::GetIcon() const
 {
 	static const FSlateIcon EmptyIcon = FSlateIcon();
@@ -587,30 +656,55 @@ AActor* UMovieGraphConditionGroupQueryBase::GetActorForCurrentWorld(AActor* InAc
 	return nullptr;
 }
 
+bool FMovieGraphActorQueryEntry::operator==(const FMovieGraphActorQueryEntry& InOther) const
+{
+	return (InOther.ActorToMatch == ActorToMatch) && (InOther.ComponentsToMatch == ComponentsToMatch);
+}
+
 void UMovieGraphConditionGroupQuery_Actor::Evaluate(const TArray<AActor*>& InActorsToQuery, const UWorld* InWorld, TSet<AActor*>& OutMatchingActors) const
 {
-	TRACE_CPUPROFILER_EVENT_SCOPE(UMovieGraphConditionGroupQuery_Actor::Evaluate);
+	TSet<UActorComponent*> UnusedComponentsArray;
+	EvaluateActorsAndComponents(InActorsToQuery, InWorld, OutMatchingActors, UnusedComponentsArray);
+}
 
-	// Convert the actors in the query to PIE (or editor) equivalents once, rather than constantly in the loop.
-	TArray<AActor*> ActorsToMatch_Converted;
-	ActorsToMatch_Converted.Reserve(ActorsToMatch.Num());
-	for (const TSoftObjectPtr<AActor>& SoftActorToMatch : ActorsToMatch)
-	{
-		if (AActor* ConvertedActor = GetActorForCurrentWorld(SoftActorToMatch.Get()))
-		{
-			OutMatchingActors.Add(ConvertedActor);
-		}
-	}
-	
-	for (AActor* Actor : InActorsToQuery)
+void UMovieGraphConditionGroupQuery_Actor::EvaluateActorsAndComponents(const TArray<AActor*>& InActorsToQuery, const UWorld* InWorld, TSet<AActor*>& OutMatchingActors, TSet<UActorComponent*>& OutMatchingComponents) const
+{
+	TRACE_CPUPROFILER_EVENT_SCOPE(UMovieGraphConditionGroupQuery_Actor::EvaluateActorsAndComponents);
+
+	// Note: This query does not need to search through InActorsToQuery because it has direct references to actors/components
+	for (const FMovieGraphActorQueryEntry& Entry : ActorsAndComponentsToMatch)
 	{
-		if (ActorsToMatch_Converted.Contains(Actor))
+		if (AActor* ConvertedActor = GetActorForCurrentWorld(Entry.ActorToMatch.Get()))
 		{
-			OutMatchingActors.Add(Actor);
+			if (!Entry.ComponentsToMatch.IsEmpty())
+			{
+				// Instead of matching the actor, get the components requested from the actor in the correct world.
+				// Since the number of components here is typically going to be very small, an array should be more efficient vs. a set when used in the matching loop below.
+				TArray<FString> ComponentNamesToMatch;
+				ComponentNamesToMatch.Reserve(Entry.ComponentsToMatch.Num());
+				Algo::Transform(Entry.ComponentsToMatch, ComponentNamesToMatch, [](const TSoftObjectPtr<UActorComponent>& InComponent) { return InComponent.IsValid() ? InComponent->GetName() : FString(); });
+
+				for (UActorComponent* ActorComponent : ConvertedActor->GetComponents())
+				{
+					if (ActorComponent && ComponentNamesToMatch.Contains(ActorComponent->GetName()))
+					{
+						OutMatchingComponents.Add(ActorComponent);
+					}
+				}
+			}
+			else
+			{
+				OutMatchingActors.Add(ConvertedActor);
+			}
 		}
 	}
 }
 
+bool UMovieGraphConditionGroupQuery_Actor::ShouldEvaluateComponents() const
+{
+	return true;
+}
+
 const FSlateIcon& UMovieGraphConditionGroupQuery_Actor::GetIcon() const
 {
 	static const FSlateIcon ActorIcon = FSlateIcon(FAppStyle::GetAppStyleSetName(), "ClassIcon.Actor");
@@ -658,6 +752,29 @@ TArray<TSharedRef<SWidget>> UMovieGraphConditionGroupQuery_Actor::GetWidgets()
 		}
 	};
 
+	// Generate the header row for the list widget
+	const TSharedPtr<SHeaderRow> HeaderRow = SNew(SHeaderRow)
+		.Visibility(EVisibility::Visible);
+	auto AddColumnToHeader = [HeaderRow](const FName& InColumnId, const FText& InColumnText)
+	{
+		HeaderRow->AddColumn(
+			SHeaderRow::Column(InColumnId)
+			.HAlignCell(HAlign_Left)
+			.HAlignHeader(HAlign_Left)
+			.VAlignCell(VAlign_Fill)
+			.HeaderContentPadding(FMargin(4.0f))
+			.HeaderContent()
+			[
+				SNew(STextBlock)
+				.Font(IDetailLayoutBuilder::GetDetailFont())
+				.Text(InColumnText)
+			]
+		);
+	};
+	AddColumnToHeader(ColumnID_ActorName, LOCTEXT("ActorConditionGroupQuery_ActorColumnName", "Actor Name"));
+	AddColumnToHeader(ColumnID_ActorType, LOCTEXT("ActorConditionGroupQuery_ActorTypeColumnName", "Actor Type"));
+	AddColumnToHeader(ColumnID_Components, LOCTEXT("ActorConditionGroupQuery_ComponentsColumnName", "Components"));
+
 	Widgets.Add(
 		SNew(SDropTarget)
 		.OnAllowDrop_Lambda([GetValidActorsFromOperation](TSharedPtr<FDragDropOperation> InDragOperation)
@@ -695,25 +812,46 @@ TArray<TSharedRef<SWidget>> UMovieGraphConditionGroupQuery_Actor::GetWidgets()
 			return FReply::Handled();
 		})
 		[
-			SAssignNew(ActorsList, SMovieGraphSimpleList<TSharedPtr<TSoftObjectPtr<AActor>>>)
+			SAssignNew(ActorsList, SMovieGraphSimpleList<TSharedPtr<FMovieGraphActorQueryEntry>>)
 			.DataSource(&ListDataSource)
-			.DataType(FText::FromString("Actor"))
-			.DataTypePlural(FText::FromString("Actors"))
-			.OnGetRowText_Static(&GetRowText)
-			.OnGetRowIcon_Static(&GetRowIcon)
-			.OnDelete_Lambda([this](const TArray<TSharedPtr<TSoftObjectPtr<AActor>>> InActors)
+			.HeaderRow(HeaderRow)
+			.OnGetCustomRow_Lambda([this](const TSharedRef<STableViewBase>& InOwnerTable, TSharedPtr<FMovieGraphActorQueryEntry> InEntry)
+			{
+				return SNew(SActorListRow, InOwnerTable, InEntry);
+			})
+			.CustomSummaryText_Lambda([this]()
+			{
+				int32 NumActors = 0;
+				int32 NumComponents = 0;
+				
+				for (const FMovieGraphActorQueryEntry& Entry : ActorsAndComponentsToMatch)
+				{
+					if (!Entry.ComponentsToMatch.IsEmpty())
+					{
+						NumComponents += Entry.ComponentsToMatch.Num();
+					}
+					else
+					{
+						NumActors++;
+					}
+				}
+				
+				return FText::Format(LOCTEXT("ActorQuery_SummaryRowText", "{0} Actor(s) and {1} Component(s)"), NumActors, NumComponents);
+			})
+			.OnDelete_Lambda([this](const TArray<TSharedPtr<FMovieGraphActorQueryEntry>> InEntries)
 			{
-				TArray<TSoftObjectPtr<AActor>> ActorsToRemove;
-				for (const TSharedPtr<TSoftObjectPtr<AActor>>& InActor : InActors)
+				TArray<FMovieGraphActorQueryEntry> ActorsToRemove;
+				for (const TSharedPtr<FMovieGraphActorQueryEntry>& Entry : InEntries)
 				{
-					if (InActor.IsValid())
+					if (Entry.IsValid())
 					{
-						ActorsToRemove.Add(*InActor.Get());
+						ActorsToRemove.Add(*Entry.Get());
 					}
 				}
 				
-				RemoveActors(ActorsToRemove);
+				RemoveEntries(ActorsToRemove);
 			})
+			.OnGetContextMenuContent_UObject(this, &UMovieGraphConditionGroupQuery_Actor::GetListContextMenu)
 			.OnRefreshDataSourceRequested_UObject(this, &UMovieGraphConditionGroupQuery_Actor::RefreshListDataSource)
 		]
 	);
@@ -843,32 +981,27 @@ TSharedRef<SWidget> UMovieGraphConditionGroupQuery_Actor::GetAddMenuContents(con
 	return MenuBuilder.MakeWidget();
 }
 
-const FSlateBrush* UMovieGraphConditionGroupQuery_Actor::GetRowIcon(TSharedPtr<TSoftObjectPtr<AActor>> InActor)
+void UMovieGraphConditionGroupQuery_Actor::PostLoad()
 {
-	if (InActor.IsValid())
-	{
-		if (InActor.Get()->IsValid())
-		{
-			// The first Get() returns the TSoftObjectPtr, the second Get() dereferences the TSoftObjectPtr
-			return FSlateIconFinder::FindIconForClass(InActor.Get()->Get()->GetClass()).GetIcon();
-		}
-	}
-
-	return FSlateIconFinder::FindIconForClass(AActor::StaticClass()).GetIcon();
-}
+	Super::PostLoad();
 
-FText UMovieGraphConditionGroupQuery_Actor::GetRowText(TSharedPtr<TSoftObjectPtr<AActor>> InActor)
-{
-	if (InActor.IsValid())
+#if WITH_EDITOR
+	// Upgrade the old ActorsToMatch (Actors-only) to ActorsAndComponentsToMatch (Actors + Components)
+	PRAGMA_DISABLE_DEPRECATION_WARNINGS
+	if (!ActorsToMatch.IsEmpty())
 	{
-		if (InActor.Get()->IsValid())
+		for (const TSoftObjectPtr<AActor>& Actor : ActorsToMatch)
 		{
-			// The first Get() returns the TSoftObjectPtr, the second Get() dereferences the TSoftObjectPtr
-			return FText::FromString(InActor.Get()->Get()->GetActorLabel());
+			FMovieGraphActorQueryEntry NewEntry;
+			NewEntry.ActorToMatch = Actor;
+			
+			ActorsAndComponentsToMatch.Add(MoveTemp(NewEntry));
 		}
-	}
 
-	return LOCTEXT("MovieGraphActorConditionGroupQuery_InvalidActor", "(invalid)");
+		ActorsToMatch.Empty();
+	}
+	PRAGMA_ENABLE_DEPRECATION_WARNINGS
+#endif	// WITH_EDITOR
 }
 
 void UMovieGraphConditionGroupQuery_Actor::AddActors(const TArray<AActor*>& InActors, const FMovieGraphConditionGroupQueryContentsChanged& InOnAddFinished, const bool bCloseAddMenu)
@@ -878,10 +1011,12 @@ void UMovieGraphConditionGroupQuery_Actor::AddActors(const TArray<AActor*>& InAc
 	
 	for (AActor* Actor : InActors)
 	{
-		if (Actor && !ActorsToMatch.Contains(Actor))
+		if (Actor && !ActorsAndComponentsToMatch.ContainsByPredicate([Actor](const FMovieGraphActorQueryEntry& InEntry) { return InEntry.ActorToMatch == Actor; }))
 		{
-			ActorsToMatch.Add(Actor);
-			ListDataSource.Add(MakeShared<TSoftObjectPtr<AActor>>(ActorsToMatch.Last()));
+			FMovieGraphActorQueryEntry NewEntry;
+			NewEntry.ActorToMatch = Actor;
+			ActorsAndComponentsToMatch.Add(MoveTemp(NewEntry));
+			ListDataSource.Add(MakeShared<FMovieGraphActorQueryEntry>(ActorsAndComponentsToMatch.Last()));
 			
 			InOnAddFinished.ExecuteIfBound();
 		}
@@ -901,34 +1036,144 @@ void UMovieGraphConditionGroupQuery_Actor::AddActors(const TArray<AActor*>& InAc
 	}
 }
 
-void UMovieGraphConditionGroupQuery_Actor::RemoveActors(const TArray<TSoftObjectPtr<AActor>>& InActors)
+void UMovieGraphConditionGroupQuery_Actor::RemoveEntries(const TArray<FMovieGraphActorQueryEntry>& InEntries)
 {
 	const FScopedTransaction Transaction(LOCTEXT("RemoveActorsFromCollection", "Remove Actors from Collection"));
 	Modify();
 	
-	for (const TSoftObjectPtr<AActor>& Actor : InActors)
-	{		
-		ListDataSource.RemoveAll([&Actor](const TSharedPtr<TSoftObjectPtr<AActor>>& ListActor)
+	for (const FMovieGraphActorQueryEntry& Entry : InEntries)
+	{
+		ListDataSource.RemoveAll([&Entry](const TSharedPtr<FMovieGraphActorQueryEntry>& ListEntry)
 		{
-			if (const TSoftObjectPtr<AActor>* SoftListActor = ListActor.Get())
+			if (const FMovieGraphActorQueryEntry* ListEntryPtr = ListEntry.Get())
 			{
-				return *SoftListActor == Actor;
+				return *ListEntryPtr == Entry;
 			}
 			return false;
 		});
 		
-		ActorsToMatch.Remove(Actor);
+		ActorsAndComponentsToMatch.Remove(Entry);
 	}
 	
 	ActorsList->Refresh();
 }
 
+void UMovieGraphConditionGroupQuery_Actor::GetListContextMenu(FMenuBuilder& InMenuBuilder, TArray<TSharedPtr<FMovieGraphActorQueryEntry>> SelectedEntries)
+{
+	if (SelectedEntries.IsEmpty() || (SelectedEntries.Num() > 1))
+	{
+		return;
+	}
+	
+	const TSharedPtr<FMovieGraphActorQueryEntry> SelectedEntry = SelectedEntries[0];
+	if (!SelectedEntry.IsValid())
+	{
+		return;
+	}
+
+
+	// Get the associated actor for this entry
+	const AActor* AssociatedActor = SelectedEntry.Get()->ActorToMatch.Get();
+	if (!AssociatedActor)
+	{
+		return;
+	}
+
+	// Find the entry in the query associated with this actor
+	FMovieGraphActorQueryEntry* Entry = ActorsAndComponentsToMatch.FindByPredicate([AssociatedActor](const FMovieGraphActorQueryEntry& InEntry) { return InEntry.ActorToMatch == AssociatedActor; });
+	if (!Entry)
+	{
+		return;
+	}
+
+	InMenuBuilder.BeginSection(NAME_None, LOCTEXT("ActorQuery_MatchTypeSection", "Only Match"));
+
+	TInlineComponentArray<UActorComponent*> ActorComponents;
+	AssociatedActor->GetComponents(ActorComponents);
+
+	if (!ActorComponents.IsEmpty())
+	{
+		const FUIAction ToggleSelectedActor(
+			FExecuteAction::CreateLambda([this, Entry]()
+			{
+				if (!Entry->ComponentsToMatch.IsEmpty())
+				{
+					const FScopedTransaction Transaction(LOCTEXT("ToggleSelectedActorInCollection", "Toggle Selected Actor in Collection"));
+					Modify();
+					
+					Entry->ComponentsToMatch.Empty();
+					ActorsList->Refresh();
+				}
+			}),
+			FCanExecuteAction::CreateLambda([Entry]() { return !Entry->ComponentsToMatch.IsEmpty(); }),
+			FIsActionChecked::CreateLambda([Entry]()
+			{
+				return Entry->ComponentsToMatch.IsEmpty();
+			})
+		);
+
+		InMenuBuilder.AddMenuEntry(
+			LOCTEXT("ActorQuery_ToggleSelectedActor", "Selected Actor"),
+			LOCTEXT("ActorQuery_ToggleSelectedActorTooltip", "Match the selected actor, or some component(s) contained in it."),
+			FSlateIcon(FAppStyle::GetAppStyleSetName(), "ClassIcon.Actor"),
+			ToggleSelectedActor,
+			NAME_None,
+			EUserInterfaceActionType::ToggleButton
+		);
+
+		InMenuBuilder.AddSeparator();
+	}
+	
+	for (UActorComponent* ActorComponent : ActorComponents)
+	{
+		// Exclude editor-only components (editor billboard icons, text, etc)
+		if (ActorComponent->IsEditorOnly())
+		{
+			continue;
+		}
+		
+		FUIAction AddRemoveComponentAction(
+			FExecuteAction::CreateLambda([this, Entry, ActorComponent]()
+			{
+				const FScopedTransaction Transaction(LOCTEXT("ToggleComponentInCollection", "Toggle Component in Collection"));
+				Modify();
+				
+				if (Entry->ComponentsToMatch.Contains(ActorComponent))
+				{
+					Entry->ComponentsToMatch.Remove(ActorComponent);
+				}
+				else
+				{
+					Entry->ComponentsToMatch.Add(ActorComponent);
+				}
+
+				ActorsList->Refresh();
+			}),
+			FCanExecuteAction::CreateLambda([]() { return true; }),
+			FIsActionChecked::CreateLambda([Entry, ActorComponent]() {
+				return Entry->ComponentsToMatch.Contains(ActorComponent);
+			})
+		);
+		
+		InMenuBuilder.AddMenuEntry(
+			FText::FromString(ActorComponent->GetName()),
+			LOCTEXT("ActorQuery_AddComponentTooltip", "Restrict to only matching the selected component types."),
+			FSlateIconFinder::FindIconForClass(ActorComponent->GetClass()),
+			AddRemoveComponentAction,
+			NAME_None,
+			EUserInterfaceActionType::ToggleButton
+		);
+	}
+	
+	InMenuBuilder.EndSection();
+}
+
 void UMovieGraphConditionGroupQuery_Actor::RefreshListDataSource()
 {
 	ListDataSource.Empty();
-	for (TSoftObjectPtr<AActor>& Actor : ActorsToMatch)
+	for (FMovieGraphActorQueryEntry& Entry : ActorsAndComponentsToMatch)
 	{
-		ListDataSource.Add(MakeShared<TSoftObjectPtr<AActor>>(Actor));
+		ListDataSource.Add(MakeShared<FMovieGraphActorQueryEntry>(Entry));
 	}
 }
 
@@ -980,9 +1225,9 @@ ECheckBoxState UMovieGraphConditionGroupQuery_Actor::FActorSelectionColumn::IsRo
 	}
 
 	const TStrongObjectPtr<UMovieGraphConditionGroupQuery_Actor> ActorQuery = WeakActorQuery.Pin();
-	for (TSoftObjectPtr<AActor>& ActorToMatch : ActorQuery->ActorsToMatch)
+	for (FMovieGraphActorQueryEntry& Entry : ActorQuery->ActorsAndComponentsToMatch)
 	{
-		if (ActorToMatch.Get() == InActorTreeItem->Actor.Get())
+		if (Entry.ActorToMatch.Get() == InActorTreeItem->Actor.Get())
 		{
 			return ECheckBoxState::Checked;
 		}
@@ -1003,7 +1248,10 @@ void UMovieGraphConditionGroupQuery_Actor::FActorSelectionColumn::OnCheckStateCh
 	{
 		if (NewState == ECheckBoxState::Unchecked)
 		{
-			ActorQuery->RemoveActors({InActorTreeItem->Actor.Get()});
+			FMovieGraphActorQueryEntry TempEntry;
+			TempEntry.ActorToMatch = InActorTreeItem->Actor.Get();
+			
+			ActorQuery->RemoveEntries({TempEntry});
 		}
 		else
 		{
@@ -1013,6 +1261,122 @@ void UMovieGraphConditionGroupQuery_Actor::FActorSelectionColumn::OnCheckStateCh
 		}
 	}
 }
+
+void UMovieGraphConditionGroupQuery_Actor::SActorListRow::Construct(const FArguments& InArgs, const TSharedRef<STableViewBase>& InOwnerTableView, TSharedPtr<FMovieGraphActorQueryEntry> InEntry)
+{
+	WeakEntry = InEntry;
+
+	SMultiColumnTableRow<TSharedPtr<FMovieGraphActorQueryEntry>>::Construct(FSuperRowType::FArguments(), InOwnerTableView);
+}
+
+TSharedRef<SWidget> UMovieGraphConditionGroupQuery_Actor::SActorListRow::GenerateWidgetForColumn(const FName& InColumnName)
+{
+	const TSharedPtr<FMovieGraphActorQueryEntry> EntryPin = WeakEntry.Pin();
+	if (!EntryPin.IsValid())
+	{
+		return SNullWidget::NullWidget;
+	}
+
+	const TWeakObjectPtr<const AActor> WeakActor = MakeWeakObjectPtr(EntryPin->ActorToMatch.Get());
+	
+	if (InColumnName == ColumnID_ActorName)
+	{
+		const FSlateBrush* ActorIcon = WeakActor.IsValid()
+			? FSlateIconFinder::FindIconForClass(WeakActor.Get()->GetClass()).GetIcon()
+			: FSlateIconFinder::FindIconForClass(AActor::StaticClass()).GetIcon();
+		
+		return SNew(SHorizontalBox)
+			+ SHorizontalBox::Slot()
+			.VAlign(VAlign_Center)
+			.Padding(4.0f, 5.0f, 2.0f, 5.0f)
+			.AutoWidth()
+			[
+				SNew(SImage)
+				.Image(ActorIcon)
+			]
+
+			+ SHorizontalBox::Slot()
+			.AutoWidth()
+			.HAlign(HAlign_Left)
+			.VAlign(VAlign_Center)
+			.Padding(4.0f, 5.0f)
+			[
+				SNew(STextBlock)
+				.Text_Lambda([WeakActor]()
+				{
+					return WeakActor.IsValid()
+						? FText::FromString(WeakActor->GetActorNameOrLabel())
+						: LOCTEXT("MovieGraphActorConditionGroupQuery_InvalidActor", "(invalid)");
+				})
+				.Font(IDetailLayoutBuilder::GetDetailFont())
+			];
+	}
+
+	if (InColumnName == ColumnID_ActorType)
+	{
+		return SNew(SHorizontalBox)
+			+ SHorizontalBox::Slot()
+			.AutoWidth()
+			.HAlign(HAlign_Left)
+			.VAlign(VAlign_Center)
+			.Padding(4.0f, 5.0f)
+			[
+				SNew(STextBlock)
+				.Text_Lambda([WeakActor]()
+				{
+					return WeakActor.IsValid() ? FText::FromString(WeakActor->GetClass()->GetName()) : FText::GetEmpty();
+				})
+				.Font(IDetailLayoutBuilder::GetDetailFont())
+			];
+	}
+
+	if (InColumnName == ColumnID_Components)
+	{
+		return SNew(SHorizontalBox)
+			+ SHorizontalBox::Slot()
+			.AutoWidth()
+			.HAlign(HAlign_Left)
+			.VAlign(VAlign_Center)
+			.Padding(5.0f, 5.0f)
+			[
+				SNew(STextBlock)
+				.Text_Lambda([EntryPin]()
+				{
+					FText ComponentText;
+					if (EntryPin->ComponentsToMatch.Num() == 1)
+					{
+						if (const UActorComponent* ActorComponent = EntryPin->ComponentsToMatch[0].Get())
+						{
+							ComponentText = FText::FromString(ActorComponent->GetName());
+						}
+					}
+					else if (EntryPin->ComponentsToMatch.Num() > 1)
+					{
+						ComponentText = FText::Format(LOCTEXT("ActorQuery_ComponentColumnComponentCount", "{0} Components"), EntryPin->ComponentsToMatch.Num());
+					}
+
+					return ComponentText;
+				})
+				.ToolTipText_Lambda([EntryPin]()
+				{
+					TArray<FString> ComponentNames;
+					for (const TSoftObjectPtr<UActorComponent>& ActorComponent : EntryPin->ComponentsToMatch)
+					{
+						if (ActorComponent.IsValid())
+						{
+							ComponentNames.Add(ActorComponent->GetName());
+						}
+					}
+
+					return FText::Format(LOCTEXT("ActorQuery_SelectedComponentsTooltip", "Selected components: {0}"),
+						FText::FromString(FString::Join(ComponentNames, TEXT(", "))));
+				})
+				.Font(IDetailLayoutBuilder::GetDetailFont())
+			];
+	}
+
+	return SNullWidget::NullWidget;
+}
 #endif	// WITH_EDITOR
 
 void UMovieGraphConditionGroupQuery_ActorTagName::Evaluate(const TArray<AActor*>& InActorsToQuery, const UWorld* InWorld, TSet<AActor*>& OutMatchingActors) const
@@ -1369,7 +1733,7 @@ void UMovieGraphConditionGroupQuery_ActorType::AddActorTypes(const TArray<UClass
 
 void UMovieGraphConditionGroupQuery_ComponentTagName::Evaluate(const TArray<AActor*>& InActorsToQuery, const UWorld* InWorld, TSet<AActor*>& OutMatchingActors) const
 {
-	TRACE_CPUPROFILER_EVENT_SCOPE(UMovieGraphConditionGroupQuery_ComponentTag::Evaluate);
+	TRACE_CPUPROFILER_EVENT_SCOPE(UMovieGraphConditionGroupQuery_ComponentTagName::Evaluate);
 	
 	// Quick early-out if "*" is used as the wildcard. Faster than doing the wildcard matching.
 	if (TagsToMatch == TEXT("*"))
@@ -2454,6 +2818,69 @@ const FText& UMovieGraphConditionGroupQuery_IsSpawnable::GetDisplayName() const
 	return DisplayName;
 }
 
+FMovieGraphEvaluationResult::FMovieGraphEvaluationResult(TSet<TObjectPtr<AActor>>&& InActors, TSet<TObjectPtr<UActorComponent>>&& InComponents)
+	: MatchingActors(MoveTemp(InActors))
+	, MatchingComponents(MoveTemp(InComponents))
+{
+}
+
+void FMovieGraphEvaluationResult::Reset()
+{
+	MatchingActors.Reset();
+	MatchingComponents.Reset();
+}
+
+void FMovieGraphEvaluationResult::Append(const FMovieGraphEvaluationResult& InOther)
+{
+	MatchingActors.Append(InOther.MatchingActors);
+	MatchingComponents.Append(InOther.MatchingComponents);
+}
+
+FMovieGraphEvaluationResult FMovieGraphEvaluationResult::Union(const FMovieGraphEvaluationResult& InOther) const
+{
+	return FMovieGraphEvaluationResult(
+		MatchingActors.Union(InOther.MatchingActors),
+		MatchingComponents.Union(InOther.MatchingComponents));
+}
+
+FMovieGraphEvaluationResult FMovieGraphEvaluationResult::Intersect(const FMovieGraphEvaluationResult& InOther) const
+{
+	return FMovieGraphEvaluationResult(
+		MatchingActors.Intersect(InOther.MatchingActors),
+		MatchingComponents.Intersect(InOther.MatchingComponents));
+}
+
+FMovieGraphEvaluationResult FMovieGraphEvaluationResult::Difference(const FMovieGraphEvaluationResult& InOther) const
+{
+	return FMovieGraphEvaluationResult(
+		MatchingActors.Difference(InOther.MatchingActors),
+		MatchingComponents.Difference(InOther.MatchingComponents));
+}
+
+template<typename T>
+TArray<T*> FMovieGraphEvaluationResult::GetAllComponentsOfType() const
+{
+	TArray<T*> ComponentsOfType;
+
+	for (const AActor* MatchedActor : MatchingActors)
+	{
+		constexpr bool bIncludeFromChildActors = true;
+		TInlineComponentArray<T*> PrimitiveComponents(MatchedActor, bIncludeFromChildActors);
+
+		ComponentsOfType.Append(MoveTemp(PrimitiveComponents));
+	}
+
+	for (UActorComponent* ActorComponent : MatchingComponents)
+	{
+		if (T* PrimitiveComponent = Cast<T>(ActorComponent))
+		{
+			ComponentsOfType.Add(PrimitiveComponent);
+		}
+	}
+
+	return ComponentsOfType;
+}
+
 UMovieGraphConditionGroup::UMovieGraphConditionGroup()
 	: OpType(EMovieGraphConditionGroupOpType::Add)
 {
@@ -2495,7 +2922,12 @@ EMovieGraphConditionGroupOpType UMovieGraphConditionGroup::GetOperationType() co
 
 TSet<AActor*> UMovieGraphConditionGroup::Evaluate(const UWorld* InWorld) const
 {
-	TRACE_CPUPROFILER_EVENT_SCOPE(UMovieGraphConditionGroup::Evaluate);
+	return ObjectPtrDecay(EvaluateActorsAndComponents(InWorld).MatchingActors);
+}
+
+FMovieGraphEvaluationResult UMovieGraphConditionGroup::EvaluateActorsAndComponents(const UWorld* InWorld) const
+{
+	TRACE_CPUPROFILER_EVENT_SCOPE(UMovieGraphConditionGroup::EvaluateActorsAndComponents);
 
 	// Reset the TSet for evaluation results; it is persisted across frames to prevent constantly re-allocating it
 	EvaluationResult.Reset();
@@ -2527,7 +2959,14 @@ TSet<AActor*> UMovieGraphConditionGroup::Evaluate(const UWorld* InWorld) const
 		// Similar to EvaluationResult, QueryResult is persisted+reset to prevent constantly re-allocating it
 		QueryResult.Reset();
 
-		Query->Evaluate(AllActors, InWorld, MutableView(QueryResult));
+		if (!Query->ShouldEvaluateComponents())
+		{
+			Query->Evaluate(AllActors, InWorld, MutableView(QueryResult.MatchingActors));
+		}
+		else
+		{
+			Query->EvaluateActorsAndComponents(AllActors, InWorld, MutableView(QueryResult.MatchingActors), MutableView(QueryResult.MatchingComponents));
+		}
 		
 		switch (Query->GetOperationType())
 		{
@@ -2663,9 +3102,14 @@ const FGuid& UMovieGraphConditionGroup::GetId() const
 
 TSet<AActor*> UMovieGraphCollection::Evaluate(const UWorld* InWorld) const
 {
-	TRACE_CPUPROFILER_EVENT_SCOPE(UMovieGraphCollection::Evaluate);
+	return ObjectPtrDecay(EvaluateActorsAndComponents(InWorld).MatchingActors);
+}
+
+FMovieGraphEvaluationResult UMovieGraphCollection::EvaluateActorsAndComponents(const UWorld* InWorld) const
+{
+	TRACE_CPUPROFILER_EVENT_SCOPE(UMovieGraphCollection::EvaluateActorsAndComponents);
 	
-	TSet<AActor*> ResultSet;
+	FMovieGraphEvaluationResult FinalResult;
 
 	for (int32 ConditionGroupIndex = 0; ConditionGroupIndex < ConditionGroups.Num(); ++ConditionGroupIndex)
 	{
@@ -2681,25 +3125,25 @@ TSet<AActor*> UMovieGraphCollection::Evaluate(const UWorld* InWorld) const
 			ensure(ConditionGroup->GetOperationType() == EMovieGraphConditionGroupOpType::Add);
 		}
 
-		const TSet<AActor*> QueryResult = ConditionGroup->Evaluate(InWorld);
+		const FMovieGraphEvaluationResult QueryResult = ConditionGroup->EvaluateActorsAndComponents(InWorld);
 		
 		switch (ConditionGroup->GetOperationType())
 		{
 		case EMovieGraphConditionGroupOpType::Add:
-			ResultSet = ResultSet.Union(QueryResult);
+			FinalResult = FinalResult.Union(QueryResult);
 			break;
 
 		case EMovieGraphConditionGroupOpType::And:
-			ResultSet = ResultSet.Intersect(QueryResult);
+			FinalResult = FinalResult.Intersect(QueryResult);
 			break;
 
 		case EMovieGraphConditionGroupOpType::Subtract:
-			ResultSet = ResultSet.Difference(QueryResult);
+			FinalResult = FinalResult.Difference(QueryResult);
 			break;
 		}
 	}
 	
-	return ResultSet;
+	return FinalResult;
 }
 
 UMovieGraphConditionGroup* UMovieGraphCollection::AddConditionGroup()
diff --git a/Engine/Plugins/MovieScene/MovieRenderPipeline/Source/MovieRenderPipelineCore/Public/Graph/MovieGraphRenderLayerSubsystem.h b/Engine/Plugins/MovieScene/MovieRenderPipeline/Source/MovieRenderPipelineCore/Public/Graph/MovieGraphRenderLayerSubsystem.h
--- a/Engine/Plugins/MovieScene/MovieRenderPipeline/Source/MovieRenderPipelineCore/Public/Graph/MovieGraphRenderLayerSubsystem.h
+++ b/Engine/Plugins/MovieScene/MovieRenderPipeline/Source/MovieRenderPipelineCore/Public/Graph/MovieGraphRenderLayerSubsystem.h
@@ -68,10 +68,21 @@ public:
 	UFUNCTION(BlueprintCallable, Category = "Settings")
 	EMovieGraphConditionGroupQueryOpType GetOperationType() const;
 
-	/** Determines which of the provided actors (in the given world) match the query. Matches are added to OutMatchingActors.  */
+	/**
+	 * Determines which of the provided actors (in the given world) match the query. Matches are added to OutMatchingActors.
+	 * Note that this method will not be called if ShouldEvaluateComponents() returns true; EvaluateActorsAndComponents() will
+	 * be called instead.
+	 */
 	UFUNCTION(BlueprintCallable, Category = "Settings")
 	virtual void Evaluate(const TArray<AActor*>& InActorsToQuery, const UWorld* InWorld, TSet<AActor*>& OutMatchingActors) const;
 
+	/**
+	 * Similar to Evaluate(), but returns both actors and components if the query can match both.
+	 * Note that this method will only be called if ShouldEvaluateComponents() returns true.
+	 */
+	UFUNCTION(BlueprintCallable, Category = "Settings")
+	virtual void EvaluateActorsAndComponents(const TArray<AActor*>& InActorsToQuery, const UWorld* InWorld, TSet<AActor*>& OutMatchingActors, TSet<UActorComponent*>& OutMatchingComponents) const;
+
 	/**
 	 * Determines if the public properties on the query class will have their names hidden in the details panel. Returns
 	 * false by default. Most query subclasses will only have one property and do not need to clutter the UI with the
@@ -81,6 +92,13 @@ public:
 	UFUNCTION(BlueprintCallable, Category = "Settings")
 	virtual bool ShouldHidePropertyNames() const;
 
+	/**
+	 * Determines if this query should additionally match components, rather than just matching actors. Off by default. If this returns true, then
+	 * EvaluateActorsAndComponents() will be called during evaluation instead of Evaluate().
+	 */
+	UFUNCTION(BlueprintCallable, Category = "Settings")
+	virtual bool ShouldEvaluateComponents() const;
+
 	/** Gets the icon that represents this query class in the UI. */
 	virtual const FSlateIcon& GetIcon() const;
 
@@ -144,6 +162,26 @@ private:
 	bool bIsEnabled;
 };
 
+/** Contains the actors and components to match within the Actor condition group query. */
+USTRUCT(BlueprintType)
+struct MOVIERENDERPIPELINECORE_API FMovieGraphActorQueryEntry
+{
+	GENERATED_BODY()
+	
+	bool operator==(const FMovieGraphActorQueryEntry& InOther) const;
+	
+	/**
+	 * The query must match this actor in order to be a match. If these are editor actors, they will be converted to PIE actors automatically.
+	 * If there are any ComponentsToMatch, they must be part of this actor.
+	 */
+	UPROPERTY(EditAnywhere, BlueprintReadWrite, Category = "General")
+	TSoftObjectPtr<AActor> ActorToMatch;
+	
+	/** The query must match these components in order to be a match. These must be contained within ActorToMatch. */
+	UPROPERTY(EditAnywhere, BlueprintReadWrite, Category = "General")
+    TArray<TSoftObjectPtr<UActorComponent>> ComponentsToMatch;
+};
+
 /** Query type which filters actors via an explicit actor list. */
 UCLASS(BlueprintType)
 class MOVIERENDERPIPELINECORE_API UMovieGraphConditionGroupQuery_Actor final : public UMovieGraphConditionGroupQueryBase
@@ -152,6 +190,8 @@ class MOVIERENDERPIPELINECORE_API UMovieGraphConditionGroupQuery_Actor final : p
 
 public:
 	virtual void Evaluate(const TArray<AActor*>& InActorsToQuery, const UWorld* InWorld, TSet<AActor*>& OutMatchingActors) const override;
+	virtual void EvaluateActorsAndComponents(const TArray<AActor*>& InActorsToQuery, const UWorld* InWorld, TSet<AActor*>& OutMatchingActors, TSet<UActorComponent*>& OutMatchingComponents) const override;
+	virtual bool ShouldEvaluateComponents() const override;
 	virtual const FSlateIcon& GetIcon() const override;
 	virtual const FText& GetDisplayName() const override;
 
@@ -161,16 +201,18 @@ public:
 	virtual TSharedRef<SWidget> GetAddMenuContents(const FMovieGraphConditionGroupQueryContentsChanged& OnAddFinished) override;
 #endif
 
+	virtual void PostLoad() override;
+
 private:
 #if WITH_EDITOR
-	static const FSlateBrush* GetRowIcon(TSharedPtr<TSoftObjectPtr<AActor>> InActor);
-	static FText GetRowText(TSharedPtr<TSoftObjectPtr<AActor>> InActor);
-
 	/** Adds the provided actors to the query, updating the UI as needed. Calls InOnAddFinished when done. Can optionally close the Add menu. */
 	void AddActors(const TArray<AActor*>& InActors, const FMovieGraphConditionGroupQueryContentsChanged& InOnAddFinished, const bool bCloseAddMenu = true);
 
-	/** Removes the provided actors from the query, updating the UI as needed. */
-	void RemoveActors(const TArray<TSoftObjectPtr<AActor>>& InActors);
+	/** Removes the provided entries from the query, updating the UI as needed. */
+	void RemoveEntries(const TArray<FMovieGraphActorQueryEntry>& InEntries);
+
+	/** Gets the context menu for an entry in the list widget. */
+	void GetListContextMenu(FMenuBuilder& InMenuBuilder, TArray<TSharedPtr<FMovieGraphActorQueryEntry>> SelectedEntries);
 
 	/** Refreshes the list's data source to reflect the data model. */
 	void RefreshListDataSource();
@@ -178,9 +220,14 @@ private:
 
 public:
 	/** The query must match one of the actors in order to be a match. If these are editor actors, they will be converted to PIE actors automatically. */
-	UPROPERTY(EditAnywhere, Category="General")
+	UE_DEPRECATED(5.6, "This property is deprecated. Please use ActorsAndComponentsToMatch instead.")
+	UPROPERTY(meta=(DeprecatedProperty, DeprecationMessage="This property is deprecated. Please use ActorsAndComponentsToMatch instead."))
 	TArray<TSoftObjectPtr<AActor>> ActorsToMatch;
 
+	/** The query must match one of the actors (or components on an actor) to be a match. If these are editor actors, they will be converted to PIE actors automatically. */
+	UPROPERTY(EditAnywhere, BlueprintReadWrite, Category="General")
+	TArray<FMovieGraphActorQueryEntry> ActorsAndComponentsToMatch;
+
 private:
 #if WITH_EDITOR
 	/** Custom outliner column that allows adding/removing an actor from an Actor condition group query (via checkbox). */
@@ -207,15 +254,31 @@ private:
 		/** The Actor condition group query that populates the data for this column. */
 		TWeakObjectPtr<UMovieGraphConditionGroupQuery_Actor> WeakActorQuery;
 	};
+
+	static const inline FName ColumnID_ActorName = FName(TEXT("ActorName"));
+	static const inline FName ColumnID_ActorType = FName(TEXT("ActorType"));
+	static const inline FName ColumnID_Components = FName(TEXT("Components"));
+
+	/** Custom row widget for the actor list so multiple columns can be populated. */
+	class SActorListRow final : public SMultiColumnTableRow<TSharedPtr<FMovieGraphActorQueryEntry>>
+	{
+	public:
+		void Construct(const FArguments& InArgs, const TSharedRef<STableViewBase>& InOwnerTableView, TSharedPtr<FMovieGraphActorQueryEntry> InEntry);
+		
+		virtual TSharedRef<SWidget> GenerateWidgetForColumn(const FName& InColumnName) override;
+
+	private:
+		TWeakPtr<FMovieGraphActorQueryEntry> WeakEntry;
+	};
 	
 	TSharedPtr<class ISceneOutliner> ActorPickerWidget;
 
 	/** Displays the actors which have been chosen. */
-	TSharedPtr<SMovieGraphSimpleList<TSharedPtr<TSoftObjectPtr<AActor>>>> ActorsList;
+	TSharedPtr<SMovieGraphSimpleList<TSharedPtr<FMovieGraphActorQueryEntry>>> ActorsList;
 
 	// Not ideal to store a duplicate of ActorsToMatch, but SListView requires TSharedPtr<...> as the data source, and UPROPERTY does not
 	// support TSharedPtr<...>
-	TArray<TSharedPtr<TSoftObjectPtr<AActor>>> ListDataSource;
+	TArray<TSharedPtr<FMovieGraphActorQueryEntry>> ListDataSource;
 #endif
 };
 
@@ -578,6 +641,43 @@ public:
 	bool bIsSpawnable;
 };
 
+/** Provides the actors/components returned by a collection or condition group evaluation. */
+USTRUCT(BlueprintType)
+struct MOVIERENDERPIPELINECORE_API FMovieGraphEvaluationResult
+{
+	GENERATED_BODY()
+
+	FMovieGraphEvaluationResult() = default;
+	FMovieGraphEvaluationResult(TSet<TObjectPtr<AActor>>&& InActors, TSet<TObjectPtr<UActorComponent>>&& InComponents);
+
+	/** Empties this evaluation result and resets it to defaults. */
+	void Reset();
+
+	/** Appends the actors/components in InOther to the actors/components in this evaluation result. */
+	void Append(const FMovieGraphEvaluationResult& InOther);
+
+	/** Unions this evaluation result with InOther. */
+	FMovieGraphEvaluationResult Union(const FMovieGraphEvaluationResult& InOther) const;
+
+	/** Intersects this evaluation result with InOther. */
+	FMovieGraphEvaluationResult Intersect(const FMovieGraphEvaluationResult& InOther) const;
+
+	/** Differences this evaluation result with InOther. */
+	FMovieGraphEvaluationResult Difference(const FMovieGraphEvaluationResult& InOther) const;
+
+	/** Gets all components of the specified type across MatchingActors and MatchingComponents. */
+	template<typename T>
+	TArray<T*> GetAllComponentsOfType() const;
+
+	/** The actors that were matched during evaluation. */
+	UPROPERTY(BlueprintReadOnly, Category = "Movie Graph")
+	TSet<TObjectPtr<AActor>> MatchingActors;
+
+	/** The components that were matched during evaluation. */
+	UPROPERTY(BlueprintReadOnly, Category = "Movie Graph")
+	TSet<TObjectPtr<UActorComponent>> MatchingComponents;
+};
+
 /** A group of queries which can be added to a collection. */
 UCLASS(BlueprintType)
 class MOVIERENDERPIPELINECORE_API UMovieGraphConditionGroup : public UObject
@@ -599,9 +699,14 @@ public:
 	EMovieGraphConditionGroupOpType GetOperationType() const;
 
 	/** Determines the actors that match the condition group by running the queries contained in it. */
-	UFUNCTION(BlueprintCallable, Category = "Settings")
+	UE_DEPRECATED(5.6, "Please use EvaluateActorsAndComponents() instead.")
+	UFUNCTION(BlueprintCallable, Category = "Settings", meta = (DeprecatedFunction, DeprecationMessage = "Please use EvaluateActorsAndComponents() instead."))
 	TSet<AActor*> Evaluate(const UWorld* InWorld) const;
 
+	/** Determines the actors and components that match the condition group by running the queries contained in it. */
+	UFUNCTION(BlueprintCallable, Category = "Settings")
+	FMovieGraphEvaluationResult EvaluateActorsAndComponents(const UWorld* InWorld) const;
+
 	/**
 	 * Adds a new condition group query to the condition group and returns a ptr to it. The condition group owns the
 	 * created query. By default the query is added to the end, but an optional index can be provided if the query
@@ -654,11 +759,11 @@ private:
 
 	/** Persisted actor set which can be re-used for query evaluations across frames to prevent constantly re-allocating it. */
 	UPROPERTY(Transient, DuplicateTransient)
-	mutable TSet<TObjectPtr<AActor>> QueryResult;
+	mutable FMovieGraphEvaluationResult QueryResult;
 
 	/** Persisted actor set which can be re-used for condition group evaluations across frames to prevent constantly re-allocating it. */
 	UPROPERTY(Transient, DuplicateTransient)
-	mutable TSet<TObjectPtr<AActor>> EvaluationResult;
+	mutable FMovieGraphEvaluationResult EvaluationResult;
 };
 
 /** A group of actors generated by actor queries. */
@@ -691,9 +796,17 @@ public:
 	 * Gets matching actors by having condition groups evaluate themselves, and performing set operations on the
 	 * condition group results (eg, union'ing condition group A and B).
 	 */
-	UFUNCTION(BlueprintCallable, Category = "Settings")
+	UE_DEPRECATED(5.6, "Please use EvaluateActorsAndComponents() instead.")
+	UFUNCTION(BlueprintCallable, Category = "Settings", meta = (DeprecatedFunction, DeprecationMessage = "Please use EvaluateActorsAndComponents() instead."))
 	TSet<AActor*> Evaluate(const UWorld* InWorld) const;
 
+	/**
+	 * Gets matching actors and components by having condition groups evaluate themselves, and performing set operations on the
+	 * condition group results (eg, union'ing condition group A and B).
+	 */
+	UFUNCTION(BlueprintCallable, Category = "Settings")
+	FMovieGraphEvaluationResult EvaluateActorsAndComponents(const UWorld* InWorld) const;
+
 	/**
 	 * Adds a new condition group to the collection and returns a ptr to it. The collection owns the created
 	 * condition group.
@@ -829,35 +942,42 @@ public:
 
 private:
 	/** Various visibility properties for an actor. */
-	struct FActorVisibilityState
+	struct FActorState
 	{
-		struct FComponentState
-		{
-			TSoftObjectPtr<USceneComponent> Component = nullptr;
-			
-			// Note: The default values specified here reflect the defaults on the scene component. If a modifier property is marked as overridden, the
-			// override will initially be a no-op due to the defaults being the same.
-			uint8 bCastsShadows : 1 = true;
-			uint8 bCastShadowWhileHidden : 1 = false;
-			uint8 bAffectIndirectLightingWhileHidden : 1 = false;
-			uint8 bHoldout : 1 = false;
-		};
-
 		TSoftObjectPtr<AActor> Actor = nullptr;
-		TArray<FComponentState> Components;
+		uint8 bIsHidden : 1 = false;
+	};
 
+	/** Various visibility properties for a component. */
+	struct FComponentState
+	{
+		TSoftObjectPtr<USceneComponent> Component = nullptr;
+		
+		// Note: The default values specified here reflect the defaults on the scene component. If a modifier property is marked as overridden, the
+		// override will initially be a no-op due to the defaults being the same.
+		uint8 bCastsShadows : 1 = true;
+		uint8 bCastShadowWhileHidden : 1 = false;
+		uint8 bAffectIndirectLightingWhileHidden : 1 = false;
+		uint8 bHoldout : 1 = false;
 		uint8 bIsHidden : 1 = false;
 	};
 
-	/** Updates an actor's visibility state to the state contained in NewVisibilityState. */
-	void SetActorVisibilityState(const FActorVisibilityState& NewVisibilityState);
+	/**
+	 * Updates actor/component state to the state contained in InActorState/InComponentState. If bUseStateFromNode is true, InActorState and InComponent
+	 * state will be ignored; the node's state will be applied to all actors and components that are currently cached. This is typically used to set
+	 * the state of actors and components before a render after their state has been cached out.
+	 */
+	void SetActorAndComponentState(const TArray<FActorState>& InActorState, const TArray<FComponentState>& InComponentState, const bool bUseStateFromNode);
 
-	/** Convienence function to ensure that output alpha and primitive alpha holdout settings are enabled if required. */
+	/** Convenience function to ensure that output alpha and primitive alpha holdout settings are enabled if required. */
 	void ValidateProjectSettings() const;
 
 private:
-	/** Tracks actor visibility state prior to having the modifier applied. */
-	TArray<FActorVisibilityState> ModifiedActors;
+	/** Tracks relevant actor state prior to having the modifier applied. Only actors that are affected are included in the cache. */
+	TArray<FActorState> CachedActorState;
+
+	/** Tracks relevant component state prior to having the modifier applied. Only components that are affected are included in the cache. */
+	TArray<FComponentState> CachedComponentState;
 
 public:
 	UPROPERTY(EditAnywhere, BlueprintReadWrite, Category = Overrides, meta = (InlineEditConditionToggle))
diff --git a/Engine/Plugins/MovieScene/MovieRenderPipeline/Source/MovieRenderPipelineCore/Public/Graph/MovieGraphSharedWidgets.h b/Engine/Plugins/MovieScene/MovieRenderPipeline/Source/MovieRenderPipelineCore/Public/Graph/MovieGraphSharedWidgets.h
--- a/Engine/Plugins/MovieScene/MovieRenderPipeline/Source/MovieRenderPipelineCore/Public/Graph/MovieGraphSharedWidgets.h
+++ b/Engine/Plugins/MovieScene/MovieRenderPipeline/Source/MovieRenderPipelineCore/Public/Graph/MovieGraphSharedWidgets.h
@@ -33,7 +33,9 @@ class SMovieGraphSimpleList final : public SCompoundWidget, public FSelfRegister
 public:
 	DECLARE_DELEGATE_RetVal_OneParam(const FSlateBrush*, FGetRowIcon, ListType);
 	DECLARE_DELEGATE_RetVal_OneParam(FText, FGetRowText, ListType);
+	DECLARE_DELEGATE_RetVal_TwoParams(TSharedRef<ITableRow>, FGetCustomRow, const TSharedRef<STableViewBase>&, ListType);
 	DECLARE_DELEGATE_OneParam(FOnDelete, TArray<ListType>);
+	DECLARE_DELEGATE_TwoParams(FOnGetContextMenuContent, FMenuBuilder&, TArray<ListType>);
 	DECLARE_DELEGATE(FOnRefreshDataSourceRequested);
 	DECLARE_DELEGATE_RetVal_OneParam(bool, FGetRowEnableState, ListType);
 	DECLARE_DELEGATE_TwoParams(FSetRowEnableState, ListType, bool);
@@ -58,18 +60,30 @@ public:
 		/** The plural of the DataType attribute. */
 		SLATE_ATTRIBUTE(FText, DataTypePlural)
 
+		/**
+		 * Gets the text that should be shown in the summary row. By default, "DataType" and "DataTypePlural" will be used to auto-generate summary
+		 * row text. If this attribute is bound, those will be ignored, and this will be used instead.
+		 */
+		SLATE_ATTRIBUTE(FText, CustomSummaryText);
+
 		/** The selection mode the list should be put into. Defaults to Multi. */
 		SLATE_ATTRIBUTE(ESelectionMode::Type, SelectionMode)
 
 		/** Gets the icon for a row in the list. */
 		SLATE_EVENT(FGetRowIcon, OnGetRowIcon);
 
-		/** Gets the text for a row in the list. */
+		/** Gets the text for a row in the list. Use OnGetCustomRow if a custom HeaderRow is in use. */
 		SLATE_EVENT(FGetRowText, OnGetRowText);
 
+		/** If needed, custom rows can be generated. This is needed if using a custom header row. */
+		SLATE_EVENT(FGetCustomRow, OnGetCustomRow);
+
 		/** Invoked when a delete operation is performed. */
 		SLATE_EVENT(FOnDelete, OnDelete)
 
+		/** Gets the menu items that should show up in the context menu when it is shown. Note that the delete operation is handled separately via OnDelete. */
+		SLATE_EVENT(FOnGetContextMenuContent, OnGetContextMenuContent)
+
 		/** Whether the enable/disable checkbox should be displayed. */
 		SLATE_ATTRIBUTE(bool, ShowEnableDisable)
 
@@ -78,6 +92,9 @@ public:
 
 		/** Sets the enable state of the row. */
 		SLATE_EVENT(FSetRowEnableState, OnSetRowEnableState)
+
+		/** Optional header row that will be displayed in the list. Custom rows should also be implemented if this is provided. */
+		SLATE_ARGUMENT_DEFAULT(TSharedPtr<SHeaderRow>, HeaderRow) = nullptr;
 	SLATE_END_ARGS()
 	
 	void Construct(const FArguments& InArgs)
@@ -86,13 +103,17 @@ public:
 		OnRefreshDataSourceRequested = InArgs._OnRefreshDataSourceRequested;
 		DataType = InArgs._DataType.Get();
 		DataTypePlural = InArgs._DataTypePlural.Get();
+		CustomSummaryText = InArgs._CustomSummaryText;
 		SelectionMode = InArgs._SelectionMode.Get();
 		OnGetRowIcon = InArgs._OnGetRowIcon;
 		OnGetRowText = InArgs._OnGetRowText;
+		OnGetCustomRow = InArgs._OnGetCustomRow;
 		OnDelete = InArgs._OnDelete;
+		OnGetContextMenuContent = InArgs._OnGetContextMenuContent;
 		bShowEnableDisable = InArgs._ShowEnableDisable.Get();
 		OnGetRowEnableState = InArgs._OnGetRowEnableState;
 		OnSetRowEnableState = InArgs._OnSetRowEnableState;
+		HeaderRow = InArgs._HeaderRow;
 
 		CommandList = MakeShared<FUICommandList>();
 
@@ -115,6 +136,7 @@ public:
 				SAssignNew(ListView, SListView<ListType>)
 				.ListItemsSource(DataSource)
 				.SelectionMode(SelectionMode)
+				.HeaderRow(HeaderRow)
 				.OnKeyDownHandler_Lambda([this](const FGeometry& InGeometry, const FKeyEvent& InKeyEvent)
 				{
 					return CommandList->ProcessCommandBindings(InKeyEvent) ? FReply::Handled() : FReply::Unhandled();
@@ -136,6 +158,11 @@ public:
 					.Font(FAppStyle::GetFontStyle(TEXT("PropertyWindow.NormalFont")))
 					.Text_Lambda([this]()
 					{
+						if (CustomSummaryText.IsSet())
+						{
+							return CustomSummaryText.Get();
+						}
+						
 						return FText::FromString(FString::Printf(TEXT("%i %s"), DataSource->Num(), DataSource->Num() == 1 ? *DataType.ToString() : *DataTypePlural.ToString()));
 					})
 				]
@@ -207,18 +234,36 @@ private:
 		constexpr bool bShouldCloseWindowAfterMenuSelection = true;
 		FMenuBuilder MenuBuilder(bShouldCloseWindowAfterMenuSelection, CommandList);
 
-		MenuBuilder.BeginSection("BasicOperations");
+		MenuBuilder.SetSearchable(true);
+		MenuBuilder.AddSearchWidget();
+
+		MenuBuilder.BeginSection(NAME_None, LOCTEXT("ContextMenuSection_Edit", "Edit"));
 		{
 			MenuBuilder.AddMenuEntry(FGenericCommands::Get().Delete);
 		}
 		MenuBuilder.EndSection();
 
+		// Other custom context menu entries can be added by the delegate if needed
+		if (OnGetContextMenuContent.IsBound())
+		{
+			TArray<ListType> SelectedItems;
+			ListView->GetSelectedItems(SelectedItems);
+			
+			OnGetContextMenuContent.Execute(MenuBuilder, SelectedItems);
+		}
+
 		return MenuBuilder.MakeWidget();
 	}
 
 	/** Generates a row in the list for the specified data. */
 	TSharedRef<ITableRow> GenerateRow(ListType InListData, const TSharedRef<STableViewBase>& InOwnerTable) const
 	{
+		// There may be completely-custom rows supplied, rather than the typical auto-generated ones that are displayed
+		if (OnGetCustomRow.IsBound())
+		{
+			return OnGetCustomRow.Execute(InOwnerTable, InListData);
+		}
+		
 		const FSlateBrush* RowIcon = GetRowIcon(InListData);
 		
 		return
@@ -258,7 +303,7 @@ private:
 				
 				+ SHorizontalBox::Slot()
 				.VAlign(VAlign_Center)
-				.Padding(0.f, 5.f, 7.f, 5.f)
+				.Padding(3.f, 5.f, 7.f, 5.f)
 				.AutoWidth()
 				[
 					SNew(SImage)
@@ -296,13 +341,17 @@ private:
 	TSharedPtr<FUICommandList> CommandList;
 	FGetRowIcon OnGetRowIcon;
 	FGetRowText OnGetRowText;
+	FGetCustomRow OnGetCustomRow;
 	FOnDelete OnDelete;
+	FOnGetContextMenuContent OnGetContextMenuContent;
 	FOnRefreshDataSourceRequested OnRefreshDataSourceRequested;
 	bool bShowEnableDisable = false;
 	FGetRowEnableState OnGetRowEnableState;
 	FSetRowEnableState OnSetRowEnableState;
+	TSharedPtr<SHeaderRow> HeaderRow;
 	FText DataType;
 	FText DataTypePlural;
+	TAttribute<FText> CustomSummaryText;
 	TArray<ListType>* DataSource = nullptr;
 	ESelectionMode::Type SelectionMode = ESelectionMode::Multi;
 };
diff --git a/Engine/Plugins/MovieScene/MovieRenderPipeline/Source/MovieRenderPipelineEditor/Private/Graph/MovieGraphAssetToolkit.cpp b/Engine/Plugins/MovieScene/MovieRenderPipeline/Source/MovieRenderPipelineEditor/Private/Graph/MovieGraphAssetToolkit.cpp
--- a/Engine/Plugins/MovieScene/MovieRenderPipeline/Source/MovieRenderPipelineEditor/Private/Graph/MovieGraphAssetToolkit.cpp
+++ b/Engine/Plugins/MovieScene/MovieRenderPipeline/Source/MovieRenderPipelineEditor/Private/Graph/MovieGraphAssetToolkit.cpp
@@ -98,15 +98,16 @@ void SMovieGraphSyncCollectionToOutlinerButton::EvaluateCollectionAndSelect() co
 	if (const UMovieGraphCollectionNode* CollectionNode = Cast<UMovieGraphCollectionNode>(SelectedObjects[0]))
 	{
 		// Evaluate the collection based on the current editor world
-		TSet<AActor*> EvaluatedActors = CollectionNode->Collection->Evaluate(GEditor->GetEditorWorldContext().World());
+		const FMovieGraphEvaluationResult EvaluationResult =
+			CollectionNode->Collection->EvaluateActorsAndComponents(GEditor->GetEditorWorldContext().World());
 
 		// Select all actors matched by the collection
 		{
 			GEditor->GetSelectedActors()->Modify();
 			GEditor->GetSelectedActors()->BeginBatchSelectOperation();
 			GEditor->GetSelectedActors()->DeselectAll();
-							
-			for (AActor* Actor : EvaluatedActors)
+			
+			for (AActor* Actor : EvaluationResult.MatchingActors)
 			{
 				constexpr bool bShouldSelect = true;
 				constexpr bool bNotifyAfterSelect = false;
