- Fixed a link error with a function that was mistakenly wrapped in WITH_EDITOR

diff --git a/Engine/Plugins/MovieScene/MovieRenderPipeline/Source/MovieRenderPipelineCore/Private/Graph/MovieGraphRenderLayerSubsystem.cpp b/Engine/Plugins/MovieScene/MovieRenderPipeline/Source/MovieRenderPipelineCore/Private/Graph/MovieGraphRenderLayerSubsystem.cpp
--- a/Engine/Plugins/MovieScene/MovieRenderPipeline/Source/MovieRenderPipelineCore/Private/Graph/MovieGraphRenderLayerSubsystem.cpp
+++ b/Engine/Plugins/MovieScene/MovieRenderPipeline/Source/MovieRenderPipelineCore/Private/Graph/MovieGraphRenderLayerSubsystem.cpp
@@ -981,6 +981,8 @@ TSharedRef<SWidget> UMovieGraphConditionGroupQuery_Actor::GetAddMenuContents(con
 	return MenuBuilder.MakeWidget();
 }
 
+#endif
+
 void UMovieGraphConditionGroupQuery_Actor::PostLoad()
 {
 	Super::PostLoad();
@@ -1004,6 +1006,7 @@ void UMovieGraphConditionGroupQuery_Actor::PostLoad()
 #endif	// WITH_EDITOR
 }
 
+#if WITH_EDITOR
 void UMovieGraphConditionGroupQuery_Actor::AddActors(const TArray<AActor*>& InActors, const FMovieGraphConditionGroupQueryContentsChanged& InOnAddFinished, const bool bCloseAddMenu)
 {
 	const FScopedTransaction Transaction(LOCTEXT("AddActorsToCollection", "Add Actors to Collection"));
@@ -3400,4 +3403,4 @@ void UMovieGraphRenderLayerSubsystem::SetAndApplyRenderLayer(UMovieGraphRenderLa
 	ActiveRenderLayer->Apply(GetWorld());
 }
 
-#undef LOCTEXT_NAMESPACE
\ No newline at end of file
+#undef LOCTEXT_NAMESPACE
