MovieGraph: Don't render warm-up frames for UI/Burn In renderers, or additional post process materials, the output merger is not expecting the data yet as we have not actually started a render. This fixes an issue where there are long stalls at the start of renders (caused by tripped ensures) and removes the ensures that trip for each warm up frame.

Tested having additional post process materials and burn-ins enabled with varying combinations of MRG's Emulate Motion Blur (on or off), Warm Up Count (0 or more), Temporal and Spatial SubSampling (1 or more) and verified that the correct number of images is produced for each one and no ensures are tripped.

diff --git a/Engine/Plugins/MovieScene/MovieRenderPipeline/Source/MovieRenderPipelineCore/Private/Graph/Nodes/MovieGraphWidgetRendererBaseNode.cpp b/Engine/Plugins/MovieScene/MovieRenderPipeline/Source/MovieRenderPipelineCore/Private/Graph/Nodes/MovieGraphWidgetRendererBaseNode.cpp
--- a/Engine/Plugins/MovieScene/MovieRenderPipeline/Source/MovieRenderPipelineCore/Private/Graph/Nodes/MovieGraphWidgetRendererBaseNode.cpp
+++ b/Engine/Plugins/MovieScene/MovieRenderPipeline/Source/MovieRenderPipelineCore/Private/Graph/Nodes/MovieGraphWidgetRendererBaseNode.cpp
@@ -15,6 +15,7 @@
 #include "TextureResource.h"
 #include "Widgets/SVirtualWindow.h"
 #include "RenderingThread.h"
+#include "MoviePipelineQueue.h"
 
 void UMovieGraphWidgetRendererBaseNode::SetupImpl(const FMovieGraphRenderPassSetupData& InSetupData)
 {
@@ -106,6 +107,12 @@ void UMovieGraphWidgetRendererBaseNode::FMovieGraphWidgetPass::Render(const FMov
 	const UMovieGraphPipeline* Pipeline = Renderer->GetOwningGraph();
 	UMovieGraphWidgetRendererBaseNode* ParentNodeThisFrame = GetParentNode(InTimeData.EvaluatedConfig);
 
+
+	if(Pipeline->GetActiveShotList()[Pipeline->GetCurrentShotIndex()]->ShotInfo.State != EMovieRenderShotState::Rendering)
+	{
+		return;
+	}
+
 	const float CameraOverscan = Renderer->GetCameraOverscan(LayerData.CameraIndex);
 	const FIntRect CropRect = UMovieGraphBlueprintLibrary::GetOverscanCropRectangle(InTimeData.EvaluatedConfig, CameraOverscan);
 	
diff --git a/Engine/Plugins/MovieScene/MovieRenderPipeline/Source/MovieRenderPipelineRenderPasses/Private/Graph/Renderers/MovieGraphDeferredPass.cpp b/Engine/Plugins/MovieScene/MovieRenderPipeline/Source/MovieRenderPipelineRenderPasses/Private/Graph/Renderers/MovieGraphDeferredPass.cpp
--- a/Engine/Plugins/MovieScene/MovieRenderPipeline/Source/MovieRenderPipelineRenderPasses/Private/Graph/Renderers/MovieGraphDeferredPass.cpp
+++ b/Engine/Plugins/MovieScene/MovieRenderPipeline/Source/MovieRenderPipelineRenderPasses/Private/Graph/Renderers/MovieGraphDeferredPass.cpp
@@ -310,6 +310,9 @@ void FMovieGraphDeferredPass::Render(const FMovieGraphTraversalContext& InFrameT
 
 		ApplyMovieGraphOverridesToSampleState(SampleState);
 
+		// If this was just to contribute to the history buffer, no need to go any further. Never discard if we're writing individual samples, though.
+		bool bDiscardOutput = (InTimeData.bDiscardOutput || ShouldDiscardOutput(ViewFamily, CameraInfo)) && !SampleState.bWriteSampleToDisk;
+
 		if (UMovieGraphImagePassBaseNode* ParentNode = GetParentNode(InFrameTraversalContext.Time.EvaluatedConfig))
 		{
 			TSet<UMaterialInterface*> HighPrecisionMaterials;
@@ -322,6 +325,13 @@ void FMovieGraphDeferredPass::Render(const FMovieGraphTraversalContext& InFrameT
 				{
 					continue;
 				}
+
+				// If we're not going to keep the output for the main pass then we don't want to create forwarding endpoints,
+				// as they'll read back data for the discarded main pass results and then try to pass them on.
+				if (bDiscardOutput)
+				{
+					continue;
+				}
 				
 				auto BufferPipe = MakeShared<FImagePixelPipe, ESPMode::ThreadSafe>();
 				
@@ -365,8 +375,7 @@ void FMovieGraphDeferredPass::Render(const FMovieGraphTraversalContext& InFrameT
 			});
 
 
-		// If this was just to contribute to the history buffer, no need to go any further. Never discard if we're writing individual samples, though.
-		bool bDiscardOutput = (InTimeData.bDiscardOutput || ShouldDiscardOutput(ViewFamily, CameraInfo)) && !SampleState.bWriteSampleToDisk;
+		// Data may not be something we actually want to read back and write to disk after rendering.
 		if (bDiscardOutput)
 		{
 			continue;
