MovieGraph: Fixed a crash when using 32-bit PPMs and BurnIns on the same layer due to additional post process materials trying to comp the burn-in on but always expecting the underlying bitdepth to be 16 bit.

Tested 47 combinations of Additional Post Process Materials, Burn Ins, Image/Video output formats, with and without OCIO.

diff --git a/Engine/Plugins/MovieScene/MovieRenderPipeline/Source/MovieRenderPipelineRenderPasses/Private/Graph/Renderers/MovieGraphDeferredPass.cpp b/Engine/Plugins/MovieScene/MovieRenderPipeline/Source/MovieRenderPipelineRenderPasses/Private/Graph/Renderers/MovieGraphDeferredPass.cpp
--- a/Engine/Plugins/MovieScene/MovieRenderPipeline/Source/MovieRenderPipelineRenderPasses/Private/Graph/Renderers/MovieGraphDeferredPass.cpp
+++ b/Engine/Plugins/MovieScene/MovieRenderPipeline/Source/MovieRenderPipelineRenderPasses/Private/Graph/Renderers/MovieGraphDeferredPass.cpp
@@ -348,6 +348,11 @@ void FMovieGraphDeferredPass::Render(const FMovieGraphTraversalContext& InFrameT
 				
 				FMovieGraphSampleState PassSampleState = SampleState;
 				PassSampleState.TraversalContext.RenderDataIdentifier = Identifier;
+
+				// Additional Post Process materials should not have things composited onto them (like burn-ins)
+				// nor should they have OCIO applied (as they're going to be data buffers like depth).
+				PassSampleState.bAllowsCompositing = false;
+				PassSampleState.bAllowOCIO = false;
 				
 				// Give a lower priority to materials so they show up after the main pass in multi-layer exrs.
 				PassSampleState.CompositingSortOrder = SampleState.CompositingSortOrder + 1;
diff --git a/Engine/Plugins/MovieScene/MovieRenderPipeline/Source/MovieRenderPipelineRenderPasses/Private/MovieGraphImageSequenceOutputNode.cpp b/Engine/Plugins/MovieScene/MovieRenderPipeline/Source/MovieRenderPipelineRenderPasses/Private/MovieGraphImageSequenceOutputNode.cpp
--- a/Engine/Plugins/MovieScene/MovieRenderPipeline/Source/MovieRenderPipelineRenderPasses/Private/MovieGraphImageSequenceOutputNode.cpp
+++ b/Engine/Plugins/MovieScene/MovieRenderPipeline/Source/MovieRenderPipelineRenderPasses/Private/MovieGraphImageSequenceOutputNode.cpp
@@ -395,8 +395,21 @@ void UMovieGraphImageSequenceOutputNode_EXR::OnReceiveImageDataImpl(UMovieGraphP
 			// This composited pass will only composite on top of renders w/ the same branch and camera
 			if (CompositedPass.Key.IsBranchAndCameraEqual(RenderData.Key))
 			{
+				EImagePixelType PixelType = RenderData.Value->GetType();
+
 				// There could be multiple renders within this branch using the composited pass, so we have to copy the image data
-				ImageWriteTask->PixelPreprocessors.FindOrAdd(LayerIndex).Add(TAsyncCompositeImage<FFloat16Color>(CompositedPass.Value->CopyImageData(), Payload->CropRectangle.Min));
+				switch (PixelType)
+				{
+				case EImagePixelType::Color:
+					ImageWriteTask->PixelPreprocessors.FindOrAdd(LayerIndex).Add(TAsyncCompositeImage<FColor>(CompositedPass.Value->CopyImageData(), Payload->CropRectangle.Min));
+					break;
+				case EImagePixelType::Float16:
+					ImageWriteTask->PixelPreprocessors.FindOrAdd(LayerIndex).Add(TAsyncCompositeImage<FFloat16Color>(CompositedPass.Value->CopyImageData(), Payload->CropRectangle.Min));
+					break;
+				case EImagePixelType::Float32:
+					ImageWriteTask->PixelPreprocessors.FindOrAdd(LayerIndex).Add(TAsyncCompositeImage<FLinearColor>(CompositedPass.Value->CopyImageData(), Payload->CropRectangle.Min));
+					break;
+				}
 			}
 		}
 
