- Add Blueprintable specifier to the UClass macro to make the MovieGraphConfig Type availlable to blueprint.

diff --git a/Engine/Plugins/MovieScene/MovieRenderPipeline/Source/MovieRenderPipelineCore/Public/Graph/MovieGraphConfig.h b/Engine/Plugins/MovieScene/MovieRenderPipeline/Source/MovieRenderPipelineCore/Public/Graph/MovieGraphConfig.h
--- a/Engine/Plugins/MovieScene/MovieRenderPipeline/Source/MovieRenderPipelineCore/Public/Graph/MovieGraphConfig.h
+++ b/Engine/Plugins/MovieScene/MovieRenderPipeline/Source/MovieRenderPipelineCore/Public/Graph/MovieGraphConfig.h
@@ -590,7 +590,7 @@ public:
 * typed graph network that is read by the MoviePipeline. There is an editor-only representation of
 * this graph (UMoviePipelineEdGraph).
 */
-UCLASS()
+UCLASS(BlueprintType)
 class MOVIERENDERPIPELINECORE_API UMovieGraphConfig : public UObject
 {
 	GENERATED_BODY()
