MovieRenderPipeline: Fixed issue where exr outputs were not using the camera's overscan when computing display window, resulting in an incorrect display window

diff --git a/Engine/Plugins/MovieScene/MovieRenderPipeline/Source/MovieRenderPipelineRenderPasses/Private/MoviePipelineImagePassBase.cpp b/Engine/Plugins/MovieScene/MovieRenderPipeline/Source/MovieRenderPipelineRenderPasses/Private/MoviePipelineImagePassBase.cpp
--- a/Engine/Plugins/MovieScene/MovieRenderPipeline/Source/MovieRenderPipelineRenderPasses/Private/MoviePipelineImagePassBase.cpp
+++ b/Engine/Plugins/MovieScene/MovieRenderPipeline/Source/MovieRenderPipelineRenderPasses/Private/MoviePipelineImagePassBase.cpp
@@ -450,6 +450,9 @@ FSceneView* UMoviePipelineImagePassBase::GetSceneViewForSampleState(FSceneViewFa
 		{
 			GetPipeline()->WarnAboutAnimatedOverscan(CachedOverscan);
 		}
+
+		// Update the sample state with this camera's overscan instead of the config overscan it is filled with initially
+		InOutSampleState.OverscanPercentage = CachedOverscan;
 	}
 	
 	ViewInitOptions.FOV = CameraInfo.ViewInfo.FOV;
