MoviePipeline: Updated more condition group queries in the graph to be capable of component-based matches (Component Type, Component Tag Name).

Also updated the "Subtract" and "And" operations within conditions and condition groups to properly work with component-based matches.

Fixed an issue that could cause a condition group or condition to not revert to "Add" if they become the first condition group/condition after a deletion. The first condition group and condition should always be an "Add".

diff --git a/Engine/Plugins/MovieScene/MovieRenderPipeline/Source/MovieRenderPipelineCore/Private/Graph/MovieGraphRenderLayerSubsystem.cpp b/Engine/Plugins/MovieScene/MovieRenderPipeline/Source/MovieRenderPipelineCore/Private/Graph/MovieGraphRenderLayerSubsystem.cpp
--- a/Engine/Plugins/MovieScene/MovieRenderPipeline/Source/MovieRenderPipelineCore/Private/Graph/MovieGraphRenderLayerSubsystem.cpp
+++ b/Engine/Plugins/MovieScene/MovieRenderPipeline/Source/MovieRenderPipelineCore/Private/Graph/MovieGraphRenderLayerSubsystem.cpp
@@ -11,6 +11,7 @@
 #include "Materials/MaterialInterface.h"
 #include "Modules/ModuleManager.h"
 #include "MovieRenderPipelineCoreModule.h"
+#include "MovieRenderPipelineCoreObjectVersion.h"
 #include "MovieSceneSpawnableAnnotation.h"
 #include "Styling/AppStyle.h"
 #include "Styling/SlateIconFinder.h"
@@ -819,24 +820,29 @@ TArray<TSharedRef<SWidget>> UMovieGraphConditionGroupQuery_Actor::GetWidgets()
 			{
 				return SNew(SActorListRow, InOwnerTable, InEntry);
 			})
-			.CustomSummaryText_Lambda([this]()
+			.CustomSummaryWidget_Lambda([this]()
 			{
-				int32 NumActors = 0;
-				int32 NumComponents = 0;
-				
-				for (const FMovieGraphActorQueryEntry& Entry : ActorsAndComponentsToMatch)
-				{
-					if (!Entry.ComponentsToMatch.IsEmpty())
-					{
-						NumComponents += Entry.ComponentsToMatch.Num();
-					}
-					else
+				return SNew(STextBlock)
+					.Font(FAppStyle::GetFontStyle(TEXT("PropertyWindow.NormalFont")))
+					.Text_Lambda([this]()
 					{
-						NumActors++;
-					}
-				}
-				
-				return FText::Format(LOCTEXT("ActorQuery_SummaryRowText", "{0} Actor(s) and {1} Component(s)"), NumActors, NumComponents);
+						int32 NumActors = 0;
+						int32 NumComponents = 0;
+						
+						for (const FMovieGraphActorQueryEntry& Entry : ActorsAndComponentsToMatch)
+						{
+							if (!Entry.ComponentsToMatch.IsEmpty())
+							{
+								NumComponents += Entry.ComponentsToMatch.Num();
+							}
+							else
+							{
+								NumActors++;
+							}
+						}
+
+						return FText::Format(LOCTEXT("ActorQuery_SummaryRowText", "{0} Actor(s) and {1} Component(s)"), NumActors, NumComponents);
+					});
 			})
 			.OnDelete_Lambda([this](const TArray<TSharedPtr<FMovieGraphActorQueryEntry>> InEntries)
 			{
@@ -1734,14 +1740,48 @@ void UMovieGraphConditionGroupQuery_ActorType::AddActorTypes(const TArray<UClass
 }
 #endif
 
+UMovieGraphConditionGroupQuery_ComponentTagName::UMovieGraphConditionGroupQuery_ComponentTagName()
+	: bOnlyMatchComponents(true)
+{
+	
+}
+
 void UMovieGraphConditionGroupQuery_ComponentTagName::Evaluate(const TArray<AActor*>& InActorsToQuery, const UWorld* InWorld, TSet<AActor*>& OutMatchingActors) const
+{
+	if (bOnlyMatchComponents)
+	{
+		UE_LOG(LogMovieRenderPipeline, Warning, TEXT("Evaluate() called on the Component Tag Name condition, but the condition only matches components. No results will be returned; EvaluateActorsAndComponents() should be used instead."));
+	}
+	
+	TSet<UActorComponent*> UnusedComponentSet;
+	EvaluateActorsAndComponents(InActorsToQuery, InWorld, OutMatchingActors, UnusedComponentSet);
+}
+
+void UMovieGraphConditionGroupQuery_ComponentTagName::EvaluateActorsAndComponents(const TArray<AActor*>& InActorsToQuery, const UWorld* InWorld, TSet<AActor*>& OutMatchingActors, TSet<UActorComponent*>& OutMatchingComponents) const
 {
 	TRACE_CPUPROFILER_EVENT_SCOPE(UMovieGraphConditionGroupQuery_ComponentTagName::Evaluate);
 	
 	// Quick early-out if "*" is used as the wildcard. Faster than doing the wildcard matching.
 	if (TagsToMatch == TEXT("*"))
 	{
-		OutMatchingActors.Append(InActorsToQuery);
+		if (bOnlyMatchComponents)
+		{
+			TInlineComponentArray<UActorComponent*> MatchingComponents;
+			
+			// This will be EXTREMELY slow if there are a lot of actors in the scene with many components
+			for (const AActor* ActorToQuery : InActorsToQuery)
+			{
+				constexpr bool bIncludeFromChildActors = true;
+				ActorToQuery->GetComponents(MatchingComponents, bIncludeFromChildActors);
+
+				OutMatchingComponents.Append(MatchingComponents);
+			}
+		}
+		else
+		{
+			OutMatchingActors.Append(InActorsToQuery);
+		}
+		
 		return;
 	}
 
@@ -1757,7 +1797,7 @@ void UMovieGraphConditionGroupQuery_ComponentTagName::Evaluate(const TArray<AAct
 		constexpr bool bIncludeFromChildActors = false;
 		Actor->GetComponents<UActorComponent*>(ActorComponents, bIncludeFromChildActors);
 		
-		for (const UActorComponent* Component : ActorComponents)
+		for (UActorComponent* Component : ActorComponents)
 		{
 			bool bMatchedTag = false;
 			
@@ -1767,7 +1807,15 @@ void UMovieGraphConditionGroupQuery_ComponentTagName::Evaluate(const TArray<AAct
 				{
 					if (ComponentTag.ToString().MatchesWildcard(TagToMatch))
 					{
-						OutMatchingActors.Add(Actor);
+						if (bOnlyMatchComponents)
+						{
+							OutMatchingComponents.Add(Component);
+						}
+						else
+						{
+							OutMatchingActors.Add(Actor);
+						}
+
 						bMatchedTag = true;
 						break;
 					}
@@ -1780,8 +1828,9 @@ void UMovieGraphConditionGroupQuery_ComponentTagName::Evaluate(const TArray<AAct
 				}
 			}
 
-			// Skip comparing the rest of the components if one component already matched
-			if (bMatchedTag)
+			// Skip comparing the rest of the components if one component already matched. If matching components instead of actors,
+			// we still have to continue processing all components.
+			if (bMatchedTag && !bOnlyMatchComponents)
 			{
 				break;
 			}
@@ -1791,6 +1840,11 @@ void UMovieGraphConditionGroupQuery_ComponentTagName::Evaluate(const TArray<AAct
 	}
 }
 
+bool UMovieGraphConditionGroupQuery_ComponentTagName::ShouldEvaluateComponents() const
+{
+	return bOnlyMatchComponents;
+}
+
 const FSlateIcon& UMovieGraphConditionGroupQuery_ComponentTagName::GetIcon() const
 {
 	static const FSlateIcon ActorTagIcon = FSlateIcon(FAppStyle::GetAppStyleSetName(), "MainFrame.OpenIssueTracker");
@@ -1809,8 +1863,11 @@ TArray<TSharedRef<SWidget>> UMovieGraphConditionGroupQuery_ComponentTagName::Get
 	TArray<TSharedRef<SWidget>> Widgets;
 
 	Widgets.Add(
-		SNew(SBox)
+		SNew(SVerticalBox)
+
+		+ SVerticalBox::Slot()
 		.HAlign(HAlign_Fill)
+		.AutoHeight()
 		.Padding(7.f, 2.f)
 		[
 			SNew(SMultiLineEditableTextBox)
@@ -1824,15 +1881,67 @@ TArray<TSharedRef<SWidget>> UMovieGraphConditionGroupQuery_ComponentTagName::Get
 			})
 			.HintText(LOCTEXT("MovieGraphComponentTagNameQueryHintText", "A component on the actor must match one or more component tags.\nWildcards allowed. Enter each tag on a separate line."))
 		]
+
+		+ SVerticalBox::Slot()
+		.HAlign(HAlign_Right)
+		.AutoHeight()
+		.Padding(8, 0, 8, 2)
+		[
+			SNew(SCheckBox)
+			.OnCheckStateChanged_Lambda([this](ECheckBoxState NewState)
+			{
+				const FScopedTransaction Transaction(LOCTEXT("ConditionGroupQueryComponentTagName_ToggleOnlyMatchComponents", "Toggle Only Match Components"));
+				Modify();
+					
+				bOnlyMatchComponents = (NewState == ECheckBoxState::Checked);
+			})
+			.IsChecked_Lambda([this]() { return bOnlyMatchComponents ? ECheckBoxState::Checked : ECheckBoxState::Unchecked; })
+			[
+				SNew(STextBlock)
+				.Font(FAppStyle::GetFontStyle(TEXT("PropertyWindow.NormalFont")))
+				.Text(LOCTEXT("ConditionGroupQueryComponentTagName_OnlyMatchComponents", "Only Match Components"))
+			]
+			.ToolTipText(LOCTEXT("ConditionGroupQueryComponentTagName_OnlyMatchComponentsTooltip", "If checked, only components of the specified type(s) will be matched. If unchecked, actors that contain components of the specified type(s) will be matched."))
+		]
 	);
 
 	return Widgets;
 }
 #endif
 
+void UMovieGraphConditionGroupQuery_ComponentTagName::Serialize(FArchive& Ar)
+{
+	Super::Serialize(Ar);
+
+	Ar.UsingCustomVersion(FMovieRenderPipelineCoreObjectVersion::GUID);
+
+	if (Ar.CustomVer(FMovieRenderPipelineCoreObjectVersion::GUID) < FMovieRenderPipelineCoreObjectVersion::OnlyMatchComponentsAdded)
+	{
+		// Default this to false if this query existed before OnlyMatchComponents was added
+		bOnlyMatchComponents = false;
+	}
+}
+
+UMovieGraphConditionGroupQuery_ComponentType::UMovieGraphConditionGroupQuery_ComponentType()
+	: bOnlyMatchComponents(true)
+{
+	
+}
+
 void UMovieGraphConditionGroupQuery_ComponentType::Evaluate(const TArray<AActor*>& InActorsToQuery, const UWorld* InWorld, TSet<AActor*>& OutMatchingActors) const
 {
-	TRACE_CPUPROFILER_EVENT_SCOPE(UMovieGraphConditionGroupQuery_ComponentType::Evaluate);
+	if (bOnlyMatchComponents)
+	{
+		UE_LOG(LogMovieRenderPipeline, Warning, TEXT("Evaluate() called on the Component Type condition, but the condition only matches components. No results will be returned; EvaluateActorsAndComponents() should be used instead."));
+	}
+	
+	TSet<UActorComponent*> UnusedComponentSet;
+	EvaluateActorsAndComponents(InActorsToQuery, InWorld, OutMatchingActors, UnusedComponentSet);
+}
+
+void UMovieGraphConditionGroupQuery_ComponentType::EvaluateActorsAndComponents(const TArray<AActor*>& InActorsToQuery, const UWorld* InWorld, TSet<AActor*>& OutMatchingActors, TSet<UActorComponent*>& OutMatchingComponents) const
+{
+	TRACE_CPUPROFILER_EVENT_SCOPE(UMovieGraphConditionGroupQuery_ComponentType::EvaluateActorsAndComponents);
 	
 	TInlineComponentArray<UActorComponent*> ActorComponents;
 
@@ -1842,11 +1951,18 @@ void UMovieGraphConditionGroupQuery_ComponentType::Evaluate(const TArray<AActor*
 		constexpr bool bIncludeFromChildActors = false;
 		Actor->GetComponents<UActorComponent*>(ActorComponents, bIncludeFromChildActors);
 		
-		for (const UActorComponent* Component : ActorComponents)
+		for (UActorComponent* Component : ActorComponents)
 		{
 			if (ComponentTypes.Contains(Component->GetClass()))
 			{
-				OutMatchingActors.Add(Actor);
+				if (bOnlyMatchComponents)
+				{
+					OutMatchingComponents.Add(Component);
+				}
+				else
+				{
+					OutMatchingActors.Add(Actor);
+				}
 			}
 		}
 
@@ -1854,6 +1970,11 @@ void UMovieGraphConditionGroupQuery_ComponentType::Evaluate(const TArray<AActor*
 	}
 }
 
+bool UMovieGraphConditionGroupQuery_ComponentType::ShouldEvaluateComponents() const
+{
+	return bOnlyMatchComponents;
+}
+
 const FSlateIcon& UMovieGraphConditionGroupQuery_ComponentType::GetIcon() const
 {
 	static const FSlateIcon ActorTagIcon = FSlateIcon(FAppStyle::GetAppStyleSetName(), "ClassIcon.ActorComponent");
@@ -1874,8 +1995,6 @@ TArray<TSharedRef<SWidget>> UMovieGraphConditionGroupQuery_ComponentType::GetWid
 	Widgets.Add(
 		SAssignNew(ComponentTypesList, SMovieGraphSimpleList<TObjectPtr<UClass>>)
 			.DataSource(&ComponentTypes)
-			.DataType(FText::FromString("Component Type"))
-			.DataTypePlural(FText::FromString("Component Types"))
 			.OnGetRowText_Static(&GetRowText)
 			.OnGetRowIcon_Static(&GetRowIcon)
 			.OnDelete_Lambda([this](const TArray<UClass*> InComponentTypes)
@@ -1890,6 +2009,43 @@ TArray<TSharedRef<SWidget>> UMovieGraphConditionGroupQuery_ComponentType::GetWid
 				
 				ComponentTypesList->Refresh();
 			})
+			.CustomSummaryWidget_Lambda([this]()
+			{
+				return SNew(SHorizontalBox)
+					+SHorizontalBox::Slot()
+					.HAlign(HAlign_Left)
+					.VAlign(VAlign_Center)
+					.AutoWidth()
+					[
+						SNew(STextBlock)
+						.Font(FAppStyle::GetFontStyle(TEXT("PropertyWindow.NormalFont")))
+						.Text_Lambda([this]()
+						{
+							return FText::Format(LOCTEXT("ConditionGroupQueryComponentType_ComponentCount", "{0} Component(s)"), FText::FromString(LexToString(ComponentTypes.Num())));
+						})
+					]
+
+					+SHorizontalBox::Slot()
+					.HAlign(HAlign_Right)
+					.VAlign(VAlign_Center)
+					[
+						SNew(SCheckBox)
+						.OnCheckStateChanged_Lambda([this](ECheckBoxState NewState)
+						{
+							const FScopedTransaction Transaction(LOCTEXT("ConditionGroupQueryComponentType_ToggleOnlyMatchComponents", "Toggle Only Match Components"));
+							Modify();
+							
+							bOnlyMatchComponents = (NewState == ECheckBoxState::Checked);
+						})
+						.IsChecked_Lambda([this]() { return bOnlyMatchComponents ? ECheckBoxState::Checked : ECheckBoxState::Unchecked; })
+						[
+							SNew(STextBlock)
+							.Font(FAppStyle::GetFontStyle(TEXT("PropertyWindow.NormalFont")))
+							.Text(LOCTEXT("ConditionGroupQueryComponentType_OnlyMatchComponents", "Only Match Components"))
+						]
+						.ToolTipText(LOCTEXT("ConditionGroupQueryComponentType_OnlyMatchComponentsTooltip", "If checked, only components with the specified tag(s) will be matched. If unchecked, actors that contain components with the specified tags(s) will be matched."))
+					];
+			})
 	);
 
 	return Widgets;
@@ -1961,6 +2117,19 @@ FText UMovieGraphConditionGroupQuery_ComponentType::GetRowText(TObjectPtr<UClass
 }
 #endif	// WITH_EDITOR
 
+void UMovieGraphConditionGroupQuery_ComponentType::Serialize(FArchive& Ar)
+{
+	Super::Serialize(Ar);
+
+	Ar.UsingCustomVersion(FMovieRenderPipelineCoreObjectVersion::GUID);
+
+	if (Ar.CustomVer(FMovieRenderPipelineCoreObjectVersion::GUID) < FMovieRenderPipelineCoreObjectVersion::OnlyMatchComponentsAdded)
+	{
+		// Default this to false if this query existed before OnlyMatchComponents was added
+		bOnlyMatchComponents = false;
+	}
+}
+
 void UMovieGraphConditionGroupQuery_EditorFolder::Evaluate(const TArray<AActor*>& InActorsToQuery, const UWorld* InWorld, TSet<AActor*>& OutMatchingActors) const
 {
 #if WITH_EDITOR
@@ -2848,16 +3017,130 @@ FMovieGraphEvaluationResult FMovieGraphEvaluationResult::Union(const FMovieGraph
 
 FMovieGraphEvaluationResult FMovieGraphEvaluationResult::Intersect(const FMovieGraphEvaluationResult& InOther) const
 {
-	return FMovieGraphEvaluationResult(
-		MatchingActors.Intersect(InOther.MatchingActors),
-		MatchingComponents.Intersect(InOther.MatchingComponents));
+	FMovieGraphEvaluationResult Result;
+
+	// Actor AND Actor is simple: the actors just need to be the same
+	for (const TObjectPtr<AActor>& MatchingActor : MatchingActors)
+	{
+		if (MatchingActor && InOther.MatchingActors.Contains(MatchingActor))
+		{
+			Result.MatchingActors.Add(MatchingActor);
+		}
+	}
+
+	// Component AND Component is also simple: the components just need to be the same
+	for (const TObjectPtr<UActorComponent>& MatchingComponent : MatchingComponents)
+	{
+		if (MatchingComponent && InOther.MatchingComponents.Contains(MatchingComponent))
+		{
+			Result.MatchingComponents.Add(MatchingComponent);
+		}
+	}
+
+	// Actor AND Component is more complex: there's a match if the component is contained *within* an actor
+	auto FindMatchingComponents = [&Result](const TSet<TObjectPtr<UActorComponent>>& InComponents, const TSet<TObjectPtr<AActor>>& InActors)
+	{
+		for (const TObjectPtr<UActorComponent>& MatchingComponent : InComponents)
+		{
+			for (const TObjectPtr<AActor>& OtherActor : InActors)
+			{
+				if (MatchingComponent && OtherActor->GetComponents().Contains(MatchingComponent.Get()))
+				{
+					Result.MatchingComponents.Add(MatchingComponent);
+				}
+			}
+		}
+	};
+	FindMatchingComponents(MatchingComponents, InOther.MatchingActors);
+	FindMatchingComponents(InOther.MatchingComponents, MatchingActors);
+	
+	return Result;
 }
 
 FMovieGraphEvaluationResult FMovieGraphEvaluationResult::Difference(const FMovieGraphEvaluationResult& InOther) const
 {
-	return FMovieGraphEvaluationResult(
-		MatchingActors.Difference(InOther.MatchingActors),
-		MatchingComponents.Difference(InOther.MatchingComponents));
+	FMovieGraphEvaluationResult Result;
+
+	// Actor MINUS Actor is simple: the actor cannot be in the other result set
+	for (const TObjectPtr<AActor>& MatchingActor : MatchingActors)
+	{
+		if (MatchingActor && !InOther.MatchingActors.Contains(MatchingActor))
+		{
+			Result.MatchingActors.Add(MatchingActor);
+		}
+	}
+
+	// Component MINUS Component is also simple: the component cannot be in the other result set
+	for (const TObjectPtr<UActorComponent>& MatchingComponent : MatchingComponents)
+	{
+		if (MatchingComponent && !InOther.MatchingComponents.Contains(MatchingComponent))
+		{
+			Result.MatchingComponents.Add(MatchingComponent);
+		}
+	}
+
+	// Actor MINUS Component is more complex: an actor is "deconstructed" into all of its components, and component(s) are subtracted out
+	TArray<TObjectPtr<UActorComponent>, TInlineAllocator<24>> ActorComponents;
+	TArray<AActor*, TInlineAllocator<10>> DeconstructedActorsToRemove;
+	for (const TObjectPtr<AActor>& MatchingActor : Result.MatchingActors)
+	{
+		constexpr bool bIncludeFromChildActors = true;
+		MatchingActor->GetComponents(ActorComponents, bIncludeFromChildActors);
+		
+		for (const TObjectPtr<UActorComponent>& OtherComponent : InOther.MatchingComponents)
+		{
+			// If a matching actor's components contains one of the components to subtract, deconstruct the actor into its components and match all of
+			// its components except the component that is being subtracted out.
+			if (ActorComponents.Contains(OtherComponent.Get()))
+			{
+				// Deconstruct the actor if needed. We can't remove the actor from MatchingActors yet because we're iterating it, but we'll remove it
+				// after the loop finishes.
+				TObjectPtr<AActor> ComponentActor = OtherComponent->GetOwner();
+				if (!DeconstructedActorsToRemove.Contains(ComponentActor) && Result.MatchingActors.Contains(ComponentActor))
+				{
+					DeconstructedActorsToRemove.Add(ComponentActor.Get());
+					Result.MatchingComponents.Append(ActorComponents);
+				}
+				
+				// Remove the component to be subtracted out
+				Result.MatchingComponents.Remove(OtherComponent);
+			}
+
+			// Reuse the array
+			ActorComponents.Reset();
+		}
+	}
+
+	// Remove all of the deconstructed actors found above
+	for (AActor* ActorToRemove : DeconstructedActorsToRemove)
+	{
+		Result.MatchingActors.Remove(ActorToRemove);
+	}
+
+	// Component MINUS Actor is a little more straightforward: if a component is contained in the actor, remove the component
+	TArray<TObjectPtr<UActorComponent>, TInlineAllocator<24>> ComponentsToRemove;
+	for (const TObjectPtr<AActor>& OtherActor : InOther.MatchingActors)
+	{
+		constexpr bool bIncludeFromChildActors = true;
+		OtherActor->GetComponents(ActorComponents, bIncludeFromChildActors);
+		
+		for (const TObjectPtr<UActorComponent>& MatchingComponent : Result.MatchingComponents)
+		{
+			if (ActorComponents.Contains(MatchingComponent))
+			{
+				ComponentsToRemove.Add(MatchingComponent);
+				Result.MatchingActors.Remove(OtherActor);
+			}
+		}
+	}
+
+	// Remove all of the components found above
+	for (UActorComponent* ComponentToRemove : ComponentsToRemove)
+	{
+		Result.MatchingComponents.Remove(ComponentToRemove);
+	}
+
+	return Result;
 }
 
 template<typename T>
@@ -3023,7 +3306,20 @@ bool UMovieGraphConditionGroup::RemoveQuery(UMovieGraphConditionGroupQueryBase*
 	Modify();
 #endif
 	
-	return Queries.RemoveSingle(InQuery) == 1;
+	const bool bRemoved = Queries.RemoveSingle(InQuery) == 1;
+
+	// Always make sure the first query is an "add"
+	if (!Queries.IsEmpty())
+	{
+		const TObjectPtr<UMovieGraphConditionGroupQueryBase> FirstQuery = Queries[0];
+		if (FirstQuery->GetOperationType() != EMovieGraphConditionGroupQueryOpType::Add)
+		{
+			FirstQuery->Modify();
+			FirstQuery->SetOperationType(EMovieGraphConditionGroupQueryOpType::Add);
+		}
+	}
+	
+	return bRemoved;
 }
 
 UMovieGraphConditionGroupQueryBase* UMovieGraphConditionGroup::DuplicateQuery(const int32 QueryIndex)
@@ -3172,7 +3468,20 @@ bool UMovieGraphCollection::RemoveConditionGroup(UMovieGraphConditionGroup* InCo
 	Modify();
 #endif
 	
-	return ConditionGroups.RemoveSingle(InConditionGroup) == 1;
+	const bool bRemoved = ConditionGroups.RemoveSingle(InConditionGroup) == 1;
+	
+	// Always make sure the first condition group is an "add"
+	if (!ConditionGroups.IsEmpty())
+	{
+		const TObjectPtr<UMovieGraphConditionGroup> FirstConditionGroup = ConditionGroups[0];
+		if (FirstConditionGroup->GetOperationType() != EMovieGraphConditionGroupOpType::Add)
+		{
+			FirstConditionGroup->Modify();
+			FirstConditionGroup->SetOperationType(EMovieGraphConditionGroupOpType::Add);
+		}
+	}
+	
+	return bRemoved;
 }
 
 bool UMovieGraphCollection::MoveConditionGroupToIndex(UMovieGraphConditionGroup* InConditionGroup, const int32 NewIndex)
diff --git a/Engine/Plugins/MovieScene/MovieRenderPipeline/Source/MovieRenderPipelineCore/Public/Graph/MovieGraphRenderLayerSubsystem.h b/Engine/Plugins/MovieScene/MovieRenderPipeline/Source/MovieRenderPipelineCore/Public/Graph/MovieGraphRenderLayerSubsystem.h
--- a/Engine/Plugins/MovieScene/MovieRenderPipeline/Source/MovieRenderPipelineCore/Public/Graph/MovieGraphRenderLayerSubsystem.h
+++ b/Engine/Plugins/MovieScene/MovieRenderPipeline/Source/MovieRenderPipelineCore/Public/Graph/MovieGraphRenderLayerSubsystem.h
@@ -395,7 +395,11 @@ class MOVIERENDERPIPELINECORE_API UMovieGraphConditionGroupQuery_ComponentTagNam
 	GENERATED_BODY()
 
 public:
+	UMovieGraphConditionGroupQuery_ComponentTagName();
+	
 	virtual void Evaluate(const TArray<AActor*>& InActorsToQuery, const UWorld* InWorld, TSet<AActor*>& OutMatchingActors) const override;
+	virtual void EvaluateActorsAndComponents(const TArray<AActor*>& InActorsToQuery, const UWorld* InWorld, TSet<AActor*>& OutMatchingActors, TSet<UActorComponent*>& OutMatchingComponents) const override;
+	virtual bool ShouldEvaluateComponents() const override;
 	virtual const FSlateIcon& GetIcon() const override;
 	virtual const FText& GetDisplayName() const override;
 
@@ -403,6 +407,10 @@ public:
 	virtual TArray<TSharedRef<SWidget>> GetWidgets() override;
 #endif
 
+	//~ Begin UObject Interface
+	virtual void Serialize(FArchive& Ar) override;
+	//~ End UObject Interface
+
 public:
 	/**
 	 * A component on the actor must have one or more of the specified tags to be a match. Not case sensitive. One tag per line. Wildcards ("?" and "*") are supported but not required.
@@ -419,6 +427,10 @@ public:
 	 */
 	UPROPERTY(EditAnywhere, Category="General")
 	FString TagsToMatch;
+
+	/** Whether this should match components or actors. If false, any components that match will instead match their parent actor. */
+	UPROPERTY(EditAnywhere, Category="General")
+	bool bOnlyMatchComponents;
 };
 
 /** Query type which filters actors via the components contained in them. */
@@ -428,7 +440,11 @@ class MOVIERENDERPIPELINECORE_API UMovieGraphConditionGroupQuery_ComponentType f
 	GENERATED_BODY()
 
 public:
+	UMovieGraphConditionGroupQuery_ComponentType();
+	
 	virtual void Evaluate(const TArray<AActor*>& InActorsToQuery, const UWorld* InWorld, TSet<AActor*>& OutMatchingActors) const override;
+	virtual void EvaluateActorsAndComponents(const TArray<AActor*>& InActorsToQuery, const UWorld* InWorld, TSet<AActor*>& OutMatchingActors, TSet<UActorComponent*>& OutMatchingComponents) const override;
+	virtual bool ShouldEvaluateComponents() const override;
 	virtual const FSlateIcon& GetIcon() const override;
 	virtual const FText& GetDisplayName() const override;
 
@@ -438,6 +454,10 @@ public:
 	virtual TSharedRef<SWidget> GetAddMenuContents(const FMovieGraphConditionGroupQueryContentsChanged& OnAddFinished) override;
 #endif
 
+	//~ Begin UObject Interface
+	virtual void Serialize(FArchive& Ar) override;
+	//~ End UObject Interface
+
 private:
 #if WITH_EDITOR
 	static const FSlateBrush* GetRowIcon(TObjectPtr<UClass> InComponentType);
@@ -454,6 +474,10 @@ public:
 	/** The actor must have one or more of the component type(s) in order to be a match. */
 	UPROPERTY(EditAnywhere, Category="General")
 	TArray<TObjectPtr<UClass>> ComponentTypes;
+
+	/** Whether this should match components or actors. If false, any components that match will instead match their parent actor. */
+	UPROPERTY(EditAnywhere, Category="General")
+	bool bOnlyMatchComponents;
 };
 
 /** Query type which filters actors via the editor folder that they're contained in. */
diff --git a/Engine/Plugins/MovieScene/MovieRenderPipeline/Source/MovieRenderPipelineCore/Public/Graph/MovieGraphSharedWidgets.h b/Engine/Plugins/MovieScene/MovieRenderPipeline/Source/MovieRenderPipelineCore/Public/Graph/MovieGraphSharedWidgets.h
--- a/Engine/Plugins/MovieScene/MovieRenderPipeline/Source/MovieRenderPipelineCore/Public/Graph/MovieGraphSharedWidgets.h
+++ b/Engine/Plugins/MovieScene/MovieRenderPipeline/Source/MovieRenderPipelineCore/Public/Graph/MovieGraphSharedWidgets.h
@@ -61,10 +61,10 @@ public:
 		SLATE_ATTRIBUTE(FText, DataTypePlural)
 
 		/**
-		 * Gets the text that should be shown in the summary row. By default, "DataType" and "DataTypePlural" will be used to auto-generate summary
-		 * row text. If this attribute is bound, those will be ignored, and this will be used instead.
+		 * Gets the widget that should be shown in the summary row. By default, "DataType" and "DataTypePlural" will be used to auto-generate the summary
+		 * row widget. If this attribute is bound, those will be ignored, and this will be used instead.
 		 */
-		SLATE_ATTRIBUTE(FText, CustomSummaryText);
+		SLATE_ATTRIBUTE(TSharedPtr<SWidget>, CustomSummaryWidget);
 
 		/** The selection mode the list should be put into. Defaults to Multi. */
 		SLATE_ATTRIBUTE(ESelectionMode::Type, SelectionMode)
@@ -103,7 +103,7 @@ public:
 		OnRefreshDataSourceRequested = InArgs._OnRefreshDataSourceRequested;
 		DataType = InArgs._DataType.Get();
 		DataTypePlural = InArgs._DataTypePlural.Get();
-		CustomSummaryText = InArgs._CustomSummaryText;
+		CustomSummaryWidget = InArgs._CustomSummaryWidget;
 		SelectionMode = InArgs._SelectionMode.Get();
 		OnGetRowIcon = InArgs._OnGetRowIcon;
 		OnGetRowText = InArgs._OnGetRowText;
@@ -125,6 +125,17 @@ public:
 				return OnDelete.IsBound();
 			}));
 
+		// The contents of the summary row is usually auto-generated to show the number of items in the data source, but it can optionally be
+		// completely custom if needed
+		const TSharedRef<SWidget> SummaryRowContents = CustomSummaryWidget.IsSet()
+			? CustomSummaryWidget.Get().ToSharedRef()
+			: SNew(STextBlock)
+				.Font(FAppStyle::GetFontStyle(TEXT("PropertyWindow.NormalFont")))
+				.Text_Lambda([this]()
+				{
+					return FText::FromString(FString::Printf(TEXT("%i %s"), DataSource->Num(), DataSource->Num() == 1 ? *DataType.ToString() : *DataTypePlural.ToString()));
+				});
+
 		ChildSlot
 		.HAlign(HAlign_Fill)
 		.VAlign(VAlign_Top)
@@ -152,19 +163,9 @@ public:
 				.BorderImage(FAppStyle::Get().GetBrush("Brushes.Header"))
 				.VAlign(VAlign_Center)
 				.HAlign(HAlign_Fill)
-				.Padding(FMargin(14, 4))
+				.Padding(FMargin(14, 4, 7, 4))
 				[
-					SNew(STextBlock)
-					.Font(FAppStyle::GetFontStyle(TEXT("PropertyWindow.NormalFont")))
-					.Text_Lambda([this]()
-					{
-						if (CustomSummaryText.IsSet())
-						{
-							return CustomSummaryText.Get();
-						}
-						
-						return FText::FromString(FString::Printf(TEXT("%i %s"), DataSource->Num(), DataSource->Num() == 1 ? *DataType.ToString() : *DataTypePlural.ToString()));
-					})
+					SummaryRowContents
 				]
 			]
 		];
@@ -351,7 +352,7 @@ private:
 	TSharedPtr<SHeaderRow> HeaderRow;
 	FText DataType;
 	FText DataTypePlural;
-	TAttribute<FText> CustomSummaryText;
+	TAttribute<TSharedPtr<SWidget>> CustomSummaryWidget;
 	TArray<ListType>* DataSource = nullptr;
 	ESelectionMode::Type SelectionMode = ESelectionMode::Multi;
 };
diff --git a/Engine/Plugins/MovieScene/MovieRenderPipeline/Source/MovieRenderPipelineCore/Public/MovieRenderPipelineCoreObjectVersion.cpp b/Engine/Plugins/MovieScene/MovieRenderPipeline/Source/MovieRenderPipelineCore/Public/MovieRenderPipelineCoreObjectVersion.cpp
new file mode 100644
--- /dev/null
+++ b/Engine/Plugins/MovieScene/MovieRenderPipeline/Source/MovieRenderPipelineCore/Public/MovieRenderPipelineCoreObjectVersion.cpp
@@ -0,0 +1,8 @@
+﻿// Copyright Epic Games, Inc. All Rights Reserved.
+
+#include "MovieRenderPipelineCoreObjectVersion.h"
+
+#include "Serialization/CustomVersion.h"
+
+const FGuid FMovieRenderPipelineCoreObjectVersion::GUID = FGuid(0x81A1B2B3, 0xAF127711, 0xE211B0B2, 0x1D235409);
+FCustomVersionRegistration GRegisterMovieRenderPipelineCoreObjectVersion(FMovieRenderPipelineCoreObjectVersion::GUID, FMovieRenderPipelineCoreObjectVersion::LatestVersion, TEXT("MovieRenderPipelineCoreObjectVersion"));
\ No newline at end of file
diff --git a/Engine/Plugins/MovieScene/MovieRenderPipeline/Source/MovieRenderPipelineCore/Public/MovieRenderPipelineCoreObjectVersion.h b/Engine/Plugins/MovieScene/MovieRenderPipeline/Source/MovieRenderPipelineCore/Public/MovieRenderPipelineCoreObjectVersion.h
new file mode 100644
--- /dev/null
+++ b/Engine/Plugins/MovieScene/MovieRenderPipeline/Source/MovieRenderPipelineCore/Public/MovieRenderPipelineCoreObjectVersion.h
@@ -0,0 +1,26 @@
+﻿// Copyright Epic Games, Inc. All Rights Reserved.
+
+#pragma once
+
+#include "CoreTypes.h"
+#include "Misc/Guid.h"
+
+/** Custom serialization version for changes made to Movie Render Pipeline Core objects. */
+struct MOVIERENDERPIPELINECORE_API FMovieRenderPipelineCoreObjectVersion
+{
+	enum Type : int32
+	{
+		PreVersioning = 0,
+
+		/** Added bOnlyMatchComponents to some conditions within graph collections, defaulting to true. */
+		OnlyMatchComponentsAdded,
+
+		// --- Add new versions above this line ---
+		VersionPlusOne,
+		LatestVersion = VersionPlusOne - 1
+	};
+
+	static const FGuid GUID;
+
+	FMovieRenderPipelineCoreObjectVersion() = delete;
+};
\ No newline at end of file
