MovieGraph: Fixed an issue where the accumulators did not take spatial samples into account. This resulted in spatial samples for the first temporal sample continually resetting the Accumulator buffer resulting in reduced spatial sample quality on the first temporal sample of each frame.

Ran basic tests and verified deferred, pt, and objectid passes work when using 1 or more spatial samples.

diff --git a/Engine/Plugins/MovieScene/MovieRenderPipeline/Source/MovieRenderPipelineRenderPasses/Private/Graph/Renderers/MovieGraphImagePassBase.cpp b/Engine/Plugins/MovieScene/MovieRenderPipeline/Source/MovieRenderPipelineRenderPasses/Private/Graph/Renderers/MovieGraphImagePassBase.cpp
--- a/Engine/Plugins/MovieScene/MovieRenderPipeline/Source/MovieRenderPipelineRenderPasses/Private/Graph/Renderers/MovieGraphImagePassBase.cpp
+++ b/Engine/Plugins/MovieScene/MovieRenderPipeline/Source/MovieRenderPipelineRenderPasses/Private/Graph/Renderers/MovieGraphImagePassBase.cpp
@@ -571,8 +571,11 @@ TSharedRef<::MoviePipeline::IMoviePipelineAccumulationArgs> FMovieGraphImagePass
 	AccumulationArgs->OutputMerger = InGraphRenderer->GetOwningGraph()->GetOutputMerger();
 	AccumulationArgs->ImageAccumulator = StaticCastSharedPtr<FImageOverlappedAccumulator>(AccumulatorInstance->Accumulator);
 	AccumulationArgs->AccumulatorInstance = SampleAccumulatorPool->GetAccumulatorInstance_GameThread<FImageOverlappedAccumulator>(InSampleState.TraversalContext.Time.OutputFrameNumber, InSampleState.TraversalContext.RenderDataIdentifier);
-	AccumulationArgs->bIsFirstSample = InSampleState.TraversalContext.Time.bIsFirstTemporalSampleForFrame;
-	AccumulationArgs->bIsLastSample = InSampleState.TraversalContext.Time.bIsLastTemporalSampleForFrame;
+	
+	const bool bFirstSpatialSample = InSampleState.TraversalContext.Time.SpatialSampleIndex == 0;
+	const bool bLastSpatialSample = InSampleState.TraversalContext.Time.SpatialSampleIndex == InSampleState.TraversalContext.Time.SpatialSampleCount - 1;
+	AccumulationArgs->bIsFirstSample = InSampleState.TraversalContext.Time.bIsFirstTemporalSampleForFrame && bFirstSpatialSample;
+	AccumulationArgs->bIsLastSample = InSampleState.TraversalContext.Time.bIsLastTemporalSampleForFrame && bLastSpatialSample;
 
 	return AccumulationArgs;
 }
