MoviePipeline: Fixed issue that could cause multiple render property modifications to be ignored when component-based modification is being used.

Regression from when component-based modification was introduced.

#rnx

diff --git a/Engine/Plugins/MovieScene/MovieRenderPipeline/Source/MovieRenderPipelineCore/Private/Graph/MovieGraphRenderLayerSubsystem.cpp b/Engine/Plugins/MovieScene/MovieRenderPipeline/Source/MovieRenderPipelineCore/Private/Graph/MovieGraphRenderLayerSubsystem.cpp
--- a/Engine/Plugins/MovieScene/MovieRenderPipeline/Source/MovieRenderPipelineCore/Private/Graph/MovieGraphRenderLayerSubsystem.cpp
+++ b/Engine/Plugins/MovieScene/MovieRenderPipeline/Source/MovieRenderPipelineCore/Private/Graph/MovieGraphRenderLayerSubsystem.cpp
@@ -370,9 +370,6 @@ void UMovieGraphRenderPropertyModifier::SetActorAndComponentState(const TArray<F
 			if (bOverride_bIsHidden)
 			{
 				AsPrimitiveComponent->SetHiddenInGame(InNewComponentState.bIsHidden);
-#if WITH_EDITOR
-				AsPrimitiveComponent->SetIsTemporarilyHiddenInEditor(InNewComponentState.bIsHidden);
-#endif
 			}
 		}
 		// Volumetrics are special cases as they don't inherit from UPrimitiveComponent, and don't support all of the flags.
