MovieGraph: Fixed a crash when using Additional Post Process Materials and constrain aspect ratio or screen percentage.

When images come back at unexpected resolutions they cause a crash in the accumulator. This isn't a proper fix (as it just resizes to fit the expected resolution, which causes stretch/squash), but this makes MRG match the behavior of MRQ.

Users can work around this issue by disabling Constrain Aspect Ratio on their cameras, and using the full output resolution instead of Screen Percentage.

Tested combinations of one or more spatial/temporal sub samples, with matching and diverging aspect ratios, while both using and not using additional post process materials.

diff --git a/Engine/Plugins/MovieScene/MovieRenderPipeline/Source/MovieRenderPipelineRenderPasses/Private/Graph/Renderers/MovieGraphImagePassBase.cpp b/Engine/Plugins/MovieScene/MovieRenderPipeline/Source/MovieRenderPipelineRenderPasses/Private/Graph/Renderers/MovieGraphImagePassBase.cpp
--- a/Engine/Plugins/MovieScene/MovieRenderPipeline/Source/MovieRenderPipelineRenderPasses/Private/Graph/Renderers/MovieGraphImagePassBase.cpp
+++ b/Engine/Plugins/MovieScene/MovieRenderPipeline/Source/MovieRenderPipelineRenderPasses/Private/Graph/Renderers/MovieGraphImagePassBase.cpp
@@ -15,6 +15,8 @@
 #include "SceneViewExtensionContext.h"
 #include "SceneViewExtension.h"
 #include "Engine/Engine.h"
+#include "ImageUtils.h"
+
 
 namespace UE::MovieGraph::Rendering
 {
@@ -636,6 +638,88 @@ void AccumulateSample_TaskThread(TUniquePtr<FImagePixelData>&& InPixelData, cons
 	// Accumulate the new sample to our target
 	{
 		TRACE_CPUPROFILER_EVENT_SCOPE(MoviePipeline_AccumulatePixelData);
+
+		// TODO: This needs to be re-updated to support high resolution tiling,
+		// and it's mostly copy/pasted from MRQ and could be refactored.
+
+		// Some samples can come back at a different size than expected (post process materials) which
+		// creates numerous issues with the accumulators. To work around this issue for now, we will resize
+		// the image to the expected resolution. 
+		FIntPoint RawSize = SamplePixelData->GetSize();
+		const bool bCorrectSize = RawSize == SampleStatePayload->BackbufferResolution;
+
+		if (!bCorrectSize)
+		{
+			const double ResizeConvertBeginTime = FPlatformTime::Seconds();
+
+			// Convert the incoming data to full floats (the accumulator would do this later normally anyways)
+			TArray64<FLinearColor> FullSizeData;
+			FullSizeData.AddUninitialized(RawSize.X * RawSize.Y);
+
+			if (SamplePixelData->GetType() == EImagePixelType::Float32)
+			{
+				const void* RawDataPtr = nullptr;
+				int64 RawDataSize;
+
+				if (SamplePixelData->GetRawData(RawDataPtr, RawDataSize) == true)
+				{
+					FMemory::Memcpy(FullSizeData.GetData(), RawDataPtr, RawDataSize);
+				}
+				else
+				{
+					UE_LOG(LogMovieRenderPipelineIO, Error, TEXT("Failed to retrieve raw data from image data for writing. Bailing."));
+					return;
+				}
+			}
+			else if (SamplePixelData->GetType() == EImagePixelType::Float16)
+			{
+				const void* RawDataPtr = nullptr;
+				int64 RawDataSize;
+
+				if (SamplePixelData->GetRawData(RawDataPtr, RawDataSize) == true)
+				{
+					const FFloat16Color* DataAsColor = reinterpret_cast<const FFloat16Color*>(RawDataPtr);
+					for (int64 Index = 0; Index < RawSize.X * RawSize.Y; Index++)
+					{
+						FullSizeData[Index] = FLinearColor(DataAsColor[Index]);
+					}
+				}
+				else
+				{
+					UE_LOG(LogMovieRenderPipelineIO, Error, TEXT("Failed to retrieve raw data from image data for writing. Bailing."));
+					return;
+				}
+				// TODO: Produce a warning when this happens as it causes stretched Additional PPMs and the user won't know why.
+			}
+			else
+			{
+				check(0);
+			}
+			const double ResizeConvertEndTime = FPlatformTime::Seconds();
+
+			// Now we can resize to our target size.
+			FIntPoint TargetSize = SampleStatePayload->BackbufferResolution; // TODO: High Resolution TIling
+
+			TArray64<FLinearColor> NewPixelData;
+			NewPixelData.SetNumUninitialized(TargetSize.X * TargetSize.Y);
+
+			FImageUtils::ImageResize(RawSize.X, RawSize.Y, MakeArrayView<FLinearColor>(FullSizeData.GetData(), FullSizeData.Num()), TargetSize.X, TargetSize.Y, MakeArrayView<FLinearColor>(NewPixelData.GetData(), NewPixelData.Num()));
+
+			const float ElapsedConvertMs = float((ResizeConvertEndTime - ResizeConvertBeginTime) * 1000.0f);
+			const float ElapsedResizeMs = float((FPlatformTime::Seconds() - ResizeConvertEndTime) * 1000.0f);
+
+			UE_LOG(LogMovieRenderPipeline, VeryVerbose, TEXT("Resize Convert Time: %8.2fms Resize Time: %8.2fms"), ElapsedConvertMs, ElapsedResizeMs);
+
+			SamplePixelData = MakeUnique<TImagePixelData<FLinearColor>>(FIntPoint(TargetSize.X, TargetSize.Y), MoveTemp(NewPixelData), SampleStatePayload);
+
+			// Update the raw size to match our new size.
+			RawSize = SamplePixelData->GetSize();
+
+			// TODO: We'd like to push this into a more central warning system so that we don't have to spam it every frame.
+			UE_LOG(LogMovieRenderPipeline, Warning, TEXT(
+				"'Additional Post Process Materials' do not support 'Aspect Ratio Constraint' or 'Screen Percentage'. The resulting data has been resized to match which may result in a stretched/squashed image."
+			));
+		}
 				
 		const FIntPoint TileSize = SampleStatePayload->BackbufferResolution;
 		const FIntPoint OverlappedPad = SampleStatePayload->OverlappedPad;
