MovieRenderPipeline: Fixed issue where having PPMs or non-composited burn-in would crash the editor when rendering a shot in MRQ or MRG

diff --git a/Engine/Plugins/MovieScene/MovieRenderPipeline/Source/MovieRenderPipelineCore/Private/Graph/Nodes/MovieGraphVideoOutputNode.cpp b/Engine/Plugins/MovieScene/MovieRenderPipeline/Source/MovieRenderPipelineCore/Private/Graph/Nodes/MovieGraphVideoOutputNode.cpp
--- a/Engine/Plugins/MovieScene/MovieRenderPipeline/Source/MovieRenderPipelineCore/Private/Graph/Nodes/MovieGraphVideoOutputNode.cpp
+++ b/Engine/Plugins/MovieScene/MovieRenderPipeline/Source/MovieRenderPipelineCore/Private/Graph/Nodes/MovieGraphVideoOutputNode.cpp
@@ -85,11 +85,40 @@ void UMovieGraphVideoOutputNode::OnReceiveImageDataImpl(UMovieGraphPipeline* InP
 		FImagePixelData* RawRenderPassData = RenderPassData.Value.Get();
 		const UE::MovieGraph::FMovieGraphSampleState* Payload = RawRenderPassData->GetPayload<UE::MovieGraph::FMovieGraphSampleState>();
 
-		FAsyncCropImage CropImagePreprocessor(Payload->CropRectangle);
-		if (ShouldCropOverscanImpl())
+		// Pixel data pointer for the cropped result to ensure it stays in memory while the image write is processed
+		TUniquePtr<FImagePixelData> CroppedPixelData;
+		const bool bIsCropRectValid = !Payload->CropRectangle.IsEmpty();
+		const bool bCanCropResolution = RenderPassData.Value->GetSize() == Payload->BackbufferResolution;
+		if (ShouldCropOverscanImpl() && bIsCropRectValid && bCanCropResolution)
 		{
-			CropImagePreprocessor(RawRenderPassData);
-			RawRenderPassData = CropImagePreprocessor.OutCroppedImage.Get();
+			switch (RenderPassData.Value->GetType())
+			{
+			case EImagePixelType::Color:
+				{
+					TAsyncCropImage<FColor> CropImagePreprocessor(Payload->CropRectangle);
+					CropImagePreprocessor(RawRenderPassData);
+					CroppedPixelData = MoveTemp(CropImagePreprocessor.OutCroppedImage);
+				}
+				break;
+
+			case EImagePixelType::Float16:
+				{
+					TAsyncCropImage<FFloat16Color> CropImagePreprocessor(Payload->CropRectangle);
+					CropImagePreprocessor(RawRenderPassData);
+					CroppedPixelData = MoveTemp(CropImagePreprocessor.OutCroppedImage);
+				}
+				break;
+
+			case EImagePixelType::Float32:
+				{
+					TAsyncCropImage<FLinearColor> CropImagePreprocessor(Payload->CropRectangle);
+					CropImagePreprocessor(RawRenderPassData);
+					CroppedPixelData = MoveTemp(CropImagePreprocessor.OutCroppedImage);
+				}
+				break;
+			}
+			
+			RawRenderPassData = CroppedPixelData.Get();
 		}
 		
 		// Write this frame's beauty pass, and pass along other passes so the encoder can composite them
@@ -291,7 +320,10 @@ UMovieGraphVideoOutputNode::FMovieGraphCodecWriterWithPromise* UMovieGraphVideoO
 	if (!OutputWriter)
 	{
 		FIntPoint OutputResolution = InRenderPassData.Value->GetSize();
-		if (ShouldCropOverscanImpl())
+		
+		const bool bIsCropRectValid = !Payload->CropRectangle.IsEmpty();
+		const bool bCanCropResolution = InRenderPassData.Value->GetSize() == Payload->BackbufferResolution;
+		if (ShouldCropOverscanImpl() && bIsCropRectValid && bCanCropResolution)
 		{
 			OutputResolution = Payload->CropRectangle.Size();
 		}
diff --git a/Engine/Plugins/MovieScene/MovieRenderPipeline/Source/MovieRenderPipelineCore/Private/MoviePipelineVideoOutputBase.cpp b/Engine/Plugins/MovieScene/MovieRenderPipeline/Source/MovieRenderPipelineCore/Private/MoviePipelineVideoOutputBase.cpp
--- a/Engine/Plugins/MovieScene/MovieRenderPipeline/Source/MovieRenderPipelineCore/Private/MoviePipelineVideoOutputBase.cpp
+++ b/Engine/Plugins/MovieScene/MovieRenderPipeline/Source/MovieRenderPipelineCore/Private/MoviePipelineVideoOutputBase.cpp
@@ -233,11 +233,40 @@ void UMoviePipelineVideoOutputBase::OnReceiveImageDataImpl(FMoviePipelineMergerO
 		FMoviePipelineBackgroundMediaTasks Task;
 		FImagePixelData* RawRenderPassData = RenderPassData.Value.Get();
 
-		FAsyncCropImage CropImagePreprocessor(Payload->SampleState.CropRectangle);
-		if (ShouldCropOverscanImpl())
+		// Pixel data pointer for the cropped result to ensure it stays in memory while the image write is processed
+		TUniquePtr<FImagePixelData> CroppedPixelData;
+		const bool bIsCropRectValid = !Payload->SampleState.CropRectangle.IsEmpty();
+		const bool bCanCropResolution = RenderPassData.Value->GetSize() == Payload->SampleState.EffectiveOutputResolution;
+		if (ShouldCropOverscanImpl() && bIsCropRectValid && bCanCropResolution)
 		{
-			CropImagePreprocessor(RawRenderPassData);
-			RawRenderPassData = CropImagePreprocessor.OutCroppedImage.Get();
+			switch (RenderPassData.Value->GetType())
+			{
+			case EImagePixelType::Color:
+				{
+					TAsyncCropImage<FColor> CropImagePreprocessor(Payload->SampleState.CropRectangle);
+					CropImagePreprocessor(RawRenderPassData);
+					CroppedPixelData = MoveTemp(CropImagePreprocessor.OutCroppedImage);
+				}
+				break;
+
+			case EImagePixelType::Float16:
+				{
+					TAsyncCropImage<FFloat16Color> CropImagePreprocessor(Payload->SampleState.CropRectangle);
+					CropImagePreprocessor(RawRenderPassData);
+					CroppedPixelData = MoveTemp(CropImagePreprocessor.OutCroppedImage);
+				}
+				break;
+
+			case EImagePixelType::Float32:
+				{
+					TAsyncCropImage<FLinearColor> CropImagePreprocessor(Payload->SampleState.CropRectangle);
+					CropImagePreprocessor(RawRenderPassData);
+					CroppedPixelData = MoveTemp(CropImagePreprocessor.OutCroppedImage);
+				}
+				break;
+			}
+			
+			RawRenderPassData = CroppedPixelData.Get();
 		}
 		
 		// Making sure that if OCIO is enabled the Quantization won't do additional color conversion.
diff --git a/Engine/Plugins/MovieScene/MovieRenderPipeline/Source/MovieRenderPipelineCore/Public/MovieRenderPipelineDataTypes.h b/Engine/Plugins/MovieScene/MovieRenderPipeline/Source/MovieRenderPipelineCore/Public/MovieRenderPipelineDataTypes.h
--- a/Engine/Plugins/MovieScene/MovieRenderPipeline/Source/MovieRenderPipelineCore/Public/MovieRenderPipelineDataTypes.h
+++ b/Engine/Plugins/MovieScene/MovieRenderPipeline/Source/MovieRenderPipelineCore/Public/MovieRenderPipelineDataTypes.h
@@ -1678,15 +1678,15 @@ struct TAsyncCompositeImage<FLinearColor>
 	FIntPoint Offset;
 };
 
-
-struct FAsyncCropImage
+template<typename PixelType>
+struct TAsyncCropImage
 {
-	FAsyncCropImage(FImageWriteTask* InParentWriteTask, const FIntRect& InCropRect)
+	TAsyncCropImage(FImageWriteTask* InParentWriteTask, const FIntRect& InCropRect)
 		: ParentWriteTask(InParentWriteTask)
 		, CropRect(InCropRect)
 	{ }
 
-	FAsyncCropImage(const FIntRect& InCropRect)
+	TAsyncCropImage(const FIntRect& InCropRect)
 		: ParentWriteTask(nullptr)
 		, CropRect(InCropRect)
 	{ }
@@ -1700,16 +1700,16 @@ struct FAsyncCropImage
 			return;
 		}
 
-		TRACE_CPUPROFILER_EVENT_SCOPE(MoviePipeline::FAsyncCropImage)
+		TRACE_CPUPROFILER_EVENT_SCOPE(MoviePipeline::TAsyncCropImage)
 		
-		TArray64<FFloat16Color> DefaultPixels;
+		TArray64<PixelType> DefaultPixels;
 		DefaultPixels.AddUninitialized(CropRect.Size().X * CropRect.Size().Y);
 
 		TSharedRef<FImagePixelDataPayload> NewPayload = PixelData->GetPayload<FImagePixelDataPayload>()->Copy();
 		
-		TUniquePtr<FImagePixelData> CroppedImage = MakeUnique<TImagePixelData<FFloat16Color>>(CropRect.Size(), MoveTemp(DefaultPixels), NewPayload);
-		TImagePixelData<FFloat16Color>* DestColorData = static_cast<TImagePixelData<FFloat16Color>*>(CroppedImage.Get());
-		TImagePixelData<FFloat16Color>* SrcColorData = static_cast<TImagePixelData<FFloat16Color>*>(PixelData);
+		TUniquePtr<FImagePixelData> CroppedImage = MakeUnique<TImagePixelData<PixelType>>(CropRect.Size(), MoveTemp(DefaultPixels), NewPayload);
+		TImagePixelData<PixelType>* DestColorData = static_cast<TImagePixelData<PixelType>*>(CroppedImage.Get());
+		TImagePixelData<PixelType>* SrcColorData = static_cast<TImagePixelData<PixelType>*>(PixelData);
 
 		ParallelFor(DestColorData->GetSize().Y,
 			[&](int32 ScanlineIndex = 0)
@@ -1717,7 +1717,7 @@ struct FAsyncCropImage
 				int64 DstIndex = int64(ScanlineIndex) * int64(DestColorData->GetSize().X);
 				int64 SrcIndex = (int64(ScanlineIndex) + int64(CropRect.Min.Y)) * int64(SrcColorData->GetSize().X) + int64(CropRect.Min.X);
 				
-				FMemory::Memcpy(&DestColorData->Pixels[DstIndex], &SrcColorData->Pixels[SrcIndex], DestColorData->GetSize().X * sizeof(FFloat16Color));
+				FMemory::Memcpy(&DestColorData->Pixels[DstIndex], &SrcColorData->Pixels[SrcIndex], DestColorData->GetSize().X * sizeof(PixelType));
 			});
 
 		// If we have a parent write task that is holding the image being processed, write to that. Otherwise,
diff --git a/Engine/Plugins/MovieScene/MovieRenderPipeline/Source/MovieRenderPipelineRenderPasses/Private/MovieGraphImageSequenceOutputNode.cpp b/Engine/Plugins/MovieScene/MovieRenderPipeline/Source/MovieRenderPipelineRenderPasses/Private/MovieGraphImageSequenceOutputNode.cpp
--- a/Engine/Plugins/MovieScene/MovieRenderPipeline/Source/MovieRenderPipelineRenderPasses/Private/MovieGraphImageSequenceOutputNode.cpp
+++ b/Engine/Plugins/MovieScene/MovieRenderPipeline/Source/MovieRenderPipelineRenderPasses/Private/MovieGraphImageSequenceOutputNode.cpp
@@ -160,10 +160,27 @@ void UMovieGraphImageSequenceOutputNode::OnReceiveImageDataImpl(UMovieGraphPipel
 		UE::MovieGraph::FMovieGraphSampleState* Payload = RenderData.Value->GetPayload<UE::MovieGraph::FMovieGraphSampleState>();
 
 		// If the overscan isn't cropped from the final image, offset any composites to the top-left of the original frustum
+		// TODO: Would prefer to check against "effective output resolution" (desired resolution scaled by overscan) but MRG does not pass that value down.
+		// This should be addressed in UE-228752
 		FIntPoint CompositeOffset = FIntPoint::ZeroValue;
-		if (ShouldCropOverscanImpl())
+		const bool bIsCropRectValid = !Payload->CropRectangle.IsEmpty();
+		const bool bCanCropResolution = RenderData.Value->GetSize() == Payload->BackbufferResolution;
+		if (ShouldCropOverscanImpl() && bIsCropRectValid && bCanCropResolution)
 		{
-			TileImageTask->PixelPreProcessors.Add(FAsyncCropImage(TileImageTask.Get(), Payload->CropRectangle));
+			switch (RenderData.Value->GetType())
+			{
+			case EImagePixelType::Color:
+				TileImageTask->PixelPreProcessors.Add(TAsyncCropImage<FColor>(TileImageTask.Get(), Payload->CropRectangle));
+				break;
+
+			case EImagePixelType::Float16:
+				TileImageTask->PixelPreProcessors.Add(TAsyncCropImage<FFloat16Color>(TileImageTask.Get(), Payload->CropRectangle));
+				break;
+
+			case EImagePixelType::Float32:
+				TileImageTask->PixelPreProcessors.Add(TAsyncCropImage<FLinearColor>(TileImageTask.Get(), Payload->CropRectangle));
+				break;
+			}
 		}
 		else
 		{
diff --git a/Engine/Plugins/MovieScene/MovieRenderPipeline/Source/MovieRenderPipelineRenderPasses/Private/MoviePipelineImageSequenceOutput.cpp b/Engine/Plugins/MovieScene/MovieRenderPipeline/Source/MovieRenderPipelineRenderPasses/Private/MoviePipelineImageSequenceOutput.cpp
--- a/Engine/Plugins/MovieScene/MovieRenderPipeline/Source/MovieRenderPipelineRenderPasses/Private/MoviePipelineImageSequenceOutput.cpp
+++ b/Engine/Plugins/MovieScene/MovieRenderPipeline/Source/MovieRenderPipelineRenderPasses/Private/MoviePipelineImageSequenceOutput.cpp
@@ -207,9 +207,26 @@ void UMoviePipelineImageSequenceOutputBase::OnReceiveImageDataImpl(FMoviePipelin
 
 		// If the overscan isn't cropped from the final image, offset any composites to the top-left of the original frustum
 		FIntPoint CompositeOffset = FIntPoint::ZeroValue;
-		if (ShouldCropOverscanImpl())
+
+		// For now, only passes that were rendered at the overscanned resolution can be cropped using the crop rectangle
+		const bool bIsCropRectValid = !Payload->SampleState.CropRectangle.IsEmpty();
+		const bool bCanCropResolution = RenderPassData.Value->GetSize() == Payload->SampleState.EffectiveOutputResolution;
+		if (ShouldCropOverscanImpl() && bIsCropRectValid && bCanCropResolution)
 		{
-			TileImageTask->PixelPreProcessors.Add(FAsyncCropImage(TileImageTask.Get(), Payload->SampleState.CropRectangle));
+			switch (RenderPassData.Value->GetType())
+			{
+			case EImagePixelType::Color:
+				TileImageTask->PixelPreProcessors.Add(TAsyncCropImage<FColor>(TileImageTask.Get(), Payload->SampleState.CropRectangle));
+				break;
+
+			case EImagePixelType::Float16:
+				TileImageTask->PixelPreProcessors.Add(TAsyncCropImage<FFloat16Color>(TileImageTask.Get(), Payload->SampleState.CropRectangle));
+				break;
+
+			case EImagePixelType::Float32:
+				TileImageTask->PixelPreProcessors.Add(TAsyncCropImage<FLinearColor>(TileImageTask.Get(), Payload->SampleState.CropRectangle));
+				break;
+			}
 		}
 		else
 		{
