MovieRenderPipeline: Added ability for the render pipeline to support animated overscan by allowing each frame to have a potentially different internal resolution

#virtualized

diff --git a/Engine/Plugins/MovieScene/MovieRenderPipeline/Source/MovieRenderPipelineCore/Private/Graph/MovieGraphDefaultRenderer.cpp b/Engine/Plugins/MovieScene/MovieRenderPipeline/Source/MovieRenderPipelineCore/Private/Graph/MovieGraphDefaultRenderer.cpp
--- a/Engine/Plugins/MovieScene/MovieRenderPipeline/Source/MovieRenderPipelineCore/Private/Graph/MovieGraphDefaultRenderer.cpp
+++ b/Engine/Plugins/MovieScene/MovieRenderPipeline/Source/MovieRenderPipelineCore/Private/Graph/MovieGraphDefaultRenderer.cpp
@@ -226,6 +226,12 @@ void UMovieGraphDefaultRenderer::Render(const FMovieGraphTimeStepData& InTimeSte
 	// the render thread uses it.
 	FlushAsyncEngineSystems(InTimeStepData.EvaluatedConfig);
 
+	// Empty the camera overscan cache at the start of each frame
+	if (InTimeStepData.bIsFirstTemporalSampleForFrame)
+	{
+		CameraOverscanCache.Empty();
+	}
+	
 	// Housekeeping: Clean up any tasks that were completed since the last frame. This lets us have a 
 	// better idea of how much work we're concurrently doing. 
 	{
@@ -247,6 +253,37 @@ void UMovieGraphDefaultRenderer::Render(const FMovieGraphTimeStepData& InTimeSte
 		}
 	}
 
+	// Clean up and shutdown any stale surface queues. This is necessary for anything that changes resolution between frames, such as animated overscan.
+	// The surface queue pool is keyed off of resolution, so if every frame has a new resolution, a new surface queue is created, and subsequently,
+	// only one surface is ever added to the queue (that for the frame that needed that resolution of surface queue). However, when a surface queue isn't full
+	// it can't properly mark surfaces as complete and ready for readback because surface queues natively track "staleness" by how far from the current surface in the queue
+	// a previously queued surface is. So, in order to prevent the surface queue from growing too large, and to force surfaces to complete their readback,
+	// we track the last frame the queue was used on, and if it has been enough frames, we clean it up, forcing any surfaces to read back. This staleness amount
+	// should give any queued surfaces enough frames to complete rendering so that they can be read back by the time Shutdown is called
+	{
+		for (auto Iter = PooledSurfaceQueues.CreateIterator(); Iter; ++Iter)
+		{
+			if (Iter->Value->IsStale())
+			{
+				Iter->Value->Shutdown();
+				Iter.RemoveCurrent();
+			}
+		}
+	}
+
+	// Clean up any stale view render targets
+	{
+		for (auto Iter = PooledViewRenderTargets.CreateIterator(); Iter; ++Iter)
+		{
+			const UE::MovieGraph::DefaultRenderer::FMovieGraphImagePreviewPoolState& State = PooledViewRenderTargetStates[Iter->Key];
+			if (State.IsStale())
+			{
+				PooledViewRenderTargetStates.Remove(Iter->Key);
+				Iter.RemoveCurrent();
+			}
+		}
+	}
+	
 	// Hide the progress widget before we render anything. This allows widget captures to not include the progress bar.
 	GetOwningGraph()->SetPreviewWidgetVisible(false);
 
@@ -487,7 +524,7 @@ void UMovieGraphDefaultRenderer::WarnAboutAnimatedOverscan(float InInitialOversc
 		UE_LOG(
 			LogMovieRenderPipeline,
 			Warning,
-			TEXT("Detected animated Camera Overscan value on shot %s for camera %s. MRG does not support changing resolution between frames, and a resolution computed from the initial overscan (%f) will be used instead. Overscan can be affected by both the Overscan camera property or distortion parameters on the Lens component"),
+			TEXT("Overscan on camera %s changed since start of frame in shot %s, using cached overscan value of %f instead"),
 			*CurrentShot->OuterName,
 			*CurrentShot->InnerName,
 			InInitialOverscan);
@@ -504,12 +541,14 @@ UTextureRenderTarget2D* UMovieGraphDefaultRenderer::GetOrCreateViewRenderTarget(
 
 	if (const TObjectPtr<UTextureRenderTarget2D>* ExistingViewRenderTarget = PooledViewRenderTargets.Find(CombinedParams))
 	{
+		PooledViewRenderTargetStates[CombinedParams].UpdateLastUsedFrame();
 		return *ExistingViewRenderTarget;
 	}
 
 	const TObjectPtr<UTextureRenderTarget2D> NewViewRenderTarget = CreateViewRenderTarget(InInitParams);
 	PooledViewRenderTargets.Emplace(CombinedParams, NewViewRenderTarget);
-
+	PooledViewRenderTargetStates.Emplace(CombinedParams, UE::MovieGraph::DefaultRenderer::FMovieGraphImagePreviewPoolState());
+	
 	return NewViewRenderTarget.Get();
 }
 
@@ -529,6 +568,7 @@ FMoviePipelineSurfaceQueuePtr UMovieGraphDefaultRenderer::GetOrCreateSurfaceQueu
 {
 	if (const FMoviePipelineSurfaceQueuePtr* ExistingSurfaceQueue = PooledSurfaceQueues.Find(InInitParams))
 	{
+		(*ExistingSurfaceQueue)->UpdateLastUsedFrame();
 		return *ExistingSurfaceQueue;
 	}
 
@@ -556,10 +596,13 @@ TArray<FMovieGraphImagePreviewData> UMovieGraphDefaultRenderer::GetPreviewData()
 { 
 	TArray<FMovieGraphImagePreviewData> Results;
 
+	uint64 MostRecentFrame = 0;
+	
 	TSet<FString> CameraNamesInUse;
 	for (const TPair<UE::MovieGraph::DefaultRenderer::FMovieGraphImagePreviewDataPoolParams, TObjectPtr<UTextureRenderTarget2D>>& RenderTarget : PooledViewRenderTargets)
 	{
 		CameraNamesInUse.Add(RenderTarget.Key.Identifier.CameraName);
+		MostRecentFrame = FMath::Max(PooledViewRenderTargetStates[RenderTarget.Key].LastUsedFrame, MostRecentFrame);
 	}
 
 	for (const TPair<UE::MovieGraph::DefaultRenderer::FMovieGraphImagePreviewDataPoolParams, TObjectPtr<UTextureRenderTarget2D>>& RenderTarget : PooledViewRenderTargets)
@@ -571,11 +614,17 @@ TArray<FMovieGraphImagePreviewData> UMovieGraphDefaultRenderer::GetPreviewData()
 		{
 			continue;
 		}
-		
-		FMovieGraphImagePreviewData& Data = Results.AddDefaulted_GetRef();
-		Data.Identifier = RenderTarget.Key.Identifier;
-		Data.Texture = RenderTarget.Value.Get();
-		Data.bMultipleCameraNames = CameraNamesInUse.Num() > 1;
+
+		// In cases of changing resolution between frames (e.g animated overscan) the render target pool will have a render target for each resolution, but
+		// that render target may be for a past frame. So filter out any render targets that aren't from the most recent frame
+		const UE::MovieGraph::DefaultRenderer::FMovieGraphImagePreviewPoolState RenderTargetState = PooledViewRenderTargetStates[RenderTarget.Key];
+		if (RenderTargetState.LastUsedFrame == MostRecentFrame)
+		{
+			FMovieGraphImagePreviewData& Data = Results.AddDefaulted_GetRef();
+			Data.Identifier = RenderTarget.Key.Identifier;
+			Data.Texture = RenderTarget.Value.Get();
+			Data.bMultipleCameraNames = CameraNamesInUse.Num() > 1;
+		}
 	}
 
 	return Results;
diff --git a/Engine/Plugins/MovieScene/MovieRenderPipeline/Source/MovieRenderPipelineCore/Private/MoviePipelineRendering.cpp b/Engine/Plugins/MovieScene/MovieRenderPipeline/Source/MovieRenderPipelineCore/Private/MoviePipelineRendering.cpp
--- a/Engine/Plugins/MovieScene/MovieRenderPipeline/Source/MovieRenderPipelineCore/Private/MoviePipelineRendering.cpp
+++ b/Engine/Plugins/MovieScene/MovieRenderPipeline/Source/MovieRenderPipelineCore/Private/MoviePipelineRendering.cpp
@@ -167,6 +167,11 @@ void UMoviePipeline::RenderFrame()
 	// Send any output frames that have been completed since the last render.
 	ProcessOutstandingFinishedFrames();
 
+	if (CachedOutputState.IsFirstTemporalSample())
+	{
+		CameraOverscanCache.Empty();
+	}
+	
 	FMoviePipelineCameraCutInfo& CurrentCameraCut = ActiveShotList[CurrentShotIndex]->ShotInfo;
 	APlayerController* LocalPlayerController = GetWorld()->GetFirstPlayerController();
 
@@ -207,9 +212,26 @@ void UMoviePipeline::RenderFrame()
 	check(HighResSettings);
 	check(OutputSettings);
 
-	// TODO: Not much support here for multi-camera, so simply get the player controller camera and use its overscan value
-	// Use the cache to get the overscan value for resolution scaling so that it doesn't vary between frames
-	const float CameraOverscan  = CameraOverscanCache[INDEX_NONE];
+	float CameraOverscan = 0.0f;
+	if (!CameraOverscanCache.Contains(INDEX_NONE))
+	{
+		// If the cache does not contain a default overscan (likely because it has been cleared as the start of a new frame),
+		// cache the default camera overscan at INDEX_NONE to ensure anything that doesn't have multi-camera support still has an overscan value to utilize
+		// TODO: Not much support here for multi-camera, so simply get the player controller camera and use its overscan value
+		if (UCameraComponent* BoundCamera = MovieSceneHelpers::CameraComponentFromRuntimeObject(LocalPlayerController->PlayerCameraManager->GetViewTarget()))
+		{
+			FMinimalViewInfo CameraViewInfo;
+			BoundCamera->GetCameraView(GetWorld()->GetDeltaSeconds(), CameraViewInfo);
+			CameraOverscan = CameraViewInfo.GetOverscan();
+		}
+		
+		CameraOverscanCache.Add(INDEX_NONE, CameraOverscan);
+	}
+	else
+	{
+		// Use the cache to get the overscan value for resolution scaling so that it doesn't vary between subsamples
+		CameraOverscan = CameraOverscanCache[INDEX_NONE];
+	}
 	
 	FIntPoint TileCount = FIntPoint(HighResSettings->TileCount, HighResSettings->TileCount);
 	FIntPoint OriginalTileCount = TileCount;
@@ -291,6 +313,11 @@ void UMoviePipeline::RenderFrame()
 	// Reset our flag for this frame.
 	bHasRenderedFirstViewThisFrame = false;
 
+	for (UMoviePipelineRenderPass* RenderPass : InputBuffers)
+	{
+		RenderPass->OnFrameStart();
+	}
+	
 	// If this is the first sample for a new frame, we want to notify the output builder that it should expect data to accumulate for this frame.
 	if (CachedOutputState.IsFirstTemporalSample())
 	{
diff --git a/Engine/Plugins/MovieScene/MovieRenderPipeline/Source/MovieRenderPipelineCore/Private/MoviePipelineSurfaceReader.cpp b/Engine/Plugins/MovieScene/MovieRenderPipeline/Source/MovieRenderPipelineCore/Private/MoviePipelineSurfaceReader.cpp
--- a/Engine/Plugins/MovieScene/MovieRenderPipeline/Source/MovieRenderPipelineCore/Private/MoviePipelineSurfaceReader.cpp
+++ b/Engine/Plugins/MovieScene/MovieRenderPipeline/Source/MovieRenderPipelineCore/Private/MoviePipelineSurfaceReader.cpp
@@ -265,6 +265,8 @@ FMoviePipelineSurfaceQueue::FMoviePipelineSurfaceQueue(FIntPoint InSurfaceSize,
 	}
 
 	FrameResolveLatency = 1;
+
+	UpdateLastUsedFrame();
 }
 
 FMoviePipelineSurfaceQueue::~FMoviePipelineSurfaceQueue()
@@ -318,6 +320,17 @@ void FMoviePipelineSurfaceQueue::Shutdown()
 	}
 }
 
+void FMoviePipelineSurfaceQueue::UpdateLastUsedFrame()
+{
+	LastUsedFrame = GFrameCounter;
+}
+
+bool FMoviePipelineSurfaceQueue::IsStale() const
+{
+	constexpr int32 StaleThreshold = 3;
+	return GFrameCounter - LastUsedFrame > StaleThreshold;
+}
+
 void FMoviePipelineSurfaceQueue::OnRenderTargetReady_RenderThread(const FTextureRHIRef InRenderTarget, TSharedRef<FImagePixelDataPayload, ESPMode::ThreadSafe> InFramePayload, TUniqueFunction<void(TUniquePtr<FImagePixelData>&&)>&& InFunctionCallback)
 {
 	ensure(IsInRenderingThread());
diff --git a/Engine/Plugins/MovieScene/MovieRenderPipeline/Source/MovieRenderPipelineCore/Public/Graph/MovieGraphDefaultRenderer.h b/Engine/Plugins/MovieScene/MovieRenderPipeline/Source/MovieRenderPipelineCore/Public/Graph/MovieGraphDefaultRenderer.h
--- a/Engine/Plugins/MovieScene/MovieRenderPipeline/Source/MovieRenderPipelineCore/Public/Graph/MovieGraphDefaultRenderer.h
+++ b/Engine/Plugins/MovieScene/MovieRenderPipeline/Source/MovieRenderPipelineCore/Public/Graph/MovieGraphDefaultRenderer.h
@@ -137,6 +137,31 @@ namespace UE::MovieGraph::DefaultRenderer
 		}
 	};
 
+	/** State object to keep track of the state of image render targets in the image render target pool */
+	struct FMovieGraphImagePreviewPoolState
+	{
+		FMovieGraphImagePreviewPoolState()
+		{
+			UpdateLastUsedFrame();
+		}
+		
+		/** The last frame that the render target was used */
+		uint64 LastUsedFrame;
+
+		/** Updates the last used frame to the current frame counter */
+		void UpdateLastUsedFrame()
+		{
+			LastUsedFrame = GFrameCounter;
+		}
+
+		/** Compares the last used frame to the current frame counter and returns true of the difference is larger than the specified threshold */
+		bool IsStale() const
+		{
+			constexpr int32 StaleThreshold = 3;
+			return GFrameCounter - LastUsedFrame > StaleThreshold;
+		}
+	};
+	
 	struct FSurfaceAccumulatorPool : public TSharedFromThis<FSurfaceAccumulatorPool>
 	{
 		struct FInstance
@@ -289,6 +314,7 @@ public:
 	float GetCameraOverscan(int32 InCameraIndex);
 
 	/** Outputs a warning message regarding animated overscan to the MRQ log if one has not already been output  */
+	UE_DEPRECATED(5.6, "WarnAboutAnimatedOverscan is no longer used; animated overscan is supported in 5.6")
 	void WarnAboutAnimatedOverscan(float InInitialOverscan);
 	
 public:
@@ -340,6 +366,9 @@ protected:
 	*/
 	TMap<UE::MovieGraph::DefaultRenderer::FMovieGraphImagePreviewDataPoolParams, TObjectPtr<UTextureRenderTarget2D>> PooledViewRenderTargets;
 
+	/** States for the render targets stored in PooledViewRenderTargets */
+	TMap<UE::MovieGraph::DefaultRenderer::FMovieGraphImagePreviewDataPoolParams, UE::MovieGraph::DefaultRenderer::FMovieGraphImagePreviewPoolState> PooledViewRenderTargetStates;
+	
 	/*
 	* This is a list of Surfaces that we can copy our render target to immediately after
 	* it is drawn to. This is read back to the CPU a frame later (to avoid blocking the
diff --git a/Engine/Plugins/MovieScene/MovieRenderPipeline/Source/MovieRenderPipelineCore/Public/MoviePipeline.h b/Engine/Plugins/MovieScene/MovieRenderPipeline/Source/MovieRenderPipelineCore/Public/MoviePipeline.h
--- a/Engine/Plugins/MovieScene/MovieRenderPipeline/Source/MovieRenderPipelineCore/Public/MoviePipeline.h
+++ b/Engine/Plugins/MovieScene/MovieRenderPipeline/Source/MovieRenderPipelineCore/Public/MoviePipeline.h
@@ -126,6 +126,7 @@ public:
 	void CacheCameraOverscan(int32 InCameraIndex, float InCameraOverscan);
 
 	/** Outputs a warning message regarding animated overscan to the MRQ log if one has not already been output  */
+	UE_DEPRECATED(5.6, "WarnAboutAnimatedOverscan is no longer used; animated overscan is supported in 5.6")
 	void WarnAboutAnimatedOverscan(float InInitialOverscan);
 	
 #if WITH_EDITOR
diff --git a/Engine/Plugins/MovieScene/MovieRenderPipeline/Source/MovieRenderPipelineCore/Public/MoviePipelineRenderPass.h b/Engine/Plugins/MovieScene/MovieRenderPipeline/Source/MovieRenderPipelineCore/Public/MoviePipelineRenderPass.h
--- a/Engine/Plugins/MovieScene/MovieRenderPipeline/Source/MovieRenderPipelineCore/Public/MoviePipelineRenderPass.h
+++ b/Engine/Plugins/MovieScene/MovieRenderPipeline/Source/MovieRenderPipelineCore/Public/MoviePipelineRenderPass.h
@@ -27,6 +27,12 @@ public:
 		TeardownImpl();
 	}
 
+	/** Called at the start of a new frame to allow the render pass to perform any pre-frame operations */
+	void OnFrameStart()
+	{
+		OnFrameStartImpl();
+	}
+	
 	/** An array of identifiers for the output buffers expected as a result of this render pass. */
 	void GatherOutputPasses(TArray<FMoviePipelinePassIdentifier>& ExpectedRenderPasses)
 	{
@@ -57,6 +63,7 @@ protected:
 	virtual void SetupImpl(const MoviePipeline::FMoviePipelineRenderPassInitSettings& InPassInitSettings) {}
 	virtual void WaitUntilTasksComplete() {}
 	virtual void TeardownImpl() {}
+	virtual void OnFrameStartImpl() {}
 	virtual void GatherOutputPassesImpl(TArray<FMoviePipelinePassIdentifier>& ExpectedRenderPasses) {}
 	virtual void RenderSample_GameThreadImpl(const FMoviePipelineRenderPassMetrics& InSampleState) {}
 	virtual bool IsAlphaInTonemapperRequiredImpl() const { return false; }
diff --git a/Engine/Plugins/MovieScene/MovieRenderPipeline/Source/MovieRenderPipelineCore/Public/MoviePipelineSurfaceReader.h b/Engine/Plugins/MovieScene/MovieRenderPipeline/Source/MovieRenderPipelineCore/Public/MoviePipelineSurfaceReader.h
--- a/Engine/Plugins/MovieScene/MovieRenderPipeline/Source/MovieRenderPipelineCore/Public/MoviePipelineSurfaceReader.h
+++ b/Engine/Plugins/MovieScene/MovieRenderPipeline/Source/MovieRenderPipelineCore/Public/MoviePipelineSurfaceReader.h
@@ -90,6 +90,13 @@ public:
 
 	void BlockUntilAnyAvailable();
 	void Shutdown();
+
+	/** Updates the stored last used frame to be the current frame */
+	void UpdateLastUsedFrame();
+
+	/** Returns true if the amount of frames between the current frame and the last used frame is greater than a specific threshold */
+	bool IsStale() const;
+	
 private:
 	struct FResolveSurface
 	{
@@ -115,4 +122,7 @@ private:
 	
 	/** How many frames do we wait before we resolve a frame? */
 	int32 FrameResolveLatency;
+
+	/** The last frame this surface was used. Surfaces that haven't been used for a certain number of frames are cleaned up */
+	uint64 LastUsedFrame;
 };
\ No newline at end of file
diff --git a/Engine/Plugins/MovieScene/MovieRenderPipeline/Source/MovieRenderPipelineRenderPasses/Private/Graph/Renderers/MovieGraphDeferredPass.cpp b/Engine/Plugins/MovieScene/MovieRenderPipeline/Source/MovieRenderPipelineRenderPasses/Private/Graph/Renderers/MovieGraphDeferredPass.cpp
--- a/Engine/Plugins/MovieScene/MovieRenderPipeline/Source/MovieRenderPipelineRenderPasses/Private/Graph/Renderers/MovieGraphDeferredPass.cpp
+++ b/Engine/Plugins/MovieScene/MovieRenderPipeline/Source/MovieRenderPipelineRenderPasses/Private/Graph/Renderers/MovieGraphDeferredPass.cpp
@@ -223,10 +223,17 @@ void FMovieGraphDeferredPass::Render(const FMovieGraphTraversalContext& InFrameT
 		}
 		else
 		{
-			// Current overscan is different from originally cached value, indicating an animated overscan value, so output a warning message
-			if (CameraInfo.ViewInfo.GetOverscan() != CameraOverscan)
+			// Current overscan is different from originally cached value, indicating overscan changed since start of frame, so output a warning message
+			if (CameraInfo.ViewInfo.GetOverscan() != CameraOverscan && InTimeData.bIsFirstTemporalSampleForFrame)
 			{
-				GetRenderer()->WarnAboutAnimatedOverscan(CameraOverscan);
+				UE_LOG(
+					LogMovieRenderPipeline,
+					Warning,
+					TEXT("Overscan on camera %s changed since start of frame %d in shot %s, scaling resolution by cached overscan value of %f instead to keep frame resolution consistent"),
+					*InFrameTraversalContext.RenderDataIdentifier.CameraName,
+					InTimeData.ShotOutputFrameNumber,
+					*InFrameTraversalContext.Shot->OuterName,
+					CameraOverscan);
 			}
 		}
 
diff --git a/Engine/Plugins/MovieScene/MovieRenderPipeline/Source/MovieRenderPipelineRenderPasses/Private/MoviePipelineImagePassBase.cpp b/Engine/Plugins/MovieScene/MovieRenderPipeline/Source/MovieRenderPipelineRenderPasses/Private/MoviePipelineImagePassBase.cpp
--- a/Engine/Plugins/MovieScene/MovieRenderPipeline/Source/MovieRenderPipelineRenderPasses/Private/MoviePipelineImagePassBase.cpp
+++ b/Engine/Plugins/MovieScene/MovieRenderPipeline/Source/MovieRenderPipelineRenderPasses/Private/MoviePipelineImagePassBase.cpp
@@ -155,10 +155,11 @@ TWeakObjectPtr<UTextureRenderTarget2D> UMoviePipelineImagePassBase::CreateViewRe
 	NewTarget->InitCustomFormat(InSize.X, InSize.Y, EPixelFormat::PF_FloatRGBA, false);
 	NewTarget->AddToRoot();
 
-	if (GetPipeline()->GetPreviewTexture() == nullptr)
-	{
-		GetPipeline()->SetPreviewTexture(NewTarget.Get());
-	}
+	// Always update the preview texture to the new texture, so that in cases where resolution is changing between frames (e.g. animated overscan)
+	// the preview texture continues to be for the most recent frame.
+	// TODO: Multi-camera support - As there is only one preview texture, and there is no way to distinguish which camera we are creating the texture for,
+	// we can't be sure that the newest preview texture is for the same camera as previous frames.
+	GetPipeline()->SetPreviewTexture(NewTarget.Get());
 
 	return NewTarget;
 }
@@ -402,6 +403,27 @@ void UMoviePipelineImagePassBase::SetupViewForViewModeOverride(FSceneView* View)
 	UE::MovieRenderPipeline::UpdateSceneViewForShowFlags(View);
 }
 
+void UMoviePipelineImagePassBase::OnFrameStartImpl()
+{
+	Super::OnFrameStartImpl();
+
+	// Clean up and shutdown any stale surface queues. This is necessary for anything that changes resolution between frames, such as animated overscan.
+	// The surface queue pool is keyed off of resolution, so if every frame has a new resolution, a new surface queue is created, and subsequently,
+	// only one surface is ever added to the queue (that for the frame that needed that resolution of surface queue). However, when a surface queue isn't full
+	// it can't properly mark surfaces as complete and ready for readback because surface queues natively track "staleness" by how far from the current surface in the queue
+	// a previously queued surface is. So, in order to prevent the surface queue from growing too large, and to force surfaces to complete their readback,
+	// we track the last frame the queue was used on, and if it has been enough frames, we clean it up, forcing any surfaces to read back. This staleness amount
+	// should give any queued surfaces enough frames to complete rendering so that they can be read back by the time Shutdown is called
+	for (auto Iter = SurfaceQueues.CreateIterator(); Iter; ++Iter)
+	{
+		if (Iter->Value->IsStale())
+		{
+			Iter->Value->Shutdown();
+			Iter.RemoveCurrent();
+		}
+	}
+}
+
 void UMoviePipelineImagePassBase::GatherOutputPassesImpl(TArray<FMoviePipelinePassIdentifier>& ExpectedRenderPasses)
 {
 	Super::GatherOutputPassesImpl(ExpectedRenderPasses);
@@ -444,11 +466,19 @@ FSceneView* UMoviePipelineImagePassBase::GetSceneViewForSampleState(FSceneViewFa
 	}
 	else
 	{
-		// Check for animated overscan, and warn if animated overscan was detected
 		const float CachedOverscan = GetPipeline()->GetCachedCameraOverscan(InOutSampleState.OutputState.CameraIndex);
-		if (CameraInfo.ViewInfo.GetOverscan() != CachedOverscan)
+		
+		// Current overscan is different from originally cached value, indicating overscan changed since start of frame, so output a warning message
+		if (CameraInfo.ViewInfo.GetOverscan() != CachedOverscan && InOutSampleState.OutputState.IsFirstTemporalSample())
 		{
-			GetPipeline()->WarnAboutAnimatedOverscan(CachedOverscan);
+			UE_LOG(
+				LogMovieRenderPipeline,
+				Warning,
+				TEXT("Overscan on camera %s changed since start of frame %d in shot %s, scaling resolution by cached overscan value of %f instead to keep frame resolution consistent"),
+				*InOutSampleState.OutputState.CameraName,
+				InOutSampleState.OutputState.ShotOutputFrameNumber,
+				*InOutSampleState.OutputState.ShotName,
+				CachedOverscan);
 		}
 
 		// Update the sample state with this camera's overscan instead of the config overscan it is filled with initially
@@ -720,6 +750,22 @@ TSharedPtr<FAccumulatorPool::FAccumulatorInstance, ESPMode::ThreadSafe> FAccumul
 				}
 			}
 		}
+
+		// If a free accumulator wasn't found, try creating a new one
+		if (AvailableIndex == INDEX_NONE)
+		{
+			if (TSharedPtr<FAccumulatorInstance, ESPMode::ThreadSafe> NewAccumulatorInstance = CreateNewAccumulatorInstance())
+			{
+				NewAccumulatorInstance->ActiveFrameNumber = InFrameNumber;
+				NewAccumulatorInstance->ActivePassIdentifier = InPassIdentifier;
+				NewAccumulatorInstance->bIsActive = true;
+				NewAccumulatorInstance->TaskPrereq = nullptr;
+			
+				AvailableIndex = Accumulators.Num();
+				Accumulators.Add(NewAccumulatorInstance);
+				UE_LOG(LogMovieRenderPipeline, Log, TEXT("Allocated a Accumulator for Pool %s, New Pool Count: %d"), *GetPoolName().ToString(), Accumulators.Num());
+			}
+		}
 	}
 
 	return Accumulators[AvailableIndex];
diff --git a/Engine/Plugins/MovieScene/MovieRenderPipeline/Source/MovieRenderPipelineRenderPasses/Public/MoviePipelineImagePassBase.h b/Engine/Plugins/MovieScene/MovieRenderPipeline/Source/MovieRenderPipelineRenderPasses/Public/MoviePipelineImagePassBase.h
--- a/Engine/Plugins/MovieScene/MovieRenderPipeline/Source/MovieRenderPipelineRenderPasses/Public/MoviePipelineImagePassBase.h
+++ b/Engine/Plugins/MovieScene/MovieRenderPipeline/Source/MovieRenderPipelineRenderPasses/Public/MoviePipelineImagePassBase.h
@@ -102,6 +102,7 @@ public:
 protected:
 
 	// UMoviePipelineRenderPass API
+	virtual void OnFrameStartImpl() override;
 	virtual void GatherOutputPassesImpl(TArray<FMoviePipelinePassIdentifier>& ExpectedRenderPasses) override;
 	virtual void SetupImpl(const MoviePipeline::FMoviePipelineRenderPassInitSettings& InPassInitSettings) override;
 	virtual void RenderSample_GameThreadImpl(const FMoviePipelineRenderPassMetrics& InSampleState) override;
@@ -184,11 +185,17 @@ struct MOVIERENDERPIPELINERENDERPASSES_API FAccumulatorPool : public TSharedFrom
 		FGraphEventRef TaskPrereq;
 	};
 
+	virtual ~FAccumulatorPool() = default;
+	
 	TArray<TSharedPtr<FAccumulatorInstance, ESPMode::ThreadSafe>> Accumulators;
 	FCriticalSection CriticalSection;
 
 
 	TSharedPtr<FAccumulatorPool::FAccumulatorInstance, ESPMode::ThreadSafe> BlockAndGetAccumulator_GameThread(int32 InFrameNumber, const FMoviePipelinePassIdentifier& InPassIdentifier);
+
+protected:
+	virtual TSharedPtr<FAccumulatorInstance, ESPMode::ThreadSafe> CreateNewAccumulatorInstance() { return nullptr; }
+	virtual FName GetPoolName() const { return NAME_None; }
 };
 
 template<typename AccumulatorType>
@@ -205,6 +212,18 @@ struct TAccumulatorPool : FAccumulatorPool
 		}
 
 	}
+
+protected:
+	virtual TSharedPtr<FAccumulatorInstance, ESPMode::ThreadSafe> CreateNewAccumulatorInstance() override
+	{
+		TSharedPtr<MoviePipeline::IMoviePipelineOverlappedAccumulator, ESPMode::ThreadSafe> Accumulator = MakeShared<AccumulatorType, ESPMode::ThreadSafe>();
+		return MakeShared<FAccumulatorInstance, ESPMode::ThreadSafe>(Accumulator);
+	}
+
+	virtual FName GetPoolName() const override
+	{
+		return AccumulatorType::GetName();
+	}
 };
 
 
