DisplayGamma: Remove OCIO plugin-declared default display gamma value, and simply use the engine render target default.

diff --git a/Engine/Plugins/Compositing/OpenColorIO/Source/OpenColorIO/Private/OpenColorIORendering.cpp b/Engine/Plugins/Compositing/OpenColorIO/Source/OpenColorIO/Private/OpenColorIORendering.cpp
--- a/Engine/Plugins/Compositing/OpenColorIO/Source/OpenColorIO/Private/OpenColorIORendering.cpp
+++ b/Engine/Plugins/Compositing/OpenColorIO/Source/OpenColorIO/Private/OpenColorIORendering.cpp
@@ -3,6 +3,7 @@
 #include "OpenColorIORendering.h"
 
 #include "Engine/RendererSettings.h"
+#include "Engine/TextureRenderTarget.h"
 #include "Engine/TextureRenderTarget2D.h"
 #include "Engine/World.h"
 #include "EngineModule.h"
@@ -151,9 +152,10 @@ void FOpenColorIORendering::AddPass_RenderThread(
 	EOpenColorIOTransformAlpha TransformAlpha) 
 {
 	const float EngineDisplayGamma = View.Family->RenderTarget->GetDisplayGamma();
+	const float DefaultDisplayGammaRT = UTextureRenderTarget::GetDefaultDisplayGamma();
 	// There is a special case where post processing and tonemapper are disabled. In this case tonemapper applies a static display Inverse of Gamma which defaults to 2.2.
 	// In the case when Both PostProcessing and ToneMapper are disabled we apply gamma manually. In every other case we apply inverse gamma before applying OCIO.
-	float DisplayGamma = (View.Family->EngineShowFlags.Tonemapper == 0) || (View.Family->EngineShowFlags.PostProcessing == 0) ? DefaultDisplayGamma : DefaultDisplayGamma / EngineDisplayGamma;
+	float DisplayGamma = (View.Family->EngineShowFlags.Tonemapper == 0) || (View.Family->EngineShowFlags.PostProcessing == 0) ? DefaultDisplayGammaRT : DefaultDisplayGammaRT / EngineDisplayGamma;
 
 	FOpenColorIORendering::AddPass_RenderThread(
 		GraphBuilder,
diff --git a/Engine/Plugins/Compositing/OpenColorIO/Source/OpenColorIO/Public/OpenColorIORendering.h b/Engine/Plugins/Compositing/OpenColorIO/Source/OpenColorIO/Public/OpenColorIORendering.h
--- a/Engine/Plugins/Compositing/OpenColorIO/Source/OpenColorIO/Public/OpenColorIORendering.h
+++ b/Engine/Plugins/Compositing/OpenColorIO/Source/OpenColorIO/Public/OpenColorIORendering.h
@@ -115,6 +115,6 @@ public:
 	*/
 	static void PrepareView(FSceneViewFamily& InViewFamily, FSceneView& InView);
 	
-	/** */
+	UE_DEPRECATED(5.6, "This constant is deprecated, please use UTextureRenderTarget::GetDefaultDisplayGamma() instead.")
 	static float DefaultDisplayGamma;
 };
\ No newline at end of file
diff --git a/Engine/Plugins/MovieScene/MovieRenderPipeline/Source/MovieRenderPipelineCore/Private/Graph/Nodes/MovieGraphWidgetRendererBaseNode.cpp b/Engine/Plugins/MovieScene/MovieRenderPipeline/Source/MovieRenderPipelineCore/Private/Graph/Nodes/MovieGraphWidgetRendererBaseNode.cpp
--- a/Engine/Plugins/MovieScene/MovieRenderPipeline/Source/MovieRenderPipelineCore/Private/Graph/Nodes/MovieGraphWidgetRendererBaseNode.cpp
+++ b/Engine/Plugins/MovieScene/MovieRenderPipeline/Source/MovieRenderPipelineCore/Private/Graph/Nodes/MovieGraphWidgetRendererBaseNode.cpp
@@ -2,6 +2,7 @@
 
 #include "Graph/Nodes/MovieGraphWidgetRendererBaseNode.h"
 
+#include "Engine/TextureRenderTarget.h"
 #include "Engine/TextureRenderTarget2D.h"
 #include "Framework/Application/SlateApplication.h"
 #include "Graph/MovieGraphBlueprintLibrary.h"
@@ -9,7 +10,6 @@
 #include "Graph/MovieGraphDefaultRenderer.h"
 #include "Graph/MovieGraphPipeline.h"
 #include "MovieRenderPipelineCoreModule.h"
-#include "OpenColorIORendering.h"
 #include "Slate/WidgetRenderer.h"
 #include "Styling/AppStyle.h"
 #include "TextureResource.h"
@@ -134,7 +134,7 @@ void UMovieGraphWidgetRendererBaseNode::FMovieGraphWidgetPass::Render(const FMov
 	// Create the render target the widget will be rendered into
 	UE::MovieGraph::DefaultRenderer::FRenderTargetInitParams RenderTargetInitParams;
 	RenderTargetInitParams.Size = OutputResolution;
-	RenderTargetInitParams.TargetGamma = FOpenColorIORendering::DefaultDisplayGamma;
+	RenderTargetInitParams.TargetGamma = UTextureRenderTarget::GetDefaultDisplayGamma();
 	RenderTargetInitParams.PixelFormat = PF_B8G8R8A8;
 	UTextureRenderTarget2D* RenderTarget = Renderer->GetOrCreateViewRenderTarget(RenderTargetInitParams, RenderDataIdentifier);
 
diff --git a/Engine/Plugins/MovieScene/MovieRenderPipeline/Source/MovieRenderPipelineRenderPasses/Private/Graph/Renderers/MovieGraphImagePassBase.cpp b/Engine/Plugins/MovieScene/MovieRenderPipeline/Source/MovieRenderPipelineRenderPasses/Private/Graph/Renderers/MovieGraphImagePassBase.cpp
--- a/Engine/Plugins/MovieScene/MovieRenderPipeline/Source/MovieRenderPipelineRenderPasses/Private/Graph/Renderers/MovieGraphImagePassBase.cpp
+++ b/Engine/Plugins/MovieScene/MovieRenderPipeline/Source/MovieRenderPipelineRenderPasses/Private/Graph/Renderers/MovieGraphImagePassBase.cpp
@@ -9,7 +9,7 @@
 #include "MoviePipelineSurfaceReader.h"
 #include "CanvasTypes.h"
 #include "MovieRenderOverlappedImage.h"
-#include "OpenColorIORendering.h"
+#include "Engine/TextureRenderTarget.h"
 #include "Engine/RendererSettings.h"
 #include "UnrealClient.h"
 #include "SceneViewExtensionContext.h"
@@ -93,7 +93,7 @@ DefaultRenderer::FRenderTargetInitParams FMovieGraphImagePassBase::GetRenderTarg
 	// OCIO: Since this is a manually created Render target we don't need Gamma to be applied.
 	// We use this render target to render to via a display extension that utilizes Display Gamma
 	// which has a default value of 2.2 (DefaultDisplayGamma), therefore we need to set Gamma on this render target to 2.2 to cancel out any unwanted effects.
-	InitParams.TargetGamma = FOpenColorIORendering::DefaultDisplayGamma;
+	InitParams.TargetGamma = UTextureRenderTarget::GetDefaultDisplayGamma();
 	InitParams.PixelFormat = PF_FloatRGBA;
 
 	return InitParams;
diff --git a/Engine/Plugins/MovieScene/MovieRenderPipeline/Source/MovieRenderPipelineRenderPasses/Private/MoviePipelineImagePassBase.cpp b/Engine/Plugins/MovieScene/MovieRenderPipeline/Source/MovieRenderPipelineRenderPasses/Private/MoviePipelineImagePassBase.cpp
--- a/Engine/Plugins/MovieScene/MovieRenderPipeline/Source/MovieRenderPipelineRenderPasses/Private/MoviePipelineImagePassBase.cpp
+++ b/Engine/Plugins/MovieScene/MovieRenderPipeline/Source/MovieRenderPipelineRenderPasses/Private/MoviePipelineImagePassBase.cpp
@@ -17,6 +17,7 @@
 #include "EngineModule.h"
 #include "Engine/LocalPlayer.h"
 #include "Engine/RendererSettings.h"
+#include "Engine/TextureRenderTarget.h"
 #include "MovieRenderOverlappedImage.h"
 #include "ImageUtils.h"
 #include "SceneManagement.h"
@@ -149,7 +150,7 @@ TWeakObjectPtr<UTextureRenderTarget2D> UMoviePipelineImagePassBase::CreateViewRe
 	// OCIO: Since this is a manually created Render target we don't need Gamma to be applied.
 	// We use this render target to render to via a display extension that utilizes Display Gamma
 	// which has a default value of 2.2 (DefaultDisplayGamma), therefore we need to set Gamma on this render target to 2.2 to cancel out any unwanted effects.
-	NewTarget->TargetGamma = FOpenColorIORendering::DefaultDisplayGamma;
+	NewTarget->TargetGamma = UTextureRenderTarget::GetDefaultDisplayGamma();
 
 	// Initialize to the tile size (not final size) and use a 16 bit back buffer to avoid precision issues when accumulating later
 	NewTarget->InitCustomFormat(InSize.X, InSize.Y, EPixelFormat::PF_FloatRGBA, false);
diff --git a/Engine/Plugins/Runtime/nDisplay/Source/DisplayCluster/Private/Render/DisplayDevice/Proxy/DisplayClusterDisplayDeviceProxy_OpenColorIO.cpp b/Engine/Plugins/Runtime/nDisplay/Source/DisplayCluster/Private/Render/DisplayDevice/Proxy/DisplayClusterDisplayDeviceProxy_OpenColorIO.cpp
--- a/Engine/Plugins/Runtime/nDisplay/Source/DisplayCluster/Private/Render/DisplayDevice/Proxy/DisplayClusterDisplayDeviceProxy_OpenColorIO.cpp
+++ b/Engine/Plugins/Runtime/nDisplay/Source/DisplayCluster/Private/Render/DisplayDevice/Proxy/DisplayClusterDisplayDeviceProxy_OpenColorIO.cpp
@@ -8,12 +8,13 @@
 #include "ScreenPass.h"
 #include "CommonRenderResources.h"
 #include "PostProcess/PostProcessMaterialInputs.h"
+#include "Engine/TextureRenderTarget.h"
 
 namespace UE::DisplayCluster
 {
 	static inline float GetDisplayGamma(const FDisplayClusterViewport_Context& InViewportContext)
 	{
-		const float DefaultDisplayGamma = FOpenColorIORendering::DefaultDisplayGamma;
+		const float DefaultDisplayGamma = UTextureRenderTarget::GetDefaultDisplayGamma();
 		const float EngineDisplayGamma = InViewportContext.RenderThreadData.EngineDisplayGamma;
 
 		// There is a special case where post processing and tonemapper are disabled. In this case tonemapper applies a static display Inverse of Gamma which defaults to 2.2.
diff --git a/Engine/Plugins/Runtime/nDisplay/Source/DisplayCluster/Private/Render/Viewport/DisplayClusterViewport_OpenColorIO.cpp b/Engine/Plugins/Runtime/nDisplay/Source/DisplayCluster/Private/Render/Viewport/DisplayClusterViewport_OpenColorIO.cpp
--- a/Engine/Plugins/Runtime/nDisplay/Source/DisplayCluster/Private/Render/Viewport/DisplayClusterViewport_OpenColorIO.cpp
+++ b/Engine/Plugins/Runtime/nDisplay/Source/DisplayCluster/Private/Render/Viewport/DisplayClusterViewport_OpenColorIO.cpp
@@ -15,6 +15,7 @@
 #include "ScreenPass.h"
 #include "CommonRenderResources.h"
 #include "PostProcess/PostProcessMaterialInputs.h"
+#include "Engine/TextureRenderTarget.h"
 
 //////////////////////////////////////////////////////////////////////////
 // FDisplayClusterViewport_OpenColorIO
@@ -51,7 +52,7 @@ bool FDisplayClusterViewport_OpenColorIO::IsConversionSettingsEqual(const FOpenC
 
 float FDisplayClusterViewport_OpenColorIO::GetDisplayGamma(const FDisplayClusterViewport_Context& InViewportContext) const
 {
-	const float DefaultDisplayGamma = FOpenColorIORendering::DefaultDisplayGamma;
+	const float DefaultDisplayGamma = UTextureRenderTarget::GetDefaultDisplayGamma();
 	const float EngineDisplayGamma = InViewportContext.RenderThreadData.EngineDisplayGamma;
 
 	// There is a special case where post processing and tonemapper are disabled. In this case tonemapper applies a static display Inverse of Gamma which defaults to 2.2.
diff --git a/Engine/Source/Runtime/Engine/Classes/Engine/TextureRenderTarget.h b/Engine/Source/Runtime/Engine/Classes/Engine/TextureRenderTarget.h
--- a/Engine/Source/Runtime/Engine/Classes/Engine/TextureRenderTarget.h
+++ b/Engine/Source/Runtime/Engine/Classes/Engine/TextureRenderTarget.h
@@ -84,7 +84,7 @@ class UTextureRenderTarget : public UTexture
 
 	// UTextureRenderTarget default display gamma if none is set
 	//	returns hard-coded 2.2
-	static float GetDefaultDisplayGamma();
+	ENGINE_API static float GetDefaultDisplayGamma();
 
 	// get the variant of ReadPixels call that should be used
 	//	either BGRA8,RGBA32F,or RGBA16F
