MovieRenderPipeline: Deprecated blueprint APIs that referred to overscanned resolution as "effective output resolution", as it is no longer accurate with overscan crop. Also fixed issue where EXR display windows were off by one pixel for some resolutions

diff --git a/Engine/Plugins/MovieScene/MovieRenderPipeline/Source/MovieRenderPipelineCore/Private/Graph/MovieGraphBlueprintLibrary.cpp b/Engine/Plugins/MovieScene/MovieRenderPipeline/Source/MovieRenderPipelineCore/Private/Graph/MovieGraphBlueprintLibrary.cpp
--- a/Engine/Plugins/MovieScene/MovieRenderPipeline/Source/MovieRenderPipelineCore/Private/Graph/MovieGraphBlueprintLibrary.cpp
+++ b/Engine/Plugins/MovieScene/MovieRenderPipeline/Source/MovieRenderPipelineCore/Private/Graph/MovieGraphBlueprintLibrary.cpp
@@ -366,6 +366,11 @@ int32 UMovieGraphBlueprintLibrary::GetCurrentVersionNumber(const UMovieGraphPipe
 }
 
 FIntPoint UMovieGraphBlueprintLibrary::GetEffectiveOutputResolution(UMovieGraphEvaluatedConfig* InEvaluatedGraph, float DefaultOverscan)
+{
+	return GetOverscannedResolution(InEvaluatedGraph, DefaultOverscan);
+}
+
+FIntPoint UMovieGraphBlueprintLibrary::GetDesiredOutputResolution(UMovieGraphEvaluatedConfig* InEvaluatedGraph)
 {
 	if (!InEvaluatedGraph)
 	{
@@ -380,16 +385,6 @@ FIntPoint UMovieGraphBlueprintLibrary::GetEffectiveOutputResolution(UMovieGraphE
 		return FIntPoint();
 	}
 
-	float RescaledOverscan = DefaultOverscan;
-	if (UMovieGraphCameraSettingNode* CameraSetting = InEvaluatedGraph->GetSettingForBranch<UMovieGraphCameraSettingNode>(UMovieGraphNode::GlobalsPinName, bIncludeCDOs))
-	{
-		if (CameraSetting->bOverride_OverscanPercentage)
-		{
-			// The old system used [0-1] range for floats, the new system will use [0-100], so we rescale down before calling through.
-			RescaledOverscan = FMath::Clamp(CameraSetting->OverscanPercentage / 100.f, 0.f, 1.f);
-		}
-	}
-
 	// We need to look at the Project Settings for the latest value for a given profile
 	FMovieGraphNamedResolution NamedResolution;
 	if (UMovieGraphBlueprintLibrary::IsNamedResolutionValid(OutputSetting->OutputResolution.ProfileName))
@@ -402,33 +397,63 @@ FIntPoint UMovieGraphBlueprintLibrary::GetEffectiveOutputResolution(UMovieGraphE
 		NamedResolution = OutputSetting->OutputResolution;
 	}
 
-	return UMoviePipelineBlueprintLibrary::Utility_GetEffectiveOutputResolution(RescaledOverscan, NamedResolution.Resolution);
+	return NamedResolution.Resolution;
 }
 
-FIntRect UMovieGraphBlueprintLibrary::GetOverscanCropRectangle(UMovieGraphEvaluatedConfig* InEvaluatedGraph, float DefaultOverscan)
+FIntPoint UMovieGraphBlueprintLibrary::GetOverscannedResolution(UMovieGraphEvaluatedConfig* InEvaluatedGraph, float DefaultOverscan)
+{
+	if (!InEvaluatedGraph)
+    {
+    	FFrame::KismetExecutionMessage(TEXT("InEvaluatedGraph cannot be null."), ELogVerbosity::Error);
+    	return FIntPoint();
+    }
+
+	FIntPoint DesiredOutputResolution = GetDesiredOutputResolution(InEvaluatedGraph);
+	
+    constexpr bool bIncludeCDOs = false;
+    float RescaledOverscan = DefaultOverscan;
+    if (UMovieGraphCameraSettingNode* CameraSetting = InEvaluatedGraph->GetSettingForBranch<UMovieGraphCameraSettingNode>(UMovieGraphNode::GlobalsPinName, bIncludeCDOs))
+    {
+    	if (CameraSetting->bOverride_OverscanPercentage)
+    	{
+    		// The old system used [0-1] range for floats, the new system will use [0-100], so we rescale down before calling through.
+    		RescaledOverscan = FMath::Clamp(CameraSetting->OverscanPercentage / 100.f, 0.f, 1.f);
+    	}
+    }
+    
+	return UMoviePipelineBlueprintLibrary::ScaleResolutionByOverscan(RescaledOverscan, DesiredOutputResolution);
+}
+
+FIntPoint UMovieGraphBlueprintLibrary::GetBackbufferResolution(UMovieGraphEvaluatedConfig* InEvaluatedGraph, float DefaultOverscan)
 {
 	if (!InEvaluatedGraph)
 	{
 		FFrame::KismetExecutionMessage(TEXT("InEvaluatedGraph cannot be null."), ELogVerbosity::Error);
-		return FIntRect();
+		return FIntPoint();
 	}
-	
-	const FIntPoint OverscannedResolution = GetEffectiveOutputResolution(InEvaluatedGraph, DefaultOverscan);
-	
-	constexpr bool bIncludeCDOs = true;
-	const UMovieGraphGlobalOutputSettingNode* OutputSetting = InEvaluatedGraph->GetSettingForBranch<UMovieGraphGlobalOutputSettingNode>(UMovieGraphNode::GlobalsPinName, bIncludeCDOs);
-	if (!ensure(OutputSetting))
+
+	FIntPoint OverscannedResolution = GetOverscannedResolution(InEvaluatedGraph, DefaultOverscan);
+
+	// TODO: When tiling support is added, compute the appropriate backbuffer resolution from OverscannedResolution using tile count and return it
+	return OverscannedResolution;
+}
+
+FIntRect UMovieGraphBlueprintLibrary::GetOverscanCropRectangle(UMovieGraphEvaluatedConfig* InEvaluatedGraph, float DefaultOverscan)
+{
+	if (!InEvaluatedGraph)
 	{
+		FFrame::KismetExecutionMessage(TEXT("InEvaluatedGraph cannot be null."), ELogVerbosity::Error);
 		return FIntRect();
 	}
 
-	const FIntPoint OutputResolution = OutputSetting->OutputResolution.Resolution;
-	
+	const FIntPoint DesiredOutputResolution = GetDesiredOutputResolution(InEvaluatedGraph);
+	const FIntPoint OverscannedResolution = GetOverscannedResolution(InEvaluatedGraph, DefaultOverscan);
+
 	FIntRect CropRect;
 
 	// Overscan is symmetrical, so this difference is always evenly divisible by 2.
-	CropRect.Min = (OverscannedResolution - OutputResolution) / 2;
-	CropRect.Max = CropRect.Min + OutputResolution;
+	CropRect.Min = (OverscannedResolution - DesiredOutputResolution) / 2;
+	CropRect.Max = CropRect.Min + DesiredOutputResolution;
 
 	return CropRect;
 }
diff --git a/Engine/Plugins/MovieScene/MovieRenderPipeline/Source/MovieRenderPipelineCore/Private/Graph/Nodes/MovieGraphVideoOutputNode.cpp b/Engine/Plugins/MovieScene/MovieRenderPipeline/Source/MovieRenderPipelineCore/Private/Graph/Nodes/MovieGraphVideoOutputNode.cpp
--- a/Engine/Plugins/MovieScene/MovieRenderPipeline/Source/MovieRenderPipelineCore/Private/Graph/Nodes/MovieGraphVideoOutputNode.cpp
+++ b/Engine/Plugins/MovieScene/MovieRenderPipeline/Source/MovieRenderPipelineCore/Private/Graph/Nodes/MovieGraphVideoOutputNode.cpp
@@ -88,7 +88,7 @@ void UMovieGraphVideoOutputNode::OnReceiveImageDataImpl(UMovieGraphPipeline* InP
 		// Pixel data pointer for the cropped result to ensure it stays in memory while the image write is processed
 		TUniquePtr<FImagePixelData> CroppedPixelData;
 		const bool bIsCropRectValid = !Payload->CropRectangle.IsEmpty();
-		const bool bCanCropResolution = RenderPassData.Value->GetSize() == Payload->BackbufferResolution;
+		const bool bCanCropResolution = RenderPassData.Value->GetSize() == Payload->OverscannedResolution;
 		if (ShouldCropOverscanImpl() && bIsCropRectValid && bCanCropResolution)
 		{
 			switch (RenderPassData.Value->GetType())
@@ -322,7 +322,7 @@ UMovieGraphVideoOutputNode::FMovieGraphCodecWriterWithPromise* UMovieGraphVideoO
 		FIntPoint OutputResolution = InRenderPassData.Value->GetSize();
 		
 		const bool bIsCropRectValid = !Payload->CropRectangle.IsEmpty();
-		const bool bCanCropResolution = InRenderPassData.Value->GetSize() == Payload->BackbufferResolution;
+		const bool bCanCropResolution = InRenderPassData.Value->GetSize() == Payload->OverscannedResolution;
 		if (ShouldCropOverscanImpl() && bIsCropRectValid && bCanCropResolution)
 		{
 			OutputResolution = Payload->CropRectangle.Size();
diff --git a/Engine/Plugins/MovieScene/MovieRenderPipeline/Source/MovieRenderPipelineCore/Private/Graph/Nodes/MovieGraphWidgetRendererBaseNode.cpp b/Engine/Plugins/MovieScene/MovieRenderPipeline/Source/MovieRenderPipelineCore/Private/Graph/Nodes/MovieGraphWidgetRendererBaseNode.cpp
--- a/Engine/Plugins/MovieScene/MovieRenderPipeline/Source/MovieRenderPipelineCore/Private/Graph/Nodes/MovieGraphWidgetRendererBaseNode.cpp
+++ b/Engine/Plugins/MovieScene/MovieRenderPipeline/Source/MovieRenderPipelineCore/Private/Graph/Nodes/MovieGraphWidgetRendererBaseNode.cpp
@@ -114,6 +114,7 @@ void UMovieGraphWidgetRendererBaseNode::FMovieGraphWidgetPass::Render(const FMov
 	}
 
 	const float CameraOverscan = Renderer->GetCameraOverscan(LayerData.CameraIndex);
+	const FIntPoint OverscannedResolution = UMovieGraphBlueprintLibrary::GetOverscannedResolution(InTimeData.EvaluatedConfig, CameraOverscan);
 	const FIntRect CropRect = UMovieGraphBlueprintLibrary::GetOverscanCropRectangle(InTimeData.EvaluatedConfig, CameraOverscan);
 	
 	// Composited elements should be sized to the original frustum size, as the final image is either cropped to that size, or the composite will be offset
@@ -164,7 +165,7 @@ void UMovieGraphWidgetRendererBaseNode::FMovieGraphWidgetPass::Render(const FMov
 		ENQUEUE_RENDER_COMMAND(WidgetRenderTargetResolveCommand)(
 			[InTimeData, InFrameTraversalContext, RenderTargetInitParams, RenderDataIdentifier = this->RenderDataIdentifier,
 			bComposite = ParentNodeThisFrame->bCompositeOntoFinalImage, CompositingSortOrder = GetCompositingSortOrder(),
-			BackbufferRenderTarget, OutputMerger, OutputResolution](FRHICommandListImmediate& RHICmdList)
+			BackbufferRenderTarget, OutputMerger, OverscannedResolution, CropRect, OutputResolution](FRHICommandListImmediate& RHICmdList)
 			{
 				const FIntRect SourceRect = FIntRect(0, 0, BackbufferRenderTarget->GetSizeXY().X, BackbufferRenderTarget->GetSizeXY().Y);
 
@@ -189,7 +190,9 @@ void UMovieGraphWidgetRendererBaseNode::FMovieGraphWidgetPass::Render(const FMov
 				TSharedRef<UE::MovieGraph::FMovieGraphSampleState> SampleStatePayload =
 					MakeShared<UE::MovieGraph::FMovieGraphSampleState>();
 				SampleStatePayload->TraversalContext = MoveTemp(UpdatedTraversalContext);
+				SampleStatePayload->OverscannedResolution = OverscannedResolution;
 				SampleStatePayload->BackbufferResolution = RenderTargetInitParams.Size;
+				SampleStatePayload->CropRectangle = CropRect;
 				SampleStatePayload->bRequiresAccumulator = false;
 				SampleStatePayload->bFetchFromAccumulator = false;
 				SampleStatePayload->bCompositeOnOtherRenders = bComposite;
diff --git a/Engine/Plugins/MovieScene/MovieRenderPipeline/Source/MovieRenderPipelineCore/Private/MoviePipelineBlueprintLibrary.cpp b/Engine/Plugins/MovieScene/MovieRenderPipeline/Source/MovieRenderPipelineCore/Private/MoviePipelineBlueprintLibrary.cpp
--- a/Engine/Plugins/MovieScene/MovieRenderPipeline/Source/MovieRenderPipelineCore/Private/MoviePipelineBlueprintLibrary.cpp
+++ b/Engine/Plugins/MovieScene/MovieRenderPipeline/Source/MovieRenderPipelineCore/Private/MoviePipelineBlueprintLibrary.cpp
@@ -43,6 +43,7 @@
 #include "GameFramework/PlayerController.h"
 #include "CineCameraActor.h"
 #include "CineCameraComponent.h"
+#include "MoviePipelineHighResSetting.h"
 
 #include UE_INLINE_GENERATED_CPP_BY_NAME(MoviePipelineBlueprintLibrary)
 
@@ -902,34 +903,67 @@ int32 UMoviePipelineBlueprintLibrary::GetCurrentVersionNumber(const UMoviePipeli
 
 FIntPoint UMoviePipelineBlueprintLibrary::GetEffectiveOutputResolution(UMoviePipelinePrimaryConfig* InPrimaryConfig, UMoviePipelineExecutorShot* InPipelineExecutorShot, float DefaultOverscan)
 {
-	if (InPrimaryConfig && InPipelineExecutorShot)
+	return GetOverscannedResolution(InPrimaryConfig, InPipelineExecutorShot, DefaultOverscan);
+}
+
+FIntPoint UMoviePipelineBlueprintLibrary::GetDesiredOutputResolution(const UMoviePipelinePrimaryConfig* InPrimaryConfig)
+{
+	if (InPrimaryConfig)
 	{
 		const UMoviePipelineOutputSetting* OutputSetting = InPrimaryConfig->FindSetting<UMoviePipelineOutputSetting>();
-		const UMoviePipelineCameraSetting* CameraSetting = Cast<const UMoviePipelineCameraSetting>(UMoviePipelineBlueprintLibrary::FindOrGetDefaultSettingForShot(UMoviePipelineCameraSetting::StaticClass(), InPrimaryConfig, InPipelineExecutorShot));
 		check(OutputSetting);
-		check(CameraSetting);
-		return Utility_GetEffectiveOutputResolution(CameraSetting->bOverrideCameraOverscan ? CameraSetting->OverscanPercentage : DefaultOverscan, OutputSetting->OutputResolution);
+		return OutputSetting->OutputResolution;
 	}
 
 	return FIntPoint();
 }
 
-FIntRect UMoviePipelineBlueprintLibrary::GetOverscanCropRectangle(const UMoviePipelinePrimaryConfig* InPrimaryConfig, const UMoviePipelineExecutorShot* InPipelineExecutorShot, float DefaultOverscan)
+FIntPoint UMoviePipelineBlueprintLibrary::GetOverscannedResolution(const UMoviePipelinePrimaryConfig* InPrimaryConfig, const UMoviePipelineExecutorShot* InPipelineExecutorShot, float InDefaultOverscan)
 {
 	if (InPrimaryConfig && InPipelineExecutorShot)
 	{
-		const UMoviePipelineOutputSetting* OutputSetting = InPrimaryConfig->FindSetting<UMoviePipelineOutputSetting>();
-		const UMoviePipelineCameraSetting* CameraSetting = Cast<const UMoviePipelineCameraSetting>(UMoviePipelineBlueprintLibrary::FindOrGetDefaultSettingForShot(UMoviePipelineCameraSetting::StaticClass(), InPrimaryConfig, InPipelineExecutorShot));
-		check(OutputSetting);
-		check(CameraSetting);
+		const FIntPoint DesiredOutputResolution = GetDesiredOutputResolution(InPrimaryConfig);
+		
+		const UMoviePipelineCameraSetting* CameraSettings = Cast<const UMoviePipelineCameraSetting>(FindOrGetDefaultSettingForShot(UMoviePipelineCameraSetting::StaticClass(), InPrimaryConfig, InPipelineExecutorShot));
+		check(CameraSettings);
+		const float Overscan = CameraSettings->bOverrideCameraOverscan ? CameraSettings->OverscanPercentage : InDefaultOverscan;
+		return ScaleResolutionByOverscan(Overscan, DesiredOutputResolution);
+	}
+
+	return FIntPoint();
+}
+
+FIntPoint UMoviePipelineBlueprintLibrary::GetBackbufferResolution(const UMoviePipelinePrimaryConfig* InPrimaryConfig, const UMoviePipelineExecutorShot* InPipelineExecutorShot, float InDefaultOverscan)
+{
+	if (InPrimaryConfig && InPipelineExecutorShot)
+	{
+		const FIntPoint OverscannedResolution = GetOverscannedResolution(InPrimaryConfig, InPipelineExecutorShot, InDefaultOverscan);
+
+		const UMoviePipelineHighResSetting* HighResSettings = Cast<const UMoviePipelineHighResSetting>(FindOrGetDefaultSettingForShot(UMoviePipelineHighResSetting::StaticClass(), InPrimaryConfig, InPipelineExecutorShot));
+		check(HighResSettings);
 		
-		const FIntPoint OverscannedResolution = Utility_GetEffectiveOutputResolution(CameraSetting->bOverrideCameraOverscan ? CameraSetting->OverscanPercentage : DefaultOverscan, OutputSetting->OutputResolution);
+		FIntPoint TileResolution = FIntPoint(
+			FMath::CeilToInt((float)OverscannedResolution.X / (float)HighResSettings->TileCount),
+			FMath::CeilToInt((float)OverscannedResolution.Y / (float)HighResSettings->TileCount));
+
+		return HighResSettings->CalculatePaddedBackbufferSize(TileResolution);
+	}
+
+	return FIntPoint();
+}
+
+FIntRect UMoviePipelineBlueprintLibrary::GetOverscanCropRectangle(const UMoviePipelinePrimaryConfig* InPrimaryConfig, const UMoviePipelineExecutorShot* InPipelineExecutorShot, float DefaultOverscan)
+{
+	if (InPrimaryConfig && InPipelineExecutorShot)
+	{
+		const FIntPoint DesiredOutputResolution = GetDesiredOutputResolution(InPrimaryConfig);
+		const FIntPoint OverscannedResolution = GetOverscannedResolution(InPrimaryConfig, InPipelineExecutorShot, DefaultOverscan);
 		
 		FIntRect CropRect;
 
 		// Overscan is symmetrical, so this difference is always evenly divisible by 2.
-		CropRect.Min = (OverscannedResolution - OutputSetting->OutputResolution) / 2;
-		CropRect.Max = CropRect.Min + OutputSetting->OutputResolution;
+		CropRect.Min = (OverscannedResolution - DesiredOutputResolution) / 2;
+		CropRect.Max = CropRect.Min + DesiredOutputResolution;
 
 		return CropRect;
 	}
@@ -937,7 +971,7 @@ FIntRect UMoviePipelineBlueprintLibrary::GetOverscanCropRectangle(const UMoviePi
 	return FIntRect();
 }
 
-FIntPoint UMoviePipelineBlueprintLibrary::Utility_GetEffectiveOutputResolution(const float OverscanPercentage, const FIntPoint& InOutputResolution)
+FIntPoint UMoviePipelineBlueprintLibrary::ScaleResolutionByOverscan(const float OverscanPercentage, const FIntPoint& InOutputResolution)
 {
 	const float ClampedOverscanPercentage = FMath::Clamp(OverscanPercentage, 0.0f, 1.0f);
 	FIntPoint EffectiveResolution(InOutputResolution);
diff --git a/Engine/Plugins/MovieScene/MovieRenderPipeline/Source/MovieRenderPipelineCore/Private/MoviePipelineRendering.cpp b/Engine/Plugins/MovieScene/MovieRenderPipeline/Source/MovieRenderPipelineCore/Private/MoviePipelineRendering.cpp
--- a/Engine/Plugins/MovieScene/MovieRenderPipeline/Source/MovieRenderPipelineCore/Private/MoviePipelineRendering.cpp
+++ b/Engine/Plugins/MovieScene/MovieRenderPipeline/Source/MovieRenderPipelineCore/Private/MoviePipelineRendering.cpp
@@ -92,15 +92,10 @@ void UMoviePipeline::SetupRenderingPipelineForShot(UMoviePipelineExecutorShot* I
 	CameraOverscanCache.Add(INDEX_NONE, CameraOverscan);
 	
 	FIntPoint BackbufferTileCount = FIntPoint(HighResSettings->TileCount, HighResSettings->TileCount);
-	FIntPoint OutputResolution = UMoviePipelineBlueprintLibrary::GetEffectiveOutputResolution(GetPipelinePrimaryConfig(), InShot, CameraOverscan);
+	FIntPoint TotalResolution = UMoviePipelineBlueprintLibrary::GetOverscannedResolution(GetPipelinePrimaryConfig(), InShot, CameraOverscan);
 
 	// Figure out how big each sub-region (tile) is.
-	FIntPoint BackbufferResolution = FIntPoint(
-		FMath::CeilToInt((float)OutputResolution.X / (float)HighResSettings->TileCount),
-		FMath::CeilToInt((float)OutputResolution.Y / (float)HighResSettings->TileCount));
-
-	// Then increase each sub-region by the overlap amount.
-	BackbufferResolution = HighResSettings->CalculatePaddedBackbufferSize(BackbufferResolution);
+	FIntPoint BackbufferResolution = UMoviePipelineBlueprintLibrary::GetBackbufferResolution(GetPipelinePrimaryConfig(), InShot, CameraOverscan);
 
 	{
 		int32 MaxResolution = GetMax2DTextureDimension();
@@ -118,9 +113,6 @@ void UMoviePipeline::SetupRenderingPipelineForShot(UMoviePipelineExecutorShot* I
 		}
 	}
 
-	// Note how many tiles we wish to render with.
-	BackbufferTileCount = FIntPoint(HighResSettings->TileCount, HighResSettings->TileCount);
-
 	const ERHIFeatureLevel::Type FeatureLevel = GetWorld()->GetFeatureLevel();
 
 	// Initialize our render pass. This is a copy of the settings to make this less coupled to the Settings UI.
@@ -137,7 +129,7 @@ void UMoviePipeline::SetupRenderingPipelineForShot(UMoviePipelineExecutorShot* I
 		NumOutputPasses++;
 	}
 
-	UE_LOG(LogMovieRenderPipeline, Log, TEXT("Finished setting up rendering for shot. Shot has %d Passes. Total resolution: (%dx%d) Individual tile resolution: (%dx%d). Tile count: (%dx%d)"), NumOutputPasses, OutputResolution.X, OutputResolution.Y, BackbufferResolution.X, BackbufferResolution.Y, BackbufferTileCount.X, BackbufferTileCount.Y);
+	UE_LOG(LogMovieRenderPipeline, Log, TEXT("Finished setting up rendering for shot. Shot has %d Passes. Total resolution: (%dx%d) Individual tile resolution: (%dx%d). Tile count: (%dx%d)"), NumOutputPasses, TotalResolution.X, TotalResolution.Y, BackbufferResolution.X, BackbufferResolution.Y, BackbufferTileCount.X, BackbufferTileCount.Y);
 }
 
 void UMoviePipeline::TeardownRenderingPipelineForShot(UMoviePipelineExecutorShot* InShot)
@@ -235,7 +227,7 @@ void UMoviePipeline::RenderFrame()
 	
 	FIntPoint TileCount = FIntPoint(HighResSettings->TileCount, HighResSettings->TileCount);
 	FIntPoint OriginalTileCount = TileCount;
-	FIntPoint OutputResolution = UMoviePipelineBlueprintLibrary::GetEffectiveOutputResolution(GetPipelinePrimaryConfig(), ActiveShotList[CurrentShotIndex], CameraOverscan);
+	FIntPoint OverscannedResolution = UMoviePipelineBlueprintLibrary::GetOverscannedResolution(GetPipelinePrimaryConfig(), ActiveShotList[CurrentShotIndex], CameraOverscan);
 
 	int32 NumSpatialSamples = AntiAliasingSettings->SpatialSampleCount;
 	int32 NumTemporalSamples = AntiAliasingSettings->TemporalSampleCount;
@@ -479,7 +471,7 @@ void UMoviePipeline::RenderFrame()
 				}
 
 				// Render each output pass
-				FMoviePipelineRenderPassMetrics SampleStateForCurrentResolution = UE::MoviePipeline::GetRenderPassMetrics(GetPipelinePrimaryConfig(), ActiveShotList[CurrentShotIndex], SampleState, OutputResolution);
+				FMoviePipelineRenderPassMetrics SampleStateForCurrentResolution = UE::MoviePipeline::GetRenderPassMetrics(GetPipelinePrimaryConfig(), ActiveShotList[CurrentShotIndex], SampleState, OverscannedResolution);
 				for (UMoviePipelineRenderPass* RenderPass : InputBuffers)
 				{
 					RenderPass->RenderSample_GameThread(SampleStateForCurrentResolution);
diff --git a/Engine/Plugins/MovieScene/MovieRenderPipeline/Source/MovieRenderPipelineCore/Private/MoviePipelineTelemetry.cpp b/Engine/Plugins/MovieScene/MovieRenderPipeline/Source/MovieRenderPipelineCore/Private/MoviePipelineTelemetry.cpp
--- a/Engine/Plugins/MovieScene/MovieRenderPipeline/Source/MovieRenderPipelineCore/Private/MoviePipelineTelemetry.cpp
+++ b/Engine/Plugins/MovieScene/MovieRenderPipeline/Source/MovieRenderPipelineCore/Private/MoviePipelineTelemetry.cpp
@@ -137,7 +137,7 @@ FMoviePipelineShotRenderTelemetry FMoviePipelineTelemetry::GatherShotRenderTelem
 	FMoviePipelineShotRenderTelemetry Telemetry;
 	Telemetry.bIsGraph = true;
 
-	const FIntPoint Resolution = UMovieGraphBlueprintLibrary::GetEffectiveOutputResolution(InEvaluatedConfig);
+	const FIntPoint Resolution = UMovieGraphBlueprintLibrary::GetOverscannedResolution(InEvaluatedConfig);
 	Telemetry.ResolutionX = Resolution.X;
 	Telemetry.ResolutionY = Resolution.Y;
 
@@ -187,7 +187,7 @@ FMoviePipelineShotRenderTelemetry FMoviePipelineTelemetry::GatherShotRenderTelem
 	Telemetry.bIsGraph = false;
 
 	UMoviePipelinePrimaryConfig* PrimaryConfig = InShot->GetTypedOuter<UMoviePipelineExecutorJob>()->GetConfiguration();
-	const FIntPoint Resolution = UMoviePipelineBlueprintLibrary::GetEffectiveOutputResolution(PrimaryConfig, InShot);
+	const FIntPoint Resolution = UMoviePipelineBlueprintLibrary::GetOverscannedResolution(PrimaryConfig, InShot);
 	Telemetry.ResolutionX = Resolution.X;
 	Telemetry.ResolutionY = Resolution.Y;
 	
diff --git a/Engine/Plugins/MovieScene/MovieRenderPipeline/Source/MovieRenderPipelineCore/Private/MoviePipelineUtils.cpp b/Engine/Plugins/MovieScene/MovieRenderPipeline/Source/MovieRenderPipelineCore/Private/MoviePipelineUtils.cpp
--- a/Engine/Plugins/MovieScene/MovieRenderPipeline/Source/MovieRenderPipelineCore/Private/MoviePipelineUtils.cpp
+++ b/Engine/Plugins/MovieScene/MovieRenderPipeline/Source/MovieRenderPipelineCore/Private/MoviePipelineUtils.cpp
@@ -1242,7 +1242,10 @@ namespace UE
 				OutRenderPassMetrics.ProjectionMatrixJitterAmount = FVector2D((float)(OutRenderPassMetrics.SpatialShiftX) * 2.0f / BackbufferResolution.X, (float)OutRenderPassMetrics.SpatialShiftY * -2.0f / BackbufferResolution.Y);
 
 				OutRenderPassMetrics.BackbufferSize = BackbufferResolution;
+				PRAGMA_DISABLE_DEPRECATION_WARNINGS
 				OutRenderPassMetrics.EffectiveOutputResolution = InEffectiveOutputResolution;
+				PRAGMA_ENABLE_DEPRECATION_WARNINGS
+				OutRenderPassMetrics.OverscannedResolution = InEffectiveOutputResolution;
 
 
 				FIntRect CropRect;
diff --git a/Engine/Plugins/MovieScene/MovieRenderPipeline/Source/MovieRenderPipelineCore/Private/MoviePipelineVideoOutputBase.cpp b/Engine/Plugins/MovieScene/MovieRenderPipeline/Source/MovieRenderPipelineCore/Private/MoviePipelineVideoOutputBase.cpp
--- a/Engine/Plugins/MovieScene/MovieRenderPipeline/Source/MovieRenderPipelineCore/Private/MoviePipelineVideoOutputBase.cpp
+++ b/Engine/Plugins/MovieScene/MovieRenderPipeline/Source/MovieRenderPipelineCore/Private/MoviePipelineVideoOutputBase.cpp
@@ -236,7 +236,7 @@ void UMoviePipelineVideoOutputBase::OnReceiveImageDataImpl(FMoviePipelineMergerO
 		// Pixel data pointer for the cropped result to ensure it stays in memory while the image write is processed
 		TUniquePtr<FImagePixelData> CroppedPixelData;
 		const bool bIsCropRectValid = !Payload->SampleState.CropRectangle.IsEmpty();
-		const bool bCanCropResolution = RenderPassData.Value->GetSize() == Payload->SampleState.EffectiveOutputResolution;
+		const bool bCanCropResolution = RenderPassData.Value->GetSize() == Payload->SampleState.OverscannedResolution;
 		if (ShouldCropOverscanImpl() && bIsCropRectValid && bCanCropResolution)
 		{
 			switch (RenderPassData.Value->GetType())
diff --git a/Engine/Plugins/MovieScene/MovieRenderPipeline/Source/MovieRenderPipelineCore/Public/Graph/MovieGraphBlueprintLibrary.h b/Engine/Plugins/MovieScene/MovieRenderPipeline/Source/MovieRenderPipelineCore/Public/Graph/MovieGraphBlueprintLibrary.h
--- a/Engine/Plugins/MovieScene/MovieRenderPipeline/Source/MovieRenderPipelineCore/Public/Graph/MovieGraphBlueprintLibrary.h
+++ b/Engine/Plugins/MovieScene/MovieRenderPipeline/Source/MovieRenderPipelineCore/Public/Graph/MovieGraphBlueprintLibrary.h
@@ -67,9 +67,38 @@ public:
 	* @param	DefaultOverscan		- The default overscan to use if there are no camera settings that provide an overscan override value, from 0.0 to 1.0
 	* @return						- The output resolution, taking into account overscan
 	*/
-	UFUNCTION(BlueprintCallable, Category = "Movie Graph")
+	UFUNCTION(BlueprintCallable, Category = "Movie Graph", meta=(DeprecatedFunction, DeprecationMessage = "Use GetOverscannedResolution instead"))
+	UE_DEPRECATED(5.6, "GetEffectiveOutputResolution is deprecated. Use GetOverscannedResolution instead")
 	static FIntPoint GetEffectiveOutputResolution(UMovieGraphEvaluatedConfig* InEvaluatedGraph, float DefaultOverscan = 0.0f);
 
+	/**
+	 * Gets the desired output resolution for the movie render graph as specified by the user. Does not include overscan, tiling, or aspect ratio constraints,
+	 * and is the target resolution that the pipeline will generally output (e.g. when cropping overscan in non-EXR formats). Use GetOverscannedResolution or
+	 * GetBackbufferResolution to get resolutions that factor in overscan and tiling, respectively
+	 * @param	InEvaluatedGraph	- The evaluated graph that will provide context for resolving the resolution
+	 * @return						- The desired output resolution
+	 */
+	UFUNCTION(BlueprintCallable, Category = "Movie Graph")
+	static FIntPoint GetDesiredOutputResolution(UMovieGraphEvaluatedConfig* InEvaluatedGraph);
+
+	/**
+	* Gets the desired output resolution scaled by any configured overscan for the movie render graph
+	* @param	InEvaluatedGraph	- The evaluated graph that will provide context for resolving the resolution
+	* @param	DefaultOverscan		- The default overscan to use if there are no camera settings that provide an overscan override value, from 0.0 to 1.0
+	* @return						- The overscan-scaled resolution
+	*/
+	UFUNCTION(BlueprintCallable, Category = "Movie Graph")
+	static FIntPoint GetOverscannedResolution(UMovieGraphEvaluatedConfig* InEvaluatedGraph, float DefaultOverscan = 0.0f);
+
+	/**
+	 * Gets the resolution that frames will actually be rendered at in MRG, which includes factors such as overscan and tiling
+	 * @param	InEvaluatedGraph	- The evaluated graph that will provide context for resolving the resolution
+	 * @param	DefaultOverscan		- The default overscan to use if there are no camera settings that provide an overscan override value, from 0.0 to 1.0
+	 * @return						- The resolution of a rendered frame
+	 */
+	UFUNCTION(BlueprintCallable, Category = "Movie Graph")
+	static FIntPoint GetBackbufferResolution(UMovieGraphEvaluatedConfig* InEvaluatedGraph, float DefaultOverscan = 0.0f);
+	
 	/**
 	 * Gets a rectangle that will crop out any overscan applied to the scene. If there is no overscan, the rectangle returned will match the backbuffer rectangle
 	 * @param InEvaluatedGraph		- The evaluated graph that will provide context for resolving the resolution
diff --git a/Engine/Plugins/MovieScene/MovieRenderPipeline/Source/MovieRenderPipelineCore/Public/Graph/MovieGraphDataTypes.h b/Engine/Plugins/MovieScene/MovieRenderPipeline/Source/MovieRenderPipelineCore/Public/Graph/MovieGraphDataTypes.h
--- a/Engine/Plugins/MovieScene/MovieRenderPipeline/Source/MovieRenderPipelineCore/Public/Graph/MovieGraphDataTypes.h
+++ b/Engine/Plugins/MovieScene/MovieRenderPipeline/Source/MovieRenderPipelineCore/Public/Graph/MovieGraphDataTypes.h
@@ -388,7 +388,10 @@ namespace UE::MovieGraph
 		/** The traversal context used to read graph values at the time of submission. */
 		FMovieGraphTraversalContext TraversalContext;
 
-		/** The backbuffer resolution of this render resource (ie: the resolution the sample was rendered at). */
+		/** The desired output resolution scaled by any configured overscan */
+		FIntPoint OverscannedResolution;
+		
+		/** The backbuffer resolution of this render resource (ie: the resolution the sample was rendered at). Includes overscan and tiling */
 		FIntPoint BackbufferResolution;
 
 		/** 
diff --git a/Engine/Plugins/MovieScene/MovieRenderPipeline/Source/MovieRenderPipelineCore/Public/MoviePipelineBlueprintLibrary.h b/Engine/Plugins/MovieScene/MovieRenderPipeline/Source/MovieRenderPipelineCore/Public/MoviePipelineBlueprintLibrary.h
--- a/Engine/Plugins/MovieScene/MovieRenderPipeline/Source/MovieRenderPipelineCore/Public/MoviePipelineBlueprintLibrary.h
+++ b/Engine/Plugins/MovieScene/MovieRenderPipeline/Source/MovieRenderPipelineCore/Public/MoviePipelineBlueprintLibrary.h
@@ -148,9 +148,26 @@ public:
 	static int32 GetCurrentVersionNumber(const UMoviePipeline* InMoviePipeline);
 	
 	/** In case of Overscan percentage being higher than 0 we render additional pixels. This function returns the resolution with overscan taken into account. */
-	UFUNCTION(BlueprintCallable, Category = "Movie Render Pipeline")
+	UFUNCTION(BlueprintCallable, Category = "Movie Render Pipeline", meta=(DeprecatedFunction, DeprecationMessage = "Use GetOverscannedResolution instead"))
+	UE_DEPRECATED(5.6, "GetEffectiveOutputResolution is deprecated, use GetOverscannedResolution to get the overscanned resolution instead")
 	static FIntPoint GetEffectiveOutputResolution(UMoviePipelinePrimaryConfig* InPrimaryConfig, UMoviePipelineExecutorShot* InPipelineExecutorShot, float DefaultOverscan = 0.0f);
 
+	/** 
+	 * Gets the desired output resolution for the movie render queue as specified by the user. Does not include overscan, tiling, or aspect ratio constraints,
+	 * and is the target resolution that the pipeline will generally output (e.g. when cropping overscan in non-EXR formats). Use GetOverscannedResolution or
+	 * GetBackbufferResolution to get resolutions that factor in overscan and tiling, respectively
+	 */
+	UFUNCTION(BlueprintCallable, Category = "Movie Render Pipeline")
+	static FIntPoint GetDesiredOutputResolution(const UMoviePipelinePrimaryConfig* InPrimaryConfig);
+
+	/** Gets the overscanned resolution, which is the target output resolution scaled by any configured overscan amount */
+	UFUNCTION(BlueprintCallable, Category = "Movie Render Pipeline")
+	static FIntPoint GetOverscannedResolution(const UMoviePipelinePrimaryConfig* InPrimaryConfig, const UMoviePipelineExecutorShot* InPipelineExecutorShot, float InDefaultOverscan = 0.0f);
+	
+	/** Gets the resolution that will be used by the engine when rendering a frame. Includes overscan and tiling */
+	UFUNCTION(BlueprintCallable, Category = "Movie Render Pipeline")
+	static FIntPoint GetBackbufferResolution(const UMoviePipelinePrimaryConfig* InPrimaryConfig, const UMoviePipelineExecutorShot* InPipelineExecutorShot, float InDefaultOverscan = 0.0f);
+	
 	/**
 	 * Gets a rectangle that will crop out any overscan applied to the scene. If there is no overscan, the rectangle returned will match the backbuffer rectangle
 	 * @param InPrimaryConfig			- The pipeline configuration that will provide context for resolving the resolution
@@ -179,8 +196,8 @@ private:
 	// Give the graph BP library access to shared utility methods
 	friend class UMovieGraphBlueprintLibrary;
 
-	/** Get the effective output resolution, taking into account overscan. */
-	static FIntPoint Utility_GetEffectiveOutputResolution(const float OverscanPercentage, const FIntPoint& InOutputResolution);
+	/** Scales the specified resolution by the specified overscan amount, performing appropriate clamping and rounding */
+	static FIntPoint ScaleResolutionByOverscan(const float OverscanPercentage, const FIntPoint& InOutputResolution);
 
 	/** Get the current cine camera in use, or nullptr if there is none. */
 	static UCineCameraComponent* Utility_GetCurrentCineCamera(const UWorld* InWorld);
diff --git a/Engine/Plugins/MovieScene/MovieRenderPipeline/Source/MovieRenderPipelineCore/Public/MovieRenderPipelineDataTypes.h b/Engine/Plugins/MovieScene/MovieRenderPipeline/Source/MovieRenderPipelineCore/Public/MovieRenderPipelineDataTypes.h
--- a/Engine/Plugins/MovieScene/MovieRenderPipeline/Source/MovieRenderPipelineCore/Public/MovieRenderPipelineDataTypes.h
+++ b/Engine/Plugins/MovieScene/MovieRenderPipeline/Source/MovieRenderPipelineCore/Public/MovieRenderPipelineDataTypes.h
@@ -930,6 +930,16 @@ struct FMoviePipelineFilenameResolveParams
 struct FMoviePipelineRenderPassMetrics
 {
 public:
+	// Manual default constructors to avoid deprecation warnings/errors on EffectiveOutputResolution
+	PRAGMA_DISABLE_DEPRECATION_WARNINGS
+	FMoviePipelineRenderPassMetrics() = default;
+	~FMoviePipelineRenderPassMetrics() = default;
+	FMoviePipelineRenderPassMetrics(const FMoviePipelineRenderPassMetrics&) = default;
+	FMoviePipelineRenderPassMetrics(FMoviePipelineRenderPassMetrics&&) = default;
+	FMoviePipelineRenderPassMetrics& operator=(const FMoviePipelineRenderPassMetrics&) = default;
+	FMoviePipelineRenderPassMetrics& operator=(FMoviePipelineRenderPassMetrics&&) = default;
+	PRAGMA_ENABLE_DEPRECATION_WARNINGS
+	
 	/** What FrameIndex should the FSceneView use for its internal index. Used to ensure samples have unique TAA/Raytracing/etc. sequences */
 	int32 FrameIndex;
 
@@ -976,8 +986,12 @@ public:
 	FIntPoint BackbufferSize;
 
 	/** Full Output image resolution for this shot. */
+	UE_DEPRECATED(5.6, "EffectiveOutputResolution is deprecated. Use OverscannedResolution instead")
 	FIntPoint EffectiveOutputResolution;
 
+	/** The overscanned resolution, computed by scaling the desired output resolution by any configured overscan */
+	FIntPoint OverscannedResolution;
+	
 	/** How big is an individual tile for this sample? This is the final target size divided by the number of tiles (without padding added) */
 	FIntPoint TileSize;
 
diff --git a/Engine/Plugins/MovieScene/MovieRenderPipeline/Source/MovieRenderPipelineRenderPasses/Private/Graph/Renderers/MovieGraphDeferredPass.cpp b/Engine/Plugins/MovieScene/MovieRenderPipeline/Source/MovieRenderPipelineRenderPasses/Private/Graph/Renderers/MovieGraphDeferredPass.cpp
--- a/Engine/Plugins/MovieScene/MovieRenderPipeline/Source/MovieRenderPipelineRenderPasses/Private/Graph/Renderers/MovieGraphDeferredPass.cpp
+++ b/Engine/Plugins/MovieScene/MovieRenderPipeline/Source/MovieRenderPipelineRenderPasses/Private/Graph/Renderers/MovieGraphDeferredPass.cpp
@@ -172,8 +172,9 @@ void FMovieGraphDeferredPass::Render(const FMovieGraphTraversalContext& InFrameT
 
 	// ToDo: When tiling is used, this should be the size of the per-tile backbuffer
 	const float CameraOverscan = GetRenderer()->GetCameraOverscan(LayerData.CameraIndex);
-	FIntPoint AccumulatorResolution = UMovieGraphBlueprintLibrary::GetEffectiveOutputResolution(InTimeData.EvaluatedConfig, CameraOverscan);
-	FIntPoint BackbufferResolution = AccumulatorResolution;
+	FIntPoint OverscannedResolution = UMovieGraphBlueprintLibrary::GetOverscannedResolution(InTimeData.EvaluatedConfig, CameraOverscan);
+	FIntPoint AccumulatorResolution = OverscannedResolution;
+	FIntPoint BackbufferResolution = UMovieGraphBlueprintLibrary::GetBackbufferResolution(InTimeData.EvaluatedConfig, CameraOverscan);
 	
 	DefaultRenderer::FRenderTargetInitParams RenderTargetInitParams = GetRenderTargetInitParams(InTimeData, AccumulatorResolution);
 	UTextureRenderTarget2D* RenderTarget = GraphRenderer->GetOrCreateViewRenderTarget(RenderTargetInitParams, RenderDataIdentifier);
@@ -299,6 +300,7 @@ void FMovieGraphDeferredPass::Render(const FMovieGraphTraversalContext& InFrameT
 			UpdatedTraversalContext.RenderDataIdentifier = RenderDataIdentifier;
 
 			SampleState.TraversalContext = MoveTemp(UpdatedTraversalContext);
+			SampleState.OverscannedResolution = OverscannedResolution;
 			SampleState.BackbufferResolution = BackbufferResolution;
 			SampleState.AccumulatorResolution = AccumulatorResolution;
 			SampleState.bWriteSampleToDisk = bWriteAllSamples;
diff --git a/Engine/Plugins/MovieScene/MovieRenderPipeline/Source/MovieRenderPipelineRenderPasses/Private/MovieGraphImageSequenceOutputNode.cpp b/Engine/Plugins/MovieScene/MovieRenderPipeline/Source/MovieRenderPipelineRenderPasses/Private/MovieGraphImageSequenceOutputNode.cpp
--- a/Engine/Plugins/MovieScene/MovieRenderPipeline/Source/MovieRenderPipelineRenderPasses/Private/MovieGraphImageSequenceOutputNode.cpp
+++ b/Engine/Plugins/MovieScene/MovieRenderPipeline/Source/MovieRenderPipelineRenderPasses/Private/MovieGraphImageSequenceOutputNode.cpp
@@ -160,11 +160,9 @@ void UMovieGraphImageSequenceOutputNode::OnReceiveImageDataImpl(UMovieGraphPipel
 		UE::MovieGraph::FMovieGraphSampleState* Payload = RenderData.Value->GetPayload<UE::MovieGraph::FMovieGraphSampleState>();
 
 		// If the overscan isn't cropped from the final image, offset any composites to the top-left of the original frustum
-		// TODO: Would prefer to check against "effective output resolution" (desired resolution scaled by overscan) but MRG does not pass that value down.
-		// This should be addressed in UE-228752
 		FIntPoint CompositeOffset = FIntPoint::ZeroValue;
 		const bool bIsCropRectValid = !Payload->CropRectangle.IsEmpty();
-		const bool bCanCropResolution = RenderData.Value->GetSize() == Payload->BackbufferResolution;
+		const bool bCanCropResolution = RenderData.Value->GetSize() == Payload->OverscannedResolution;
 		if (ShouldCropOverscanImpl() && bIsCropRectValid && bCanCropResolution)
 		{
 			switch (RenderData.Value->GetType())
@@ -311,6 +309,8 @@ void UMovieGraphImageSequenceOutputNode_EXR::UpdateTaskPerLayer(
 		InOutImageTask.Height = Resolution.Y;
 
 		InOutImageTask.OverscanPercentage = Payload->OverscanFraction;
+		InOutImageTask.CropRectangle = Payload->CropRectangle;
+		
 #if WITH_OCIO
 		if (bEnabledOCIO)
 		{
diff --git a/Engine/Plugins/MovieScene/MovieRenderPipeline/Source/MovieRenderPipelineRenderPasses/Private/MoviePipelineDeferredPasses.cpp b/Engine/Plugins/MovieScene/MovieRenderPipeline/Source/MovieRenderPipelineRenderPasses/Private/MoviePipelineDeferredPasses.cpp
--- a/Engine/Plugins/MovieScene/MovieRenderPipeline/Source/MovieRenderPipelineRenderPasses/Private/MoviePipelineDeferredPasses.cpp
+++ b/Engine/Plugins/MovieScene/MovieRenderPipeline/Source/MovieRenderPipelineRenderPasses/Private/MoviePipelineDeferredPasses.cpp
@@ -87,8 +87,11 @@ void UMoviePipelineDeferredPassBase::PostLoad()
 
 FIntPoint UMoviePipelineDeferredPassBase::GetEffectiveOutputResolutionForCamera(const int32 InCameraIndex) const
 {
-	// Add here the the output resolution for each camera. Now only one size is implemented, obtained from the settings.
+	return GetOverscannedResolutionForCamera(InCameraIndex);
+}
 
+FIntPoint UMoviePipelineDeferredPassBase::GetOverscannedResolutionForCamera(const int32 InCameraIndex) const
+{
 	UMoviePipelinePrimaryConfig* PrimaryConfig = GetPipeline()->GetPipelinePrimaryConfig();
 	UMoviePipelineExecutorShot* CurrentShot = GetPipeline()->GetActiveShotList()[GetPipeline()->GetCurrentShotIndex()];
 
@@ -100,7 +103,7 @@ FIntPoint UMoviePipelineDeferredPassBase::GetEffectiveOutputResolutionForCamera(
 	}
 	else
 	{
-		// Get the camera view info to retrieve the camera's overscan, which is used when the settings to not override the overscan
+		// Get the camera view info to retrieve the camera's overscan, which is used when the settings do not override the overscan
 		FMinimalViewInfo CameraViewInfo;
 
 		if (GetNumCamerasToRender() == 1)
@@ -117,9 +120,7 @@ FIntPoint UMoviePipelineDeferredPassBase::GetEffectiveOutputResolutionForCamera(
 		GetPipeline()->CacheCameraOverscan(InCameraIndex, CameraOverscan);
 	}
 	
-	const FIntPoint OutputResolution = UMoviePipelineBlueprintLibrary::GetEffectiveOutputResolution(PrimaryConfig, CurrentShot, CameraOverscan);
-
-	return OutputResolution;
+	return UMoviePipelineBlueprintLibrary::GetOverscannedResolution(PrimaryConfig, CurrentShot, CameraOverscan);
 }
 
 FMoviePipelineRenderPassMetrics UMoviePipelineDeferredPassBase::GetRenderPassMetricsForCamera(const int32 InCameraIndex, const FMoviePipelineRenderPassMetrics& InSampleState) const
@@ -130,7 +131,7 @@ FMoviePipelineRenderPassMetrics UMoviePipelineDeferredPassBase::GetRenderPassMet
 	check(PrimaryConfig);
 	check(CurrentShot);
 
-	return UE::MoviePipeline::GetRenderPassMetrics(PrimaryConfig, CurrentShot, InSampleState, GetEffectiveOutputResolutionForCamera(InCameraIndex));
+	return UE::MoviePipeline::GetRenderPassMetrics(PrimaryConfig, CurrentShot, InSampleState, GetOverscannedResolutionForCamera(InCameraIndex));
 }
 
 int32 UMoviePipelineDeferredPassBase::GetNumCamerasToRender() const
@@ -225,12 +226,12 @@ void UMoviePipelineDeferredPassBase::SetupImpl(const MoviePipeline::FMoviePipeli
 	int32 TotalNumberOfAccumulators = 0;
 	for (int32 CamIndex = 0; CamIndex < NumCameras; CamIndex++)
 	{
-		const FIntPoint OutputCameraResolution = GetEffectiveOutputResolutionForCamera(CamIndex);
+		const FIntPoint CameraResolution = GetOverscannedResolutionForCamera(CamIndex);
 
 		// Figure out how big each sub-region (tile) is.
 		FIntPoint BackbufferResolution = FIntPoint(
-			FMath::CeilToInt((float)OutputCameraResolution.X / (float)HighResSettings->TileCount),
-			FMath::CeilToInt((float)OutputCameraResolution.Y / (float)HighResSettings->TileCount));
+			FMath::CeilToInt((float)CameraResolution.X / (float)HighResSettings->TileCount),
+			FMath::CeilToInt((float)CameraResolution.Y / (float)HighResSettings->TileCount));
 
 		// Then increase each sub-region by the overlap amount.
 		BackbufferResolution = HighResSettings->CalculatePaddedBackbufferSize(BackbufferResolution);
@@ -923,7 +924,7 @@ void UMoviePipelineDeferredPassBase::PostRendererSubmission(const FMoviePipeline
 		check(OutputSettings);
 		
 		// Taking overscan into account.
-		const FIntPoint FullOutputSize = InSampleState.EffectiveOutputResolution;
+		const FIntPoint FullOutputSize = InSampleState.OverscannedResolution;
 
 		const float OutputSizeAspectRatio = FullOutputSize.X / (float)FullOutputSize.Y;
 		const float CameraAspectRatio = bAllowCameraAspectRatio ? CameraCache.AspectRatio : OutputSizeAspectRatio;
diff --git a/Engine/Plugins/MovieScene/MovieRenderPipeline/Source/MovieRenderPipelineRenderPasses/Private/MoviePipelineEXROutput.cpp b/Engine/Plugins/MovieScene/MovieRenderPipeline/Source/MovieRenderPipelineRenderPasses/Private/MoviePipelineEXROutput.cpp
--- a/Engine/Plugins/MovieScene/MovieRenderPipeline/Source/MovieRenderPipelineRenderPasses/Private/MoviePipelineEXROutput.cpp
+++ b/Engine/Plugins/MovieScene/MovieRenderPipeline/Source/MovieRenderPipelineRenderPasses/Private/MoviePipelineEXROutput.cpp
@@ -130,13 +130,13 @@ bool FEXRImageWriteTask::WriteToDisk()
 		// Data Window specifies how much data is in the actual file, ie: 1920x1080
 		IMATH_NAMESPACE::Box2i DataWindow = IMATH_NAMESPACE::Box2i(IMATH_NAMESPACE::V2i(0,0), IMATH_NAMESPACE::V2i(Width - 1, Height - 1));
 		
-		float BorderRatio = FMath::Clamp((OverscanPercentage / (1.0f + OverscanPercentage)) / 2.0f, 0.0f, 0.5f);
-		int32 OverscanBorderWidth = FMath::RoundToInt(Width * BorderRatio);
-		int32 OverscanBorderHeight = FMath::RoundToInt(Height * BorderRatio);
-
+		const bool bIsCropRectValid = !CropRectangle.IsEmpty() &&
+			CropRectangle.Min.X >= 0 && CropRectangle.Min.Y >= 0 &&
+			CropRectangle.Max.X <= Width && CropRectangle.Max.Y <= Height;
+		
 		// Taking Overscan into account.
-		IMATH_NAMESPACE::V2i TopLeft = IMATH_NAMESPACE::V2i(OverscanBorderWidth, OverscanBorderHeight);
-		IMATH_NAMESPACE::V2i BottomRight = IMATH_NAMESPACE::V2i(Width - OverscanBorderWidth - 1, Height - OverscanBorderHeight - 1);
+		IMATH_NAMESPACE::V2i TopLeft = IMATH_NAMESPACE::V2i(bIsCropRectValid ? CropRectangle.Min.X : 0, bIsCropRectValid ? CropRectangle.Min.Y : 0);
+		IMATH_NAMESPACE::V2i BottomRight = IMATH_NAMESPACE::V2i(bIsCropRectValid ? CropRectangle.Max.X - 1 : Width - 1, bIsCropRectValid ? CropRectangle.Max.Y - 1 : Height - 1);
 
 		// Display Window specifies the total 'visible' area of the output file. 
 		// The Display Window always starts at 0,0, but Data Window can go negative to
@@ -662,6 +662,7 @@ void UMoviePipelineImageSequenceOutput_EXR::OnReceiveImageDataImpl(FMoviePipelin
 				// Only check the main image pass for transparent output since that's generally considered the 'preview'.
 				bRequiresTransparentOutput = Payload->bRequireTransparentOutput;
 				MultiLayerImageTask->OverscanPercentage = Payload->SampleState.OverscanPercentage;
+				MultiLayerImageTask->CropRectangle = Payload->SampleState.CropRectangle;
 			}
 			else
 			{
diff --git a/Engine/Plugins/MovieScene/MovieRenderPipeline/Source/MovieRenderPipelineRenderPasses/Private/MoviePipelineImageSequenceOutput.cpp b/Engine/Plugins/MovieScene/MovieRenderPipeline/Source/MovieRenderPipelineRenderPasses/Private/MoviePipelineImageSequenceOutput.cpp
--- a/Engine/Plugins/MovieScene/MovieRenderPipeline/Source/MovieRenderPipelineRenderPasses/Private/MoviePipelineImageSequenceOutput.cpp
+++ b/Engine/Plugins/MovieScene/MovieRenderPipeline/Source/MovieRenderPipelineRenderPasses/Private/MoviePipelineImageSequenceOutput.cpp
@@ -210,7 +210,7 @@ void UMoviePipelineImageSequenceOutputBase::OnReceiveImageDataImpl(FMoviePipelin
 
 		// For now, only passes that were rendered at the overscanned resolution can be cropped using the crop rectangle
 		const bool bIsCropRectValid = !Payload->SampleState.CropRectangle.IsEmpty();
-		const bool bCanCropResolution = RenderPassData.Value->GetSize() == Payload->SampleState.EffectiveOutputResolution;
+		const bool bCanCropResolution = RenderPassData.Value->GetSize() == Payload->SampleState.OverscannedResolution;
 		if (ShouldCropOverscanImpl() && bIsCropRectValid && bCanCropResolution)
 		{
 			switch (RenderPassData.Value->GetType())
diff --git a/Engine/Plugins/MovieScene/MovieRenderPipeline/Source/MovieRenderPipelineRenderPasses/Public/MoviePipelineDeferredPasses.h b/Engine/Plugins/MovieScene/MovieRenderPipeline/Source/MovieRenderPipelineRenderPasses/Public/MoviePipelineDeferredPasses.h
--- a/Engine/Plugins/MovieScene/MovieRenderPipeline/Source/MovieRenderPipelineRenderPasses/Public/MoviePipelineDeferredPasses.h
+++ b/Engine/Plugins/MovieScene/MovieRenderPipeline/Source/MovieRenderPipelineRenderPasses/Public/MoviePipelineDeferredPasses.h
@@ -87,8 +87,11 @@ protected:
 	virtual FString GetCameraNameOverride(const int32 InCameraIndex) const;
 
 	virtual FMoviePipelineRenderPassMetrics GetRenderPassMetricsForCamera(const int32 InCameraIndex, const FMoviePipelineRenderPassMetrics& InSampleState) const;
-	virtual FIntPoint GetEffectiveOutputResolutionForCamera(const int32 InCameraIndex) const;
 
+	UE_DEPRECATED(5.6, "GetEffectiveOutputResolutionForCamera is deprecated. Use GetOverscannedResolutionForCamera instead")
+	virtual FIntPoint GetEffectiveOutputResolutionForCamera(const int32 InCameraIndex) const;
+	virtual FIntPoint GetOverscannedResolutionForCamera(const int32 InCameraIndex) const;
+	
 	bool IsUsingDataLayers() const;
 	int32 GetNumStencilLayers() const;
 	TArray<FString> GetStencilLayerNames() const;
diff --git a/Engine/Plugins/MovieScene/MovieRenderPipeline/Source/MovieRenderPipelineRenderPasses/Public/MoviePipelineEXROutput.h b/Engine/Plugins/MovieScene/MovieRenderPipeline/Source/MovieRenderPipelineRenderPasses/Public/MoviePipelineEXROutput.h
--- a/Engine/Plugins/MovieScene/MovieRenderPipeline/Source/MovieRenderPipelineRenderPasses/Public/MoviePipelineEXROutput.h
+++ b/Engine/Plugins/MovieScene/MovieRenderPipeline/Source/MovieRenderPipelineRenderPasses/Public/MoviePipelineEXROutput.h
@@ -113,6 +113,9 @@ public:
 	/** Overscan info used to create apropriate dataWindow for EXR output. Goes from 0.0 to 1.0. */
 	float OverscanPercentage;
 
+	/** Crop rectangle used to set the displayWindow for EXR output. Used to crop out overscan */
+	FIntRect CropRectangle;
+	
 	/** Color space chromaticity metadata. */
 	TArray<FVector2d> ColorSpaceChromaticities;
 
@@ -121,6 +124,7 @@ public:
 		, Compression(EEXRCompressionFormat::PIZ)
 		, CompressionLevel(45)
 		, OverscanPercentage(0.0f)
+		, CropRectangle(FIntRect())
 	{}
 
 public:
diff --git a/Engine/Plugins/Runtime/nDisplay/Source/DisplayClusterMoviePipeline/Private/DisplayClusterMoviePipelineViewportPass.cpp b/Engine/Plugins/Runtime/nDisplay/Source/DisplayClusterMoviePipeline/Private/DisplayClusterMoviePipelineViewportPass.cpp
--- a/Engine/Plugins/Runtime/nDisplay/Source/DisplayClusterMoviePipeline/Private/DisplayClusterMoviePipelineViewportPass.cpp
+++ b/Engine/Plugins/Runtime/nDisplay/Source/DisplayClusterMoviePipeline/Private/DisplayClusterMoviePipelineViewportPass.cpp
@@ -124,7 +124,7 @@ void UDisplayClusterMoviePipelineViewportPassBase::TeardownImpl()
 	Super::TeardownImpl();
 }
 
-FIntPoint UDisplayClusterMoviePipelineViewportPassBase::GetEffectiveOutputResolutionForCamera(const int32 InCameraIndex) const
+FIntPoint UDisplayClusterMoviePipelineViewportPassBase::GetOverscannedResolutionForCamera(const int32 InCameraIndex) const
 {
 	if (DisplayClusterViewportSizes.IsValidIndex(InCameraIndex))
 	{
@@ -134,7 +134,7 @@ FIntPoint UDisplayClusterMoviePipelineViewportPassBase::GetEffectiveOutputResolu
 	// Use one size from OutputSettings
 	UMoviePipelinePrimaryConfig* PrimaryConfig = GetPipeline()->GetPipelinePrimaryConfig();
 	UMoviePipelineExecutorShot* CurrentShot = GetPipeline()->GetActiveShotList()[GetPipeline()->GetCurrentShotIndex()];
-	const FIntPoint OutputResolution = UMoviePipelineBlueprintLibrary::GetEffectiveOutputResolution(PrimaryConfig, CurrentShot);
+	const FIntPoint OutputResolution = UMoviePipelineBlueprintLibrary::GetOverscannedResolution(PrimaryConfig, CurrentShot);
 
 	return OutputResolution;
 }
@@ -399,7 +399,7 @@ void UDisplayClusterMoviePipelineViewportPassBase::GetViewportCutOffset(const FM
 		const FMinimalViewInfo CameraCache = PlayerCameraManager->GetCameraCacheView();
 
 		// Taking overscan into account.
-		FIntPoint FullOutputSize = UMoviePipelineBlueprintLibrary::GetEffectiveOutputResolution(GetPipeline()->GetPipelinePrimaryConfig(), GetPipeline()->GetActiveShotList()[GetPipeline()->GetCurrentShotIndex()]);
+		FIntPoint FullOutputSize = UMoviePipelineBlueprintLibrary::GetOverscannedResolution(GetPipeline()->GetPipelinePrimaryConfig(), GetPipeline()->GetActiveShotList()[GetPipeline()->GetCurrentShotIndex()]);
 
 		float OutputSizeAspectRatio = FullOutputSize.X / (float)FullOutputSize.Y;
 		const FIntPoint ConstrainedFullSize = CameraCache.AspectRatio > OutputSizeAspectRatio ?
diff --git a/Engine/Plugins/Runtime/nDisplay/Source/DisplayClusterMoviePipeline/Public/DisplayClusterMoviePipelineViewportPass.h b/Engine/Plugins/Runtime/nDisplay/Source/DisplayClusterMoviePipeline/Public/DisplayClusterMoviePipelineViewportPass.h
--- a/Engine/Plugins/Runtime/nDisplay/Source/DisplayClusterMoviePipeline/Public/DisplayClusterMoviePipelineViewportPass.h
+++ b/Engine/Plugins/Runtime/nDisplay/Source/DisplayClusterMoviePipeline/Public/DisplayClusterMoviePipelineViewportPass.h
@@ -44,7 +44,7 @@ public:
 	virtual UE::MoviePipeline::FImagePassCameraViewData GetCameraInfo(FMoviePipelineRenderPassMetrics& InOutSampleState, IViewCalcPayload* OptPayload = nullptr) const override;
 	virtual void BlendPostProcessSettings(FSceneView* InView, FMoviePipelineRenderPassMetrics& InOutSampleState, IViewCalcPayload* OptPayload = nullptr) override;
 
-	virtual FIntPoint GetEffectiveOutputResolutionForCamera(const int32 InCameraIndex) const override;
+	virtual FIntPoint GetOverscannedResolutionForCamera(const int32 InCameraIndex) const override;
 	// ~UMoviePipelineDeferredPassBase
 
 #if WITH_EDITOR
