[MRG] Seedoffset feature + metadata for it and camera matrices.
Added a seed offset exposed to the path tracer via MRQ render graph pathtracer node
Added metadata export of the seedoffset and of the camera matrices

Problems:
* To generate data for ML training we need a way to generate individual samples using deterministic but differend seeds for the pathtracer random generator. This parameter should be able to be controlled via the MRG to be able to extract data from animation clips deterministically.
* We also need the frame camera matrices to be exported from MRG toward the same generate data for ML training goal.

Solutions:
* Add a seed offset exposed to the path tracer via MRG pathtracer node as blueprint only that can be controlled with MTG scripting capabilities.
* Export seedoffset and camera matrices using MRG metadata functionality.

Test:
* Generate ML data for a sequence a 65 frames from an animation sequence using 8 different seed and verified reproducibility of sampling patterns.
* Verified the generated frame are containing the seed offset and camera matrices in the file metadata.

#rnx

diff --git a/Engine/Plugins/MovieScene/MovieRenderPipeline/Source/MovieRenderPipelineCore/Public/Graph/MovieGraphConfig.h b/Engine/Plugins/MovieScene/MovieRenderPipeline/Source/MovieRenderPipelineCore/Public/Graph/MovieGraphConfig.h
--- a/Engine/Plugins/MovieScene/MovieRenderPipeline/Source/MovieRenderPipelineCore/Public/Graph/MovieGraphConfig.h
+++ b/Engine/Plugins/MovieScene/MovieRenderPipeline/Source/MovieRenderPipelineCore/Public/Graph/MovieGraphConfig.h
@@ -496,9 +496,9 @@ public:
 			}
 		}
 
-		// If they didn't return above, then either they specified an invalid branch (for which the ensure tripped)
+		// If we didn't found results above, then either they specified an invalid branch (for which the ensure tripped)
 		// or the config simply didn't override that setting class, at which point we might try to return a default
-		if (bIncludeCDOs)
+		if (bIncludeCDOs && ResultNodes.Num() == 0)
 		{
 			ResultNodes.Add(Cast<UMovieGraphSettingNode>(InClass->GetDefaultObject()));
 		}
diff --git a/Engine/Plugins/MovieScene/MovieRenderPipeline/Source/MovieRenderPipelineCore/Public/Graph/MovieGraphDefaultRenderer.h b/Engine/Plugins/MovieScene/MovieRenderPipeline/Source/MovieRenderPipelineCore/Public/Graph/MovieGraphDefaultRenderer.h
--- a/Engine/Plugins/MovieScene/MovieRenderPipeline/Source/MovieRenderPipelineCore/Public/Graph/MovieGraphDefaultRenderer.h
+++ b/Engine/Plugins/MovieScene/MovieRenderPipeline/Source/MovieRenderPipelineCore/Public/Graph/MovieGraphDefaultRenderer.h
@@ -41,6 +41,7 @@ namespace UE::MovieGraph::DefaultRenderer
 		int32 TemporalSampleCount;
 		int32 SpatialSampleIndex;
 		int32 SpatialSampleCount;
+		int32 SeedOffset;
 	};
 
 	struct FCameraInfo
diff --git a/Engine/Plugins/MovieScene/MovieRenderPipeline/Source/MovieRenderPipelineCore/Public/MovieRenderPipelineDataTypes.h b/Engine/Plugins/MovieScene/MovieRenderPipeline/Source/MovieRenderPipelineCore/Public/MovieRenderPipelineDataTypes.h
--- a/Engine/Plugins/MovieScene/MovieRenderPipeline/Source/MovieRenderPipelineCore/Public/MovieRenderPipelineDataTypes.h
+++ b/Engine/Plugins/MovieScene/MovieRenderPipeline/Source/MovieRenderPipelineCore/Public/MovieRenderPipelineDataTypes.h
@@ -1007,6 +1007,9 @@ public:
 	/** Of the SpatialSampleCount, which index is this? */
 	int32 SpatialSampleIndex;
 
+	/** Offset to apply to random number generator seed. */
+	int32 SeedOffset;
+
 	/** How many temporal jitters will there be total for this particular output frame? */
 	int32 TemporalSampleCount;
 
diff --git a/Engine/Plugins/MovieScene/MovieRenderPipeline/Source/MovieRenderPipelineRenderPasses/Private/Graph/Nodes/MovieGraphPathTracerPassNode.cpp b/Engine/Plugins/MovieScene/MovieRenderPipeline/Source/MovieRenderPipelineRenderPasses/Private/Graph/Nodes/MovieGraphPathTracerPassNode.cpp
--- a/Engine/Plugins/MovieScene/MovieRenderPipeline/Source/MovieRenderPipelineRenderPasses/Private/Graph/Nodes/MovieGraphPathTracerPassNode.cpp
+++ b/Engine/Plugins/MovieScene/MovieRenderPipeline/Source/MovieRenderPipelineRenderPasses/Private/Graph/Nodes/MovieGraphPathTracerPassNode.cpp
@@ -16,6 +16,7 @@ TUniquePtr<UE::MovieGraph::Rendering::FMovieGraphImagePassBase> UMovieGraphPathT
 
 UMovieGraphPathTracerRenderPassNode::UMovieGraphPathTracerRenderPassNode()
 	: SpatialSampleCount(1)
+	, SeedOffset(0)
 	, bEnableReferenceMotionBlur(true)
 	, bEnableDenoiser(true)
 	, DenoiserType(EMovieGraphPathTracerDenoiserType::Spatial)
@@ -41,6 +42,15 @@ UMovieGraphPathTracerRenderPassNode::UMovieGraphPathTracerRenderPassNode()
 	}
 }
 
+void UMovieGraphPathTracerRenderPassNode::GetFormatResolveArgs(FMovieGraphResolveArgs& OutMergedFormatArgs, const FMovieGraphRenderDataIdentifier& InRenderDataIdentifier) const
+{
+	OutMergedFormatArgs.FilenameArguments.Add(TEXT("path_tracer_sampling_count"), FString::FromInt(SpatialSampleCount));
+	OutMergedFormatArgs.FilenameArguments.Add(TEXT("path_tracer_seed_offset"), FString::FromInt(SeedOffset));
+
+	OutMergedFormatArgs.FileMetadata.Add(TEXT("unreal/pathTracer/samplingCount"), FString::FromInt(SpatialSampleCount));
+	OutMergedFormatArgs.FileMetadata.Add(TEXT("unreal/pathTracer/seedOffset"), FString::FromInt(SeedOffset));
+}
+
 #if WITH_EDITOR
 FText UMovieGraphPathTracerRenderPassNode::GetNodeTitle(const bool bGetDescriptive) const
 {
@@ -197,6 +207,11 @@ int32 UMovieGraphPathTracerRenderPassNode::GetNumSpatialSamples() const
 	return SpatialSampleCount;
 }
 
+int32 UMovieGraphPathTracerRenderPassNode::GetSeedOffset() const
+{
+	return SeedOffset;
+}
+
 int32 UMovieGraphPathTracerRenderPassNode::GetNumSpatialSamplesDuringWarmUp() const
 {
 	// Path Tracer doesn't have an image history like the deferred renderer, so it doesn't need
diff --git a/Engine/Plugins/MovieScene/MovieRenderPipeline/Source/MovieRenderPipelineRenderPasses/Private/Graph/Renderers/MovieGraphDeferredPass.cpp b/Engine/Plugins/MovieScene/MovieRenderPipeline/Source/MovieRenderPipelineRenderPasses/Private/Graph/Renderers/MovieGraphDeferredPass.cpp
--- a/Engine/Plugins/MovieScene/MovieRenderPipeline/Source/MovieRenderPipelineRenderPasses/Private/Graph/Renderers/MovieGraphDeferredPass.cpp
+++ b/Engine/Plugins/MovieScene/MovieRenderPipeline/Source/MovieRenderPipelineRenderPasses/Private/Graph/Renderers/MovieGraphDeferredPass.cpp
@@ -30,6 +30,14 @@
 
 namespace UE::MovieGraph::Rendering
 {
+	namespace MetadataHelper
+	{
+		void UpdateSpatialSampleMetadata(const FSceneViewProjectionData& InProjectionData, TMap<FString, FString>& InOutMetadataMap)
+		{
+			InOutMetadataMap.Add(FString::Printf(TEXT("unreal/sampling/matrix/worldToCamera")), InProjectionData.ViewRotationMatrix.ToString());
+			InOutMetadataMap.Add(FString::Printf(TEXT("unreal/sampling/matrix/worldToUEClip")), InProjectionData.ProjectionMatrix.ToString());
+		}
+	}
 
 void FMovieGraphDeferredPass::Setup(TWeakObjectPtr<UMovieGraphDefaultRenderer> InRenderer, TWeakObjectPtr<UMovieGraphImagePassBaseNode> InRenderPassNode, const FMovieGraphRenderPassLayerData& InLayer)
 {
@@ -152,6 +160,7 @@ void FMovieGraphDeferredPass::Render(const FMovieGraphTraversalContext& InFrameT
 	const bool bIsRenderingState = InFrameTraversalContext.Shot->ShotInfo.State == EMovieRenderShotState::Rendering ||
 									InFrameTraversalContext.Shot->ShotInfo.State == EMovieRenderShotState::CoolingDown;
 	int32 NumSpatialSamples = FMath::Max(1, bIsRenderingState ? ParentNodeThisFrame->GetNumSpatialSamples() : ParentNodeThisFrame->GetNumSpatialSamplesDuringWarmUp());
+	const int32 SeedOffset = ParentNodeThisFrame->GetSeedOffset();
 
 	const ESceneCaptureSource SceneCaptureSource = ParentNodeThisFrame->GetDisableToneCurve() ? ESceneCaptureSource::SCS_FinalColorHDR : ESceneCaptureSource::SCS_FinalToneCurveHDR;
 	const EAntiAliasingMethod AntiAliasingMethod = ParentNodeThisFrame->GetAntiAliasingMethod();
@@ -214,6 +223,7 @@ void FMovieGraphDeferredPass::Render(const FMovieGraphTraversalContext& InFrameT
 		CameraInfo.SamplingParams.TemporalSampleCount = InTimeData.TemporalSampleCount;
 		CameraInfo.SamplingParams.SpatialSampleIndex = SpatialIndex;
 		CameraInfo.SamplingParams.SpatialSampleCount = NumSpatialSamples;
+		CameraInfo.SamplingParams.SeedOffset = SeedOffset;
 		CameraInfo.ProjectionMatrixJitterAmount = FVector2D((SpatialShiftAmount.X) * 2.0f / (float)BackbufferResolution.X, SpatialShiftAmount.Y * -2.0f / (float)BackbufferResolution.Y);
 		CameraInfo.bUseCameraManagerPostProcess = !bRenderAllCameras;
 		
@@ -319,6 +329,8 @@ void FMovieGraphDeferredPass::Render(const FMovieGraphTraversalContext& InFrameT
 
 		ApplyMovieGraphOverridesToSampleState(SampleState);
 
+		MetadataHelper::UpdateSpatialSampleMetadata(SceneViewInitOptions, SampleState.AdditionalFileMetadata);
+
 		// If this was just to contribute to the history buffer, no need to go any further. Never discard if we're writing individual samples, though.
 		bool bDiscardOutput = (InTimeData.bDiscardOutput || ShouldDiscardOutput(ViewFamily, CameraInfo)) && !SampleState.bWriteSampleToDisk;
 
diff --git a/Engine/Plugins/MovieScene/MovieRenderPipeline/Source/MovieRenderPipelineRenderPasses/Private/Graph/Renderers/MovieGraphPathTracerPass.cpp b/Engine/Plugins/MovieScene/MovieRenderPipeline/Source/MovieRenderPipelineRenderPasses/Private/Graph/Renderers/MovieGraphPathTracerPass.cpp
--- a/Engine/Plugins/MovieScene/MovieRenderPipeline/Source/MovieRenderPipelineRenderPasses/Private/Graph/Renderers/MovieGraphPathTracerPass.cpp
+++ b/Engine/Plugins/MovieScene/MovieRenderPipeline/Source/MovieRenderPipelineRenderPasses/Private/Graph/Renderers/MovieGraphPathTracerPass.cpp
@@ -101,6 +101,8 @@ void FMovieGraphPathTracerPass::ApplyMovieGraphOverridesToSceneView(TSharedRef<F
 	
 	// reset path tracer's accumulation at the start of each sample
 	View->bForcePathTracerReset = SampleIndex == 0;
+	// set path tracer seed from sampling parameters
+	View->PathTracerSeedOffset = InCameraInfo.SamplingParams.SeedOffset;
 }
 
 void FMovieGraphPathTracerPass::ApplyMovieGraphOverridesToSampleState(FMovieGraphSampleState& SampleState) const
diff --git a/Engine/Plugins/MovieScene/MovieRenderPipeline/Source/MovieRenderPipelineRenderPasses/Public/Graph/Nodes/MovieGraphImagePassBaseNode.h b/Engine/Plugins/MovieScene/MovieRenderPipeline/Source/MovieRenderPipelineRenderPasses/Public/Graph/Nodes/MovieGraphImagePassBaseNode.h
--- a/Engine/Plugins/MovieScene/MovieRenderPipeline/Source/MovieRenderPipelineRenderPasses/Public/Graph/Nodes/MovieGraphImagePassBaseNode.h
+++ b/Engine/Plugins/MovieScene/MovieRenderPipeline/Source/MovieRenderPipelineRenderPasses/Public/Graph/Nodes/MovieGraphImagePassBaseNode.h
@@ -50,6 +50,9 @@ public:
 	/** How many spatial samples should be used during warm-up frames? */
 	virtual int32 GetNumSpatialSamplesDuringWarmUp() const { return GetNumSpatialSamples(); }
 
+	/** Offset to apply to random number generator seed. */
+	virtual int32 GetSeedOffset() const { return 0; }
+
 	/** Should the tone curve be disabled while rendering? Allows for linear values in exrs but changes the look of the final image. */
 	virtual bool GetDisableToneCurve() const { return false; }
 
diff --git a/Engine/Plugins/MovieScene/MovieRenderPipeline/Source/MovieRenderPipelineRenderPasses/Public/Graph/Nodes/MovieGraphPathTracerPassNode.h b/Engine/Plugins/MovieScene/MovieRenderPipeline/Source/MovieRenderPipelineRenderPasses/Public/Graph/Nodes/MovieGraphPathTracerPassNode.h
--- a/Engine/Plugins/MovieScene/MovieRenderPipeline/Source/MovieRenderPipelineRenderPasses/Public/Graph/Nodes/MovieGraphPathTracerPassNode.h
+++ b/Engine/Plugins/MovieScene/MovieRenderPipeline/Source/MovieRenderPipelineRenderPasses/Public/Graph/Nodes/MovieGraphPathTracerPassNode.h
@@ -50,6 +50,7 @@ public:
 	virtual TArray<FMoviePipelinePostProcessPass> GetAdditionalPostProcessMaterials() const override;
 	virtual int32 GetNumSpatialSamples() const override;
 	virtual int32 GetNumSpatialSamplesDuringWarmUp() const override;
+	virtual int32 GetSeedOffset() const override;
 	virtual bool GetDisableToneCurve() const override;
 	virtual bool GetAllowOCIO() const override;
 	virtual bool GetAllowDenoiser() const override;
@@ -57,6 +58,10 @@ public:
 	virtual FEngineShowFlags GetShowFlags() const override;
 	// ~UMovieGraphImagePassBaseNode Interface
 
+	// UMovieGraphSettingNode Interface
+	virtual void GetFormatResolveArgs(FMovieGraphResolveArgs& OutMergedFormatArgs, const FMovieGraphRenderDataIdentifier& InRenderDataIdentifier) const override;
+	// ~UMovieGraphSettingNode Interface
+
 	virtual void UpdateTelemetry(FMoviePipelineShotRenderTelemetry* InTelemetry) const override;
 
 protected:
@@ -74,6 +79,9 @@ public:
 	UPROPERTY(EditAnywhere, BlueprintReadWrite, Category = Overrides, meta = (InlineEditConditionToggle))
 	uint8 bOverride_SpatialSampleCount : 1;
 
+	UPROPERTY(BlueprintReadWrite, Category = Overrides, meta = (InlineEditConditionToggle))
+	uint8 bOverride_SeedOffset : 1;
+
 	UPROPERTY(EditAnywhere, BlueprintReadWrite, Category = Overrides, meta = (InlineEditConditionToggle))
 	uint8 bOverride_bEnableReferenceMotionBlur : 1;
 
@@ -128,6 +136,13 @@ public:
 	UPROPERTY(EditAnywhere, BlueprintReadWrite, meta = (UIMin = 1, ClampMin = 1), Category = "Sampling", meta = (EditCondition = "bOverride_SpatialSampleCount"))
 	int32 SpatialSampleCount;
 
+	/**
+	* Offset to apply to random number generator seed.
+	* Intentionally not exposed to the UI as meant for automated pipelines. Should be edited via scripting if needs to be changed.
+	*/
+	UPROPERTY(BlueprintReadWrite, Category = "Sampling", AdvancedDisplay, meta = (EditCondition = "bOverride_SeedOffset"))
+	int32 SeedOffset;
+
 	/** 
 	 *  When enabled, the path tracer will blend all spatial and temporal samples prior to the denoising and will disable post-processed motion blur.
 	 *  In this mode it is possible to use higher temporal sample counts to improve the motion blur quality. This mode also automatically enabled reference DOF.
diff --git a/Engine/Source/Runtime/Engine/Private/SceneView.cpp b/Engine/Source/Runtime/Engine/Private/SceneView.cpp
--- a/Engine/Source/Runtime/Engine/Private/SceneView.cpp
+++ b/Engine/Source/Runtime/Engine/Private/SceneView.cpp
@@ -894,10 +894,11 @@ PRAGMA_ENABLE_DEPRECATION_WARNINGS
 	, bIsMultiViewportEnabled(false)
 	, bIsMobileMultiViewEnabled(false)
 	, bShouldBindInstancedViewUB(false)
-	, UnderwaterDepth(-1.0f)
 	, bForceCameraVisibilityReset(false)
-	, bForcePathTracerReset(false)
 	, bDisableDistanceBasedFadeTransitions(false)
+	, bForcePathTracerReset(false)
+	, PathTracerSeedOffset(0)
+	, UnderwaterDepth(-1.0f)
 	, GlobalClippingPlane(FPlane(0, 0, 0, 0))
 	, LensPrincipalPointOffsetScale(0.0f, 0.0f, 1.0f, 1.0f)
 #if WITH_EDITOR
diff --git a/Engine/Source/Runtime/Engine/Public/SceneView.h b/Engine/Source/Runtime/Engine/Public/SceneView.h
--- a/Engine/Source/Runtime/Engine/Public/SceneView.h
+++ b/Engine/Source/Runtime/Engine/Public/SceneView.h
@@ -1595,19 +1595,22 @@ public:
 	/** True if we need to bind the instanced view uniform buffer parameters. */
 	bool bShouldBindInstancedViewUB;
 
-	/** How far below the water surface this view is. -1 means the view is out of water. */
-	float UnderwaterDepth;
-
 	/** True if we need to force the camera to discard previous frames occlusion. Necessary for overlapped tile rendering
 	 * where we discard previous frame occlusion because the projection matrix changes.
 	 */
 	bool bForceCameraVisibilityReset;
 
+	/** Whether we should disable distance-based fade transitions for this frame (usually after a large camera movement.) */
+	bool bDisableDistanceBasedFadeTransitions;
+
 	/** True if we should force the path tracer to reset its internal accumulation state */
 	bool bForcePathTracerReset;
 
-	/** Whether we should disable distance-based fade transitions for this frame (usually after a large camera movement.) */
-	bool bDisableDistanceBasedFadeTransitions;
+	/** Sets an offset to the seed the path tracer is using. */
+	int32 PathTracerSeedOffset;
+
+	/** How far below the water surface this view is. -1 means the view is out of water. */
+	float UnderwaterDepth;
 
 	/** Global clipping plane being applied to the scene, or all 0's if disabled.  This is used when rendering the planar reflection pass. */
 	FPlane GlobalClippingPlane;
diff --git a/Engine/Source/Runtime/Renderer/Private/PathTracing.cpp b/Engine/Source/Runtime/Renderer/Private/PathTracing.cpp
--- a/Engine/Source/Runtime/Renderer/Private/PathTracing.cpp
+++ b/Engine/Source/Runtime/Renderer/Private/PathTracing.cpp
@@ -3228,6 +3228,8 @@ void FDeferredShadingSceneRenderer::RenderPathTracing(
 	uint32 MaxSPP = SamplesPerPixelCVar > -1 ? SamplesPerPixelCVar : View.FinalPostProcessSettings.PathTracingSamplesPerPixel;
 	MaxSPP = FMath::Max(MaxSPP, 1u);
 
+	int32 SeedOffset = View.PathTracerSeedOffset;
+
 	const bool bUseExperimental = CVarPathTracingExperimental.GetValueOnRenderThread();
 
 	Config.LockedSamplingPattern = CVarPathTracingFrameIndependentTemporalSeed.GetValueOnRenderThread() == 0;
@@ -3530,6 +3532,8 @@ void FDeferredShadingSceneRenderer::RenderPathTracing(
 			// Count samples from an ever-increasing counter to avoid screen-door effect
 			Config.PathTracingData.TemporalSeed = PathTracingState->FrameIndex;
 		}
+		Config.PathTracingData.TemporalSeed += SeedOffset;
+
 		Config.PathTracingData.Iteration = PathTracingState->SampleIndex;
 		Config.PathTracingData.BlendFactor = 1.0f / (Config.PathTracingData.Iteration + 1);
 
@@ -3660,7 +3664,7 @@ void FDeferredShadingSceneRenderer::RenderPathTracing(
 					FPathTracingBuildCloudAccelerationMapCS::FParameters* PassParameters = GraphBuilder.AllocParameters<FPathTracingBuildCloudAccelerationMapCS::FParameters>();
 					PassParameters->NumSamples = NumSamples;
 					PassParameters->Iteration = PathTracingState->SampleIndex;
-					PassParameters->TemporalSeed = PathTracingState->FrameIndex;
+					PassParameters->TemporalSeed = Config.PathTracingData.TemporalSeed;
 					PassParameters->CloudParameters = CloudParameters;
 					PassParameters->ViewUniformBuffer = View.ViewUniformBuffer;
 					PassParameters->CloudAccelerationMap = GraphBuilder.CreateUAV(CloudAccelerationMap);
