OCIO: Minor simplification of transform processor getter function for MRG.

diff --git a/Engine/Plugins/Compositing/OpenColorIO/Source/OpenColorIO/Private/OpenColorIOColorTransform.cpp b/Engine/Plugins/Compositing/OpenColorIO/Source/OpenColorIO/Private/OpenColorIOColorTransform.cpp
--- a/Engine/Plugins/Compositing/OpenColorIO/Source/OpenColorIO/Private/OpenColorIOColorTransform.cpp
+++ b/Engine/Plugins/Compositing/OpenColorIO/Source/OpenColorIO/Private/OpenColorIOColorTransform.cpp
@@ -567,7 +567,7 @@ bool UOpenColorIOColorTransform::IsTransform(const FString& InSourceColorSpace,
 	return false;
 }
 
-bool UOpenColorIOColorTransform::GetTransformProcessor(FOpenColorIOWrapperProcessor& OutProcessor) const
+bool UOpenColorIOColorTransform::GetTransformProcessor(FOpenColorIOWrapperProcessor& OutProcessor, const TMap<FString, FString>& InContextOverride) const
 {
 #if WITH_OCIO
 	UOpenColorIOConfiguration* ConfigurationOwner = GetTypedOuter<UOpenColorIOConfiguration>();
@@ -583,11 +583,23 @@ bool UOpenColorIOColorTransform::GetTransformProcessor(FOpenColorIOWrapperProces
 	EOpenColorIOViewTransformDirection CurrentDisplayViewDirection;
 	if (GetDisplayViewDirection(CurrentDisplayViewDirection))
 	{
-		OutProcessor = FOpenColorIOWrapperProcessor(ConfigWrapper, SourceColorSpace, Display, View, CurrentDisplayViewDirection == EOpenColorIOViewTransformDirection::Inverse, ConfigurationOwner->Context);
+		OutProcessor = FOpenColorIOWrapperProcessor(
+			ConfigWrapper,
+			SourceColorSpace,
+			Display,
+			View,
+			CurrentDisplayViewDirection == EOpenColorIOViewTransformDirection::Inverse,
+			InContextOverride.IsEmpty() ? ConfigurationOwner->Context : InContextOverride
+		);
 	}
 	else
 	{
-		OutProcessor = FOpenColorIOWrapperProcessor(ConfigWrapper, SourceColorSpace, DestinationColorSpace, ConfigurationOwner->Context);
+		OutProcessor = FOpenColorIOWrapperProcessor(
+			ConfigWrapper,
+			SourceColorSpace,
+			DestinationColorSpace,
+			InContextOverride.IsEmpty() ? ConfigurationOwner->Context : InContextOverride
+		);
 	}
 
 	return OutProcessor.IsValid();
diff --git a/Engine/Plugins/Compositing/OpenColorIO/Source/OpenColorIO/Public/OpenColorIOColorTransform.h b/Engine/Plugins/Compositing/OpenColorIO/Source/OpenColorIO/Public/OpenColorIOColorTransform.h
--- a/Engine/Plugins/Compositing/OpenColorIO/Source/OpenColorIO/Public/OpenColorIOColorTransform.h
+++ b/Engine/Plugins/Compositing/OpenColorIO/Source/OpenColorIO/Public/OpenColorIOColorTransform.h
@@ -73,7 +73,7 @@ public:
 	 * @param OutProcessor Processor wrapper for the current transform.
 	 * @return True if the processor was created and is valid.
 	 */
-	bool GetTransformProcessor(FOpenColorIOWrapperProcessor& OutProcessor) const;
+	bool GetTransformProcessor(FOpenColorIOWrapperProcessor& OutProcessor, const TMap<FString, FString>& InContextOverride = {}) const;
 
 	/** Apply the color transform in-place to the specified color. */
 	bool TransformColor(FLinearColor& InOutColor) const;
diff --git a/Engine/Plugins/MovieScene/MovieRenderPipeline/Source/MovieRenderPipelineCore/Private/Graph/MovieGraphOCIOHelper.cpp b/Engine/Plugins/MovieScene/MovieRenderPipeline/Source/MovieRenderPipelineCore/Private/Graph/MovieGraphOCIOHelper.cpp
--- a/Engine/Plugins/MovieScene/MovieRenderPipeline/Source/MovieRenderPipelineCore/Private/Graph/MovieGraphOCIOHelper.cpp
+++ b/Engine/Plugins/MovieScene/MovieRenderPipeline/Source/MovieRenderPipelineCore/Private/Graph/MovieGraphOCIOHelper.cpp
@@ -1,4 +1,4 @@
-﻿// Copyright Epic Games, Inc. All Rights Reserved.
+// Copyright Epic Games, Inc. All Rights Reserved.
 
 #include "Graph/MovieGraphOCIOHelper.h"
 
@@ -103,35 +103,11 @@ FPixelPreProcessor FMovieGraphOCIOHelper::CreateOpenColorIOPixelPreProcessor(con
 	const TObjectPtr<UOpenColorIOConfiguration>& ConfigurationSource = InConversionSettings.ConfigurationSource;
 	if (IsValid(ConfigurationSource))
 	{
-		const FOpenColorIOWrapperConfig* ConfigWrapper = ConfigurationSource->GetOrCreateConfigWrapper();
 		const TObjectPtr<const UOpenColorIOColorTransform> ColorTransform = ConfigurationSource->FindTransform(InConversionSettings);
 		if (IsValid(ColorTransform))
 		{
 			FOpenColorIOWrapperProcessor Processor;
-			EOpenColorIOViewTransformDirection CurrentDisplayViewDirection;
-
-			if (ColorTransform->GetDisplayViewDirection(CurrentDisplayViewDirection))
-			{
-				Processor = FOpenColorIOWrapperProcessor(
-						ConfigWrapper,
-						ColorTransform->SourceColorSpace,
-						ColorTransform->Display,
-						ColorTransform->View,
-						static_cast<bool>(CurrentDisplayViewDirection),
-						InContext
-					);
-			}
-			else
-			{
-				Processor = FOpenColorIOWrapperProcessor(
-						ConfigWrapper,
-						ColorTransform->SourceColorSpace,
-						ColorTransform->DestinationColorSpace,
-						InContext
-					);
-			}
-			
-			if (Processor.IsValid())
+			if (ColorTransform->GetTransformProcessor(Processor, InContext))
 			{
 				return FOpenColorIOPixelPreProcessor(MoveTemp(Processor));
 			}
