MRQ: Fix support for scene captures rendering with the main view, which would otherwise fail to update.

diff --git a/Engine/Plugins/MovieScene/MovieRenderPipeline/Source/MovieRenderPipelineRenderPasses/Private/Graph/Renderers/MovieGraphImagePassBase.cpp b/Engine/Plugins/MovieScene/MovieRenderPipeline/Source/MovieRenderPipelineRenderPasses/Private/Graph/Renderers/MovieGraphImagePassBase.cpp
--- a/Engine/Plugins/MovieScene/MovieRenderPipeline/Source/MovieRenderPipelineRenderPasses/Private/Graph/Renderers/MovieGraphImagePassBase.cpp
+++ b/Engine/Plugins/MovieScene/MovieRenderPipeline/Source/MovieRenderPipelineRenderPasses/Private/Graph/Renderers/MovieGraphImagePassBase.cpp
@@ -175,6 +175,9 @@ TSharedRef<FSceneViewFamilyContext> FMovieGraphImagePassBase::CreateSceneViewFam
 		ViewExt->SetupViewFamily(*OutViewFamily);
 	}
 
+	// Support scene captures with the "bMainViewFamily" flag set
+	OutViewFamily->bIsMainViewFamily = true;
+
 	return OutViewFamily;
 }
 
diff --git a/Engine/Plugins/MovieScene/MovieRenderPipeline/Source/MovieRenderPipelineRenderPasses/Private/MoviePipelineImagePassBase.cpp b/Engine/Plugins/MovieScene/MovieRenderPipeline/Source/MovieRenderPipelineRenderPasses/Private/MoviePipelineImagePassBase.cpp
--- a/Engine/Plugins/MovieScene/MovieRenderPipeline/Source/MovieRenderPipelineRenderPasses/Private/MoviePipelineImagePassBase.cpp
+++ b/Engine/Plugins/MovieScene/MovieRenderPipeline/Source/MovieRenderPipelineRenderPasses/Private/MoviePipelineImagePassBase.cpp
@@ -306,6 +306,9 @@ TSharedPtr<FSceneViewFamilyContext> UMoviePipelineImagePassBase::CalculateViewFa
 		OutViewFamily->ViewExtensions[ViewExt]->SetupView(*OutViewFamily.Get(), *View);
 	}
 
+	// Support scene captures with the "bMainViewFamily" flag set
+	OutViewFamily->bIsMainViewFamily = true;
+
 	// The requested configuration may not be supported, warn user and fall back. We can't call
 	// FSceneView::SetupAntiAliasingMethod because it reads the value from the cvar which would
 	// cause the value set by the MoviePipeline UI to be ignored.
diff --git a/Engine/Source/Runtime/Engine/Public/SceneView.h b/Engine/Source/Runtime/Engine/Public/SceneView.h
--- a/Engine/Source/Runtime/Engine/Public/SceneView.h
+++ b/Engine/Source/Runtime/Engine/Public/SceneView.h
@@ -2313,7 +2313,7 @@ public:
 
 	/**
 	 * Whether this is the "main" view family.  Should be set to true for the view family of the main editor or game viewport.  Affects
-	 * where rendering occurs for Scene Captures with the "bRenderWithMainViewFamily" flag set.
+	 * where rendering occurs for Scene Captures with the "bMainViewFamily" flag set.
 	 */
 	bool bIsMainViewFamily = false;
 
