MoviePipeline: Added a post-render file display system, so MRG and other MRG-adjacent systems can display rendered media after a render based on user preferences.

There are no systems that are currently wired up to use this; that will come in subsequent CLs.

diff --git a/Engine/Plugins/MovieScene/MovieRenderPipeline/Source/MovieRenderPipelineCore/Private/Graph/MovieGraphDefaultRenderer.cpp b/Engine/Plugins/MovieScene/MovieRenderPipeline/Source/MovieRenderPipelineCore/Private/Graph/MovieGraphDefaultRenderer.cpp
--- a/Engine/Plugins/MovieScene/MovieRenderPipeline/Source/MovieRenderPipelineCore/Private/Graph/MovieGraphDefaultRenderer.cpp
+++ b/Engine/Plugins/MovieScene/MovieRenderPipeline/Source/MovieRenderPipelineCore/Private/Graph/MovieGraphDefaultRenderer.cpp
@@ -74,6 +74,9 @@ void UMovieGraphDefaultRenderer::SetupRenderingPipelineForShot(UMoviePipelineExe
 
 	CameraOverscanCache.Empty();
 	bHasWarnedAboutAnimatedOverscan = false;
+
+	int32 RenderLayerIndex = 0;
+	TMap<FMovieGraphRenderDataIdentifier, int32> IdentifierToLayerIndex;
 	
 	for (const FName& Branch : GraphBranches)
 	{
@@ -111,6 +114,8 @@ void UMovieGraphDefaultRenderer::SetupRenderingPipelineForShot(UMoviePipelineExe
 			FMovieGraphRenderDataIdentifier Identifier;
 			Identifier.RootBranchName = Branch;
 			Identifier.LayerName = RenderLayerNode->LayerName;
+
+			IdentifierToLayerIndex.Add(Identifier, RenderLayerIndex++);
 			
 			ExistingPass->BranchRenderers.Add(Identifier, RenderPassNode);
 		}
@@ -135,6 +140,7 @@ void UMovieGraphDefaultRenderer::SetupRenderingPipelineForShot(UMoviePipelineExe
 				FMovieGraphRenderPassLayerData& LayerData = SetupData.Layers.AddDefaulted_GetRef();
 				LayerData.BranchName = BranchRenderer.Key.RootBranchName;
 				LayerData.LayerName = BranchRenderer.Key.LayerName;
+				LayerData.LayerIndex = *IdentifierToLayerIndex.Find(BranchRenderer.Key);
 				LayerData.RenderPassNode = BranchRenderer.Value;
 				// Result could be -1 or 0...n
 				LayerData.CameraIndex = CameraIndexes[CameraIndex]; 
diff --git a/Engine/Plugins/MovieScene/MovieRenderPipeline/Source/MovieRenderPipelineCore/Private/Graph/MovieGraphPipeline.cpp b/Engine/Plugins/MovieScene/MovieRenderPipeline/Source/MovieRenderPipelineCore/Private/Graph/MovieGraphPipeline.cpp
--- a/Engine/Plugins/MovieScene/MovieRenderPipeline/Source/MovieRenderPipelineCore/Private/Graph/MovieGraphPipeline.cpp
+++ b/Engine/Plugins/MovieScene/MovieRenderPipeline/Source/MovieRenderPipelineCore/Private/Graph/MovieGraphPipeline.cpp
@@ -1674,8 +1674,15 @@ void UMovieGraphPipeline::ProcessOutstandingFutures()
 			ShotOutputData->Shot = FutureData.Shot;
 		}
 
-		// Add the filepath to the render layer data.
-		ShotOutputData->RenderLayerData.FindOrAdd(FutureData.DataIdentifier).FilePaths.Add(FutureData.FilePath);
+		// Add the filepath and other data to the render layer.
+		FMovieGraphRenderLayerOutputData& OutputData = ShotOutputData->RenderLayerData.FindOrAdd(FutureData.DataIdentifier);
+		OutputData.FilePaths.Add(FutureData.FilePath);
+		OutputData.FrameTemplatedFilePaths.AddUnique(FutureData.FrameTemplatedFilePath);
+		OutputData.RenderLayerIndex = FutureData.RenderLayerIndex;
+
+		// Also associate this file with a specific origin node class.
+		FMovieGraphStringArray& FilesFromNodeClass = OutputData.NodeTypeToFilePaths.FindOrAdd(FutureData.OriginNodeClass);
+		FilesFromNodeClass.Array.Add(FutureData.FilePath);
 
 		// Sometime futures can be completed, but will be set to an error state (such as we couldn't write to the specified disk path.)
 		if (!OutputFuture.Get<0>().Get())
diff --git a/Engine/Plugins/MovieScene/MovieRenderPipeline/Source/MovieRenderPipelineCore/Private/Graph/MovieRenderGraphEditorSettings.cpp b/Engine/Plugins/MovieScene/MovieRenderPipeline/Source/MovieRenderPipelineCore/Private/Graph/MovieRenderGraphEditorSettings.cpp
new file mode 100644
--- /dev/null
+++ b/Engine/Plugins/MovieScene/MovieRenderPipeline/Source/MovieRenderPipelineCore/Private/Graph/MovieRenderGraphEditorSettings.cpp
@@ -0,0 +1,9 @@
+﻿// Copyright Epic Games, Inc. All Rights Reserved.
+
+#include "Graph/MovieRenderGraphEditorSettings.h"
+
+UMovieRenderGraphEditorSettings::UMovieRenderGraphEditorSettings()
+	: PostRenderSettings(FMovieGraphPostRenderSettings())
+{
+	
+}
\ No newline at end of file
diff --git a/Engine/Plugins/MovieScene/MovieRenderPipeline/Source/MovieRenderPipelineCore/Private/Graph/Nodes/MovieGraphAudioOutputNode.cpp b/Engine/Plugins/MovieScene/MovieRenderPipeline/Source/MovieRenderPipelineCore/Private/Graph/Nodes/MovieGraphAudioOutputNode.cpp
--- a/Engine/Plugins/MovieScene/MovieRenderPipeline/Source/MovieRenderPipelineCore/Private/Graph/Nodes/MovieGraphAudioOutputNode.cpp
+++ b/Engine/Plugins/MovieScene/MovieRenderPipeline/Source/MovieRenderPipelineCore/Private/Graph/Nodes/MovieGraphAudioOutputNode.cpp
@@ -243,6 +243,7 @@ void UMovieGraphAudioOutputNode::StartAudioExport()
 		UE::MovieGraph::FMovieGraphOutputFutureData OutputData;
 		OutputData.Shot = CachedPipeline->GetActiveShotList()[AudioData.ShotIndex];
 		OutputData.DataIdentifier = AudioData.RenderIdentifier;
+		OutputData.OriginNodeClass = GetClass();
 
 		// Do this before we start manipulating the filepath for the audio API
 		OutputData.FilePath = AudioData.FilePath;
diff --git a/Engine/Plugins/MovieScene/MovieRenderPipeline/Source/MovieRenderPipelineCore/Private/Graph/Nodes/MovieGraphVideoOutputNode.cpp b/Engine/Plugins/MovieScene/MovieRenderPipeline/Source/MovieRenderPipelineCore/Private/Graph/Nodes/MovieGraphVideoOutputNode.cpp
--- a/Engine/Plugins/MovieScene/MovieRenderPipeline/Source/MovieRenderPipelineCore/Private/Graph/Nodes/MovieGraphVideoOutputNode.cpp
+++ b/Engine/Plugins/MovieScene/MovieRenderPipeline/Source/MovieRenderPipelineCore/Private/Graph/Nodes/MovieGraphVideoOutputNode.cpp
@@ -347,6 +347,8 @@ UMovieGraphVideoOutputNode::FMovieGraphCodecWriterWithPromise* UMovieGraphVideoO
 			OutputData.Shot = InPipeline->GetActiveShotList()[Payload->TraversalContext.ShotIndex];
 			OutputData.DataIdentifier = InRenderPassData.Key;
 			OutputData.FilePath = FinalFilePath;
+			OutputData.OriginNodeClass = GetClass();
+			OutputData.RenderLayerIndex = Payload->RenderLayerIndex;
 
 			TPromise<bool> Completed;
 			InPipeline->AddOutputFuture(Completed.GetFuture(), OutputData);
diff --git a/Engine/Plugins/MovieScene/MovieRenderPipeline/Source/MovieRenderPipelineCore/Private/Graph/Nodes/MovieGraphWidgetRendererBaseNode.cpp b/Engine/Plugins/MovieScene/MovieRenderPipeline/Source/MovieRenderPipelineCore/Private/Graph/Nodes/MovieGraphWidgetRendererBaseNode.cpp
--- a/Engine/Plugins/MovieScene/MovieRenderPipeline/Source/MovieRenderPipelineCore/Private/Graph/Nodes/MovieGraphWidgetRendererBaseNode.cpp
+++ b/Engine/Plugins/MovieScene/MovieRenderPipeline/Source/MovieRenderPipelineCore/Private/Graph/Nodes/MovieGraphWidgetRendererBaseNode.cpp
@@ -165,7 +165,8 @@ void UMovieGraphWidgetRendererBaseNode::FMovieGraphWidgetPass::Render(const FMov
 		ENQUEUE_RENDER_COMMAND(WidgetRenderTargetResolveCommand)(
 			[InTimeData, InFrameTraversalContext, RenderTargetInitParams, RenderDataIdentifier = this->RenderDataIdentifier,
 			bComposite = ParentNodeThisFrame->bCompositeOntoFinalImage, CompositingSortOrder = GetCompositingSortOrder(),
-			BackbufferRenderTarget, OutputMerger, OverscannedResolution, CropRect, OutputResolution](FRHICommandListImmediate& RHICmdList)
+			BackbufferRenderTarget, OutputMerger, OverscannedResolution, CropRect, OutputResolution,
+			RenderLayerIndex = this->LayerData.LayerIndex](FRHICommandListImmediate& RHICmdList)
 			{
 				const FIntRect SourceRect = FIntRect(0, 0, BackbufferRenderTarget->GetSizeXY().X, BackbufferRenderTarget->GetSizeXY().Y);
 
@@ -197,6 +198,7 @@ void UMovieGraphWidgetRendererBaseNode::FMovieGraphWidgetPass::Render(const FMov
 				SampleStatePayload->bFetchFromAccumulator = false;
 				SampleStatePayload->bCompositeOnOtherRenders = bComposite;
 				SampleStatePayload->CompositingSortOrder = CompositingSortOrder;
+				SampleStatePayload->RenderLayerIndex = RenderLayerIndex;
 
 				TUniquePtr<FImagePixelData> PixelData =
 					MakeUnique<TImagePixelData<FColor>>(OutputResolution, TArray64<FColor>(MoveTemp(RawPixels)), SampleStatePayload);
diff --git a/Engine/Plugins/MovieScene/MovieRenderPipeline/Source/MovieRenderPipelineCore/Private/MoviePipelinePostRenderSettings.cpp b/Engine/Plugins/MovieScene/MovieRenderPipeline/Source/MovieRenderPipelineCore/Private/MoviePipelinePostRenderSettings.cpp
new file mode 100644
--- /dev/null
+++ b/Engine/Plugins/MovieScene/MovieRenderPipeline/Source/MovieRenderPipelineCore/Private/MoviePipelinePostRenderSettings.cpp
@@ -0,0 +1,25 @@
+﻿// Copyright Epic Games, Inc. All Rights Reserved.
+
+#include "MoviePipelinePostRenderSettings.h"
+
+FMovieGraphPostRenderVideoPlayOptions::FMovieGraphPostRenderVideoPlayOptions()
+	: PlaybackMethod(EMovieGraphPlaybackMethod::OperatingSystem)
+	, JobPlayback(EMovieGraphJobPlaybackRange::FirstJobOnly)
+	, RenderLayerPlayback(EMovieGraphRenderLayerPlaybackRange::FirstRenderLayerOnly)
+{
+	
+}
+
+FMovieGraphPostRenderImageSequencePlayOptions::FMovieGraphPostRenderImageSequencePlayOptions()
+	: FrameRangeNotation(EMovieGraphFrameRangeNotation::Hash)
+	, PlaybackRange(EMovieGraphImageSequencePlaybackRange::FirstFrameOnly)
+{
+	
+}
+
+FMovieGraphPostRenderSettings::FMovieGraphPostRenderSettings()
+	: OutputTypePriorityOrder({TEXT("EXR"), TEXT("PNG"), TEXT("JPEG"), TEXT("BMP"), TEXT("MOV"), TEXT("MXF")})
+	, OutputTypePlayback(EMovieGraphOutputTypePlayback::UsePriorityOrder)
+{
+	
+}
\ No newline at end of file
diff --git a/Engine/Plugins/MovieScene/MovieRenderPipeline/Source/MovieRenderPipelineCore/Public/Graph/MovieGraphDataTypes.h b/Engine/Plugins/MovieScene/MovieRenderPipeline/Source/MovieRenderPipelineCore/Public/Graph/MovieGraphDataTypes.h
--- a/Engine/Plugins/MovieScene/MovieRenderPipeline/Source/MovieRenderPipelineCore/Public/Graph/MovieGraphDataTypes.h
+++ b/Engine/Plugins/MovieScene/MovieRenderPipeline/Source/MovieRenderPipelineCore/Public/Graph/MovieGraphDataTypes.h
@@ -320,6 +320,9 @@ struct FMovieGraphRenderPassLayerData
 	int32 CameraIndex;
 	FString CameraName;
 
+	// The index of this layer in the list of layers being rendered (0 being the first layer)
+	int32 LayerIndex = 0;
+
 	// The number of cameras that are being rendered.
 	uint32 NumCameras = 1;
 	
@@ -341,6 +344,8 @@ namespace UE::MovieGraph
 	{
 		FMovieGraphOutputFutureData()
 			: Shot(nullptr)
+			, OriginNodeClass(nullptr)
+			, RenderLayerIndex(0)
 		{}
 
 		/** Which shot is this for? */
@@ -349,8 +354,17 @@ namespace UE::MovieGraph
 		/** What was the filepath on disk that was generated by the future? */
 		FString FilePath;
 
+		/** The "frame templated" filepath. This groups together a set of output frames by frame number, eg "Seq.Shot.#.exr" if the frame template was "#". */
+		FString FrameTemplatedFilePath;
+
 		/** Which render pass is this for? */
 		FMovieGraphRenderDataIdentifier DataIdentifier;
+
+		/** The node class that this file originated from. */
+		UClass* OriginNodeClass;
+
+		/** The index of the render layer this future belongs to. Index 0 is the first render layer rendered, and so on. */
+		int32 RenderLayerIndex;
 	};
 	
 	/**
@@ -381,6 +395,7 @@ namespace UE::MovieGraph
 			, bAllowsCompositing(true)
 			, OverscanFraction(0.f)
 			, CompositingSortOrder(0)
+			, RenderLayerIndex(0)
 			, bAllowOCIO(true)
 		{}
 
@@ -444,6 +459,9 @@ namespace UE::MovieGraph
 		*/
 		int32 CompositingSortOrder;
 
+		/** The index of the render layer this sample belongs to. Index 0 is the first render layer rendered, and so on. */
+		int32 RenderLayerIndex;
+
 		/** Allow OpenColorIO transform to be used on this render. */
 		bool bAllowOCIO;
 
diff --git a/Engine/Plugins/MovieScene/MovieRenderPipeline/Source/MovieRenderPipelineCore/Public/Graph/MovieGraphRenderDataIdentifier.h b/Engine/Plugins/MovieScene/MovieRenderPipeline/Source/MovieRenderPipelineCore/Public/Graph/MovieGraphRenderDataIdentifier.h
--- a/Engine/Plugins/MovieScene/MovieRenderPipeline/Source/MovieRenderPipelineCore/Public/Graph/MovieGraphRenderDataIdentifier.h
+++ b/Engine/Plugins/MovieScene/MovieRenderPipeline/Source/MovieRenderPipelineCore/Public/Graph/MovieGraphRenderDataIdentifier.h
@@ -92,6 +92,16 @@ public:
 	FString CameraName;
 };
 
+/** A simple struct wrapper around a string array so it can be included in a TMap. */
+USTRUCT(BlueprintType)
+struct MOVIERENDERPIPELINECORE_API FMovieGraphStringArray
+{
+	GENERATED_BODY()
+
+	UPROPERTY(VisibleAnywhere, BlueprintReadOnly, Category = "Movie Pipeline")
+	TArray<FString> Array;
+};
+
 USTRUCT(BlueprintType)
 struct MOVIERENDERPIPELINECORE_API FMovieGraphRenderLayerOutputData
 {
@@ -100,6 +110,22 @@ public:
 	/** A list of file paths on disk (in order) that were generated for this particular render pass. */
 	UPROPERTY(VisibleAnywhere, BlueprintReadOnly, Category = "Movie Pipeline")
 	TArray<FString> FilePaths;
+	
+	/**
+	 * The "frame templated" file paths. This groups together a set of output frames by frame number, eg "Seq.Shot.#.exr" if the frame template was
+	 * "#". If the render layer has more than one output type in it (eg, EXR and JPEG), there may be more than one path here, but there should only
+	 * be one per output type.
+	 */
+	UPROPERTY(VisibleAnywhere, BlueprintReadOnly, Category = "Movie Pipeline")
+	TArray<FString> FrameTemplatedFilePaths;
+	
+	/** Maps a node type to all of the files that the node type generated. */
+	UPROPERTY(VisibleAnywhere, BlueprintReadOnly, Category = "Movie Pipeline")
+	TMap<FSoftClassPath, FMovieGraphStringArray> NodeTypeToFilePaths;
+
+	/** The index of the render layer these files belong to. Index 0 is the first render layer rendered, and so on. */
+	UPROPERTY(VisibleAnywhere, BlueprintReadOnly, Category = "Movie Pipeline")
+	int32 RenderLayerIndex = 0;
 };
 
 USTRUCT(BlueprintType)
diff --git a/Engine/Plugins/MovieScene/MovieRenderPipeline/Source/MovieRenderPipelineCore/Public/Graph/MovieRenderGraphEditorSettings.h b/Engine/Plugins/MovieScene/MovieRenderPipeline/Source/MovieRenderPipelineCore/Public/Graph/MovieRenderGraphEditorSettings.h
new file mode 100644
--- /dev/null
+++ b/Engine/Plugins/MovieScene/MovieRenderPipeline/Source/MovieRenderPipelineCore/Public/Graph/MovieRenderGraphEditorSettings.h
@@ -0,0 +1,23 @@
+﻿// Copyright Epic Games, Inc. All Rights Reserved.
+
+#pragma once
+
+#include "Engine/DeveloperSettings.h"
+
+#include "MovieGraphDataTypes.h"
+#include "MoviePipelinePostRenderSettings.h"
+
+#include "MovieRenderGraphEditorSettings.generated.h"
+
+UCLASS(BlueprintType, Config=EditorPerProjectUserSettings, meta=(DisplayName="Movie Render Graph"))
+class MOVIERENDERPIPELINECORE_API UMovieRenderGraphEditorSettings : public UDeveloperSettings
+{
+	GENERATED_BODY()
+
+public:
+	UMovieRenderGraphEditorSettings();
+
+	/** If PostRenderBehavior is set to PlayRenderOutput, these settings are used to determine how to play media. */
+	UPROPERTY(Config, EditAnywhere, Category="Play Render Output", meta=(ShowOnlyInnerProperties))
+	FMovieGraphPostRenderSettings PostRenderSettings;
+};
diff --git a/Engine/Plugins/MovieScene/MovieRenderPipeline/Source/MovieRenderPipelineCore/Public/MoviePipelinePostRenderSettings.h b/Engine/Plugins/MovieScene/MovieRenderPipeline/Source/MovieRenderPipelineCore/Public/MoviePipelinePostRenderSettings.h
new file mode 100644
--- /dev/null
+++ b/Engine/Plugins/MovieScene/MovieRenderPipeline/Source/MovieRenderPipelineCore/Public/MoviePipelinePostRenderSettings.h
@@ -0,0 +1,163 @@
+﻿// Copyright Epic Games, Inc. All Rights Reserved.
+
+#pragma once
+
+#include "UObject/ObjectMacros.h"
+
+#include "MoviePipelinePostRenderSettings.generated.h"
+
+/** The type of viewer that should be used to play back media. */
+UENUM(BlueprintType)
+enum class EMovieGraphPlaybackMethod : uint8
+{
+	/** Use the operating system's default application for playing back media of the type that's being viewed. */
+	OperatingSystem,
+
+	/** Use a specific application for playing back media. */
+	CustomViewer
+};
+
+/** The frame range for image sequences that should be used during playback. */
+UENUM(BlueprintType)
+enum class EMovieGraphImageSequencePlaybackRange : uint8
+{
+	/** Only play back the first frame. */
+	FirstFrameOnly,
+
+	/** Play back all available frames. */
+	FullRange
+};
+
+/** Dictates which jobs should be used to get media from for playback if more than one job generated media. */
+UENUM(BlueprintType)
+enum class EMovieGraphJobPlaybackRange : uint8
+{
+	/** Only play back media from the first job that was rendered. */
+	FirstJobOnly,
+
+	/** Play back media from all of the jobs that were rendered. */
+	AllJobs
+};
+
+/** Dictates the selection of render layers that should be used in playback if more than one was generated. */
+UENUM(BlueprintType)
+enum class EMovieGraphRenderLayerPlaybackRange : uint8
+{
+	/** Only play back media from the first render layer that was rendered. */
+	FirstRenderLayerOnly,
+
+	/** Play back media from all render layers that were rendered. */
+	AllRenderLayers
+};
+
+/** The symbol that should be used to substitute a frame number within image paths. */
+UENUM(BlueprintType)
+enum class EMovieGraphFrameRangeNotation : uint8
+{
+	Hash UMETA(DisplayName = "#"),
+	DollarF UMETA(DisplayName = "$F"),
+	StartFrame
+};
+
+/** How media should be selected for playback when more than one type was generated by a render. */
+UENUM(BlueprintType)
+enum class EMovieGraphOutputTypePlayback : uint8
+{
+	/**
+	 * Use the priority order when determining the media to play (the media with the extension that has the highest priority is played).
+	 * The priority order is specified in Editor Preferences.
+	 */
+	UsePriorityOrder,
+
+	/** Play all media types that were generated from the render. */
+	PlayAllOutputTypes
+};
+
+/** Options that dictate how videos should be played back after a render completes. */
+USTRUCT(BlueprintType)
+struct MOVIERENDERPIPELINECORE_API FMovieGraphPostRenderVideoPlayOptions
+{
+	GENERATED_BODY()
+	
+	FMovieGraphPostRenderVideoPlayOptions();
+	
+	/** Determines which application is used to open rendered media for playback. */
+	UPROPERTY(EditAnywhere, BlueprintReadWrite, Category = "Settings")
+	EMovieGraphPlaybackMethod PlaybackMethod;
+
+	/** If using a custom viewer, this is the path to the executable to use. */
+	UPROPERTY(EditAnywhere, BlueprintReadWrite, Category = "Settings", meta = (EditCondition = "PlaybackMethod == EMovieGraphPlaybackMethod::CustomViewer", EditConditionHides))
+	FFilePath PlayerExecutable;
+
+	/** If using a custom viewer, these command line arguments are passed along. */
+	UPROPERTY(EditAnywhere, BlueprintReadWrite, Category = "Settings", meta = (EditCondition = "PlaybackMethod == EMovieGraphPlaybackMethod::CustomViewer", EditConditionHides))
+	FString AdditionalCommandLineArguments;
+
+	/** The jobs that should have their media played. */
+	UPROPERTY(EditAnywhere, BlueprintReadWrite, Category = "Settings", meta = (EditCondition = "PlaybackMethod == EMovieGraphPlaybackMethod::CustomViewer"))
+	EMovieGraphJobPlaybackRange JobPlayback;
+
+	/** The render layers that should have their media played. */
+	UPROPERTY(EditAnywhere, BlueprintReadWrite, Category = "Settings", meta = (EditCondition = "PlaybackMethod == EMovieGraphPlaybackMethod::CustomViewer"))
+	EMovieGraphRenderLayerPlaybackRange RenderLayerPlayback;
+};
+
+/** Options that dictate how images should be played back after a render completes. */
+USTRUCT(BlueprintType)
+struct MOVIERENDERPIPELINECORE_API FMovieGraphPostRenderImageSequencePlayOptions : public FMovieGraphPostRenderVideoPlayOptions
+{
+	GENERATED_BODY()
+
+	FMovieGraphPostRenderImageSequencePlayOptions();
+
+	/** The symbol that will be used to replace the frame number in image paths when providing those paths to external viewers. */
+	UPROPERTY(EditAnywhere, BlueprintReadWrite, Category = "Settings", meta = (DisplayAfter = "AdditionalCommandLineArguments", EditCondition = "PlaybackMethod == EMovieGraphPlaybackMethod::CustomViewer", EditConditionHides))
+	EMovieGraphFrameRangeNotation FrameRangeNotation;
+
+	/** The range of frames to play back (only applicable to image sequences). */
+	UPROPERTY(EditAnywhere, BlueprintReadWrite, Category = "Settings", meta = (DisplayAfter = "FrameRangeNotation", EditCondition = "PlaybackMethod == EMovieGraphPlaybackMethod::CustomViewer"))
+	EMovieGraphImageSequencePlaybackRange PlaybackRange;
+};
+
+/** Settings that dictate how media is played back after a render finishes. */
+USTRUCT(BlueprintType)
+struct MOVIERENDERPIPELINECORE_API FMovieGraphPostRenderSettings
+{
+	GENERATED_BODY()
+
+	FMovieGraphPostRenderSettings();
+
+	/**
+	 * When deciding which media to play, media with extensions ranked higher in this list will play instead of media ranked lower.
+	 * For example, if EXR is listed above PNG, and both EXRs and PNGs were rendered, the EXRs will play instead of the PNGs.
+	 */
+	UPROPERTY(EditAnywhere, BlueprintReadWrite, Category = "Play Render Output")
+	TArray<FString> OutputTypePriorityOrder;
+
+	/**
+	 * If set to Use Priority Order, the Output Type Priority Order list will be used to determine which media to play (if there were multiple types
+	 * of images and/or videos generated). If set to Play All Output Types, all media will be opened up in the player if possible. The media that
+	 * is highest in the Output Type Priority Order will have its player used, and the other media will also be opened if it has been assigned this
+	 * player as well (ie, two or more players will not be opened). If using the Operating System to open media, this setting is effectively ignored
+	 * because only the highest-ranked media in Output Type Priority Order will be opened.
+	 */
+	UPROPERTY(EditAnywhere, BlueprintReadWrite, Category = "Play Render Output")
+	EMovieGraphOutputTypePlayback OutputTypePlayback;
+
+	/** These options are used to determine how to play the rendered images. */
+	UPROPERTY(EditAnywhere, BlueprintReadWrite, Category = "Play Render Output", DisplayName = "Image Sequences")
+	FMovieGraphPostRenderImageSequencePlayOptions PostRenderImageSequencePlayOptions;
+
+	/** These options are used to determine how to play Apple ProRes movies. */
+	UPROPERTY(EditAnywhere, BlueprintReadWrite, Category = "Play Render Output", DisplayName = "Apple ProRes Movies")
+	FMovieGraphPostRenderVideoPlayOptions PostRenderAppleProResPlayOptions;
+
+	/** These options are used to determine how to play Avid DNxHR movies. */
+	UPROPERTY(EditAnywhere, BlueprintReadWrite, Category = "Play Render Output", DisplayName = "Avid DNxHR Movies")
+	FMovieGraphPostRenderVideoPlayOptions PostRenderAvidDNxHRPlayOptions;
+
+	// TODO: No CLE support yet.
+	/** If PostRenderBehavior is set to PlayRenderOutput, these options are used to determine how to play Command Line Encoder output. */
+	// UPROPERTY(EditAnywhere, BlueprintReadWrite, Category = "Play Render Output", DisplayName = "Command Line Encoder")
+	// FMovieGraphPostRenderVideoPlayOptions PostRenderCommandLineEncoderOptions;
+};
\ No newline at end of file
diff --git a/Engine/Plugins/MovieScene/MovieRenderPipeline/Source/MovieRenderPipelineEditor/MovieRenderPipelineEditor.Build.cs b/Engine/Plugins/MovieScene/MovieRenderPipeline/Source/MovieRenderPipelineEditor/MovieRenderPipelineEditor.Build.cs
--- a/Engine/Plugins/MovieScene/MovieRenderPipeline/Source/MovieRenderPipelineEditor/MovieRenderPipelineEditor.Build.cs
+++ b/Engine/Plugins/MovieScene/MovieRenderPipeline/Source/MovieRenderPipelineEditor/MovieRenderPipelineEditor.Build.cs
@@ -71,6 +71,8 @@ public class MovieRenderPipelineEditor : ModuleRules
 				"ApplicationCore",
 				"KismetWidgets",
 				"BlueprintGraph",
+				"Imath",		// This and UEOpenExr are only required for use of UMovieGraphImageSequenceOutputNode. Ideally these are refactored
+				"UEOpenExr",	// out of the ImageSequenceOutputNode header at some point so these dependencies can be removed.
 			}
         );
 
diff --git a/Engine/Plugins/MovieScene/MovieRenderPipeline/Source/MovieRenderPipelineEditor/Private/MoviePipelinePostRenderFileDisplayProcessor.cpp b/Engine/Plugins/MovieScene/MovieRenderPipeline/Source/MovieRenderPipelineEditor/Private/MoviePipelinePostRenderFileDisplayProcessor.cpp
new file mode 100644
--- /dev/null
+++ b/Engine/Plugins/MovieScene/MovieRenderPipeline/Source/MovieRenderPipelineEditor/Private/MoviePipelinePostRenderFileDisplayProcessor.cpp
@@ -0,0 +1,427 @@
+﻿// Copyright Epic Games, Inc. All Rights Reserved.
+
+#include "MoviePipelinePostRenderFileDisplayProcessor.h"
+
+#include "MovieGraphImageSequenceOutputNode.h"
+#include "MovieRenderPipelineCoreModule.h"
+
+FMoviePipelinePostRenderFileDisplayProcessor::FMoviePipelinePostRenderFileDisplayProcessor(const FMovieGraphPostRenderSettings& InPostRenderSettings)
+	: PostRenderSettings(InPostRenderSettings)
+{
+	
+}
+
+void FMoviePipelinePostRenderFileDisplayProcessor::AddFiles(const FMoviePipelineOutputData& InRenderOutputData)
+{
+	PipelineOutputData.Add(InRenderOutputData);
+}
+
+void FMoviePipelinePostRenderFileDisplayProcessor::OpenFiles() const
+{
+	TArray<FString> FinalListOfFilesToOpen;
+	const FMovieGraphPostRenderVideoPlayOptions* FinalPlayOptions = nullptr;
+	
+	for (FMoviePipelineOutputData& OutputData : PipelineOutputData)
+	{
+		FFilesToOpen FilesToOpen = FindFilesToOpen(OutputData);
+
+		// FindFilesToOpen() may return files we want to filter out; this is the final list of what should be opened, as well as the settings that
+		// contain the application executable to use for opening the files.
+		GetFilteredFilesToOpen(FilesToOpen, FinalListOfFilesToOpen, FinalPlayOptions);
+
+		// Settings may dictate that only one job's worth of media should be opened
+		if (PostRenderSettings.PostRenderImageSequencePlayOptions.JobPlayback != EMovieGraphJobPlaybackRange::AllJobs)
+		{
+			break;
+		}
+	}
+
+	// Open the application with the filtered list of files
+	if (!FinalListOfFilesToOpen.IsEmpty() && FinalPlayOptions)
+	{
+		LaunchFilesWithSettings(FinalListOfFilesToOpen, FinalPlayOptions);
+	}
+}
+
+FMoviePipelinePostRenderFileDisplayProcessor::FFilesToOpen FMoviePipelinePostRenderFileDisplayProcessor::FindFilesToOpen(FMoviePipelineOutputData& InPipelineOutputData) const
+{
+	FFilesToOpen FilesToOpen;
+	
+	if (InPipelineOutputData.GraphData.IsEmpty())
+	{
+		return FilesToOpen;
+	}
+
+	// Each entry in the array is an individual shot's image/video data. The map is keyed by file extension, and the values are the file paths.
+	TArray<TMap<FString, TArray<FString>>> ImagesGroupedByShot;
+	TArray<TMap<FString, TArray<FString>>> FrameTemplateImagesGroupedByShot;
+	TArray<TMap<FString, TArray<FString>>> VideosGroupedByShot;
+	
+	// Group all rendered images and videos by extension and by frame templated string. If only playing back the first render layer, only capture
+	// images generated from the first layer.
+	GroupFilesByShot(InPipelineOutputData, ImagesGroupedByShot, FrameTemplateImagesGroupedByShot, VideosGroupedByShot);
+
+	// If getting images by priority order, find the image and video extensions that were used that have the highest priority. Note that the first
+	// shot will take precedence over all other shots for this determination.
+	FString HighestPriorityImageExtension;
+	FString HighestPriorityVideoExtension;
+	GetHighestPriorityExtensions(ImagesGroupedByShot, VideosGroupedByShot, HighestPriorityImageExtension, HighestPriorityVideoExtension);
+
+	auto AddImages = [this, &FilesToOpen, &FrameTemplateImagesGroupedByShot](const TArray<FString>& InImagePaths, const uint32 ShotIndex)
+	{
+		const EMovieGraphPlaybackMethod PlaybackMethod = PostRenderSettings.PostRenderImageSequencePlayOptions.PlaybackMethod;
+		const EMovieGraphImageSequencePlaybackRange PlaybackRange = PostRenderSettings.PostRenderImageSequencePlayOptions.PlaybackRange;
+
+		// Only provide the first non-templated frame in some cases
+		if ((PlaybackMethod == EMovieGraphPlaybackMethod::OperatingSystem) || (PlaybackRange == EMovieGraphImageSequencePlaybackRange::FirstFrameOnly))
+		{
+			FilesToOpen.Images.Add(InImagePaths[0]);
+			return;
+		}
+		
+		// Otherwise, provide the frame-templated image paths
+		const TMap<FString, TArray<FString>>& ImagesByFrameTemplate = FrameTemplateImagesGroupedByShot[ShotIndex];
+		for (const TPair<FString, TArray<FString>>& FrameTemplateToFilePaths : ImagesByFrameTemplate)
+		{
+			for (const FString& ImagePath : InImagePaths)
+			{
+				if (FrameTemplateToFilePaths.Value.Contains(ImagePath))
+				{
+					FilesToOpen.Images.Add(FrameTemplateToFilePaths.Key);
+
+					// Found an image that matched this frame templated path; move on to the next frame template and see if a path matches it
+					break;
+				}
+			}
+		}
+	};
+
+	auto AddVideos = [this, &FilesToOpen](const TArray<FString>& InVideoPaths)
+	{
+		for (const FString& VideoPath : InVideoPaths)
+		{
+			const FString& UpperExtension = FPaths::GetExtension(VideoPath).ToUpper();
+			
+			if (UpperExtension == TEXT("MOV"))
+			{
+				FilesToOpen.AppleProResMovies.Add(VideoPath);
+			}
+			else
+			{
+				FilesToOpen.AvidDNxHRMovies.Add(VideoPath);
+			}
+		}
+	};
+
+	// With all images/videos grouped, and the highest-priority extensions found, find the images or videos that should be returned
+	for (int32 ShotIndex = 0; ShotIndex < ImagesGroupedByShot.Num(); ++ShotIndex)
+	{
+		const TMap<FString, TArray<FString>>& ShotImages = ImagesGroupedByShot[ShotIndex];
+		const TMap<FString, TArray<FString>>& ShotVideos = VideosGroupedByShot[ShotIndex];
+		
+		// Only add the images OR videos with the highest priority if that's what was requested
+		if (PostRenderSettings.OutputTypePlayback == EMovieGraphOutputTypePlayback::UsePriorityOrder)
+		{
+			const int32 ImageExtensionPriorityIndex = GetExtensionPriorityIndex(HighestPriorityImageExtension);
+			const int32 VideoExtensionPriorityIndex = GetExtensionPriorityIndex(HighestPriorityVideoExtension);
+
+			// Images have a higher priority
+			if (ImageExtensionPriorityIndex < VideoExtensionPriorityIndex)
+			{
+				if (const TArray<FString>* ImagePaths = ShotImages.Find(HighestPriorityImageExtension))
+				{
+					AddImages(*ImagePaths, ShotIndex);
+				}
+			}
+			// Videos have a higher priority
+			else
+			{
+				if (const TArray<FString>* VideoPaths = ShotVideos.Find(HighestPriorityVideoExtension))
+				{
+					AddVideos(*VideoPaths);
+				}
+			}
+		}
+		else
+		{
+			// If not adding just the images/videos with the highest priority extension, add all images/videos that were found
+			for (const TPair<FString, TArray<FString>>& ImagesByExtension : ShotImages)
+			{
+				AddImages(ImagesByExtension.Value, ShotIndex);
+			}
+
+			for (const TPair<FString, TArray<FString>>& VideosByExtension : ShotVideos)
+			{
+				AddVideos(VideosByExtension.Value);
+			}
+		}
+	}
+
+	FilesToOpen.HighestPriorityImageExtension = HighestPriorityImageExtension;
+	FilesToOpen.HighestPriorityVideoExtension = HighestPriorityVideoExtension;
+
+	return FilesToOpen;
+}
+
+void FMoviePipelinePostRenderFileDisplayProcessor::GetFilteredFilesToOpen(const FFilesToOpen& InFilesToOpen, TArray<FString>& OutFilteredFilesToOpen, const FMovieGraphPostRenderVideoPlayOptions*& OutFilteredPlayOptions) const
+{
+	// Only update the play options if no filtered files have been provided yet. If there ARE filtered files already, that means that we're processing
+	// files that are not part of the first job. The first job's files should dictate what the play options are.
+	auto MaybeAssignPlayOptions = [&OutFilteredPlayOptions, &OutFilteredFilesToOpen](const FMovieGraphPostRenderVideoPlayOptions* InNewPlayOptions)
+	{
+		if (OutFilteredFilesToOpen.IsEmpty())
+		{
+			OutFilteredPlayOptions = InNewPlayOptions;
+		}
+	};
+	
+	// If there are both images and videos returned, and Play All Output Types is specified on the images, find the media type with the highest
+	// priority and use the player for that.
+	const bool bPlayAllOutputTypes = (PostRenderSettings.OutputTypePlayback == EMovieGraphOutputTypePlayback::PlayAllOutputTypes);
+	const bool bHasImagesAndVideos = !InFilesToOpen.Images.IsEmpty() && (!InFilesToOpen.AppleProResMovies.IsEmpty() || !InFilesToOpen.AvidDNxHRMovies.IsEmpty());
+	if (bPlayAllOutputTypes && bHasImagesAndVideos)
+	{
+		const int32 ImagePriority = GetExtensionPriorityIndex(InFilesToOpen.HighestPriorityImageExtension);
+		const int32 ProResPriority = GetExtensionPriorityIndex(TEXT("MOV"));
+		const int32 AvidPriority = GetExtensionPriorityIndex(TEXT("MXF"));
+
+		FString PlayerWithPriority;
+
+		// Find the player that should be used, based on the media type with the highest priority
+		if ((ImagePriority < ProResPriority) && (ImagePriority < AvidPriority))
+		{
+			MaybeAssignPlayOptions(&PostRenderSettings.PostRenderImageSequencePlayOptions);
+			PlayerWithPriority = PostRenderSettings.PostRenderImageSequencePlayOptions.PlayerExecutable.FilePath;
+		}
+		else if ((ProResPriority < ImagePriority) && (ProResPriority < AvidPriority))
+		{
+			MaybeAssignPlayOptions(&PostRenderSettings.PostRenderAppleProResPlayOptions);
+			PlayerWithPriority = PostRenderSettings.PostRenderAppleProResPlayOptions.PlayerExecutable.FilePath;
+		}
+		else
+		{
+			MaybeAssignPlayOptions(&PostRenderSettings.PostRenderAvidDNxHRPlayOptions);
+			PlayerWithPriority = PostRenderSettings.PostRenderAvidDNxHRPlayOptions.PlayerExecutable.FilePath;
+		}
+
+		// Only return the media that opens with the player associated with the highest-priority media
+		if (PostRenderSettings.PostRenderImageSequencePlayOptions.PlayerExecutable.FilePath == PlayerWithPriority)
+		{
+			OutFilteredFilesToOpen.Append(InFilesToOpen.Images);
+		}
+		if (PostRenderSettings.PostRenderAppleProResPlayOptions.PlayerExecutable.FilePath == PlayerWithPriority)
+		{
+			OutFilteredFilesToOpen.Append(InFilesToOpen.AppleProResMovies);
+		}
+		if (PostRenderSettings.PostRenderAvidDNxHRPlayOptions.PlayerExecutable.FilePath == PlayerWithPriority)
+		{
+			OutFilteredFilesToOpen.Append(InFilesToOpen.AvidDNxHRMovies);
+		}
+	}
+	else
+	{
+		// If the Playback Type isn't PlayAllOutputTypes, only images or movies will be specified in FilesToOpen; just provide whatever files were
+		// given by FilesToOpen in this case.
+		if (!InFilesToOpen.Images.IsEmpty())
+		{
+			MaybeAssignPlayOptions(&PostRenderSettings.PostRenderImageSequencePlayOptions);
+			OutFilteredFilesToOpen.Append(InFilesToOpen.Images);
+		}
+		else if (!InFilesToOpen.AppleProResMovies.IsEmpty())
+		{
+			MaybeAssignPlayOptions(&PostRenderSettings.PostRenderAppleProResPlayOptions);
+			OutFilteredFilesToOpen.Append(InFilesToOpen.AppleProResMovies);
+		}
+		else
+		{
+			MaybeAssignPlayOptions(&PostRenderSettings.PostRenderAvidDNxHRPlayOptions);
+			OutFilteredFilesToOpen.Append(InFilesToOpen.AvidDNxHRMovies);
+		}
+	}
+}
+
+void FMoviePipelinePostRenderFileDisplayProcessor::GroupFilesByShot(
+	FMoviePipelineOutputData& InPipelineOutputData,
+	TArray<TMap<FString, TArray<FString>>>& InImagesGroupedByShot,
+	TArray<TMap<FString, TArray<FString>>>& InFrameTemplateImagesGroupedByShot,
+	TArray<TMap<FString, TArray<FString>>>& InVideosGroupedByShot) const
+{
+	// OutputData is an individual shot's collection of rendered output (images + movies)
+	for (FMovieGraphRenderOutputData& OutputData : InPipelineOutputData.GraphData)
+	{
+		TMap<FString, TArray<FString>>& ImagesByExtension = InImagesGroupedByShot.AddDefaulted_GetRef();
+		TMap<FString, TArray<FString>>& ImagesByFrameTemplate = InFrameTemplateImagesGroupedByShot.AddDefaulted_GetRef();
+		TMap<FString, TArray<FString>>& VideosByExtension = InVideosGroupedByShot.AddDefaulted_GetRef();
+
+		const bool bFirstRenderLayerOnly_Images = (PostRenderSettings.PostRenderImageSequencePlayOptions.RenderLayerPlayback == EMovieGraphRenderLayerPlaybackRange::FirstRenderLayerOnly);
+		const bool bFirstRenderLayerOnly_ProRes = (PostRenderSettings.PostRenderAppleProResPlayOptions.RenderLayerPlayback == EMovieGraphRenderLayerPlaybackRange::FirstRenderLayerOnly);
+		const bool bFirstRenderLayerOnly_Avid = (PostRenderSettings.PostRenderAvidDNxHRPlayOptions.RenderLayerPlayback == EMovieGraphRenderLayerPlaybackRange::FirstRenderLayerOnly);
+		
+		// Sort the render layers by index. This will allow us to only get files from the first render layer rendered if the "First Render Layer Only"
+		// setting is turned on.
+		OutputData.RenderLayerData.ValueSort([](const FMovieGraphRenderLayerOutputData& A, const FMovieGraphRenderLayerOutputData& B)
+		{
+			return A.RenderLayerIndex < B.RenderLayerIndex;
+		});
+
+		// Add the images/videos for each render layer within the shot (only the first render layer may be considered depending on settings)
+		for (const TPair<FMovieGraphRenderDataIdentifier, FMovieGraphRenderLayerOutputData>& PassData : OutputData.RenderLayerData)
+		{
+			// For each node type that generated files within the render layer, add its files
+			for (const TPair<FSoftClassPath, FMovieGraphStringArray>& NodeToFiles : PassData.Value.NodeTypeToFilePaths)
+			{
+				const UClass* NodeClass = NodeToFiles.Key.ResolveClass();
+				if (!NodeClass)
+				{
+					UE_LOG(LogMovieRenderPipeline, Warning, TEXT("Found an invalid node type while determining the media to display post-render. This node's media will be skipped."));
+					continue;
+				}
+
+				// It's unfortunate that we have to identify nodes by name here in some cases, but we cannot compare directly against the static class
+				// due to some nodes being in plugins.
+				const bool bIsImageNode = NodeClass->IsChildOf(UMovieGraphImageSequenceOutputNode::StaticClass());
+				const bool bIsProResNode = NodeClass->GetName().Contains(TEXT("ProRes"), ESearchCase::CaseSensitive);
+				const bool bIsAvidNode = NodeClass->GetName().Contains(TEXT("Avid"), ESearchCase::CaseSensitive);
+				
+				// Don't continue to the next render layer if we're only getting files for one render layer
+				if (bIsImageNode && bFirstRenderLayerOnly_Images)
+				{
+					continue;
+				}
+				if (bIsProResNode && bFirstRenderLayerOnly_ProRes)
+				{
+					continue;
+				}
+				if (bIsAvidNode && bFirstRenderLayerOnly_Avid)
+				{
+					continue;
+				}
+
+				// For all files generated by the node, categorize it either by extension and/or frame template (for images)
+				for (const FString& NodeFile : NodeToFiles.Value.Array)
+				{
+					const FString Extension = FPaths::GetExtension(NodeFile).ToUpper();
+					if (Extension.IsEmpty())
+					{
+						continue;
+					}
+					
+					if (bIsImageNode)
+					{
+						ImagesByExtension.FindOrAdd(Extension).Add(NodeFile);
+
+						// The render layer will contain a set of frame templated file paths, in addition to normal file paths. Determine if this
+						// file should be associated with any of these frame templated paths (matched by extension).
+						for (const FString& FrameTemplatedFilePath : PassData.Value.FrameTemplatedFilePaths)
+						{
+							if (FPaths::GetExtension(FrameTemplatedFilePath).ToUpper() == Extension)
+							{
+								ImagesByFrameTemplate.FindOrAdd(FrameTemplatedFilePath).Add(NodeFile);
+								break;
+							}
+						}
+					}
+					else
+					{
+						VideosByExtension.FindOrAdd(Extension).Add(NodeFile);
+					}
+				}
+			}
+		}
+	}
+}
+
+void FMoviePipelinePostRenderFileDisplayProcessor::GetHighestPriorityExtensions(
+	const TArray<TMap<FString, TArray<FString>>>& InImagesGroupedByShot,
+	const TArray<TMap<FString, TArray<FString>>>& InVideosGroupedByShot,
+	FString& OutHighestPriorityImageExtension,
+	FString& OutHighestPriorityVideoExtension) const
+{
+	// Assign the given extension (InExtensionToTest_Uppercase) as the highest-priority extension (OutHighestPriorityExtension) if any media contained in
+	// InMediaGroupedByShot matches the extension.
+	auto MaybeAssignExtensionAsHighestPriority = [](const TArray<TMap<FString, TArray<FString>>>& InMediaGroupedByShot, const FString& InExtensionToTest_Uppercase, FString& OutHighestPriorityExtension)
+	{
+		// The first shot always takes precedence over all other shots as far as which extensions are considered highest priority
+		const TMap<FString, TArray<FString>>& FirstShotMedia = InMediaGroupedByShot[0];
+
+		// Media within the shot is grouped by extension (hence why the keys are used)
+		TArray<FString> ExtensionsUsedWithinShot;
+		TArray<FString> ExtensionsUsedWithinShot_Uppercase;
+		FirstShotMedia.GetKeys(ExtensionsUsedWithinShot);
+		Algo::Transform(ExtensionsUsedWithinShot, ExtensionsUsedWithinShot_Uppercase, [](const FString& InExtension) { return InExtension.ToUpper(); });
+
+		if (ExtensionsUsedWithinShot_Uppercase.Contains(InExtensionToTest_Uppercase))
+		{
+			OutHighestPriorityExtension = InExtensionToTest_Uppercase;
+		}
+	};
+	
+	// Iterate the priority list in order of highest priority to lowest priority. The first extensions in this list that actually had media generated
+	// for them will be deemed the highest priority extensions (eg, if EXR is first in the priority list, and EXR files were actually generated, EXR
+	// would be the highest priority image extension).
+	for (const FString& Extension : PostRenderSettings.OutputTypePriorityOrder)
+	{
+		// Do all extension comparisons in uppercase for consistency
+		const FString Extension_Uppercase = Extension.ToUpper();
+
+		// Find the highest priority image extension
+		if (!InImagesGroupedByShot.IsEmpty() && OutHighestPriorityImageExtension.IsEmpty())
+		{
+			MaybeAssignExtensionAsHighestPriority(InImagesGroupedByShot, Extension_Uppercase, OutHighestPriorityImageExtension);
+		}
+
+		// Find the highest priority video extension
+		if (!InVideosGroupedByShot.IsEmpty() && OutHighestPriorityVideoExtension.IsEmpty())
+		{
+			MaybeAssignExtensionAsHighestPriority(InVideosGroupedByShot, Extension_Uppercase, OutHighestPriorityVideoExtension);
+		}
+	}
+}
+
+void FMoviePipelinePostRenderFileDisplayProcessor::LaunchFilesWithSettings(const TArray<FString>& InFilesToOpen, const FMovieGraphPostRenderVideoPlayOptions* InPlayOptions) const
+{
+	if (InPlayOptions->PlaybackMethod == EMovieGraphPlaybackMethod::OperatingSystem)
+	{
+		// Only one frame will be displayed if the OS is being used to open images
+		FPlatformProcess::LaunchFileInDefaultExternalApplication(*InFilesToOpen[0]);
+	}
+	else
+	{
+		// No quotes around the executable path
+		const FString ExecutablePath = InPlayOptions->PlayerExecutable.FilePath;
+		FString ExecutablePathNoQuotes = ExecutablePath.Replace(TEXT("\""), TEXT(""));
+		FPaths::NormalizeFilename(ExecutablePathNoQuotes);
+		
+		// Make sure the paths are in the right format for the platform
+		TArray<FString> TransformedFilePaths = InFilesToOpen;
+		for (FString& FileToOpen : TransformedFilePaths)
+		{
+			FPaths::MakePlatformFilename(FileToOpen);
+		}
+		
+		const FString CommandLineArguments = FString::Format(TEXT("{0} {1}"), {
+			InPlayOptions->AdditionalCommandLineArguments,
+			FString::Join(TransformedFilePaths, TEXT(" "))
+		});
+		
+		// Open the images in the selected application
+		constexpr bool bLaunchDetached = true;
+		constexpr bool bLaunchHidden = false;
+		constexpr bool bLaunchReallyHidden = false;
+		const FProcHandle ProcHandle = FPlatformProcess::CreateProc(*ExecutablePathNoQuotes, *CommandLineArguments, bLaunchDetached, bLaunchHidden, bLaunchReallyHidden, nullptr, 0, nullptr, nullptr, nullptr);
+		if (!ProcHandle.IsValid())
+		{
+			const FString FinalCommandString = FString::Format(TEXT("{0} {1}"), {ExecutablePathNoQuotes, CommandLineArguments});
+			UE_LOG(LogMovieRenderPipeline, Warning, TEXT("Unable to open the post-render media viewer. Ensure the executable and its command line arguments have been specified correctly in Editor Preferences. The command that was run: [%s]"), *FinalCommandString);
+		}
+	}
+}
+
+int32 FMoviePipelinePostRenderFileDisplayProcessor::GetExtensionPriorityIndex(const FString& InExtension) const
+{
+	const int32 ExtensionPriority = PostRenderSettings.OutputTypePriorityOrder.Find(InExtension);
+
+	// Extensions that have a lower index in the array have the highest priority, thus if the extension was not found (-1), it should be
+	// assigned a very high priority number so it is not considered a high priority (10000 is arbitrary).
+	return (ExtensionPriority == INDEX_NONE) ? 10000 : ExtensionPriority;
+}
\ No newline at end of file
diff --git a/Engine/Plugins/MovieScene/MovieRenderPipeline/Source/MovieRenderPipelineEditor/Private/MoviePipelinePostRenderFileDisplayProcessor.h b/Engine/Plugins/MovieScene/MovieRenderPipeline/Source/MovieRenderPipelineEditor/Private/MoviePipelinePostRenderFileDisplayProcessor.h
new file mode 100644
--- /dev/null
+++ b/Engine/Plugins/MovieScene/MovieRenderPipeline/Source/MovieRenderPipelineEditor/Private/MoviePipelinePostRenderFileDisplayProcessor.h
@@ -0,0 +1,72 @@
+﻿// Copyright Epic Games, Inc. All Rights Reserved.
+
+#pragma once
+
+#include "CoreMinimal.h"
+
+#include "MoviePipelineQueue.h"
+#include "MoviePipelinePostRenderSettings.h"
+
+/** Handles sorting through all of the images/videos generated by a render(s), and opening them according to post-render editor preferences. */
+class FMoviePipelinePostRenderFileDisplayProcessor
+{
+public:
+	explicit FMoviePipelinePostRenderFileDisplayProcessor(const FMovieGraphPostRenderSettings& InPostRenderSettings);
+
+	/**
+	 * Adds a set of files to the processor. This should only be called once per job, otherwise duplicate data may be added.
+	 * Only graph-related output data is added.
+	 */
+	void AddFiles(const FMoviePipelineOutputData& InRenderOutputData);
+
+	/** Opens/displays rendered images/videos according to the post-render settings the class was initialized with. */
+	void OpenFiles() const;
+
+private:
+	/** Keeps track of the working set of images/videos that should be opened, and any other associated transient data. */
+	struct FFilesToOpen
+	{
+		TArray<FString> Images;
+		TArray<FString> AppleProResMovies;
+		TArray<FString> AvidDNxHRMovies;
+
+		FString HighestPriorityImageExtension;
+		FString HighestPriorityVideoExtension;
+	};
+
+	/** Sorts through all of the given rendered data and determines which images/videos should be opened. */
+	FFilesToOpen FindFilesToOpen(FMoviePipelineOutputData& InPipelineOutputData) const;
+
+	/** Reduces the files returned to FindFilesToOpen() to a set of files that can be opened in the same application. */
+	void GetFilteredFilesToOpen(const FFilesToOpen& InFilesToOpen, TArray<FString>& OutFilteredFilesToOpen, const FMovieGraphPostRenderVideoPlayOptions*& OutFilteredPlayOptions) const;
+
+	/** Group all pipeline output data by shot (and additionally extension within each shot) for three categories: images, frame-templated images, and movies. */
+	void GroupFilesByShot(
+		FMoviePipelineOutputData& InPipelineOutputData,
+		TArray<TMap<FString, TArray<FString>>>& InImagesGroupedByShot,
+		TArray<TMap<FString, TArray<FString>>>& InFrameTemplateImagesGroupedByShot,
+		TArray<TMap<FString, TArray<FString>>>& InVideosGroupedByShot) const;
+
+	void GetHighestPriorityExtensions(
+		const TArray<TMap<FString, TArray<FString>>>& InImagesGroupedByShot,
+		const TArray<TMap<FString, TArray<FString>>>& InVideosGroupedByShot,
+		FString& OutHighestPriorityImageExtension,
+		FString& OutHighestPriorityVideoExtension) const;
+
+	/** Opens the given files according to the playback settings provided. */
+	void LaunchFilesWithSettings(const TArray<FString>& InFilesToOpen, const FMovieGraphPostRenderVideoPlayOptions* InPlayOptions) const;
+
+	/**
+	 * Gets the extension's priority according to the Output Type Priority Order list, or 10000 if the extension was not found. A lower index indicates
+	 * higher priority.
+	 */
+	int32 GetExtensionPriorityIndex(const FString& InExtension) const;
+
+private:
+	/** The post-render settings the class was initialized with. Usually sourced from the editor preferences. */
+	FMovieGraphPostRenderSettings PostRenderSettings;
+
+	// Marked as mutable so the map (within the FMoviePipelineOutputData data structure) can be sorted internally to this class
+	/** Pipeline output data sourced from AddFiles(). Each entry in the array is one call from AddFiles(). */
+	mutable TArray<FMoviePipelineOutputData> PipelineOutputData;
+};
\ No newline at end of file
diff --git a/Engine/Plugins/MovieScene/MovieRenderPipeline/Source/MovieRenderPipelineRenderPasses/Private/Graph/Renderers/MovieGraphDeferredPass.cpp b/Engine/Plugins/MovieScene/MovieRenderPipeline/Source/MovieRenderPipelineRenderPasses/Private/Graph/Renderers/MovieGraphDeferredPass.cpp
--- a/Engine/Plugins/MovieScene/MovieRenderPipeline/Source/MovieRenderPipelineRenderPasses/Private/Graph/Renderers/MovieGraphDeferredPass.cpp
+++ b/Engine/Plugins/MovieScene/MovieRenderPipeline/Source/MovieRenderPipelineRenderPasses/Private/Graph/Renderers/MovieGraphDeferredPass.cpp
@@ -45,12 +45,10 @@ void FMovieGraphDeferredPass::Setup(TWeakObjectPtr<UMovieGraphDefaultRenderer> I
 	
 	LayerData = InLayer;
 
-
 	RenderDataIdentifier.RootBranchName = LayerData.BranchName;
 	RenderDataIdentifier.LayerName = LayerData.LayerName;
 	RenderDataIdentifier.RendererName = InRenderPassNode->GetRendererName();
 	RenderDataIdentifier.SubResourceName = TEXT("beauty");
-	
 	RenderDataIdentifier.CameraName =  InLayer.CameraName;
 
 	SceneViewState.Allocate(InRenderer->GetWorld()->GetFeatureLevel());
@@ -325,6 +323,7 @@ void FMovieGraphDeferredPass::Render(const FMovieGraphTraversalContext& InFrameT
 			SampleState.bAllowsCompositing = ParentNodeThisFrame->GetAllowsCompositing();
 			SampleState.SceneCaptureSource = SceneCaptureSource;
 			SampleState.CompositingSortOrder = 10;
+			SampleState.RenderLayerIndex = LayerData.LayerIndex;
 		}
 
 		ApplyMovieGraphOverridesToSampleState(SampleState);
diff --git a/Engine/Plugins/MovieScene/MovieRenderPipeline/Source/MovieRenderPipelineRenderPasses/Private/MovieGraphImageSequenceOutputNode.cpp b/Engine/Plugins/MovieScene/MovieRenderPipeline/Source/MovieRenderPipelineRenderPasses/Private/MovieGraphImageSequenceOutputNode.cpp
--- a/Engine/Plugins/MovieScene/MovieRenderPipeline/Source/MovieRenderPipelineRenderPasses/Private/MovieGraphImageSequenceOutputNode.cpp
+++ b/Engine/Plugins/MovieScene/MovieRenderPipeline/Source/MovieRenderPipelineRenderPasses/Private/MovieGraphImageSequenceOutputNode.cpp
@@ -5,10 +5,12 @@
 #include "Graph/Nodes/MovieGraphGlobalOutputSettingNode.h"
 #include "Graph/Nodes/MovieGraphRenderLayerNode.h"
 #include "Graph/MovieGraphDataTypes.h"
+#include "Graph/MovieGraphOCIOHelper.h"
 #include "Graph/MovieGraphPipeline.h"
 #include "Graph/MovieGraphConfig.h"
 #include "Graph/MovieGraphFilenameResolveParams.h"
 #include "Graph/MovieGraphBlueprintLibrary.h"
+#include "Graph/MovieRenderGraphEditorSettings.h"
 #include "MoviePipelineUtils.h"
 #include "MoviePipelineImageSequenceOutput.h" // for FAsyncImageQuantization
 #include "MovieRenderPipelineCoreModule.h"
@@ -17,7 +19,6 @@
 #include "ImageWriteQueue.h"
 #include "Misc/Paths.h"
 #include "Async/TaskGraphInterfaces.h"
-#include "Graph/MovieGraphOCIOHelper.h"
 
 UMovieGraphImageSequenceOutputNode::UMovieGraphImageSequenceOutputNode()
 {
@@ -39,9 +40,11 @@ FString UMovieGraphImageSequenceOutputNode::CreateFileName(
 	UE::MovieGraph::FMovieGraphOutputMergerFrame* InRawFrameData,
 	const UMovieGraphImageSequenceOutputNode* InParentNode,
 	const UMovieGraphPipeline* InPipeline,
-	const TPair<FMovieGraphRenderDataIdentifier, TUniquePtr<FImagePixelData>>& InRenderData,
+	const TPair<FMovieGraphRenderDataIdentifier,
+	TUniquePtr<FImagePixelData>>& InRenderData,
 	const EImageFormat InImageFormat,
-	FMovieGraphResolveArgs& OutMergedFormatArgs) const
+	FMovieGraphResolveArgs& OutMergedFormatArgs,
+	FString& OutFrameTemplatedFileName) const
 {
 	const TCHAR* Extension = TEXT("");
 	switch (InImageFormat)
@@ -84,7 +87,45 @@ FString UMovieGraphImageSequenceOutputNode::CreateFileName(
 		AdditionalFormatArgs);
 
 	// Take our string path from the Output Setting and resolve it.
-	return UMovieGraphBlueprintLibrary::ResolveFilenameFormatArguments(FileNameFormatString, Params, OutMergedFormatArgs);
+	const FString ResolvedFileName = UMovieGraphBlueprintLibrary::ResolveFilenameFormatArguments(FileNameFormatString, Params, OutMergedFormatArgs);
+
+	OutFrameTemplatedFileName = GetFrameTemplatedFileName(InPipeline, Params, FileNameFormatString, OutMergedFormatArgs);
+
+	return ResolvedFileName;
+}
+
+FString UMovieGraphImageSequenceOutputNode::GetFrameTemplatedFileName(
+	const UMovieGraphPipeline* InPipeline, const FMovieGraphFilenameResolveParams& InParams, const FString& InFileNameFormatString, FMovieGraphResolveArgs& OutMergedFormatArgs) const
+{
+	const EMovieGraphFrameRangeNotation FrameRangeNotationType = GetDefault<UMovieRenderGraphEditorSettings>()->PostRenderSettings.PostRenderImageSequencePlayOptions.FrameRangeNotation;
+	FString FramePlaceholder;
+
+	switch (FrameRangeNotationType)
+	{
+	case EMovieGraphFrameRangeNotation::Hash:
+		FramePlaceholder = TEXT("#");
+		break;
+	case EMovieGraphFrameRangeNotation::DollarF:
+		FramePlaceholder = TEXT("$F");
+		break;
+	case EMovieGraphFrameRangeNotation::StartFrame:
+		{
+			const int32 RootStartFrameNumber = InPipeline->GetActiveShotList()[InPipeline->GetCurrentShotIndex()]->ShotInfo.TotalOutputRangeRoot.GetLowerBoundValue().Value;
+			FramePlaceholder = UE::MoviePipeline::GetPaddingFormatString(InParams.ZeroPadFrameNumberCount, RootStartFrameNumber);
+			break;
+		}
+	default:
+		UE_LOG(LogMovieRenderPipeline, Warning, TEXT("An invalid value was specified for the Frame Range Notation property. This property will be ignored."));
+		break;
+	}
+
+	if (!FramePlaceholder.IsEmpty())
+	{
+		const FString FrameTemplatedFormatString = InFileNameFormatString.Replace(TEXT("{frame_number}"), *FramePlaceholder);
+		return UMovieGraphBlueprintLibrary::ResolveFilenameFormatArguments(FrameTemplatedFormatString, InParams, OutMergedFormatArgs);
+	}
+
+	return FString();
 }
 
 void UMovieGraphImageSequenceOutputNode::OnReceiveImageDataImpl(UMovieGraphPipeline* InPipeline, UE::MovieGraph::FMovieGraphOutputMergerFrame* InRawFrameData, const TSet<FMovieGraphRenderDataIdentifier>& InMask)
@@ -136,7 +177,8 @@ void UMovieGraphImageSequenceOutputNode::OnReceiveImageDataImpl(UMovieGraphPipel
 		checkf(ParentNode, TEXT("Image sequence output should not exist without a parent node in the graph."));
 		
 		FMovieGraphResolveArgs FinalResolvedKVPs;
-		FString FileName = CreateFileName(InRawFrameData, ParentNode, InPipeline, RenderData, PreferredOutputFormat, FinalResolvedKVPs);
+		FString FrameTemplatedFileName;
+		FString FileName = CreateFileName(InRawFrameData, ParentNode, InPipeline, RenderData, PreferredOutputFormat, FinalResolvedKVPs, FrameTemplatedFileName);
 		if (!ensureMsgf(!FileName.IsEmpty(), TEXT("Unexpected empty file name, skipping frame.")))
 		{
 			continue;
@@ -235,7 +277,10 @@ void UMovieGraphImageSequenceOutputNode::OnReceiveImageDataImpl(UMovieGraphPipel
 		UE::MovieGraph::FMovieGraphOutputFutureData OutputData;
 		OutputData.Shot = InPipeline->GetActiveShotList()[Payload->TraversalContext.ShotIndex];
 		OutputData.FilePath = FileName;
+		OutputData.FrameTemplatedFilePath = FrameTemplatedFileName;
 		OutputData.DataIdentifier = RenderData.Key;
+		OutputData.OriginNodeClass = GetClass();
+		OutputData.RenderLayerIndex = Payload->RenderLayerIndex;
 
 		TFuture<bool> Future = ImageWriteQueue->Enqueue(MoveTemp(TileImageTask));
 
@@ -363,7 +408,8 @@ void UMovieGraphImageSequenceOutputNode_EXR::OnReceiveImageDataImpl(UMovieGraphP
 		checkf(ParentNode, TEXT("Single-layer EXR should not exist without a parent node in the graph."));
 
 		FMovieGraphResolveArgs ResolvedFormatArgs;
-		FString FileName = CreateFileName(InRawFrameData, ParentNode, InPipeline, RenderData, OutputFormat, ResolvedFormatArgs);
+		FString FrameTemplatedFileName;
+		FString FileName = CreateFileName(InRawFrameData, ParentNode, InPipeline, RenderData, OutputFormat, ResolvedFormatArgs, FrameTemplatedFileName);
 		if (!ensureMsgf(!FileName.IsEmpty(), TEXT("Unexpected empty file name, skipping frame.")))
 		{
 			continue;
@@ -433,7 +479,10 @@ void UMovieGraphImageSequenceOutputNode_EXR::OnReceiveImageDataImpl(UMovieGraphP
 		UE::MovieGraph::FMovieGraphOutputFutureData OutputFutureData;
 		OutputFutureData.Shot = InPipeline->GetActiveShotList()[Payload->TraversalContext.ShotIndex];
 		OutputFutureData.FilePath = FileName;
+		OutputFutureData.FrameTemplatedFilePath = FrameTemplatedFileName;
 		OutputFutureData.DataIdentifier = RenderData.Key;
+		OutputFutureData.OriginNodeClass = GetClass();
+		OutputFutureData.RenderLayerIndex = Payload->RenderLayerIndex;
 
 		TFuture<bool> Future = ImageWriteQueue->Enqueue(MoveTemp(ImageWriteTask));
 
@@ -452,12 +501,13 @@ void UMovieGraphImageSequenceOutputNode_MultiLayerEXR::OnReceiveImageDataImpl(UM
 		UMovieGraphNode::GlobalsPinName, bIncludeCDOs, bExactMatch);
 	checkf(ParentNode, TEXT("Multi-Layer EXR should not exist without a parent node in the graph."));
 
-	// Generate a mapping of resolved filename -> RenderIDs, and filename -> resolve args. The generated EXRs can only
+	// Generate a mapping of resolved filename -> RenderIDs, filename -> resolve args, and filename -> frame templated filename. The generated EXRs can only
 	// store layers with a common resolution, and this takes care of ensuring that only renderIDs with a common resolution
 	// map to the same filename.
 	TMap<FString, TArray<FMovieGraphRenderDataIdentifier>> FilenameToRenderIDs;
 	TMap<FString, FMovieGraphResolveArgs> FilenameToResolveArgs;
-	GetFilenameToRenderIDMappings(ParentNode, InPipeline, InRawFrameData, FilenameToRenderIDs, FilenameToResolveArgs);
+	TMap<FString, FString> FilenameToFrameTemplatedFilename;
+	GetFilenameToRenderIDMappings(ParentNode, InPipeline, InRawFrameData, FilenameToRenderIDs, FilenameToResolveArgs, FilenameToFrameTemplatedFilename);
 
 	// Write an EXR for each filename, which potentially contains multiple passes (render IDs).
 	for (const TPair<FString, TArray<FMovieGraphRenderDataIdentifier>>& RenderIDsForFilename : FilenameToRenderIDs)
@@ -468,6 +518,11 @@ void UMovieGraphImageSequenceOutputNode_MultiLayerEXR::OnReceiveImageDataImpl(UM
 		TUniquePtr<FEXRImageWriteTask> MultiLayerImageTask = CreateImageWriteTask(Filename, ParentNode->Compression);
 		PrepareTaskGlobalMetadata(*MultiLayerImageTask, InRawFrameData, FilenameToResolveArgs[Filename].FileMetadata);
 
+		// Keep track of the lowest render layer index found among the layers that are included. This will be used as the index provided to the output
+		// future. This index is used to determine what the first render layer is when "First Render Layer Only" is turned on for displaying media
+		// post-render, so for multi-layer EXRs, the layer with the lowest index should be used as the index for the file.
+		int32 LowestRenderLayerIndex = 100000;
+
 		// Add each render pass as a layer to the EXR
 		bool bHasGeneratedPrimaryRGBALayer = false;
 		int32 LayerIndex = 0;
@@ -480,6 +535,8 @@ void UMovieGraphImageSequenceOutputNode_MultiLayerEXR::OnReceiveImageDataImpl(UM
 			const UE::MovieGraph::FMovieGraphSampleState* Payload = ImageData->GetPayload<UE::MovieGraph::FMovieGraphSampleState>();
 			ShotIndex = Payload->TraversalContext.ShotIndex;
 
+			LowestRenderLayerIndex = FMath::Min(LowestRenderLayerIndex, Payload->RenderLayerIndex);
+
 			// The first layer that doesn't have an explicitly-specified name will have an empty layer name in the EXR -- this is the "primary"/RGBA
 			// layer. Otherwise, the layer name will be procedurally generated if the payload didn't specify a layer name override. Having a "primary"
 			// layer in the EXR expands compatibility with a number of applications that read EXRs.
@@ -564,10 +621,19 @@ void UMovieGraphImageSequenceOutputNode_MultiLayerEXR::OnReceiveImageDataImpl(UM
 			LayerIndex++;
 		}
 		
+		FString FrameTemplatedFilename;
+		if (FString* FrameTemplatedFilenamePtr = FilenameToFrameTemplatedFilename.Find(Filename))
+		{
+			FrameTemplatedFilename = *FrameTemplatedFilenamePtr;
+		}
+		
 		UE::MovieGraph::FMovieGraphOutputFutureData OutputFutureData;
 		OutputFutureData.Shot = InPipeline->GetActiveShotList()[ShotIndex];
 		OutputFutureData.FilePath = Filename;
+		OutputFutureData.FrameTemplatedFilePath = FrameTemplatedFilename;
 		OutputFutureData.DataIdentifier = FMovieGraphRenderDataIdentifier(); // EXRs put all the render passes internally so this resolves to a ""
+		OutputFutureData.OriginNodeClass = GetClass();
+		OutputFutureData.RenderLayerIndex = LowestRenderLayerIndex;
 
 		InPipeline->AddOutputFuture(ImageWriteQueue->Enqueue(MoveTemp(MultiLayerImageTask)), OutputFutureData);
 	}
@@ -577,7 +643,8 @@ void UMovieGraphImageSequenceOutputNode_MultiLayerEXR::GetFilenameToRenderIDMapp
 	const UMovieGraphImageSequenceOutputNode_MultiLayerEXR* InParentNode,
 	UMovieGraphPipeline* InPipeline, UE::MovieGraph::FMovieGraphOutputMergerFrame* InRawFrameData,
 	TMap<FString, TArray<FMovieGraphRenderDataIdentifier>>& OutFilenameToRenderIDs,
-	TMap<FString, FMovieGraphResolveArgs>& OutFilenameToResolveArgs) const
+	TMap<FString, FMovieGraphResolveArgs>& OutFilenameToResolveArgs,
+	TMap<FString, FString>& OutFilenameToFrameTemplatedFilename) const
 {
 	// Merge one layer's resolve args (InNewResolveArgs) into an existing set of resolve args (InExistingResolveArgs).
 	auto MergeResolveArgs = [](FMovieGraphResolveArgs& InNewResolveArgs, FMovieGraphResolveArgs& InExistingResolveArgs)
@@ -621,13 +688,16 @@ void UMovieGraphImageSequenceOutputNode_MultiLayerEXR::GetFilenameToRenderIDMapp
 	// problem.
 	for (const TPair<FMovieGraphRenderDataIdentifier, TUniquePtr<FImagePixelData>>& RenderPassData : InRawFrameData->ImageOutputData)
 	{
+		FString FrameTemplatedFilename;
 		constexpr int32 ResolutionIndex = 0;
 		FMovieGraphResolveArgs ResolveArgs;
-		const FString PreliminaryFileName = ResolveOutputFilename(InParentNode, InPipeline, ResolutionIndex, InRawFrameData, RenderPassData.Key, ResolveArgs);
+		const FString PreliminaryFileName = ResolveOutputFilename(InParentNode, InPipeline, ResolutionIndex, InRawFrameData, RenderPassData.Key, ResolveArgs, FrameTemplatedFilename);
 		
 		TArray<FMovieGraphRenderDataIdentifier>& RenderIDs = OutFilenameToRenderIDs.FindOrAdd(PreliminaryFileName);
 		RenderIDs.Add(RenderPassData.Key);
 
+		OutFilenameToFrameTemplatedFilename.Add(PreliminaryFileName, FrameTemplatedFilename);
+
 		TArray<FIntPoint>& Resolutions = FilenameToResolutions.FindOrAdd(PreliminaryFileName);
 		Resolutions.AddUnique(RenderPassData.Value->GetSize());
 
@@ -656,12 +726,16 @@ void UMovieGraphImageSequenceOutputNode_MultiLayerEXR::GetFilenameToRenderIDMapp
 
 			// Re-resolve the filename, this time using the resolution index to generate a filename that will only contain
 			// passes with this particular resolution
+			FString FrameTemplatedFilename;
 			FMovieGraphResolveArgs ResolveArgs;
-			const FString FinalFilename = ResolveOutputFilename(InParentNode, InPipeline, ResolutionIndex, InRawFrameData, RenderID, ResolveArgs);
+			const FString FinalFilename = ResolveOutputFilename(InParentNode, InPipeline, ResolutionIndex, InRawFrameData, RenderID, ResolveArgs, FrameTemplatedFilename);
 
 			TArray<FMovieGraphRenderDataIdentifier>& RenderIDs = OutFilenameToRenderIDs.FindOrAdd(FinalFilename);
 			RenderIDs.Add(RenderID);
 
+			OutFilenameToFrameTemplatedFilename.Remove(PreliminaryFilename);
+			OutFilenameToFrameTemplatedFilename.Add(FinalFilename, FrameTemplatedFilename);
+
 			MergeResolveArgs(ResolveArgs, OutFilenameToResolveArgs.FindOrAdd(FinalFilename));
 		}
 	}
@@ -670,8 +744,11 @@ void UMovieGraphImageSequenceOutputNode_MultiLayerEXR::GetFilenameToRenderIDMapp
 FString UMovieGraphImageSequenceOutputNode_MultiLayerEXR::ResolveOutputFilename(
 	const UMovieGraphImageSequenceOutputNode_MultiLayerEXR* InParentNode,
 	const UMovieGraphPipeline* InPipeline,
-	const int32 ResolutionIndex, const UE::MovieGraph::FMovieGraphOutputMergerFrame* InRawFrameData,
-	const FMovieGraphRenderDataIdentifier& InRenderDataIdentifier, FMovieGraphResolveArgs& OutResolveArgs) const
+	const int32 ResolutionIndex,
+	const UE::MovieGraph::FMovieGraphOutputMergerFrame* InRawFrameData,
+	const FMovieGraphRenderDataIdentifier& InRenderDataIdentifier,
+	FMovieGraphResolveArgs& OutResolveArgs,
+	FString& OutFrameTemplatedFilename) const
 {
 	const TCHAR* Extension = TEXT("exr");
 
@@ -725,6 +802,8 @@ FString UMovieGraphImageSequenceOutputNode_MultiLayerEXR::ResolveOutputFilename(
 
 	FString FinalFilePath = UMovieGraphBlueprintLibrary::ResolveFilenameFormatArguments(FilePathFormatString, Params, OutResolveArgs);
 
+	OutFrameTemplatedFilename = GetFrameTemplatedFileName(InPipeline, Params, FileNameFormatString, OutResolveArgs);
+
 	if (FPaths::IsRelative(FinalFilePath))
 	{
 		FinalFilePath = FPaths::ConvertRelativePathToFull(FinalFilePath);
diff --git a/Engine/Plugins/MovieScene/MovieRenderPipeline/Source/MovieRenderPipelineRenderPasses/Public/MovieGraphImageSequenceOutputNode.h b/Engine/Plugins/MovieScene/MovieRenderPipeline/Source/MovieRenderPipelineRenderPasses/Public/MovieGraphImageSequenceOutputNode.h
--- a/Engine/Plugins/MovieScene/MovieRenderPipeline/Source/MovieRenderPipelineRenderPasses/Public/MovieGraphImageSequenceOutputNode.h
+++ b/Engine/Plugins/MovieScene/MovieRenderPipeline/Source/MovieRenderPipelineRenderPasses/Public/MovieGraphImageSequenceOutputNode.h
@@ -20,7 +20,7 @@ class IImageWriteQueue;
 * UMovieGraphImageSequenceOutputNode_JPG) instead of this abstract base class.
 */
 UCLASS(Abstract)
-class UMovieGraphImageSequenceOutputNode : public UMovieGraphFileOutputNode
+class MOVIERENDERPIPELINERENDERPASSES_API UMovieGraphImageSequenceOutputNode : public UMovieGraphFileOutputNode
 {
 	GENERATED_BODY()
 public:
@@ -66,6 +66,14 @@ protected:
 		const UMovieGraphPipeline* InPipeline,
 		const TPair<FMovieGraphRenderDataIdentifier, TUniquePtr<FImagePixelData>>& InRenderData,
 		const EImageFormat InImageFormat,
+		FMovieGraphResolveArgs& OutMergedFormatArgs,
+		FString& OutFrameTemplatedFileName) const;
+
+	/** Gets a "frame-number templated" filename (eg, Seq.Shot.#.exr) where '#' is used in place of the frame number. */
+	FString GetFrameTemplatedFileName(
+		const UMovieGraphPipeline* InPipeline,
+		const FMovieGraphFilenameResolveParams& InParams,
+		const FString& InFileNameFormatString,
 		FMovieGraphResolveArgs& OutMergedFormatArgs) const;
 
 protected:
@@ -226,13 +234,15 @@ private:
 	 * Generates filenames for each render pass (renderID), placed in OutFilenameToRenderIDs. Also provides the resolve
 	 * args that were created when resolving the filename in OutFilenameToResolveArgs. The generated mapping ensures
 	 * that each filename only points to renderIDs with a common resolution, since currently EXRs can only contain layers
-	 * with the same resolution.
+	 * with the same resolution. Frame-templated filenames (ie, filenames that have placeholders in the place of frame numbers)
+	 * are also provided -- the key in OutFilenameToFrameTemplatedFilename is the non-templated filename, and the value is the templated filename.
 	 */
 	void GetFilenameToRenderIDMappings(
 		const UMovieGraphImageSequenceOutputNode_MultiLayerEXR* InParentNode,
 		UMovieGraphPipeline* InPipeline, UE::MovieGraph::FMovieGraphOutputMergerFrame* InRawFrameData,
 		TMap<FString, TArray<FMovieGraphRenderDataIdentifier>>& OutFilenameToRenderIDs,
-		TMap<FString, FMovieGraphResolveArgs>& OutFilenameToResolveArgs) const;
+		TMap<FString, FMovieGraphResolveArgs>& OutFilenameToResolveArgs,
+		TMap<FString, FString>& OutFilenameToFrameTemplatedFilename) const;
 	
 	/**
 	 * Generates the filename that the EXR will be written to, as well as the resulting resolve args via OutResolveArgs.
@@ -241,7 +251,8 @@ private:
 	FString ResolveOutputFilename(
 		const UMovieGraphImageSequenceOutputNode_MultiLayerEXR* InParentNode,
 		const UMovieGraphPipeline* InPipeline, const int32 ResolutionIndex, const UE::MovieGraph::FMovieGraphOutputMergerFrame* InRawFrameData,
-		const FMovieGraphRenderDataIdentifier& InRenderDataIdentifier, FMovieGraphResolveArgs& OutResolveArgs) const;
+		const FMovieGraphRenderDataIdentifier& InRenderDataIdentifier, FMovieGraphResolveArgs& OutResolveArgs,
+		FString& OutFrameTemplatedFilename) const;
 };
 
 /**
