MoviePipeline: Added _API exports to debug-widget related classes so they can be inherited from in C++ for custom pipelines.

diff --git a/Engine/Plugins/MovieScene/MovieRenderPipeline/Source/MovieRenderPipelineCore/Public/MovieRenderDebugWidget.h b/Engine/Plugins/MovieScene/MovieRenderPipeline/Source/MovieRenderPipelineCore/Public/MovieRenderDebugWidget.h
--- a/Engine/Plugins/MovieScene/MovieRenderPipeline/Source/MovieRenderPipelineCore/Public/MovieRenderDebugWidget.h
+++ b/Engine/Plugins/MovieScene/MovieRenderPipeline/Source/MovieRenderPipelineCore/Public/MovieRenderDebugWidget.h
@@ -12,7 +12,7 @@ class UMoviePipeline;
 * the current state of the pipeline.
 */
 UCLASS(Blueprintable, Abstract)
-class UMovieRenderDebugWidget : public UUserWidget
+class MOVIERENDERPIPELINECORE_API UMovieRenderDebugWidget : public UUserWidget
 {
 	GENERATED_BODY()
 public:
@@ -26,7 +26,7 @@ public:
 * the current state of the pipeline. Used for Graph Based pipelines.
 */
 UCLASS(Blueprintable, Abstract)
-class UMovieGraphRenderPreviewWidget : public UUserWidget
+class MOVIERENDERPIPELINECORE_API UMovieGraphRenderPreviewWidget : public UUserWidget
 {
 	GENERATED_BODY()
 public:
