MoviePipeline: Added the basics for a Quick Render feature. Still many more features and improvements to add.

This is currently hidden behind a new cvar: MoviePipeline.EnableQuickRender. Not all options work yet, and the Settings dialog still needs to be implemented. It is, however, working with the post-render file display system that was recently added.

Also added a new Apply Viewport Look node that Quick Render graphs should use if they want to make renders appear like they do in the viewport. More updates to this node will be coming soon.

Misc:
* Fixed a bug in the PIE executor that caused the OnIndividualJobWorkFinished delegate to not pass along graph output data. Outside of Quick Render, this would mostly impact scripting.
* Fixed a bug in the file display processor that did not correctly handle the First Render Layer Only setting in some cases.
* Moved some previously private utilities to public utilities within Core (including a helper to duplicate graphs).

diff --git a/Engine/Plugins/MovieScene/MovieRenderPipeline/Source/MovieRenderPipelineCore/MovieRenderPipelineCore.Build.cs b/Engine/Plugins/MovieScene/MovieRenderPipeline/Source/MovieRenderPipelineCore/MovieRenderPipelineCore.Build.cs
--- a/Engine/Plugins/MovieScene/MovieRenderPipeline/Source/MovieRenderPipelineCore/MovieRenderPipelineCore.Build.cs
+++ b/Engine/Plugins/MovieScene/MovieRenderPipeline/Source/MovieRenderPipelineCore/MovieRenderPipelineCore.Build.cs
@@ -26,6 +26,7 @@ public class MovieRenderPipelineCore : ModuleRules
 				"OpenColorIOWrapper",
 				"ImageCore",
 				"ChaosClothAssetEngine",
+				"Renderer",
 			}
 		);
 
diff --git a/Engine/Plugins/MovieScene/MovieRenderPipeline/Source/MovieRenderPipelineCore/Private/Graph/MovieGraphPipeline.cpp b/Engine/Plugins/MovieScene/MovieRenderPipeline/Source/MovieRenderPipelineCore/Private/Graph/MovieGraphPipeline.cpp
--- a/Engine/Plugins/MovieScene/MovieRenderPipeline/Source/MovieRenderPipelineCore/Private/Graph/MovieGraphPipeline.cpp
+++ b/Engine/Plugins/MovieScene/MovieRenderPipeline/Source/MovieRenderPipelineCore/Private/Graph/MovieGraphPipeline.cpp
@@ -182,7 +182,7 @@ void UMovieGraphPipeline::DuplicateJobAndConfiguration()
 	// don't have the job as their outer will not have been duplicate (which is good for the World/Sequence),
 	// but this also means that the graph configurations were not duplicated (as they are assets), so we need
 	// to manually duplicate them and update the pointers in the duplicated job.
-	UMovieGraphConfig* DuplicatePrimaryConfig = DuplicateConfigRecursive(CurrentJob->GetGraphPreset(), DuplicatedGraphs);
+	UMovieGraphConfig* DuplicatePrimaryConfig = MoviePipeline::DuplicateConfigRecursive(CurrentJob->GetGraphPreset(), DuplicatedGraphs);
 	UpdateVariableAssignmentsHelper(CurrentJobDuplicate.Get(), DuplicatedGraphs);
 
 	// Set the graph preset WITHOUT updating assignments. We're doing all updates manually here. SetGraphPreset() will attempt to update the
@@ -196,7 +196,7 @@ void UMovieGraphPipeline::DuplicateJobAndConfiguration()
 		// Now for each shot we need to duplicate its config (if any)
 		if (UMovieGraphConfig* ShotGraphPreset = CurrentJob->ShotInfo[Index]->GetGraphPreset())
 		{
-			UMovieGraphConfig* DuplicateShotConfig = DuplicateConfigRecursive(ShotGraphPreset, DuplicatedGraphs);
+			UMovieGraphConfig* DuplicateShotConfig = MoviePipeline::DuplicateConfigRecursive(ShotGraphPreset, DuplicatedGraphs);
 			CurrentJobDuplicate->ShotInfo[Index]->SetGraphPreset(DuplicateShotConfig, bUpdateVariableAssignments);
 		}
 
@@ -235,70 +235,6 @@ void UMovieGraphPipeline::DuplicateJobAndConfiguration()
 			}
 		}
 	}
-
-}
-
-UMovieGraphConfig* UMovieGraphPipeline::DuplicateConfigRecursive(UMovieGraphConfig* InGraphToDuplicate, TMap<TObjectPtr<UMovieGraphConfig>, TObjectPtr<UMovieGraphConfig>>& OutDuplicatedGraphs)
-{
-	UMovieGraphConfig* DuplicateConfig;
-
-	// Duplicate the graph. If the graph has been duplicated already, don't re-duplicate it, but continue updating variable assignments.
-	if (OutDuplicatedGraphs.Contains(InGraphToDuplicate))
-	{
-		DuplicateConfig = OutDuplicatedGraphs[InGraphToDuplicate];
-	}
-	else
-	{
-		// The transient package is used because graphs don't belong to the executor job usually (they belong to an asset package)
-		FObjectDuplicationParameters GraphDuplicationParams(InGraphToDuplicate, GetTransientPackage());
-		GraphDuplicationParams.DestName = FName(FString::Format(TEXT("{0}_Duplicate"), {InGraphToDuplicate->GetFName().ToString()}));
-		GraphDuplicationParams.FlagMask = RF_AllFlags & ~(RF_Standalone | RF_Transactional);
-		GraphDuplicationParams.ApplyFlags = RF_Transient;
-		DuplicateConfig = Cast<UMovieGraphConfig>(StaticDuplicateObjectEx(GraphDuplicationParams));
-		
-		OutDuplicatedGraphs.Add(InGraphToDuplicate, DuplicateConfig);
-	}
-
-	// Duplicate sub-graphs also.
-	for (const TObjectPtr<UMovieGraphNode>& Node : DuplicateConfig->GetNodes())
-	{
-		UMovieGraphSubgraphNode* SubgraphNode = Cast<UMovieGraphSubgraphNode>(Node);
-		if (!SubgraphNode)
-		{
-			continue;
-		}
-
-		// Only duplicate if the subgraph node has a graph asset assigned to it.
-		if (UMovieGraphConfig* SubgraphConfig = SubgraphNode->GetSubgraphAsset())
-		{
-			// Don't recurse into this graph if it was already duplicated. Check BOTH the keys (the original graph) AND value (the duplicated graph)
-			// to prevent recursion. Checking the key ensures that we only duplicate if this graph has never been encountered. Checking the value
-			// ensures that we don't re-duplicate a graph that has already been duplicated (the subgraph node was already updated).
-			bool bHasBeenDuplicated = false;
-			for (const TPair<TObjectPtr<UMovieGraphConfig>, TObjectPtr<UMovieGraphConfig>>& DuplicateMapping : OutDuplicatedGraphs)
-			{
-				if ((DuplicateMapping.Key == SubgraphConfig) || (DuplicateMapping.Value == SubgraphConfig))
-				{
-					bHasBeenDuplicated = true;
-					break;
-				}
-			}
-			
-			if (!bHasBeenDuplicated)
-			{
-				DuplicateConfigRecursive(SubgraphConfig, OutDuplicatedGraphs);
-			}
-
-			// Update the subgraph node to use the duplicated graph. This should always be done, even if the graph was already duplicated (since
-			// a graph can be included as a subgraph in multiple locations).
-			if (const TObjectPtr<UMovieGraphConfig>* DuplicatedGraph = OutDuplicatedGraphs.Find(SubgraphConfig))
-			{
-				SubgraphNode->SetSubGraphAsset(*DuplicatedGraph);
-			}
-		}
-	}
-	
-	return DuplicateConfig;
 }
 
 template <typename JobType>
diff --git a/Engine/Plugins/MovieScene/MovieRenderPipeline/Source/MovieRenderPipelineCore/Private/Graph/Nodes/MovieGraphApplyViewportLookNode.cpp b/Engine/Plugins/MovieScene/MovieRenderPipeline/Source/MovieRenderPipelineCore/Private/Graph/Nodes/MovieGraphApplyViewportLookNode.cpp
new file mode 100644
--- /dev/null
+++ b/Engine/Plugins/MovieScene/MovieRenderPipeline/Source/MovieRenderPipelineCore/Private/Graph/Nodes/MovieGraphApplyViewportLookNode.cpp
@@ -0,0 +1,91 @@
+﻿// Copyright Epic Games, Inc. All Rights Reserved.
+
+#include "Graph/Nodes/MovieGraphApplyViewportLookNode.h"
+
+#include "MeshEdges.h"
+#include "SceneView.h"
+#include "Styling/AppStyle.h"
+#include "UnrealClient.h"
+
+#if WITH_EDITOR
+#include "Editor.h"
+#include "EditorViewportClient.h"
+#endif
+
+EMovieGraphBranchRestriction UMovieGraphApplyViewportLookNode::GetBranchRestriction() const
+{
+	return EMovieGraphBranchRestriction::Globals;
+}
+
+#if WITH_EDITOR
+FEditorViewportClient* UMovieGraphApplyViewportLookNode::GetViewportClient()
+{
+	if (const FViewport* ActiveViewport = GEditor->GetActiveViewport())
+	{
+		return static_cast<FEditorViewportClient*>(ActiveViewport->GetClient());
+	}
+
+	return nullptr;
+}
+
+bool UMovieGraphApplyViewportLookNode::GetViewportInfo(FEngineShowFlags& OutShowFlags, EViewModeIndex& OutViewModeIndex) const
+{
+	if (const FEditorViewportClient* ViewportClient = GetViewportClient())
+	{
+		OutShowFlags = ViewportClient->EngineShowFlags;
+		OutViewModeIndex = ViewportClient->GetViewMode();
+
+		return true;
+	}
+
+	return false;
+}
+
+void UMovieGraphApplyViewportLookNode::UpdateSceneView(FSceneView* InSceneView) const
+{
+	const FEditorViewportClient* ViewportClient = GetViewportClient();
+	
+	if (InSceneView && ViewportClient)
+	{
+		// Note: Most of the scene view is set up in the deferred pass; these are the settings that are not handled there.
+		
+		InSceneView->CurrentBufferVisualizationMode = ViewportClient->CurrentBufferVisualizationMode;
+		InSceneView->CurrentNaniteVisualizationMode = ViewportClient->CurrentNaniteVisualizationMode;
+		InSceneView->CurrentLumenVisualizationMode = ViewportClient->CurrentLumenVisualizationMode;
+		InSceneView->CurrentSubstrateVisualizationMode = ViewportClient->CurrentSubstrateVisualizationMode;
+		InSceneView->CurrentGroomVisualizationMode = ViewportClient->CurrentGroomVisualizationMode;
+		InSceneView->CurrentVirtualShadowMapVisualizationMode = ViewportClient->CurrentVirtualShadowMapVisualizationMode;
+		InSceneView->CurrentGPUSkinCacheVisualizationMode = ViewportClient->CurrentGPUSkinCacheVisualizationMode;
+		InSceneView->CurrentRayTracingDebugVisualizationMode = ViewportClient->CurrentRayTracingDebugVisualizationMode;
+
+		// Wireframe opacity
+		GetMeshEdgesViewSettings(*InSceneView).Opacity = ViewportClient->WireframeOpacity;
+	}
+}
+
+FText UMovieGraphApplyViewportLookNode::GetNodeTitle(const bool bGetDescriptive) const
+{
+	static const FText NodeName = NSLOCTEXT("MovieGraphNodes", "NodeName_ApplyViewportLook", "Apply Viewport Look");
+	return NodeName;
+}
+
+FText UMovieGraphApplyViewportLookNode::GetMenuCategory() const
+{
+	static const FText NodeCategory_Globals = NSLOCTEXT("MovieGraphNodes", "NodeCategory_Globals", "Globals");
+	return NodeCategory_Globals;
+}
+
+FLinearColor UMovieGraphApplyViewportLookNode::GetNodeTitleColor() const
+{
+	static const FLinearColor NodeColor = FLinearColor(0.04f, 0.22f, 0.36f);
+	return NodeColor;
+}
+
+FSlateIcon UMovieGraphApplyViewportLookNode::GetIconAndTint(FLinearColor& OutColor) const
+{
+	static const FSlateIcon Icon = FSlateIcon(FAppStyle::GetAppStyleSetName(), "LevelEditor.Tabs.Viewports");
+
+	OutColor = FLinearColor::White;
+	return Icon;
+}
+#endif // WITH_EDITOR
diff --git a/Engine/Plugins/MovieScene/MovieRenderPipeline/Source/MovieRenderPipelineCore/Private/MoviePipelineQueue.cpp b/Engine/Plugins/MovieScene/MovieRenderPipeline/Source/MovieRenderPipelineCore/Private/MoviePipelineQueue.cpp
--- a/Engine/Plugins/MovieScene/MovieRenderPipeline/Source/MovieRenderPipelineCore/Private/MoviePipelineQueue.cpp
+++ b/Engine/Plugins/MovieScene/MovieRenderPipeline/Source/MovieRenderPipelineCore/Private/MoviePipelineQueue.cpp
@@ -5,71 +5,13 @@
 #include "MoviePipelineSetting.h"
 #include "LevelSequence.h"
 #include "MoviePipelineBlueprintLibrary.h"
+#include "MoviePipelineUtils.h"
 #include "MovieRenderPipelineCoreModule.h"
 
 #include "UObject/UObjectGlobals.h"
 
 #include UE_INLINE_GENERATED_CPP_BY_NAME(MoviePipelineQueue)
 
-namespace UE::MovieGraph::Private
-{
-	/**
-	 * Gets the job variable assignments for a specific graph. Creates a new variable assignments container if one was not found for the given graph.
-	 * The owner of the variable assignments must be provided in case a new assignment needs to be created.
-	 */
-	TObjectPtr<UMovieJobVariableAssignmentContainer> GetOrCreateJobVariableAssignmentsForGraph(
-		const UMovieGraphConfig* InGraph, TArray<TObjectPtr<UMovieJobVariableAssignmentContainer>>& InVariableAssignments, UObject* InAssignmentsOwner)
-	{
-		for (TObjectPtr<UMovieJobVariableAssignmentContainer>& VariableAssignment : InVariableAssignments)
-		{
-			const TSoftObjectPtr<UMovieGraphConfig> SoftGraphConfig = VariableAssignment->GetGraphConfig();
-			if (SoftGraphConfig.Get() == InGraph)
-			{
-#if WITH_EDITOR
-				VariableAssignment->UpdateGraphVariableOverrides();
-#endif
-				
-				return VariableAssignment;
-			}
-		}
-
-		// Create the variable assignments container if it wasn't found
-		TObjectPtr<UMovieJobVariableAssignmentContainer> NewVariableAssignments = NewObject<UMovieJobVariableAssignmentContainer>(InAssignmentsOwner);
-		InVariableAssignments.Add(NewVariableAssignments);
-		NewVariableAssignments->SetGraphConfig(const_cast<UMovieGraphConfig*>(InGraph));
-
-#if WITH_EDITOR
-		NewVariableAssignments->UpdateGraphVariableOverrides();
-#endif
-	
-		return NewVariableAssignments;
-	}
-
-	/** Refreshes the variable assignments for the given graph and its associated subgraphs. */
-	void RefreshVariableAssignments(UMovieGraphConfig* InRootGraph, TArray<TObjectPtr<UMovieJobVariableAssignmentContainer>>& InVariableAssignments, UObject* InAssignmentsOwner)
-	{
-		if (!InRootGraph)
-		{
-			return;
-		}
-	
-		TSet<UMovieGraphConfig*> AllGraphs = {InRootGraph};
-		InRootGraph->GetAllContainedSubgraphs(AllGraphs);
-
-		// Add/update variable assignments for the graph on the job and all of its subgraphs
-		for (const UMovieGraphConfig* Graph : AllGraphs)
-		{
-			GetOrCreateJobVariableAssignmentsForGraph(Graph, InVariableAssignments, InAssignmentsOwner);
-		}
-
-		// Remove any stale variable assignments for graphs/subgraphs which are no longer part of the job
-		InVariableAssignments.RemoveAll([&AllGraphs](const TObjectPtr<UMovieJobVariableAssignmentContainer>& VariableAssignment)
-		{
-			return !AllGraphs.Contains(VariableAssignment->GetGraphConfig().LoadSynchronous());
-		});
-	}
-}
-
 UMoviePipelineQueue::UMoviePipelineQueue()
 	: QueueSerialNumber(0)
 	, bIsDirty(false)
@@ -298,7 +240,7 @@ TArray<TObjectPtr<UMovieJobVariableAssignmentContainer>>& UMoviePipelineExecutor
 
 void UMoviePipelineExecutorJob::RefreshAllVariableAssignments()
 {
-	UE::MovieGraph::Private::RefreshVariableAssignments(GraphPreset.LoadSynchronous(), GraphVariableAssignments, this);
+	MoviePipeline::RefreshVariableAssignments(GraphPreset.LoadSynchronous(), GraphVariableAssignments, this);
 
 	// Notify all shots that they should refresh their assignments as well (this could be needed, for example, when the primary graph preset changes)
 	for (TObjectPtr<UMoviePipelineExecutorShot>& ShotJob : ShotInfo)
@@ -362,7 +304,7 @@ UMovieJobVariableAssignmentContainer* UMoviePipelineExecutorJob::GetOrCreateJobV
 {
 	if (InGraph)
 	{
-		return UE::MovieGraph::Private::GetOrCreateJobVariableAssignmentsForGraph(InGraph, GraphVariableAssignments, this);
+		return MoviePipeline::GetOrCreateJobVariableAssignmentsForGraph(InGraph, GraphVariableAssignments, this);
 	}
 
 	return nullptr;
@@ -425,7 +367,7 @@ UMovieJobVariableAssignmentContainer* UMoviePipelineExecutorShot::GetOrCreateJob
 {
 	if (InGraph)
 	{
-		return UE::MovieGraph::Private::GetOrCreateJobVariableAssignmentsForGraph(InGraph, bIsForPrimaryOverrides ? PrimaryGraphVariableAssignments : GraphVariableAssignments, this);
+		return MoviePipeline::GetOrCreateJobVariableAssignmentsForGraph(InGraph, bIsForPrimaryOverrides ? PrimaryGraphVariableAssignments : GraphVariableAssignments, this);
 	}
 
 	return nullptr;
@@ -545,14 +487,14 @@ TArray<TObjectPtr<UMovieJobVariableAssignmentContainer>>& UMoviePipelineExecutor
 
 void UMoviePipelineExecutorShot::RefreshAllVariableAssignments()
 {
-	UE::MovieGraph::Private::RefreshVariableAssignments(GraphPreset.LoadSynchronous(), GraphVariableAssignments, this);
+	MoviePipeline::RefreshVariableAssignments(GraphPreset.LoadSynchronous(), GraphVariableAssignments, this);
 
 	// Refresh the associated primary graph variable override assignments as well
 	if (const UMoviePipelineExecutorJob* PrimaryJob = GetTypedOuter<UMoviePipelineExecutorJob>())
 	{
 		if (UMovieGraphConfig* PrimaryGraph = PrimaryJob->GetGraphPreset())
 		{
-			UE::MovieGraph::Private::RefreshVariableAssignments(PrimaryGraph, PrimaryGraphVariableAssignments, this);
+			MoviePipeline::RefreshVariableAssignments(PrimaryGraph, PrimaryGraphVariableAssignments, this);
 		}
 	}
 }
diff --git a/Engine/Plugins/MovieScene/MovieRenderPipeline/Source/MovieRenderPipelineCore/Private/MoviePipelineUtils.cpp b/Engine/Plugins/MovieScene/MovieRenderPipeline/Source/MovieRenderPipelineCore/Private/MoviePipelineUtils.cpp
--- a/Engine/Plugins/MovieScene/MovieRenderPipeline/Source/MovieRenderPipelineCore/Private/MoviePipelineUtils.cpp
+++ b/Engine/Plugins/MovieScene/MovieRenderPipeline/Source/MovieRenderPipelineCore/Private/MoviePipelineUtils.cpp
@@ -42,6 +42,7 @@
 #include "ChaosClothAsset/ClothComponent.h"
 #include "ChaosClothAsset/ClothAssetInteractor.h"
 #include "Components/SkeletalMeshComponent.h"
+#include "Graph/Nodes/MovieGraphSubgraphNode.h"
 
 namespace UE
 {
@@ -235,6 +236,118 @@ namespace MoviePipeline
 		InOutFileMetadata.Add(TEXT("unreal/shotName"), ShotName);
 	}
 
+	TObjectPtr<UMovieJobVariableAssignmentContainer> GetOrCreateJobVariableAssignmentsForGraph(const UMovieGraphConfig* InGraph, TArray<TObjectPtr<UMovieJobVariableAssignmentContainer>>& InVariableAssignments, UObject* InAssignmentsOwner)
+	{
+		for (TObjectPtr<UMovieJobVariableAssignmentContainer>& VariableAssignment : InVariableAssignments)
+		{
+			const TSoftObjectPtr<UMovieGraphConfig> SoftGraphConfig = VariableAssignment->GetGraphConfig();
+			if (SoftGraphConfig.Get() == InGraph)
+			{
+#if WITH_EDITOR
+				VariableAssignment->UpdateGraphVariableOverrides();
+#endif
+				
+				return VariableAssignment;
+			}
+		}
+
+		// Create the variable assignments container if it wasn't found
+		TObjectPtr<UMovieJobVariableAssignmentContainer> NewVariableAssignments = NewObject<UMovieJobVariableAssignmentContainer>(InAssignmentsOwner);
+		InVariableAssignments.Add(NewVariableAssignments);
+		NewVariableAssignments->SetGraphConfig(const_cast<UMovieGraphConfig*>(InGraph));
+
+#if WITH_EDITOR
+		NewVariableAssignments->UpdateGraphVariableOverrides();
+#endif
+	
+		return NewVariableAssignments;
+	}
+
+	void RefreshVariableAssignments(UMovieGraphConfig* InRootGraph, TArray<TObjectPtr<UMovieJobVariableAssignmentContainer>>& InVariableAssignments, UObject* InAssignmentsOwner)
+	{
+		if (!InRootGraph)
+		{
+			return;
+		}
+	
+		TSet<UMovieGraphConfig*> AllGraphs = {InRootGraph};
+		InRootGraph->GetAllContainedSubgraphs(AllGraphs);
+
+		// Add/update variable assignments for the graph on the job and all of its subgraphs
+		for (const UMovieGraphConfig* Graph : AllGraphs)
+		{
+			GetOrCreateJobVariableAssignmentsForGraph(Graph, InVariableAssignments, InAssignmentsOwner);
+		}
+
+		// Remove any stale variable assignments for graphs/subgraphs which are no longer part of the job
+		InVariableAssignments.RemoveAll([&AllGraphs](const TObjectPtr<UMovieJobVariableAssignmentContainer>& VariableAssignment)
+		{
+			return !AllGraphs.Contains(VariableAssignment->GetGraphConfig().LoadSynchronous());
+		});
+	}
+
+	UMovieGraphConfig* DuplicateConfigRecursive(UMovieGraphConfig* InGraphToDuplicate, TMap<TObjectPtr<UMovieGraphConfig>, TObjectPtr<UMovieGraphConfig>>& OutDuplicatedGraphs)
+	{
+		UMovieGraphConfig* DuplicateConfig;
+
+		// Duplicate the graph. If the graph has been duplicated already, don't re-duplicate it, but continue updating variable assignments.
+		if (OutDuplicatedGraphs.Contains(InGraphToDuplicate))
+		{
+			DuplicateConfig = OutDuplicatedGraphs[InGraphToDuplicate];
+		}
+		else
+		{
+			// The transient package is used because graphs don't belong to the executor job usually (they belong to an asset package)
+			FObjectDuplicationParameters GraphDuplicationParams(InGraphToDuplicate, GetTransientPackage());
+			GraphDuplicationParams.DestName = FName(FString::Format(TEXT("{0}_Duplicate"), {InGraphToDuplicate->GetFName().ToString()}));
+			GraphDuplicationParams.FlagMask = RF_AllFlags & ~(RF_Standalone | RF_Transactional);
+			GraphDuplicationParams.ApplyFlags = RF_Transient;
+			DuplicateConfig = Cast<UMovieGraphConfig>(StaticDuplicateObjectEx(GraphDuplicationParams));
+			
+			OutDuplicatedGraphs.Add(InGraphToDuplicate, DuplicateConfig);
+		}
+
+		// Duplicate sub-graphs also.
+		for (const TObjectPtr<UMovieGraphNode>& Node : DuplicateConfig->GetNodes())
+		{
+			UMovieGraphSubgraphNode* SubgraphNode = Cast<UMovieGraphSubgraphNode>(Node);
+			if (!SubgraphNode)
+			{
+				continue;
+			}
+
+			// Only duplicate if the subgraph node has a graph asset assigned to it.
+			if (UMovieGraphConfig* SubgraphConfig = SubgraphNode->GetSubgraphAsset())
+			{
+				// Don't recurse into this graph if it was already duplicated. Check BOTH the keys (the original graph) AND value (the duplicated graph)
+				// to prevent recursion. Checking the key ensures that we only duplicate if this graph has never been encountered. Checking the value
+				// ensures that we don't re-duplicate a graph that has already been duplicated (the subgraph node was already updated).
+				bool bHasBeenDuplicated = false;
+				for (const TPair<TObjectPtr<UMovieGraphConfig>, TObjectPtr<UMovieGraphConfig>>& DuplicateMapping : OutDuplicatedGraphs)
+				{
+					if ((DuplicateMapping.Key == SubgraphConfig) || (DuplicateMapping.Value == SubgraphConfig))
+					{
+						bHasBeenDuplicated = true;
+						break;
+					}
+				}
+				
+				if (!bHasBeenDuplicated)
+				{
+					DuplicateConfigRecursive(SubgraphConfig, OutDuplicatedGraphs);
+				}
+
+				// Update the subgraph node to use the duplicated graph. This should always be done, even if the graph was already duplicated (since
+				// a graph can be included as a subgraph in multiple locations).
+				if (const TObjectPtr<UMovieGraphConfig>* DuplicatedGraph = OutDuplicatedGraphs.Find(SubgraphConfig))
+				{
+					SubgraphNode->SetSubGraphAsset(*DuplicatedGraph);
+				}
+			}
+		}
+		
+		return DuplicateConfig;
+	}
 }
 
 //namespace UE
diff --git a/Engine/Plugins/MovieScene/MovieRenderPipeline/Source/MovieRenderPipelineCore/Public/Graph/MovieGraphPipeline.h b/Engine/Plugins/MovieScene/MovieRenderPipeline/Source/MovieRenderPipelineCore/Public/Graph/MovieGraphPipeline.h
--- a/Engine/Plugins/MovieScene/MovieRenderPipeline/Source/MovieRenderPipelineCore/Public/Graph/MovieGraphPipeline.h
+++ b/Engine/Plugins/MovieScene/MovieRenderPipeline/Source/MovieRenderPipelineCore/Public/Graph/MovieGraphPipeline.h
@@ -174,12 +174,6 @@ protected:
 	virtual void ExecutePreShotScripts(const TObjectPtr<UMoviePipelineExecutorShot>& InShot);
 	virtual void ExecutePostShotScripts(const FMoviePipelineOutputData& InData);
 
-	/**
-	 * Helps duplicate graph configs. Prevents re-duplications, duplicates sub-graphs (updates subgraph nodes accordingly), and potentially more.
-	 * Returns the duplicated graph.
-	 */
-	UMovieGraphConfig* DuplicateConfigRecursive(UMovieGraphConfig* InGraphToDuplicate, TMap<TObjectPtr<UMovieGraphConfig>, TObjectPtr<UMovieGraphConfig>>& OutDuplicatedGraphs);
-
 	/** Helps update variable assignments on the provided job to use duplicated graphs (reflected in the original-to-duplicate graph mapping). */
 	template<typename JobType>
 	void UpdateVariableAssignmentsHelper(JobType* InTargetJob, TMap<TObjectPtr<UMovieGraphConfig>, TObjectPtr<UMovieGraphConfig>>& InOriginalToDuplicateGraphMap);
diff --git a/Engine/Plugins/MovieScene/MovieRenderPipeline/Source/MovieRenderPipelineCore/Public/Graph/MovieGraphQuickRenderSettings.h b/Engine/Plugins/MovieScene/MovieRenderPipeline/Source/MovieRenderPipelineCore/Public/Graph/MovieGraphQuickRenderSettings.h
new file mode 100644
--- /dev/null
+++ b/Engine/Plugins/MovieScene/MovieRenderPipeline/Source/MovieRenderPipelineCore/Public/Graph/MovieGraphQuickRenderSettings.h
@@ -0,0 +1,128 @@
+﻿// Copyright Epic Games, Inc. All Rights Reserved.
+
+#pragma once
+
+#include "CoreMinimal.h"
+
+#include "MovieGraphNamedResolution.h"
+
+#include "MovieGraphQuickRenderSettings.generated.h"
+
+class UMovieGraphConfig;
+class UMovieJobVariableAssignmentContainer;
+
+/** The type of render that the Quick Render toolbar button performs. */
+UENUM(BlueprintType)
+enum class EMovieGraphQuickRenderButtonMode : uint8
+{
+	/** Uses the active queue in the Movie Render Queue editor to do a render. */
+	NormalMovieRenderQueue,
+
+	/** Use the Quick Render configuration to do a render. */
+	QuickRender
+};
+
+/** The available modes for Quick Render to be put into. */
+UENUM(BlueprintType)
+enum class EMovieGraphQuickRenderMode : uint8
+{
+	CurrentSequence,
+	CurrentShotAtPlayhead,
+	CurrentViewport,
+	SelectedCameras,
+	UseViewportCameraInSequence
+};
+
+/** The frame range within the active level sequence that a Quick Render should use. */
+UENUM(BlueprintType)
+enum class EMovieGraphQuickRenderFrameRangeType : uint8
+{
+	/** Use the frame range specified by the sequence's start and end frame. */
+	PlaybackRange,
+
+	/** Use the sequence's selection as the start and end frame. */
+	SelectionRange,
+
+	/** Use a custom frame range. */
+	Custom
+};
+
+/** The action that quick render should take after finishing the render. */
+UENUM(BlueprintType)
+enum class EMovieGraphQuickRenderPostRenderActionType : uint8
+{
+	/** Do nothing after the render. */
+	DoNothing,
+
+	/** Open the movie that was generated from Quick Render. Opens via the operating system's default viewing application for the type of media that was created. */
+	PlayRenderOutput
+};
+
+/** The aspects of the viewport that should be applied to the quick render. */
+UENUM(BlueprintType, meta = (Bitflags, UseEnumValuesAsMaskValuesInEditor = "true"))
+enum class EMovieGraphQuickRenderViewportLookFlags : uint8
+{
+	None		= 0	UMETA(Hidden),
+	Ocio		= 1 << 0 UMETA(DisplayName = "OCIO"),
+	ShowFlags	= 1 << 1,
+	ViewMode	= 1 << 2,
+};
+ENUM_CLASS_FLAGS(EMovieGraphQuickRenderViewportLookFlags)
+
+/** Data that sets up user-configurable aspects of a quick render. */
+USTRUCT(BlueprintType)
+struct MOVIERENDERPIPELINECORE_API FMovieGraphQuickRenderSettings
+{
+	GENERATED_BODY()
+
+	/** How to determine what the start/end frame should be. */
+	UPROPERTY(Config)
+	EMovieGraphQuickRenderFrameRangeType FrameRangeType = EMovieGraphQuickRenderFrameRangeType::PlaybackRange;
+
+	/** The frame in the level sequence that the render will begin on. */
+	UPROPERTY(Config)
+	int32 CustomStartFrame = 0;
+
+	/** The frame in the level sequence that the render will end on. */
+	UPROPERTY(Config)
+	int32 CustomEndFrame = 0;
+
+	/** The properties of the viewport that should be applied to the quick render. These will override any equivalent properties specified in the graph. */
+	UPROPERTY(Config)
+	EMovieGraphQuickRenderViewportLookFlags ViewportLookFlags = EMovieGraphQuickRenderViewportLookFlags::None;
+
+	// TODO: This should probably be deleted, since it will be controlled by variable assignments
+	/** The directory that render outputs are saved to. */
+	UPROPERTY(Config)
+	FDirectoryPath OutputDirectory;
+
+	/** The graph preset that is used to configure the render. If not specified, the quick render default graph will be used. */
+	UPROPERTY(Config)
+	TSoftObjectPtr<UMovieGraphConfig> GraphPreset = nullptr;
+
+	/** What mode the quick render mode button is using. */
+	UPROPERTY(Config)
+	EMovieGraphQuickRenderButtonMode QuickAccessButtonMode = EMovieGraphQuickRenderButtonMode::NormalMovieRenderQueue;
+
+	/** The mode that Quick Render will use when a render begins. Generally dictates the level sequence and camera(s) that will be used. */
+	UPROPERTY(Config)
+	EMovieGraphQuickRenderMode QuickRenderMode = EMovieGraphQuickRenderMode::CurrentSequence;
+
+	/** Whether the settings dialog should be shown before a quick render is started. */
+	UPROPERTY(Config)
+	bool bShowSettingsBeforeQuickRender = false;
+
+	/** The action that quick render should perform after a render is finished. */
+	UPROPERTY(Config)
+	EMovieGraphQuickRenderPostRenderActionType PostRenderBehavior = EMovieGraphQuickRenderPostRenderActionType::DoNothing;
+
+	/** Values that are set on the variables contained within the graph preset. */
+	// TODO: This can't be stored as config... how to fix this?
+	UPROPERTY()
+	TArray<TObjectPtr<UMovieJobVariableAssignmentContainer>> GraphVariableAssignments;
+
+	// TODO: This should probably be deleted, since it will be controlled by variable assignments
+	/** The resolution to render at. */
+	UPROPERTY(Config)
+	FMovieGraphNamedResolution Resolution = FMovieGraphNamedResolution(FName(), FIntPoint(-1, -1), FString());
+};
\ No newline at end of file
diff --git a/Engine/Plugins/MovieScene/MovieRenderPipeline/Source/MovieRenderPipelineCore/Public/Graph/MovieRenderGraphEditorSettings.h b/Engine/Plugins/MovieScene/MovieRenderPipeline/Source/MovieRenderPipelineCore/Public/Graph/MovieRenderGraphEditorSettings.h
--- a/Engine/Plugins/MovieScene/MovieRenderPipeline/Source/MovieRenderPipelineCore/Public/Graph/MovieRenderGraphEditorSettings.h
+++ b/Engine/Plugins/MovieScene/MovieRenderPipeline/Source/MovieRenderPipelineCore/Public/Graph/MovieRenderGraphEditorSettings.h
@@ -5,6 +5,7 @@
 #include "Engine/DeveloperSettings.h"
 
 #include "MovieGraphDataTypes.h"
+#include "MovieGraphQuickRenderSettings.h"
 #include "MoviePipelinePostRenderSettings.h"
 
 #include "MovieRenderGraphEditorSettings.generated.h"
@@ -20,4 +21,8 @@ public:
 	/** If PostRenderBehavior is set to PlayRenderOutput, these settings are used to determine how to play media. */
 	UPROPERTY(Config, EditAnywhere, Category="Play Render Output", meta=(ShowOnlyInnerProperties))
 	FMovieGraphPostRenderSettings PostRenderSettings;
+
+	/** Settings that are used to generate a quick render. */
+	UPROPERTY(Config)
+	FMovieGraphQuickRenderSettings QuickRenderSettings;
 };
diff --git a/Engine/Plugins/MovieScene/MovieRenderPipeline/Source/MovieRenderPipelineCore/Public/Graph/Nodes/MovieGraphApplyViewportLookNode.h b/Engine/Plugins/MovieScene/MovieRenderPipeline/Source/MovieRenderPipelineCore/Public/Graph/Nodes/MovieGraphApplyViewportLookNode.h
new file mode 100644
--- /dev/null
+++ b/Engine/Plugins/MovieScene/MovieRenderPipeline/Source/MovieRenderPipelineCore/Public/Graph/Nodes/MovieGraphApplyViewportLookNode.h
@@ -0,0 +1,45 @@
+﻿// Copyright Epic Games, Inc. All Rights Reserved.
+
+#pragma once
+
+#include "Graph/MovieGraphNode.h"
+#include "ShowFlags.h"
+
+#include "MovieGraphApplyViewportLookNode.generated.h"
+
+class FEditorViewportClient;
+class FSceneView;
+
+/** A node which applies the look of the viewport to the render (show flags, view mode, OCIO, etc). */
+UCLASS()
+class MOVIERENDERPIPELINECORE_API UMovieGraphApplyViewportLookNode : public UMovieGraphSettingNode
+{
+	GENERATED_BODY()
+
+public:
+	UMovieGraphApplyViewportLookNode() = default;
+
+	//~ Begin UMovieGraphNode interface (not editor-only)
+	virtual EMovieGraphBranchRestriction GetBranchRestriction() const override;
+	//~ End UMovieGraphNode interface
+
+	// TODO: Expose a property which allows a selection of the following to be applied: OCIO, Show Flags, View Mode
+
+#if WITH_EDITOR
+	/** Gets the viewport client for the currently active viewport. */
+	static FEditorViewportClient* GetViewportClient();
+	
+	/** Gets the show flags and view mode index for the currently active viewport. */
+	bool GetViewportInfo(FEngineShowFlags& OutShowFlags, EViewModeIndex& OutViewModeIndex) const;
+
+	/** Updates the given scene view to be like the current viewport's scene view. */
+	void UpdateSceneView(FSceneView* InSceneView) const;
+
+	//~ Begin UMovieGraphNode interface (editor-only)
+	virtual FText GetNodeTitle(const bool bGetDescriptive = false) const override;
+	virtual FText GetMenuCategory() const override;
+	virtual FLinearColor GetNodeTitleColor() const override;
+	virtual FSlateIcon GetIconAndTint(FLinearColor& OutColor) const override;
+	//~ End UMovieGraphNode interface
+#endif
+};
\ No newline at end of file
diff --git a/Engine/Plugins/MovieScene/MovieRenderPipeline/Source/MovieRenderPipelineCore/Public/MoviePipelineUtils.h b/Engine/Plugins/MovieScene/MovieRenderPipeline/Source/MovieRenderPipelineCore/Public/MoviePipelineUtils.h
--- a/Engine/Plugins/MovieScene/MovieRenderPipeline/Source/MovieRenderPipelineCore/Public/MoviePipelineUtils.h
+++ b/Engine/Plugins/MovieScene/MovieRenderPipeline/Source/MovieRenderPipelineCore/Public/MoviePipelineUtils.h
@@ -167,6 +167,21 @@ namespace UE
 namespace MoviePipeline
 {
 	MOVIERENDERPIPELINECORE_API void GetOutputStateFormatArgs(TMap<FString, FString>& InFilenameArguments, TMap<FString, FString>& InFileMetadata, const FString FrameNumber, const FString FrameNumberShot, const FString FrameNumberRel, const FString FrameNumberShotRel, const FString CameraName, const FString ShotName);
+
+	/**
+	 * Gets the job variable assignments for a specific graph. Creates a new variable assignments container if one was not found for the given graph.
+	 * The owner of the variable assignments must be provided in case a new assignment needs to be created.
+	 */
+	MOVIERENDERPIPELINECORE_API TObjectPtr<UMovieJobVariableAssignmentContainer> GetOrCreateJobVariableAssignmentsForGraph(const UMovieGraphConfig* InGraph, TArray<TObjectPtr<UMovieJobVariableAssignmentContainer>>& InVariableAssignments, UObject* InAssignmentsOwner);
+
+	/** Refreshes the variable assignments for the given graph and its associated subgraphs. */
+	MOVIERENDERPIPELINECORE_API void RefreshVariableAssignments(UMovieGraphConfig* InRootGraph, TArray<TObjectPtr<UMovieJobVariableAssignmentContainer>>& InVariableAssignments, UObject* InAssignmentsOwner);
+	
+	/**
+	 * Helps duplicate graph configs. Prevents re-duplications, duplicates sub-graphs (updates subgraph nodes accordingly), and potentially more.
+	 * Returns the duplicated graph.
+	 */
+	MOVIERENDERPIPELINECORE_API UMovieGraphConfig* DuplicateConfigRecursive(UMovieGraphConfig* InGraphToDuplicate, TMap<TObjectPtr<UMovieGraphConfig>, TObjectPtr<UMovieGraphConfig>>& OutDuplicatedGraphs);
 	
 	/** Iterate root-to-tails and generate a HierarchyNode for each level. Caches the sub-section range, playback range, camera cut range, etc. */
 	void CacheCompleteSequenceHierarchy(UMovieSceneSequence* InSequence, TSharedPtr<FCameraCutSubSectionHierarchyNode> InRootNode);
diff --git a/Engine/Plugins/MovieScene/MovieRenderPipeline/Source/MovieRenderPipelineEditor/MovieRenderPipelineEditor.Build.cs b/Engine/Plugins/MovieScene/MovieRenderPipeline/Source/MovieRenderPipelineEditor/MovieRenderPipelineEditor.Build.cs
--- a/Engine/Plugins/MovieScene/MovieRenderPipeline/Source/MovieRenderPipelineEditor/MovieRenderPipelineEditor.Build.cs
+++ b/Engine/Plugins/MovieScene/MovieRenderPipeline/Source/MovieRenderPipelineEditor/MovieRenderPipelineEditor.Build.cs
@@ -71,6 +71,8 @@ public class MovieRenderPipelineEditor : ModuleRules
 				"ApplicationCore",
 				"KismetWidgets",
 				"BlueprintGraph",
+				"OpenColorIO",
+				"OpenColorIOEditor",
 				"Imath",		// This and UEOpenExr are only required for use of UMovieGraphImageSequenceOutputNode. Ideally these are refactored
 				"UEOpenExr",	// out of the ImageSequenceOutputNode header at some point so these dependencies can be removed.
 			}
diff --git a/Engine/Plugins/MovieScene/MovieRenderPipeline/Source/MovieRenderPipelineEditor/Private/Graph/MovieGraphQuickRender.cpp b/Engine/Plugins/MovieScene/MovieRenderPipeline/Source/MovieRenderPipelineEditor/Private/Graph/MovieGraphQuickRender.cpp
new file mode 100644
--- /dev/null
+++ b/Engine/Plugins/MovieScene/MovieRenderPipeline/Source/MovieRenderPipelineEditor/Private/Graph/MovieGraphQuickRender.cpp
@@ -0,0 +1,387 @@
+// Copyright Epic Games, Inc. All Rights Reserved.
+
+#include "Graph/MovieGraphQuickRender.h"
+
+#include "Editor.h"
+#include "EngineUtils.h"
+#include "LevelSequence.h"
+#include "LevelSequenceEditorBlueprintLibrary.h"
+#include "Misc/MessageDialog.h"
+#include "MoviePipelinePIEExecutor.h"
+#include "MoviePipelinePostRenderFileDisplayProcessor.h"
+#include "MoviePipelineQueueSubsystem.h"
+#include "MoviePipelineUtils.h"
+#include "MovieRenderPipelineCoreModule.h"
+#include "MovieRenderPipelineSettings.h"
+#include "ObjectEditorUtils.h"
+#include "Graph/MovieGraphRenderLayerSubsystem.h"
+#include "Graph/MovieRenderGraphEditorSettings.h"
+#include "Graph/Nodes/MovieGraphApplyViewportLookNode.h"
+#include "Graph/Nodes/MovieGraphCollectionNode.h"
+#include "Graph/Nodes/MovieGraphDeferredPassNode.h"
+#include "Graph/Nodes/MovieGraphGlobalOutputSettingNode.h"
+#include "Graph/Nodes/MovieGraphModifierNode.h"
+#include "Graph/Nodes/MovieGraphVideoOutputNode.h"
+
+#if WITH_EDITOR
+#include "EditorViewportClient.h"
+#endif
+
+#if WITH_OCIO
+#include "IOpenColorIOModule.h"
+#include "OpenColorIODisplayManager.h"
+#endif
+
+#define LOCTEXT_NAMESPACE "FMovieGraphQuickRender"
+
+// Quick Render should be hidden in the UI while under development.
+static TAutoConsoleVariable<bool> CVarMoviePipelineEnableQuickRender(
+	TEXT("MoviePipeline.EnableQuickRender"),
+	false,
+	TEXT("Enables or disables whether the Quick Render feature is exposed in the UI (defaults to off)."));
+
+void FMovieGraphQuickRender::BeginQuickRender(const FMovieGraphQuickRenderSettings& InQuickRenderSettings)
+{
+	UMoviePipelineQueueSubsystem* Subsystem = GEditor->GetEditorSubsystem<UMoviePipelineQueueSubsystem>();
+	check(Subsystem);
+
+	if (Subsystem->IsRendering())
+	{
+		UE_LOG(LogMovieRenderPipeline, Error, TEXT("Unable to Quick Render: Rendering already in progress."));
+		return;
+	}
+
+	// Always use the PIE executor for quick render, rather than the executor specified in the project settings for normal MRG.
+	const TSubclassOf<UMoviePipelineExecutorBase> ExecutorClass = UMoviePipelinePIEExecutor::StaticClass();
+
+	// Get the level sequence that Sequencer is currently displaying
+	const ULevelSequence* CurrentLevelSequence = ULevelSequenceEditorBlueprintLibrary::GetCurrentLevelSequence();
+	if (!CurrentLevelSequence)
+	{
+		UE_LOG(LogMovieRenderPipeline, Error, TEXT("Unable to Quick Render: No level sequence currently active to render from."));
+		return;
+	}
+
+	// Find the editor world
+	UWorld* EditorWorld = nullptr;
+	for (const FWorldContext& Context : GEngine->GetWorldContexts())
+	{
+		if (Context.WorldType == EWorldType::Editor)
+		{
+			EditorWorld = Context.World();
+			break;
+		}
+	}
+	
+	if (!EditorWorld)
+	{
+		UE_LOG(LogMovieRenderPipeline, Error, TEXT("Unable to Quick Render: Could not find a currently active editor world."));
+		return;
+	}
+
+	// Allocate a new queue that will be used temporarily
+	TemporaryQueue = TStrongObjectPtr(NewObject<UMoviePipelineQueue>(GetTransientPackage()));
+
+	// Add a temp job to the queue, and give it the editor's level and the current level sequence being edited in Sequencer
+	UMoviePipelineExecutorJob* NewJob = TemporaryQueue->AllocateNewJob(UMoviePipelineExecutorJob::StaticClass());
+	NewJob->SetSequence(FSoftObjectPath(CurrentLevelSequence));
+	NewJob->Map = FSoftObjectPath(EditorWorld);
+	NewJob->JobName = FString(TEXT("QuickRender"));
+
+	// Make sure the level has been saved (unsaved levels won't be usable in PIE)
+	if (!FPackageName::IsValidLongPackageName(NewJob->Map.GetLongPackageName()))
+	{
+		FMessageDialog::Open(
+			EAppMsgType::Ok,
+			LOCTEXT("UnsavedMapWarning", "Quick Render needs a saved level to function correctly. Please save your level, then try again."));
+		
+		return;
+	}
+
+	// Update the job to use the quick render graph (either the default, or the one specified in the settings)
+	TMap<TObjectPtr<UMovieGraphConfig>, TObjectPtr<UMovieGraphConfig>> OriginalGraphToDupeMap;
+	const UMovieGraphConfig* QuickRenderGraph = GetQuickRenderGraph(InQuickRenderSettings.GraphPreset.LoadSynchronous(), OriginalGraphToDupeMap);
+	if (!QuickRenderGraph)
+	{
+		// GetQuickRenderGraph() will output an error to the log if a graph was not returned
+		return;
+	}
+	NewJob->SetGraphPreset(QuickRenderGraph);
+
+	ApplyQuickRenderUpdatesToDuplicateGraph(InQuickRenderSettings, EditorWorld);
+
+	// Apply any job-level variable overrides, if they were specified
+	for (const TObjectPtr<UMovieJobVariableAssignmentContainer>& IncomingVariableAssignment : InQuickRenderSettings.GraphVariableAssignments)
+	{
+		// Map the setting's graph to the duplicate graph
+		const UMovieGraphConfig* AssignmentGraph = IncomingVariableAssignment->GetGraphConfig().LoadSynchronous();
+		const TObjectPtr<UMovieGraphConfig>* DuplicateGraph = OriginalGraphToDupeMap.Find(AssignmentGraph);
+		if (!DuplicateGraph)
+		{
+			UE_LOG(LogMovieRenderPipeline, Warning, TEXT("Quick Render: Unable to properly set variable overrides for graph [%s]."), *AssignmentGraph->GetName());
+			continue;
+		}
+		
+		UMovieJobVariableAssignmentContainer* NewVariableAssignments = NewJob->GetOrCreateJobVariableAssignmentsForGraph(*DuplicateGraph);
+
+		for (const UMovieGraphVariable* GraphVariable : (*DuplicateGraph)->GetVariables())
+		{
+			// Copying values by serialized string isn't ideal, but it's the easiest way to transfer values between the variable assignments
+			NewVariableAssignments->SetValueSerializedString(GraphVariable, IncomingVariableAssignment->GetValueSerializedString(GraphVariable));
+			NewVariableAssignments->SetVariableAssignmentEnableState(GraphVariable, true);
+		}
+	}
+
+	// Do the render
+	UMoviePipelinePIEExecutor* Executor = Cast<UMoviePipelinePIEExecutor>(Subsystem->RenderQueueInstanceWithExecutor(TemporaryQueue.Get(), ExecutorClass));
+
+	// View the job's rendered frame(s)/video(s) if requested
+	Executor->OnIndividualJobWorkFinished().AddLambda([InQuickRenderSettings](FMoviePipelineOutputData OutputData)
+	{
+		if (InQuickRenderSettings.PostRenderBehavior == EMovieGraphQuickRenderPostRenderActionType::PlayRenderOutput)
+		{
+			FMoviePipelinePostRenderFileDisplayProcessor PostRenderFileDisplayProcessor(GetDefault<UMovieRenderGraphEditorSettings>()->PostRenderSettings);
+			PostRenderFileDisplayProcessor.AddFiles(OutputData);
+			PostRenderFileDisplayProcessor.OpenFiles();
+		}
+	});
+
+	// Clean up
+	Executor->OnExecutorFinished().AddLambda([](UMoviePipelineExecutorBase* InExecutor, bool bSuccess)
+	{
+		// The temp queue and graph are part of the transient package, so manually release them when the render is finished
+		TemporaryQueue.Reset();
+		TemporaryGraph.Reset();
+	});
+}
+
+UMovieGraphConfig* FMovieGraphQuickRender::GetQuickRenderGraph(UMovieGraphConfig* InUserSpecifiedGraph, TMap<TObjectPtr<UMovieGraphConfig>, TObjectPtr<UMovieGraphConfig>>& OutOriginalToDupeMap)
+{
+	UMovieGraphConfig* QuickRenderGraph = InUserSpecifiedGraph;
+	if (!QuickRenderGraph)
+	{
+		if (const UMovieRenderPipelineProjectSettings* MoviePipelineProjectSettings = GetDefault<UMovieRenderPipelineProjectSettings>())
+		{
+			QuickRenderGraph = MoviePipelineProjectSettings->DefaultQuickRenderGraph.LoadSynchronous();
+			if (!QuickRenderGraph)
+			{
+				UE_LOG(LogMovieRenderPipeline, Error, TEXT("Unable to Quick Render: The default graph specified in project settings could not be loaded."));
+				return nullptr;
+			}
+		}
+		else
+		{
+			UE_LOG(LogMovieRenderPipeline, Error, TEXT("Unable to Quick Render: Could not get the movie pipeline project settings."));
+			return nullptr;
+		}
+	}
+
+	if (!QuickRenderGraph)
+	{
+		UE_LOG(LogMovieRenderPipeline, Error, TEXT("Unable to Quick Render: No valid graph could be loaded."));
+		return nullptr;
+	}
+
+	// Duplicate the graph so changes can be made to it without affecting the graph asset being used
+	UMovieGraphConfig* DuplicateGraph = MoviePipeline::DuplicateConfigRecursive(QuickRenderGraph, OutOriginalToDupeMap);
+	TemporaryGraph = TStrongObjectPtr(DuplicateGraph);
+
+	return TemporaryGraph.Get();
+}
+
+void FMovieGraphQuickRender::ApplyQuickRenderUpdatesToDuplicateGraph(const FMovieGraphQuickRenderSettings& InQuickRenderSettings, UWorld* InEditorWorld)
+{
+	check(TemporaryGraph.IsValid());
+
+	// Global Output Settings
+	{
+		UMovieGraphGlobalOutputSettingNode* GlobalOutputSettings = Cast<UMovieGraphGlobalOutputSettingNode>(InjectNodeIntoGlobals(UMovieGraphGlobalOutputSettingNode::StaticClass()));
+		
+		if (!InQuickRenderSettings.OutputDirectory.Path.IsEmpty())
+		{
+			GlobalOutputSettings->bOverride_OutputDirectory = true;
+			GlobalOutputSettings->OutputDirectory = InQuickRenderSettings.OutputDirectory;
+		}
+
+		// If the frame range type is StartEnd, the frame range doesn't need to be modified
+		if ((InQuickRenderSettings.FrameRangeType == EMovieGraphQuickRenderFrameRangeType::SelectionRange) || (InQuickRenderSettings.FrameRangeType == EMovieGraphQuickRenderFrameRangeType::Custom))
+		{
+			const int32 StartFrame = (InQuickRenderSettings.FrameRangeType == EMovieGraphQuickRenderFrameRangeType::SelectionRange)
+				? ULevelSequenceEditorBlueprintLibrary::GetSelectionRangeStart()
+				: InQuickRenderSettings.CustomStartFrame;
+			const int32 EndFrame = (InQuickRenderSettings.FrameRangeType == EMovieGraphQuickRenderFrameRangeType::SelectionRange)
+				? ULevelSequenceEditorBlueprintLibrary::GetSelectionRangeEnd()
+				: InQuickRenderSettings.CustomEndFrame;
+		
+			GlobalOutputSettings->bOverride_CustomPlaybackRangeStartFrame = true;
+			GlobalOutputSettings->CustomPlaybackRangeStartFrame = StartFrame;
+
+			GlobalOutputSettings->bOverride_CustomPlaybackRangeEndFrame = true;
+			GlobalOutputSettings->CustomPlaybackRangeEndFrame = EndFrame;
+		}
+
+		if (InQuickRenderSettings.Resolution.IsValid())
+		{
+			GlobalOutputSettings->bOverride_OutputResolution = true;
+			GlobalOutputSettings->OutputResolution = InQuickRenderSettings.Resolution;
+		}
+	}
+
+	// Some updates are only done if there's an Apply Viewport Look node present
+	const TObjectPtr<UMovieGraphNode>* ViewportNode = TemporaryGraph->GetNodes().FindByPredicate([](const TObjectPtr<UMovieGraphNode>& InNode)
+	{
+		return InNode && InNode->IsA<UMovieGraphApplyViewportLookNode>();
+	});
+
+	if (ViewportNode)
+	{
+		// Apply viewport actor visibility
+		{
+			TArray<FMovieGraphActorQueryEntry> VisibleActors;
+			TArray<FMovieGraphActorQueryEntry> HiddenActors;
+			
+			for (FActorIterator It(InEditorWorld); It; ++It)
+			{
+				AActor* Actor = *It;
+
+				if (!Actor)
+				{
+					continue;
+				}
+
+				if (Actor->IsHiddenEd())
+				{
+					FMovieGraphActorQueryEntry& HiddenActorEntry = HiddenActors.AddDefaulted_GetRef();
+					HiddenActorEntry.ActorToMatch = Actor;
+				}
+				else
+				{
+					FMovieGraphActorQueryEntry& VisibleActorEntry = VisibleActors.AddDefaulted_GetRef();
+					VisibleActorEntry.ActorToMatch = Actor;
+				}
+			}
+
+			const UMovieGraphCollectionNode* VisibleActorsCollectionNode = Cast<UMovieGraphCollectionNode>(InjectNodeIntoGlobals(UMovieGraphCollectionNode::StaticClass()));
+			const UMovieGraphCollectionNode* HiddenActorsCollectionNode = Cast<UMovieGraphCollectionNode>(InjectNodeIntoGlobals(UMovieGraphCollectionNode::StaticClass()));
+
+			const FName VisibleActorsCollectionName(TEXT("__AUTOGEN_COLLECTION_VISIBLE_ACTORS"));
+			const FName HiddenActorsCollectionName(TEXT("__AUTOGEN_COLLECTION_HIDDEN_ACTORS"));
+			VisibleActorsCollectionNode->Collection->SetCollectionName(VisibleActorsCollectionName.ToString());
+			HiddenActorsCollectionNode->Collection->SetCollectionName(HiddenActorsCollectionName.ToString());
+
+			UMovieGraphConditionGroupQuery_Actor* VisibleActorsQuery =
+				Cast<UMovieGraphConditionGroupQuery_Actor>(VisibleActorsCollectionNode->Collection->AddConditionGroup()->AddQuery(UMovieGraphConditionGroupQuery_Actor::StaticClass()));
+			VisibleActorsQuery->ActorsAndComponentsToMatch = VisibleActors;
+
+			UMovieGraphConditionGroupQuery_Actor* HiddenActorsQuery =
+				Cast<UMovieGraphConditionGroupQuery_Actor>(HiddenActorsCollectionNode->Collection->AddConditionGroup()->AddQuery(UMovieGraphConditionGroupQuery_Actor::StaticClass()));
+			HiddenActorsQuery->ActorsAndComponentsToMatch = HiddenActors;
+
+			UMovieGraphModifierNode* VisibleActorsModifierNode = Cast<UMovieGraphModifierNode>(InjectNodeIntoGlobals(UMovieGraphModifierNode::StaticClass()));
+			UMovieGraphModifierNode* HiddenActorsModifierNode = Cast<UMovieGraphModifierNode>(InjectNodeIntoGlobals(UMovieGraphModifierNode::StaticClass()));
+			VisibleActorsModifierNode->ModifierName = TEXT("__AUTOGEN_MODIFIER_VISIBLE_ACTORS");
+			HiddenActorsModifierNode->ModifierName = TEXT("__AUTOGEN_MODIFIER_HIDDEN_ACTORS");
+			VisibleActorsModifierNode->AddCollection(VisibleActorsCollectionName);
+			HiddenActorsModifierNode->AddCollection(HiddenActorsCollectionName);
+
+			UMovieGraphRenderPropertyModifier* VisibleActorsModifier = Cast<UMovieGraphRenderPropertyModifier>(VisibleActorsModifierNode->GetModifier(UMovieGraphRenderPropertyModifier::StaticClass()));
+			UMovieGraphRenderPropertyModifier* HiddenActorsModifier = Cast<UMovieGraphRenderPropertyModifier>(HiddenActorsModifierNode->GetModifier(UMovieGraphRenderPropertyModifier::StaticClass()));
+
+			VisibleActorsModifier->bOverride_bIsHidden = true;
+			VisibleActorsModifier->bIsHidden = false;
+			HiddenActorsModifier->bOverride_bIsHidden = true;
+			HiddenActorsModifier->bIsHidden = true;
+		}
+		
+		// Apply OCIO if it is activated on the viewport
+#if WITH_OCIO
+		if (const FEditorViewportClient* ViewportClient = UMovieGraphApplyViewportLookNode::GetViewportClient())
+		{
+			if (const FOpenColorIODisplayConfiguration* OCIOConfiguration = IOpenColorIOModule::Get().GetDisplayManager().GetDisplayConfiguration(ViewportClient))
+			{
+				if (OCIOConfiguration->bIsEnabled)
+				{
+					TArray<TObjectPtr<UMovieGraphNode>> FileOutputNodes = TemporaryGraph->GetNodes().FilterByPredicate([](const TObjectPtr<UMovieGraphNode>& InNode)
+					{
+						return InNode && InNode->IsA<UMovieGraphFileOutputNode>();
+					});
+
+					TArray<TObjectPtr<UMovieGraphNode>> RendererNodes = TemporaryGraph->GetNodes().FilterByPredicate([](const TObjectPtr<UMovieGraphNode>& InNode)
+					{
+						return InNode && InNode->IsA<UMovieGraphDeferredRenderPassNode>();
+					});
+
+					// Specify the OCIO configuration on all nodes that support it
+					for (const TObjectPtr<UMovieGraphNode>& OutputNode : FileOutputNodes)
+					{
+						static const FName OCIOOverrideName(TEXT("bOverride_OCIOConfiguration"));
+						static const FName OCIOConfigName(TEXT("OCIOConfiguration"));
+						
+						// This is hacky, but OCIO properties are not uniformly inherited across image/video nodes. The properties are, however,
+						// named the same.
+						const FBoolProperty* OverrideProp = FindFProperty<FBoolProperty>(OutputNode->GetClass(), FName(OCIOOverrideName));
+						const FStructProperty* OCIOProp = FindFProperty<FStructProperty>(OutputNode->GetClass(), FName(OCIOConfigName));
+						
+						if (OverrideProp && OCIOProp)
+						{
+							OverrideProp->SetPropertyValue_InContainer(OutputNode.Get(), true);
+							OCIOProp->SetValue_InContainer(OutputNode.Get(), OCIOConfiguration);
+						}
+					}
+					
+					// Also enable OCIO on the renderer nodes
+					for (const TObjectPtr<UMovieGraphNode>& RendererNode : RendererNodes)
+					{
+						UMovieGraphDeferredRenderPassNode* DeferredRenderNode = Cast<UMovieGraphDeferredRenderPassNode>(RendererNode.Get());
+						if (!DeferredRenderNode)
+						{
+							continue;
+						}
+
+						if (!DeferredRenderNode->bAllowOCIO)
+						{
+							UE_LOG(LogMovieRenderPipeline, Warning, TEXT("Quick Render was set to apply Viewport OCIO, but the renderer node [%s] has Allow OCIO turned off. OCIO will not be applied in this node's render."), *DeferredRenderNode->GetName());
+							continue;
+						}
+						
+						DeferredRenderNode->bOverride_bAllowOCIO = true;
+						DeferredRenderNode->bOverride_bDisableToneCurve = true;
+						DeferredRenderNode->bAllowOCIO = true;
+						DeferredRenderNode->bDisableToneCurve = true;
+					}
+				}
+			}
+		}
+#endif	// WITH_OCIO
+	}
+}
+
+UMovieGraphNode* FMovieGraphQuickRender::InjectNodeIntoGlobals(const TSubclassOf<UMovieGraphNode> NodeType)
+{
+	check(TemporaryGraph.IsValid());
+
+	UMovieGraphNode* NewNode = TemporaryGraph->CreateNodeByClass(NodeType);
+
+	UMovieGraphNode* OutputNode = TemporaryGraph->GetOutputNode();
+	const UMovieGraphPin* GlobalsOutputPin = OutputNode->GetInputPin(UMovieGraphNode::GlobalsPinName);
+
+	// Get the node on the Globals branch that is most downstream (if any)
+	UMovieGraphNode* MostDownstreamGlobalsNode = nullptr;
+	if (const UMovieGraphPin* ConnectedPin = GlobalsOutputPin->GetFirstConnectedPin())
+	{
+		MostDownstreamGlobalsNode = ConnectedPin->Node;
+	}
+
+	// Add in the new node to Globals, downstream from the existing node that was previously furthest downstream
+	TemporaryGraph->AddLabeledEdge(NewNode, FName(), OutputNode, UMovieGraphNode::GlobalsPinName);
+
+	// Re-connect the previously connected node upstream of the new node, if there was one
+	if (MostDownstreamGlobalsNode)
+	{
+		TemporaryGraph->AddLabeledEdge(MostDownstreamGlobalsNode, FName(), NewNode, FName());
+	}
+
+	return NewNode;
+}
+
+#undef LOCTEXT_NAMESPACE
diff --git a/Engine/Plugins/MovieScene/MovieRenderPipeline/Source/MovieRenderPipelineEditor/Private/MoviePipelinePIEExecutor.cpp b/Engine/Plugins/MovieScene/MovieRenderPipeline/Source/MovieRenderPipelineEditor/Private/MoviePipelinePIEExecutor.cpp
--- a/Engine/Plugins/MovieScene/MovieRenderPipeline/Source/MovieRenderPipelineEditor/Private/MoviePipelinePIEExecutor.cpp
+++ b/Engine/Plugins/MovieScene/MovieRenderPipeline/Source/MovieRenderPipelineEditor/Private/MoviePipelinePIEExecutor.cpp
@@ -371,7 +371,6 @@ void UMoviePipelinePIEExecutor::OnPIEEnded(bool)
 	FEditorDelegates::EndPIE.RemoveAll(this);
 
 	CachedOutputDataParams = FMoviePipelineOutputData();
-	UMoviePipeline* PipelineAsLegacy = Cast<UMoviePipeline>(ActiveMoviePipeline);
 
 	// Only call Shutdown if the pipeline hasn't been finished.
 	if (ActiveMoviePipeline && ActiveMoviePipeline->GetPipelineState() != EMovieRenderPipelineState::Finished)
@@ -382,11 +381,17 @@ void UMoviePipelinePIEExecutor::OnPIEEnded(bool)
 		ActiveMoviePipeline->Shutdown(true);
 		UE_LOG(LogMovieRenderPipeline, Log, TEXT("MoviePipelinePIEExecutor: Stalling finished, pipeline has shut down."));
 	}
-	if (PipelineAsLegacy)
+
+	// Cache off the output data so we can use it in DelayedFinishNotification, at which point ActiveMoviePipeline is null.	
+	if (UMoviePipeline* PipelineAsLegacy = Cast<UMoviePipeline>(ActiveMoviePipeline))
 	{
-		// Cache this off so we can use it in DelayedFinishNotification, at which point ActiveMoviePipeline is null.
 		CachedOutputDataParams = PipelineAsLegacy->GetOutputDataParams();
 	}
+	else if (UMovieGraphPipeline* GraphPipeline = Cast<UMovieGraphPipeline>(ActiveMoviePipeline))
+	{
+		// We only set GraphData for the graph pipeline
+		CachedOutputDataParams.GraphData = GraphPipeline->GetGeneratedOutputData();
+	}
 
 	// We need to null out this reference this frame, otherwise we try to hold onto a PIE
 	// object after PIE finishes which causes a GC leak.
diff --git a/Engine/Plugins/MovieScene/MovieRenderPipeline/Source/MovieRenderPipelineEditor/Private/MoviePipelinePostRenderFileDisplayProcessor.cpp b/Engine/Plugins/MovieScene/MovieRenderPipeline/Source/MovieRenderPipelineEditor/Private/MoviePipelinePostRenderFileDisplayProcessor.cpp
--- a/Engine/Plugins/MovieScene/MovieRenderPipeline/Source/MovieRenderPipelineEditor/Private/MoviePipelinePostRenderFileDisplayProcessor.cpp
+++ b/Engine/Plugins/MovieScene/MovieRenderPipeline/Source/MovieRenderPipelineEditor/Private/MoviePipelinePostRenderFileDisplayProcessor.cpp
@@ -264,6 +264,11 @@ void FMoviePipelinePostRenderFileDisplayProcessor::GroupFilesByShot(
 			return A.RenderLayerIndex < B.RenderLayerIndex;
 		});
 
+		// Keep track of the render layer that media types were first found in
+		FString ImagesFirstFoundOnLayer;
+		FString ProResFirstFoundOnLayer;
+		FString AvidFirstFoundOnLayer;
+
 		// Add the images/videos for each render layer within the shot (only the first render layer may be considered depending on settings)
 		for (const TPair<FMovieGraphRenderDataIdentifier, FMovieGraphRenderLayerOutputData>& PassData : OutputData.RenderLayerData)
 		{
@@ -283,16 +288,18 @@ void FMoviePipelinePostRenderFileDisplayProcessor::GroupFilesByShot(
 				const bool bIsProResNode = NodeClass->GetName().Contains(TEXT("ProRes"), ESearchCase::CaseSensitive);
 				const bool bIsAvidNode = NodeClass->GetName().Contains(TEXT("Avid"), ESearchCase::CaseSensitive);
 				
-				// Don't continue to the next render layer if we're only getting files for one render layer
-				if (bIsImageNode && bFirstRenderLayerOnly_Images)
+				// Don't continue processing this node if the node's media type has already been found for this render layer. Note that the render
+				// layer name check is to allow media of the same category (eg, several image types like PNG and EXR) to be processed for a single
+				// render layer. This loop should process all nodes within the render layer.
+				if (bIsImageNode && bFirstRenderLayerOnly_Images && (!ImagesFirstFoundOnLayer.IsEmpty() && ImagesFirstFoundOnLayer != PassData.Key.LayerName))
 				{
 					continue;
 				}
-				if (bIsProResNode && bFirstRenderLayerOnly_ProRes)
+				if (bIsProResNode && bFirstRenderLayerOnly_ProRes && (!ProResFirstFoundOnLayer.IsEmpty() && ProResFirstFoundOnLayer != PassData.Key.LayerName))
 				{
 					continue;
 				}
-				if (bIsAvidNode && bFirstRenderLayerOnly_Avid)
+				if (bIsAvidNode && bFirstRenderLayerOnly_Avid && (!AvidFirstFoundOnLayer.IsEmpty() && AvidFirstFoundOnLayer != PassData.Key.LayerName))
 				{
 					continue;
 				}
@@ -308,6 +315,7 @@ void FMoviePipelinePostRenderFileDisplayProcessor::GroupFilesByShot(
 					
 					if (bIsImageNode)
 					{
+						ImagesFirstFoundOnLayer = PassData.Key.LayerName;
 						ImagesByExtension.FindOrAdd(Extension).Add(NodeFile);
 
 						// The render layer will contain a set of frame templated file paths, in addition to normal file paths. Determine if this
@@ -323,6 +331,15 @@ void FMoviePipelinePostRenderFileDisplayProcessor::GroupFilesByShot(
 					}
 					else
 					{
+						if (bIsProResNode)
+						{
+							ProResFirstFoundOnLayer = PassData.Key.LayerName;
+						}
+						else if (bIsAvidNode)
+						{
+							AvidFirstFoundOnLayer = PassData.Key.LayerName;
+						}
+						
 						VideosByExtension.FindOrAdd(Extension).Add(NodeFile);
 					}
 				}
diff --git a/Engine/Plugins/MovieScene/MovieRenderPipeline/Source/MovieRenderPipelineEditor/Private/MoviePipelineQueueSubsystem.cpp b/Engine/Plugins/MovieScene/MovieRenderPipeline/Source/MovieRenderPipelineEditor/Private/MoviePipelineQueueSubsystem.cpp
--- a/Engine/Plugins/MovieScene/MovieRenderPipeline/Source/MovieRenderPipelineEditor/Private/MoviePipelineQueueSubsystem.cpp
+++ b/Engine/Plugins/MovieScene/MovieRenderPipeline/Source/MovieRenderPipelineEditor/Private/MoviePipelineQueueSubsystem.cpp
@@ -89,7 +89,7 @@ bool UMoviePipelineQueueSubsystem::IsQueueDirty() const
 	return bHasNoQueueOrigin || Queue->IsDirty();
 }
 
-UMoviePipelineExecutorBase* UMoviePipelineQueueSubsystem::RenderQueueWithExecutor(TSubclassOf<UMoviePipelineExecutorBase> InExecutorType)
+UMoviePipelineExecutorBase* UMoviePipelineQueueSubsystem::RenderQueueInstanceWithExecutor(UMoviePipelineQueue* InQueue, TSubclassOf<UMoviePipelineExecutorBase> InExecutorType)
 {
 	if(!ensureMsgf(!IsRendering(), TEXT("RenderQueueWithExecutor cannot be called while already rendering!")))
 	{
@@ -97,11 +97,25 @@ UMoviePipelineExecutorBase* UMoviePipelineQueueSubsystem::RenderQueueWithExecuto
 	}
 
 	ActiveExecutor = NewObject<UMoviePipelineExecutorBase>(this, InExecutorType);
-	RenderQueueWithExecutorInstance(ActiveExecutor);
+
+	if (InQueue)
+	{
+		RenderQueueInstanceWithExecutorInstance(InQueue, ActiveExecutor);
+	}
+	else
+	{
+		RenderQueueWithExecutorInstance(ActiveExecutor);
+	}
+	
 	return ActiveExecutor;
 }
 
-void UMoviePipelineQueueSubsystem::RenderQueueWithExecutorInstance(UMoviePipelineExecutorBase* InExecutor)
+UMoviePipelineExecutorBase* UMoviePipelineQueueSubsystem::RenderQueueWithExecutor(TSubclassOf<UMoviePipelineExecutorBase> InExecutorType)
+{
+	return RenderQueueInstanceWithExecutor(nullptr, InExecutorType);
+}
+
+void UMoviePipelineQueueSubsystem::RenderQueueInstanceWithExecutorInstance(UMoviePipelineQueue* InQueue, UMoviePipelineExecutorBase* InExecutor)
 {
 	if (!ensureMsgf(!IsRendering(), TEXT("RenderQueueWithExecutor cannot be called while already rendering!")))
 	{
@@ -123,7 +137,12 @@ void UMoviePipelineQueueSubsystem::RenderQueueWithExecutorInstance(UMoviePipelin
 
 	ActiveExecutor = InExecutor;
 	ActiveExecutor->OnExecutorFinished().AddUObject(this, &UMoviePipelineQueueSubsystem::OnExecutorFinished);
-	ActiveExecutor->Execute(GetQueue());
+	ActiveExecutor->Execute(InQueue);
+}
+
+void UMoviePipelineQueueSubsystem::RenderQueueWithExecutorInstance(UMoviePipelineExecutorBase* InExecutor)
+{
+	RenderQueueInstanceWithExecutorInstance(GetQueue(), InExecutor);
 }
 
 void UMoviePipelineQueueSubsystem::OnSequencerContextBinding(bool& bAllowBinding)
diff --git a/Engine/Plugins/MovieScene/MovieRenderPipeline/Source/MovieRenderPipelineEditor/Private/MoviePipelineQuickRenderMenu.cpp b/Engine/Plugins/MovieScene/MovieRenderPipeline/Source/MovieRenderPipelineEditor/Private/MoviePipelineQuickRenderMenu.cpp
new file mode 100644
--- /dev/null
+++ b/Engine/Plugins/MovieScene/MovieRenderPipeline/Source/MovieRenderPipelineEditor/Private/MoviePipelineQuickRenderMenu.cpp
@@ -0,0 +1,377 @@
+﻿// Copyright Epic Games, Inc. All Rights Reserved.
+
+#include "MoviePipelineQuickRenderMenu.h"
+
+#include "Editor.h"
+#include "Framework/Application/SlateApplication.h"
+#include "Framework/Docking/TabManager.h"
+#include "Graph/MovieGraphQuickRender.h"
+#include "Graph/MovieRenderGraphEditorSettings.h"
+#include "HAL/IConsoleManager.h"
+#include "MoviePipelineQueueSubsystem.h"
+#include "MoviePipelineUtils.h"
+#include "MovieRenderPipelineCoreModule.h"
+#include "MovieRenderPipelineEditorModule.h"
+#include "MovieRenderPipelineEditorUtils.h"
+#include "MovieRenderPipelineSettings.h"
+#include "ToolMenu.h"
+#include "ToolMenus.h"
+#include "Widgets/Text/SRichTextBlock.h"
+
+class UMovieGraphVariable;
+class UMovieGraphConfig;
+
+#define LOCTEXT_NAMESPACE "FMoviePipelineQuickRenderMenu"
+
+bool FMoviePipelineQuickRenderMenu::IsVisible()
+{
+	if (const IConsoleVariable* EnableQuickRenderCVar = IConsoleManager::Get().FindConsoleVariable(TEXT("MoviePipeline.EnableQuickRender")))
+	{
+		return EnableQuickRenderCVar->GetBool();
+	}
+
+	return false;
+}
+
+void FMoviePipelineQuickRenderMenu::AddQuickRenderButtonToToolMenu(UToolMenu* InMenu)
+{
+	if (!InMenu)
+	{
+		return;
+	}
+	
+	FToolMenuSection& Section = InMenu->AddSection(QuickRenderToolbarId);
+	Section.AddSeparator("QuickRenderSeparator");
+	
+	const FToolMenuEntry QuickRenderLaunchButton =
+		FToolMenuEntry::InitToolBarButton(
+			"QuickRenderToolbarButton",
+			FUIAction(
+				FExecuteAction::CreateStatic(&FMoviePipelineQuickRenderMenu::QuickRenderButtonPressed)),
+			TAttribute<FText>(),
+			TAttribute<FText>::Create(
+				TAttribute<FText>::FGetter::CreateLambda([]()
+				{
+					return LOCTEXT("BeginQuickRender", "Begin Quick Render\n\nQuick Render performs a render without having to manually configure a queue and a graph. A typical quick render uses the current map and level sequence, the level sequence's playback range, and the viewport's look, to generate frames. Several different types of quick renders are available -- see the options drop-down.");
+				})
+			),
+			FSlateIcon(FAppStyle::Get().GetStyleSetName(), "LevelEditor.OpenCinematic")
+		);
+	
+	Section.AddEntry(QuickRenderLaunchButton);
+
+	const FToolMenuEntry QuickRenderOptionsButton = FToolMenuEntry::InitComboButton(
+		"QuickRenderOptions",
+		FUIAction(),
+		FOnGetContent::CreateStatic(&FMoviePipelineQuickRenderMenu::GenerateQuickRenderOptionsMenu),
+		LOCTEXT("QuickRenderOptionsLabel", "Quick Render Options"),
+		LOCTEXT("QuickRenderOptionsToolTip", "Quick Render Options"),
+		TAttribute<FSlateIcon>(),
+		true
+	);
+
+	Section.AddEntry(QuickRenderOptionsButton);
+}
+
+void FMoviePipelineQuickRenderMenu::RemoveQuickRenderButtonToolMenu(UToolMenu* InMenu)
+{
+	if (!InMenu)
+	{
+		return;
+	}
+	
+	if (InMenu->FindSection(QuickRenderToolbarId))
+	{
+		InMenu->RemoveSection(QuickRenderToolbarId);
+	}
+}
+
+void FMoviePipelineQuickRenderMenu::QuickRenderButtonPressed()
+{
+	const UMovieRenderGraphEditorSettings* EditorSettings = GetDefault<UMovieRenderGraphEditorSettings>();
+
+	FDirectoryPath OutputDirectory = FDirectoryPath();
+	OutputDirectory.Path = FString(TEXT("{project_dir}/Saved/MovieRenders/QuickRender"));
+	
+	// TODO: These should come from the Quick Render Settings dialog (not finished yet)
+	FMovieGraphQuickRenderSettings QuickRenderSettings;
+	QuickRenderSettings.PostRenderBehavior = EMovieGraphQuickRenderPostRenderActionType::PlayRenderOutput;
+	QuickRenderSettings.OutputDirectory = OutputDirectory;
+
+	if (EditorSettings->QuickRenderSettings.QuickAccessButtonMode == EMovieGraphQuickRenderButtonMode::QuickRender)
+	{
+		FMovieGraphQuickRender::BeginQuickRender(QuickRenderSettings);
+	}
+	else
+	{
+		if (!UE::MovieRenderPipelineEditor::Private::PerformLocalRender())
+		{
+			UE_LOG(LogMovieRenderPipeline, Warning, TEXT("Unable to start local render. Make sure an executor was specified in Project Settings, a job is present in the queue, and a render is not currently running."));
+		}
+	}
+}
+
+TSharedRef<SWidget> FMoviePipelineQuickRenderMenu::GenerateQuickRenderOptionsMenu()
+{
+	static const FName MenuName("MoviePipeline.QuickRenderOptionsMenu");
+
+	if (!UToolMenus::Get()->IsMenuRegistered(MenuName))
+	{
+		UToolMenu* Menu = UToolMenus::Get()->RegisterMenu(MenuName);
+
+		GenerateModesMenuSection(Menu);
+		GenerateQuickRenderMenuSection(Menu);
+		GenerateQuickRenderConfigurationMenuSection(Menu);
+		GenerateCreateMenuSection(Menu);
+		GenerateSettingsMenuSection(Menu);
+	}
+	
+	return UToolMenus::Get()->GenerateWidget(MenuName, FToolMenuContext());
+}
+
+void FMoviePipelineQuickRenderMenu::OpenQuickRenderSettingsWindow(const FToolMenuContext& InToolMenuContext)
+{
+	if (const TSharedPtr<SWindow> SettingsWindow = WeakQuickRenderSettingsWindow.Pin())
+	{
+		SettingsWindow->DrawAttention(FWindowDrawAttentionParameters());
+		return;
+	}
+
+	const TSharedRef<SWindow> QuickRenderSettingsWindow = SNew(SWindow)
+		.Title(LOCTEXT("QuickRenderSettingsWindow_Title", "Quick Render"))
+		.SupportsMaximize(false)
+		.ClientSize(FVector2D(600.f, 510.f))
+		.Content()
+		[
+			SNew(STextBlock)
+			.Text(FText::FromString(TEXT("Settings Window")))
+		];
+
+	if (const TSharedPtr<SWindow> RootWindow = FGlobalTabmanager::Get()->GetRootWindow())
+	{
+		FSlateApplication::Get().AddWindowAsNativeChild(QuickRenderSettingsWindow, RootWindow.ToSharedRef());
+	}
+	else
+	{
+		FSlateApplication::Get().AddWindow(QuickRenderSettingsWindow);
+	}
+
+	WeakQuickRenderSettingsWindow = QuickRenderSettingsWindow;
+}
+
+void FMoviePipelineQuickRenderMenu::GenerateModesMenuSection(UToolMenu* InMenu)
+{
+	FToolMenuSection& Section = InMenu->AddSection("MoviePipelineQuickRenderModes", LOCTEXT("QuickRenderModesSection", "Modes"));
+
+	auto AddModeAction = [&Section](const FName& InName, const FText& InLabel, const FText& InToolTip, const FSlateIcon& InIcon, EMovieGraphQuickRenderButtonMode CheckedValue)
+	{
+		FToolUIAction NewAction;
+		NewAction.ExecuteAction = FToolMenuExecuteAction::CreateLambda([CheckedValue](const FToolMenuContext&)
+		{
+			UMovieRenderGraphEditorSettings* EditorSettings = GetMutableDefault<UMovieRenderGraphEditorSettings>();
+			EditorSettings->QuickRenderSettings.QuickAccessButtonMode = CheckedValue;
+			EditorSettings->PostEditChange();
+			EditorSettings->SaveConfig();
+		});
+		NewAction.CanExecuteAction = FToolMenuCanExecuteAction::CreateLambda([](const FToolMenuContext&) { return true; });
+		NewAction.GetActionCheckState = FToolMenuGetActionCheckState::CreateLambda([CheckedValue](const FToolMenuContext&)
+		{
+			return (GetDefault<UMovieRenderGraphEditorSettings>()->QuickRenderSettings.QuickAccessButtonMode == CheckedValue)
+				? ECheckBoxState::Checked
+				: ECheckBoxState::Unchecked;
+		});
+		Section.AddMenuEntry(InName, InLabel, InToolTip, InIcon, NewAction, EUserInterfaceActionType::RadioButton);
+	};
+
+	AddModeAction(
+		FName("QuickRender_MovieRenderQueueMode"),
+		LOCTEXT("QuickRenderMode_MovieRenderQueueLabel", "Movie Render Queue"),
+		LOCTEXT("QuickRenderMode_MovieRenderQueueToolTip", "Render with Movie Render Queue"),
+		FSlateIcon(FAppStyle::Get().GetStyleSetName(), "LevelEditor.OpenCinematic"),
+		EMovieGraphQuickRenderButtonMode::NormalMovieRenderQueue);
+
+	AddModeAction(
+		FName("QuickRender_QuickRenderMode"),
+		LOCTEXT("QuickRenderMode_QuickRenderLabel", "Quick Render"),
+		LOCTEXT("QuickRenderMode_QuickRenderToolTip", "Render with Quick Render"),
+		FSlateIcon(FAppStyle::Get().GetStyleSetName(), "LevelEditor.OpenCinematic"),
+		EMovieGraphQuickRenderButtonMode::QuickRender);
+}
+
+void FMoviePipelineQuickRenderMenu::GenerateQuickRenderMenuSection(UToolMenu* InMenu)
+{
+	FToolMenuSection& Section = InMenu->AddSection("MoviePipelineQuickRenderType", LOCTEXT("QuickRenderType", "Quick Render"));
+
+	auto AddQuickRenderModeAction = [&Section](const FName& InName, const FText& InLabel, const FText& InToolTip, const FSlateIcon& InIcon, EMovieGraphQuickRenderMode CheckedValue)
+	{
+		FToolUIAction NewAction;
+		NewAction.ExecuteAction = FToolMenuExecuteAction::CreateLambda([CheckedValue](const FToolMenuContext&)
+		{
+			UMovieRenderGraphEditorSettings* EditorSettings = GetMutableDefault<UMovieRenderGraphEditorSettings>();
+			EditorSettings->QuickRenderSettings.QuickRenderMode = CheckedValue;
+			EditorSettings->PostEditChange();
+			EditorSettings->SaveConfig();
+		});
+		NewAction.CanExecuteAction = FToolMenuCanExecuteAction::CreateLambda([](const FToolMenuContext&) { return true; });
+		NewAction.GetActionCheckState = FToolMenuGetActionCheckState::CreateLambda([CheckedValue](const FToolMenuContext&)
+		{
+			return (GetDefault<UMovieRenderGraphEditorSettings>()->QuickRenderSettings.QuickRenderMode == CheckedValue)
+				? ECheckBoxState::Checked
+				: ECheckBoxState::Unchecked;
+		});
+		Section.AddMenuEntry(InName, InLabel, InToolTip, InIcon, NewAction, EUserInterfaceActionType::RadioButton);
+	};
+
+	AddQuickRenderModeAction(
+		FName("QuickRender_CurrentSequence"),
+		LOCTEXT("QuickRender_CurrentSequenceLabel", "Current Sequence"),
+		LOCTEXT("QuickRender_CurrentSequenceToolTip", "Render the sequence that is currently active in Sequencer."),
+		FSlateIcon("LevelSequenceEditorStyle", "LevelSequenceEditor.CreateNewLevelSequenceInLevel"),
+		EMovieGraphQuickRenderMode::CurrentSequence);
+
+	AddQuickRenderModeAction(
+		FName("QuickRender_CurrentShotAtPlayhead"),
+		LOCTEXT("QuickRender_CurrentShotAtPlayheadLabel", "Current Shot at Playhead"),
+		LOCTEXT("QuickRender_CurrentShotAtPlayheadToolTip", "Render the current shot starting at the playhead."),
+		FSlateIcon("LevelSequenceEditorStyle", "LevelSequenceEditor.CreateNewLevelSequenceWithShotsInLevel"),
+		EMovieGraphQuickRenderMode::CurrentShotAtPlayhead);
+
+	AddQuickRenderModeAction(
+		FName("QuickRender_CurrentViewport"),
+		LOCTEXT("QuickRender_CurrentViewportLabel", "Current Viewport"),
+		LOCTEXT("QuickRender_CurrentViewportToolTip", "Render the current viewport."),
+		FSlateIcon(FAppStyle::Get().GetStyleSetName(), "LevelEditor.Tabs.Viewports"),
+		EMovieGraphQuickRenderMode::CurrentViewport);
+
+	AddQuickRenderModeAction(
+		FName("QuickRender_SelectedCameras"),
+		LOCTEXT("QuickRender_SelectedCamerasLabel", "Selected Camera(s)"),
+		LOCTEXT("QuickRender_SelectedCamerasToolTip", "Render the selected cameras."),
+		FSlateIcon(FAppStyle::Get().GetStyleSetName(), "ClassIcon.CameraActor"),
+		EMovieGraphQuickRenderMode::SelectedCameras);
+
+	AddQuickRenderModeAction(
+		FName("QuickRender_UseViewportCamera"),
+		LOCTEXT("QuickRender_UseViewportCameraLabel", "Use Viewport Camera in Sequence"),
+		LOCTEXT("QuickRender_UseViewportCameraToolTip", "Use the current viewport camera in the active level sequence."),
+		FSlateIcon(FAppStyle::Get().GetStyleSetName(), "ClassIcon.CameraActor"),
+		EMovieGraphQuickRenderMode::UseViewportCameraInSequence);
+}
+
+void FMoviePipelineQuickRenderMenu::GenerateQuickRenderConfigurationMenuSection(UToolMenu* InMenu)
+{
+	FToolMenuSection& Section = InMenu->AddSection("MoviePipelineQuickRenderConfiguration", LOCTEXT("QuickRenderConfiguration", "Quick Render Configuration"));
+	
+	Section.AddEntry(FToolMenuEntry::InitWidget(
+		"QuickRenderConfiguration",
+		SNew(SBox)
+		.Padding(FMargin(16.f, 3.f))
+		[
+			SNew(SRichTextBlock)
+			.TextStyle(FAppStyle::Get(), "NormalText.Subdued")
+			.Text_Lambda([]()
+			{
+				const UMovieRenderGraphEditorSettings* EditorSettings = GetDefault<UMovieRenderGraphEditorSettings>();
+				const FMovieGraphQuickRenderSettings& QuickRenderSettings = EditorSettings->QuickRenderSettings;
+
+				const FText GraphName = QuickRenderSettings.GraphPreset.IsValid()
+					? FText::FromString(QuickRenderSettings.GraphPreset->GetName())
+					: LOCTEXT("QuickRenderConfigMenu_InvalidGraph", "Invalid");
+				
+				const FText AfterRender = (QuickRenderSettings.PostRenderBehavior == EMovieGraphQuickRenderPostRenderActionType::PlayRenderOutput)
+					? LOCTEXT("QuickRenderConfigMenu_PlayRenderOutput", "Play Render Output")
+					: LOCTEXT("QuickRenderConfigMenu_DoNothing", "Do Nothing");
+				
+				TArray<FText> ViewportFlags;
+				if (EnumHasAnyFlags(QuickRenderSettings.ViewportLookFlags, EMovieGraphQuickRenderViewportLookFlags::Ocio))
+				{
+					ViewportFlags.Add(LOCTEXT("QuickRenderConfigMenu_OCIO", "OCIO"));
+				}
+				if (EnumHasAnyFlags(QuickRenderSettings.ViewportLookFlags, EMovieGraphQuickRenderViewportLookFlags::ShowFlags))
+				{
+					ViewportFlags.Add(LOCTEXT("QuickRenderConfigMenu_ShowFlags", "Show Flags"));
+				}
+				if (EnumHasAnyFlags(QuickRenderSettings.ViewportLookFlags, EMovieGraphQuickRenderViewportLookFlags::ViewMode))
+				{
+					ViewportFlags.Add(LOCTEXT("QuickRenderConfigMenu_ViewMode", "View Mode"));
+				}
+				if (QuickRenderSettings.ViewportLookFlags == EMovieGraphQuickRenderViewportLookFlags::None)
+				{
+					ViewportFlags.Add(LOCTEXT("QuickRenderConfigMenu_DontApply", "Don't Apply"));
+				}
+
+				// TODO: "Sequencer Frame Range" needs to be sourced from data
+				// TODO: Add in graph variables summary as well (although this might make a summary that's too long in some cases)
+				FText ConfigurationSummaryText = FText::Format(
+					LOCTEXT("QuickRenderConfigurationSummary", "Configuration: <RichTextBlock.BoldHighlight>{0}</>\nAfter Render: {1}\nApply Viewport Look: {2}\nSequencer Frame Range: {3}"),
+					{
+						GraphName,
+						AfterRender,
+						FText::Join(LOCTEXT("CommaDelim", ", "), ViewportFlags),
+						FText::FromString(FString(TEXT("Section"))),
+					});
+
+				return ConfigurationSummaryText;
+			})
+			.LineHeightPercentage(1.3f)
+		],
+		FText::GetEmpty()
+	));
+}
+
+void FMoviePipelineQuickRenderMenu::GenerateCreateMenuSection(UToolMenu* InMenu)
+{
+	FToolMenuSection& Section = InMenu->AddSection("MoviePipelineQuickRenderCreate", LOCTEXT("QuickRenderCreate", "Create"));
+		
+	FToolUIAction CreateSequenceAction;
+	CreateSequenceAction.ExecuteAction = FToolMenuExecuteAction::CreateLambda([](const FToolMenuContext&) { return; });
+	CreateSequenceAction.CanExecuteAction = FToolMenuCanExecuteAction::CreateLambda([](const FToolMenuContext&) { return true; });
+	Section.AddMenuEntry(
+		"QuickRender_CreateSequenceFromCurrentViewport",
+		LOCTEXT("QuickRender_CreateSequenceFromCurrentViewportLabel", "Create Sequence from Current Viewport"),
+		LOCTEXT("QuickRender_CreateSequenceFromCurrentViewportToolTip", "Create a sequence from the current viewport."),
+		FSlateIcon(),
+		CreateSequenceAction);
+}
+
+void FMoviePipelineQuickRenderMenu::GenerateSettingsMenuSection(UToolMenu* InMenu)
+{
+	FToolMenuSection& Section = InMenu->AddSection("MoviePipelineQuickRenderSettings");
+
+	Section.AddSeparator("SettingsSeparator");
+
+	FToolUIAction ShowSettingsAction;
+        ShowSettingsAction.ExecuteAction = FToolMenuExecuteAction::CreateLambda([](const FToolMenuContext&)
+        {
+            UMovieRenderGraphEditorSettings* EditorSettings = GetMutableDefault<UMovieRenderGraphEditorSettings>();
+			EditorSettings->QuickRenderSettings.bShowSettingsBeforeQuickRender = !EditorSettings->QuickRenderSettings.bShowSettingsBeforeQuickRender;
+			EditorSettings->PostEditChange();
+			EditorSettings->SaveConfig();
+        });
+        ShowSettingsAction.CanExecuteAction = FToolMenuCanExecuteAction::CreateLambda([](const FToolMenuContext&) { return true; });
+        ShowSettingsAction.GetActionCheckState = FToolMenuGetActionCheckState::CreateLambda([](const FToolMenuContext&)
+        {
+            return GetDefault<UMovieRenderGraphEditorSettings>()->QuickRenderSettings.bShowSettingsBeforeQuickRender
+				? ECheckBoxState::Checked
+				: ECheckBoxState::Unchecked;
+        });
+        Section.AddMenuEntry(
+            "QuickRender_ShowSettingsBefore",
+            LOCTEXT("QuickRender_ShowSettingsBeforeLabel", "Show Settings Before Quick Render"),
+            LOCTEXT("QuickRender_ShowSettingsBeforeToolTip", "Show the quick render settings before starting a quick render."),
+            FSlateIcon(FAppStyle::Get().GetStyleSetName(), "Icons.OpenInExternalEditor"),
+            ShowSettingsAction,
+            EUserInterfaceActionType::ToggleButton);
+
+	FToolUIAction OpenSettingsAction;
+	OpenSettingsAction.ExecuteAction = FToolMenuExecuteAction::CreateStatic(&FMoviePipelineQuickRenderMenu::OpenQuickRenderSettingsWindow);
+	OpenSettingsAction.CanExecuteAction = FToolMenuCanExecuteAction::CreateLambda([](const FToolMenuContext&) { return true; });
+	Section.AddMenuEntry(
+		"QuickRender_OpenSettings",
+		LOCTEXT("QuickRender_OpenSettingsLabel", "Quick Render Settings"),
+		LOCTEXT("QuickRender_OpenSettingsToolTip", "Open the quick render settings."),
+		FSlateIcon(FAppStyle::Get().GetStyleSetName(), "Icons.Settings"),
+		OpenSettingsAction);
+}
+
+#undef LOCTEXT_NAMESPACE
\ No newline at end of file
diff --git a/Engine/Plugins/MovieScene/MovieRenderPipeline/Source/MovieRenderPipelineEditor/Private/MoviePipelineQuickRenderMenu.h b/Engine/Plugins/MovieScene/MovieRenderPipeline/Source/MovieRenderPipelineEditor/Private/MoviePipelineQuickRenderMenu.h
new file mode 100644
--- /dev/null
+++ b/Engine/Plugins/MovieScene/MovieRenderPipeline/Source/MovieRenderPipelineEditor/Private/MoviePipelineQuickRenderMenu.h
@@ -0,0 +1,42 @@
+﻿// Copyright Epic Games, Inc. All Rights Reserved.
+
+#pragma once
+
+#include "CoreMinimal.h"
+
+struct FToolMenuContext;
+class SWidget;
+class SWindow;
+class UToolMenu;
+
+/** Generates the Quick Render button, options menu, and settings dialog. */
+class FMoviePipelineQuickRenderMenu
+{
+public:
+	FMoviePipelineQuickRenderMenu() = default;
+
+	/** Determines if the Quick Render button should be visible. */
+	static bool IsVisible();
+
+	/** Adds the Quick Render button to the specified tool menu. */
+	static void AddQuickRenderButtonToToolMenu(UToolMenu* InMenu);
+
+	/** Removes the Quick Render button from the specified tool menu. */
+	static void RemoveQuickRenderButtonToolMenu(UToolMenu* InMenu);
+
+private:
+	static void QuickRenderButtonPressed();
+	
+	static TSharedRef<SWidget> GenerateQuickRenderOptionsMenu();
+	static void OpenQuickRenderSettingsWindow(const FToolMenuContext& InToolMenuContext);
+
+	static void GenerateModesMenuSection(UToolMenu* InMenu);
+	static void GenerateQuickRenderMenuSection(UToolMenu* InMenu);
+	static void GenerateQuickRenderConfigurationMenuSection(UToolMenu* InMenu);
+	static void GenerateCreateMenuSection(UToolMenu* InMenu);
+	static void GenerateSettingsMenuSection(UToolMenu* InMenu);
+
+private:
+	inline static const FName QuickRenderToolbarId = TEXT("MovieRenderGraphQuickRender");
+	inline static TWeakPtr<SWindow> WeakQuickRenderSettingsWindow;
+};
diff --git a/Engine/Plugins/MovieScene/MovieRenderPipeline/Source/MovieRenderPipelineEditor/Private/MovieRenderPipelineEditorModule.cpp b/Engine/Plugins/MovieScene/MovieRenderPipeline/Source/MovieRenderPipelineEditor/Private/MovieRenderPipelineEditorModule.cpp
--- a/Engine/Plugins/MovieScene/MovieRenderPipeline/Source/MovieRenderPipelineEditor/Private/MovieRenderPipelineEditorModule.cpp
+++ b/Engine/Plugins/MovieScene/MovieRenderPipeline/Source/MovieRenderPipelineEditor/Private/MovieRenderPipelineEditorModule.cpp
@@ -26,16 +26,20 @@
 #include "Templates/SubclassOf.h"
 #include "MovieRenderPipelineStyle.h"
 #include "MoviePipelineCommands.h"
+#include "MoviePipelineDeferredPasses.h"
 #include "MoviePipelineQueue.h"
 #include "MoviePipelineQueueSubsystem.h"
+#include "MoviePipelineQuickRenderMenu.h"
 #include "MoviePipelineEditorBlueprintLibrary.h"
 #include "MovieSceneSection.h"
 #include "Sections/MovieSceneCinematicShotSection.h"
 #include "MoviePipelineConsoleVariableSetting.h"
 #include "PropertyEditorModule.h"
 #include "Customizations/ConsoleVariableCustomization.h"
-#include "Graph/MovieGraphPinFactory.h"
 #include "LevelSequence.h"
+#include "ToolMenus.h"
+#include "Graph/MovieGraphDataTypes.h"
+#include "Graph/MovieGraphPinFactory.h"
 
 #define LOCTEXT_NAMESPACE "FMovieRenderPipelineEditorModule"
 
@@ -231,6 +235,27 @@ void FMovieRenderPipelineEditorModule::UnregisterTypeCustomizations()
 	}
 }
 
+void FMovieRenderPipelineEditorModule::RegisterToolbarItem()
+{
+	UToolMenu* Menu = UToolMenus::Get()->ExtendMenu("LevelEditor.LevelEditorToolBar.User");
+	
+	UnregisterToolbarItem();
+	
+	if (FMoviePipelineQuickRenderMenu::IsVisible())
+	{
+		FMoviePipelineQuickRenderMenu::AddQuickRenderButtonToToolMenu(Menu);
+	}
+}
+
+void FMovieRenderPipelineEditorModule::UnregisterToolbarItem()
+{
+	if (UToolMenus* ToolMenus = UToolMenus::Get())
+	{
+		UToolMenu* Menu = ToolMenus->FindMenu("LevelEditor.LevelEditorToolBar.User");
+		FMoviePipelineQuickRenderMenu::RemoveQuickRenderButtonToolMenu(Menu);
+	}
+}
+
 void FMovieRenderPipelineEditorModule::StartupModule()
 {
 	// Initialize our custom style
@@ -241,6 +266,7 @@ void FMovieRenderPipelineEditorModule::StartupModule()
 	RegisterSettings();
 	RegisterTypeCustomizations();
 	RegisterMovieRenderer();
+	RegisterToolbarItem();
 
 	GraphPanelPinFactory = MakeShared<FMovieGraphPanelPinFactory>();
 	FEdGraphUtilities::RegisterVisualPinFactory(GraphPanelPinFactory);
@@ -248,6 +274,7 @@ void FMovieRenderPipelineEditorModule::StartupModule()
 
 void FMovieRenderPipelineEditorModule::ShutdownModule()
 {
+	UnregisterToolbarItem();
 	UnregisterMovieRenderer();
 	UnregisterTypeCustomizations();
 	UnregisterSettings();
diff --git a/Engine/Plugins/MovieScene/MovieRenderPipeline/Source/MovieRenderPipelineEditor/Private/MovieRenderPipelineEditorModule.h b/Engine/Plugins/MovieScene/MovieRenderPipeline/Source/MovieRenderPipelineEditor/Private/MovieRenderPipelineEditorModule.h
--- a/Engine/Plugins/MovieScene/MovieRenderPipeline/Source/MovieRenderPipelineEditor/Private/MovieRenderPipelineEditorModule.h
+++ b/Engine/Plugins/MovieScene/MovieRenderPipeline/Source/MovieRenderPipelineEditor/Private/MovieRenderPipelineEditorModule.h
@@ -15,6 +15,7 @@ public:
 	/** IModuleInterface implementation */
 	virtual void StartupModule() override;
 	virtual void ShutdownModule() override;
+	
 private:
 	void RegisterSettings();
 	void UnregisterSettings();
@@ -25,6 +26,9 @@ private:
 	void RegisterTypeCustomizations();
 	void UnregisterTypeCustomizations();
 
+	void RegisterToolbarItem();
+	void UnregisterToolbarItem();
+
 	FDelegateHandle MovieRendererDelegate;
 
 	/** Pin factory for the render graph. */
diff --git a/Engine/Plugins/MovieScene/MovieRenderPipeline/Source/MovieRenderPipelineEditor/Private/MovieRenderPipelineEditorUtils.cpp b/Engine/Plugins/MovieScene/MovieRenderPipeline/Source/MovieRenderPipelineEditor/Private/MovieRenderPipelineEditorUtils.cpp
new file mode 100644
--- /dev/null
+++ b/Engine/Plugins/MovieScene/MovieRenderPipeline/Source/MovieRenderPipelineEditor/Private/MovieRenderPipelineEditorUtils.cpp
@@ -0,0 +1,58 @@
+﻿// Copyright Epic Games, Inc. All Rights Reserved.
+
+#include "MovieRenderPipelineEditorUtils.h"
+
+#include "Editor.h"
+#include "MoviePipelineQueueSubsystem.h"
+#include "MoviePipelineTelemetry.h"
+#include "MovieRenderPipelineSettings.h"
+
+namespace UE::MovieRenderPipelineEditor::Private
+{
+	bool PerformLocalRender()
+	{
+		if (!CanPerformLocalRender())
+		{
+			return false;
+		}
+		
+		UMoviePipelineQueueSubsystem* Subsystem = GEditor->GetEditorSubsystem<UMoviePipelineQueueSubsystem>();
+		check(Subsystem);
+
+		const UMovieRenderPipelineProjectSettings* ProjectSettings = GetDefault<UMovieRenderPipelineProjectSettings>();
+		
+		const TSubclassOf<UMoviePipelineExecutorBase> ExecutorClass = ProjectSettings->DefaultLocalExecutor.TryLoadClass<UMoviePipelineExecutorBase>();
+		check(ExecutorClass != nullptr);
+		
+		Subsystem->RenderQueueWithExecutor(ExecutorClass);
+
+		constexpr bool bIsLocal = true;
+		FMoviePipelineTelemetry::SendRendersRequestedTelemetry(bIsLocal, Subsystem->GetQueue()->GetJobs());
+
+		return true;
+	}
+
+	bool CanPerformLocalRender()
+	{
+		const UMoviePipelineQueueSubsystem* Subsystem = GEditor->GetEditorSubsystem<UMoviePipelineQueueSubsystem>();
+		check(Subsystem);
+
+		const UMovieRenderPipelineProjectSettings* ProjectSettings = GetDefault<UMovieRenderPipelineProjectSettings>();
+		const bool bHasExecutor = ProjectSettings->DefaultLocalExecutor.TryLoadClass<UMoviePipelineExecutorBase>() != nullptr;
+		const bool bNotRendering = !Subsystem->IsRendering();
+
+		bool bAtLeastOneJobAvailable = false;
+		for (const UMoviePipelineExecutorJob* Job : Subsystem->GetQueue()->GetJobs())
+		{
+			if (!Job->IsConsumed() && Job->IsEnabled())
+			{
+				bAtLeastOneJobAvailable = true;
+				break;
+			}
+		}
+
+		const bool bWorldIsActive = GEditor->IsPlaySessionInProgress();
+		
+		return bHasExecutor && bNotRendering && bAtLeastOneJobAvailable && !bWorldIsActive;
+	}
+}
diff --git a/Engine/Plugins/MovieScene/MovieRenderPipeline/Source/MovieRenderPipelineEditor/Private/MovieRenderPipelineEditorUtils.h b/Engine/Plugins/MovieScene/MovieRenderPipeline/Source/MovieRenderPipelineEditor/Private/MovieRenderPipelineEditorUtils.h
new file mode 100644
--- /dev/null
+++ b/Engine/Plugins/MovieScene/MovieRenderPipeline/Source/MovieRenderPipelineEditor/Private/MovieRenderPipelineEditorUtils.h
@@ -0,0 +1,12 @@
+﻿// Copyright Epic Games, Inc. All Rights Reserved.
+
+#pragma once
+
+namespace UE::MovieRenderPipelineEditor::Private
+{
+	/** Performs a local render using the queue currently loaded up in the Movie Render Queue editor. Returns true if a render was started, else false. */
+	bool PerformLocalRender();
+
+	/** Determines if a local render is currently able to be performed. */
+	bool CanPerformLocalRender();
+}
\ No newline at end of file
diff --git a/Engine/Plugins/MovieScene/MovieRenderPipeline/Source/MovieRenderPipelineEditor/Private/MovieRenderPipelineSettings.cpp b/Engine/Plugins/MovieScene/MovieRenderPipeline/Source/MovieRenderPipelineEditor/Private/MovieRenderPipelineSettings.cpp
--- a/Engine/Plugins/MovieScene/MovieRenderPipeline/Source/MovieRenderPipelineEditor/Private/MovieRenderPipelineSettings.cpp
+++ b/Engine/Plugins/MovieScene/MovieRenderPipeline/Source/MovieRenderPipelineEditor/Private/MovieRenderPipelineSettings.cpp
@@ -19,6 +19,7 @@ UMovieRenderPipelineProjectSettings::UMovieRenderPipelineProjectSettings()
 	DefaultRemoteExecutor = UMoviePipelineNewProcessExecutor::StaticClass();
 	DefaultExecutorJob = UMoviePipelineExecutorJob::StaticClass();
 	DefaultGraph = GetDefaultGraphPath();
+	DefaultQuickRenderGraph = GetDefaultQuickRenderGraphPath();
 
 	DefaultClasses.Add(UMoviePipelineImageSequenceOutput_JPG::StaticClass());
 	DefaultClasses.Add(UMoviePipelineDeferredPassBase::StaticClass());
@@ -30,3 +31,10 @@ FSoftObjectPath UMovieRenderPipelineProjectSettings::GetDefaultGraphPath()
 	
 	return DefaultGraphPath;
 }
+
+FSoftObjectPath UMovieRenderPipelineProjectSettings::GetDefaultQuickRenderGraphPath()
+{
+	static const FSoftObjectPath DefaultQuickRenderGraphPath(TEXT("/MovieRenderPipeline/DefaultQuickRenderGraph.DefaultQuickRenderGraph"));
+	
+	return DefaultQuickRenderGraphPath;
+}
diff --git a/Engine/Plugins/MovieScene/MovieRenderPipeline/Source/MovieRenderPipelineEditor/Private/Widgets/SMoviePipelineQueuePanel.cpp b/Engine/Plugins/MovieScene/MovieRenderPipeline/Source/MovieRenderPipelineEditor/Private/Widgets/SMoviePipelineQueuePanel.cpp
--- a/Engine/Plugins/MovieScene/MovieRenderPipeline/Source/MovieRenderPipelineEditor/Private/Widgets/SMoviePipelineQueuePanel.cpp
+++ b/Engine/Plugins/MovieScene/MovieRenderPipeline/Source/MovieRenderPipelineEditor/Private/Widgets/SMoviePipelineQueuePanel.cpp
@@ -8,6 +8,7 @@
 #include "Widgets/MoviePipelineWidgetConstants.h"
 #include "SMoviePipelineQueueEditor.h"
 #include "SMoviePipelineConfigPanel.h"
+#include "MovieRenderPipelineEditorUtils.h"
 #include "MovieRenderPipelineSettings.h"
 #include "MovieRenderPipelineStyle.h"
 #include "Sections/MovieSceneCinematicShotSection.h"
@@ -277,44 +278,16 @@ UE_ENABLE_OPTIMIZATION_SHIP
 
 FReply SMoviePipelineQueuePanel::OnRenderLocalRequested()
 {
-	UMoviePipelineQueueSubsystem* Subsystem = GEditor->GetEditorSubsystem<UMoviePipelineQueueSubsystem>();
-	check(Subsystem);
-
-	const UMovieRenderPipelineProjectSettings* ProjectSettings = GetDefault<UMovieRenderPipelineProjectSettings>();
-	TSubclassOf<UMoviePipelineExecutorBase> ExecutorClass = ProjectSettings->DefaultLocalExecutor.TryLoadClass<UMoviePipelineExecutorBase>();
-
-	// OnRenderLocalRequested should only get called if IsRenderLocalEnabled() returns true, meaning there's a valid class.
-	check(ExecutorClass != nullptr);
-	Subsystem->RenderQueueWithExecutor(ExecutorClass);
-
-	constexpr bool bIsLocal = true;
-	FMoviePipelineTelemetry::SendRendersRequestedTelemetry(bIsLocal, Subsystem->GetQueue()->GetJobs());
+	UE::MovieRenderPipelineEditor::Private::PerformLocalRender();
 	
 	return FReply::Handled();
 }
 
 bool SMoviePipelineQueuePanel::IsRenderLocalEnabled() const
 {
-	UMoviePipelineQueueSubsystem* Subsystem = GEditor->GetEditorSubsystem<UMoviePipelineQueueSubsystem>();
-	check(Subsystem);
-
-	const UMovieRenderPipelineProjectSettings* ProjectSettings = GetDefault<UMovieRenderPipelineProjectSettings>();
-	const bool bHasExecutor = ProjectSettings->DefaultLocalExecutor.TryLoadClass<UMoviePipelineExecutorBase>() != nullptr;
-	const bool bNotRendering = !Subsystem->IsRendering();
 	const bool bConfigWindowIsOpen = WeakEditorWindow.IsValid();
 
-	bool bAtLeastOneJobAvailable = false;
-	for (UMoviePipelineExecutorJob* Job : Subsystem->GetQueue()->GetJobs())
-	{
-		if (!Job->IsConsumed() && Job->IsEnabled())
-		{
-			bAtLeastOneJobAvailable = true;
-			break;
-		}
-	}
-
-	const bool bWorldIsActive = GEditor->IsPlaySessionInProgress();
-	return bHasExecutor && bNotRendering && bAtLeastOneJobAvailable && !bWorldIsActive && !bConfigWindowIsOpen;
+	return UE::MovieRenderPipelineEditor::Private::CanPerformLocalRender() && !bConfigWindowIsOpen;
 }
 
 FReply SMoviePipelineQueuePanel::OnRenderRemoteRequested()
diff --git a/Engine/Plugins/MovieScene/MovieRenderPipeline/Source/MovieRenderPipelineEditor/Public/Graph/MovieGraphQuickRender.h b/Engine/Plugins/MovieScene/MovieRenderPipeline/Source/MovieRenderPipelineEditor/Public/Graph/MovieGraphQuickRender.h
new file mode 100644
--- /dev/null
+++ b/Engine/Plugins/MovieScene/MovieRenderPipeline/Source/MovieRenderPipelineEditor/Public/Graph/MovieGraphQuickRender.h
@@ -0,0 +1,46 @@
+// Copyright Epic Games, Inc. All Rights Reserved.
+
+#pragma once
+
+#include "CoreMinimal.h"
+#include "UObject/StrongObjectPtr.h"
+
+#include "Graph/MovieGraphDataTypes.h"
+#include "Graph/MovieGraphQuickRenderSettings.h"
+#include "MoviePipelineQueue.h"
+
+class FEditorViewportClient;
+
+/**
+ * Provides the ability to perform a "Quick Render". A Quick Render is a render which requires no conventional Movie Render Queue setup, like
+ * creating a queue, adding a job(s) to it, specifying the level sequence to use, providing a custom configuration or graph, etc. Quick Renders are
+ * designed to get you a render as fast as possible while providing minimal configuration input, for use in things like approving animation.
+ */
+class MOVIERENDERPIPELINEEDITOR_API FMovieGraphQuickRender
+{
+public:
+	/* Begins a quick render using the provided settings. */
+	static void BeginQuickRender(const FMovieGraphQuickRenderSettings& InQuickRenderSettings);
+
+	// TODO: Add python-exposed helper to refresh the variable assignments on the context?
+
+private:
+	/**
+	 * Gets the graph that should be used for the quick render. This will be a duplicate of either the user-specified graph or the default quick render graph.
+	 * OutOriginalToDupeMap will provide a mapping of the pre-duplicate -> duplicate graphs (may contain multiple graphs if subgraphs are involved).
+	 */
+	static UMovieGraphConfig* GetQuickRenderGraph(UMovieGraphConfig* InUserSpecifiedGraph, TMap<TObjectPtr<UMovieGraphConfig>, TObjectPtr<UMovieGraphConfig>>& OutOriginalToDupeMap);
+
+	/** Applies all necessary updates to the duplicated graph to prepare it for a quick render (adjusting start/end frames, etc). */
+	static void ApplyQuickRenderUpdatesToDuplicateGraph(const FMovieGraphQuickRenderSettings& InQuickRenderSettings, UWorld* InEditorWorld);
+
+	/** Injects a new node in the furthest-downstream position within the Globals branch in the duplicated graph. */
+	static UMovieGraphNode* InjectNodeIntoGlobals(const TSubclassOf<UMovieGraphNode> NodeType);
+
+private:
+	/** The temporary queue that is used by quick render. */
+	inline static TStrongObjectPtr<UMoviePipelineQueue> TemporaryQueue = nullptr;
+
+	/** The temporary graph that is used by quick render. Usually a duplicate of a graph created outside of quick render, but will be modified by quick render. */
+	inline static TStrongObjectPtr<UMovieGraphConfig> TemporaryGraph = nullptr;
+};
diff --git a/Engine/Plugins/MovieScene/MovieRenderPipeline/Source/MovieRenderPipelineEditor/Public/MoviePipelineQueueSubsystem.h b/Engine/Plugins/MovieScene/MovieRenderPipeline/Source/MovieRenderPipelineEditor/Public/MoviePipelineQueueSubsystem.h
--- a/Engine/Plugins/MovieScene/MovieRenderPipeline/Source/MovieRenderPipelineEditor/Public/MoviePipelineQueueSubsystem.h
+++ b/Engine/Plugins/MovieScene/MovieRenderPipeline/Source/MovieRenderPipelineEditor/Public/MoviePipelineQueueSubsystem.h
@@ -64,6 +64,10 @@ public:
 	UFUNCTION(BlueprintCallable, meta = (DeterminesOutputType = "InExecutorType"), Category = "Movie Render Pipeline|Rendering")
 	UMoviePipelineExecutorBase* RenderQueueWithExecutor(TSubclassOf<UMoviePipelineExecutorBase> InExecutorType);
 
+	/** Like RenderQueueWithExecutor(), but a specific queue instance is used rather than the current queue. */
+	UFUNCTION(BlueprintCallable, meta = (DeterminesOutputType = "InExecutorType"), Category = "Movie Render Pipeline|Rendering")
+	UMoviePipelineExecutorBase* RenderQueueInstanceWithExecutor(UMoviePipelineQueue* InQueue, TSubclassOf<UMoviePipelineExecutorBase> InExecutorType);
+
 	/**
 	* Starts processing the current queue with the supplied executor. This starts an async process which
 	* may or may not run in a separate process (or on separate machines), determined by the executor implementation.
@@ -74,6 +78,10 @@ public:
 	UFUNCTION(BlueprintCallable, Category = "Movie Render Pipeline|Rendering")
 	void RenderQueueWithExecutorInstance(UMoviePipelineExecutorBase* InExecutor);
 
+	/** Like RenderQueueWithExecutorInstance(), but a specific queue instance is used rather than the current queue. */
+	UFUNCTION(BlueprintCallable, Category = "Movie Render Pipeline|Rendering")
+	void RenderQueueInstanceWithExecutorInstance(UMoviePipelineQueue* InQueue, UMoviePipelineExecutorBase* InExecutor);
+
 	/**
 	* Returns true if there is an active executor working on producing a movie. This could be actively rendering frames,
 	* or working on post processing (finalizing file writes, etc.). Use GetActiveExecutor() and query it directly for
diff --git a/Engine/Plugins/MovieScene/MovieRenderPipeline/Source/MovieRenderPipelineEditor/Public/MovieRenderPipelineSettings.h b/Engine/Plugins/MovieScene/MovieRenderPipeline/Source/MovieRenderPipelineEditor/Public/MovieRenderPipelineSettings.h
--- a/Engine/Plugins/MovieScene/MovieRenderPipeline/Source/MovieRenderPipelineEditor/Public/MovieRenderPipelineSettings.h
+++ b/Engine/Plugins/MovieScene/MovieRenderPipeline/Source/MovieRenderPipelineEditor/Public/MovieRenderPipelineSettings.h
@@ -75,6 +75,11 @@ public:
 	UPROPERTY(config, EditAnywhere, BlueprintReadWrite, Category="Movie Render Pipeline")
 	TSoftObjectPtr<UMovieGraphConfig> DefaultGraph;
 
+	// TODO: Delete this once there's a full Settings dialog for Quick Render
+	/** The graph that quick renders will use. */
+	UPROPERTY(config, EditAnywhere, BlueprintReadWrite, Category="Movie Render Pipeline")
+	TSoftObjectPtr<UMovieGraphConfig> DefaultQuickRenderGraph;
+
 	/**
 	* The settings specified here will automatically be added to a Movie Pipeline Primary Configuration when using the UI. 
 	* This does not apply to scripting and does not apply to runtime. It is only a convenience function so that when a job is
@@ -89,4 +94,7 @@ public:
 public:
 	/** Gets the path to the default graph asset supplied by MRQ. */
 	static FSoftObjectPath GetDefaultGraphPath();
+
+	/** Gets the path to the default quick render graph asset supplied by MRQ. */
+	static FSoftObjectPath GetDefaultQuickRenderGraphPath();
 };
\ No newline at end of file
diff --git a/Engine/Plugins/MovieScene/MovieRenderPipeline/Source/MovieRenderPipelineRenderPasses/Private/Graph/Renderers/MovieGraphDeferredPass.cpp b/Engine/Plugins/MovieScene/MovieRenderPipeline/Source/MovieRenderPipelineRenderPasses/Private/Graph/Renderers/MovieGraphDeferredPass.cpp
--- a/Engine/Plugins/MovieScene/MovieRenderPipeline/Source/MovieRenderPipelineRenderPasses/Private/Graph/Renderers/MovieGraphDeferredPass.cpp
+++ b/Engine/Plugins/MovieScene/MovieRenderPipeline/Source/MovieRenderPipelineRenderPasses/Private/Graph/Renderers/MovieGraphDeferredPass.cpp
@@ -1,6 +1,7 @@
 // Copyright Epic Games, Inc. All Rights Reserved.
 
 #include "Graph/Renderers/MovieGraphDeferredPass.h"
+#include "Graph/Nodes/MovieGraphApplyViewportLookNode.h"
 #include "Graph/Nodes/MovieGraphDeferredPassNode.h"
 #include "Graph/Nodes/MovieGraphCameraNode.h"
 #include "Graph/MovieGraphDataTypes.h"
@@ -285,6 +286,15 @@ void FMovieGraphDeferredPass::Render(const FMovieGraphTraversalContext& InFrameT
 		 
 		// Construct a View to go within this family.
 		FSceneView* NewView = CreateSceneView(SceneViewInitOptions, ViewFamily, CameraInfo);
+
+		// Viewport-look mode may need to apply additional customizations to the view
+#if WITH_EDITOR
+		constexpr bool bExactMatch = true;
+		if (const UMovieGraphApplyViewportLookNode* ViewportLookNode = InTimeData.EvaluatedConfig->GetSettingForBranch<UMovieGraphApplyViewportLookNode>(UMovieGraphNode::GlobalsPinName, bIncludeCDOs, bExactMatch))
+		{
+			ViewportLookNode->UpdateSceneView(NewView);
+		}
+#endif	// WITH_EDITOR
 		
 		// Then apply Movie Render Queue specific overrides to the ViewFamily, and then to the SceneView.
 		ApplyMovieGraphOverridesToViewFamily(ViewFamily, ViewFamilyInitData);
diff --git a/Engine/Plugins/MovieScene/MovieRenderPipeline/Source/MovieRenderPipelineRenderPasses/Private/Graph/Renderers/MovieGraphImagePassBase.cpp b/Engine/Plugins/MovieScene/MovieRenderPipeline/Source/MovieRenderPipelineRenderPasses/Private/Graph/Renderers/MovieGraphImagePassBase.cpp
--- a/Engine/Plugins/MovieScene/MovieRenderPipeline/Source/MovieRenderPipelineRenderPasses/Private/Graph/Renderers/MovieGraphImagePassBase.cpp
+++ b/Engine/Plugins/MovieScene/MovieRenderPipeline/Source/MovieRenderPipelineRenderPasses/Private/Graph/Renderers/MovieGraphImagePassBase.cpp
@@ -1,8 +1,10 @@
 // Copyright Epic Games, Inc. All Rights Reserved.
 
 #include "Graph/Renderers/MovieGraphImagePassBase.h"
+
 #include "Graph/MovieGraphDataTypes.h"
 #include "Graph/MovieGraphPipeline.h"
+#include "Graph/Nodes/MovieGraphApplyViewportLookNode.h"
 #include "SceneView.h"
 #include "MoviePipelineUtils.h"
 #include "LegacyScreenPercentageDriver.h"
@@ -151,6 +153,28 @@ TSharedRef<FSceneViewFamilyContext> FMovieGraphImagePassBase::CreateSceneViewFam
 	EViewModeIndex ViewModeIndex = InInitData.ViewModeIndex;
 	FEngineShowFlags ShowFlags = InInitData.ShowFlags;
 
+	// Viewport rendering mode may need to customize the show flags and view mode index
+	bool bIsViewportLook = false;
+#if WITH_EDITOR
+	if (InInitData.TimeData.EvaluatedConfig)
+	{
+		constexpr bool bIncludeCDOs = false;
+		constexpr bool bExactMatch = true;
+		if (const UMovieGraphApplyViewportLookNode* ViewportLookNode = InInitData.TimeData.EvaluatedConfig->GetSettingForBranch<UMovieGraphApplyViewportLookNode>(UMovieGraphNode::GlobalsPinName, bIncludeCDOs, bExactMatch))
+		{
+			EViewModeIndex ViewportViewModeIndex;
+			FEngineShowFlags ViewportShowFlags(ESFIM_Editor);
+			bIsViewportLook = true;
+			
+			if (ViewportLookNode->GetViewportInfo(ViewportShowFlags, ViewportViewModeIndex))
+			{
+				ViewModeIndex = ViewportViewModeIndex;
+				ShowFlags = ViewportShowFlags;
+			}
+		}
+	}
+#endif	// WITH_EDITOR
+
 	const bool bIsPerspective = InInitData.ProjectionMode == ECameraProjectionMode::Type::Perspective;
 
 	// Allow the Engine Showflag system to override our engine showflags, based on our view mode index.
@@ -158,7 +182,7 @@ TSharedRef<FSceneViewFamilyContext> FMovieGraphImagePassBase::CreateSceneViewFam
 	ApplyViewMode(/*In*/ ViewModeIndex, bIsPerspective, /*InOut*/ShowFlags);
 
 	// And then we have to let another system override them again (based on cvars, etc.)
-	EngineShowFlagOverride(ESFIM_Game, ViewModeIndex, ShowFlags, false);
+	EngineShowFlagOverride(bIsViewportLook ? ESFIM_Editor : ESFIM_Game, ViewModeIndex, ShowFlags, false);
 
 	TSharedRef<FSceneViewFamilyContext> OutViewFamily = MakeShared<FSceneViewFamilyContext>(FSceneViewFamily::ConstructionValues(
 		InInitData.RenderTarget,
