MovieGraph: Improvements to the UMovieGraphConfig API for finding nodes by type, and for finding nodes by tag.

Added:
Ability to get the branch names (of the Output node) for a graph.
Ability to find setting(s) in a graph by setting class and branch name.
Ability to find setting(s) in a graph by tag, setting class, and branch name.

Specifically tested:
Using these operations with both Globals and non-global branches
Verifying that settings class filters restrict the returned values to just those settings.
Verifying that branch cilters restrict the returned values to just settings on that branch.
Verifying that tags can return across all branches, one branch, and when filtered by class that no additional classes are returned even if they have a matching tag.

#virtualized

diff --git a/Engine/Plugins/MovieScene/MovieRenderPipeline/Source/MovieRenderPipelineCore/Private/Graph/MovieGraphConfig.cpp b/Engine/Plugins/MovieScene/MovieRenderPipeline/Source/MovieRenderPipelineCore/Private/Graph/MovieGraphConfig.cpp
--- a/Engine/Plugins/MovieScene/MovieRenderPipeline/Source/MovieRenderPipelineCore/Private/Graph/MovieGraphConfig.cpp
+++ b/Engine/Plugins/MovieScene/MovieRenderPipeline/Source/MovieRenderPipelineCore/Private/Graph/MovieGraphConfig.cpp
@@ -1292,6 +1292,145 @@ void UMovieGraphConfig::RecurseUpGlobalsBranchToFindOutputDirectory(const UMovie
 	}
 };
 
+const TArray<FName> UMovieGraphConfig::GetBranchNames() const
+{
+	TArray<FName> BranchNames;
+	for (UMovieGraphOutput* Output : GetOutputs())
+	{
+		if (Output->IsA<UMovieGraphOutputNode>())
+		{
+			BranchNames.AddUnique(FName(Output->GetMemberName()));
+		}
+	}
+	return BranchNames;
+}
+
+UMovieGraphSettingNode* UMovieGraphConfig::GetSettingForBranch(UClass* InClass, const FName& InBranchName, bool bExactMatch) const
+{
+	TArray<UMovieGraphSettingNode*> AllSettings = GetSettingsForBranch(InClass, InBranchName, bExactMatch);
+	if (AllSettings.Num() > 0)
+	{
+		return AllSettings[0];
+	}
+
+	return nullptr;
+}
+
+TArray<UMovieGraphSettingNode*> UMovieGraphConfig::GetSettingsForBranch(UClass* InClass, const FName& InBranchName, bool bExactMatch) const
+{
+	TArray<UMovieGraphSettingNode*> ResultNodes;
+
+	// We need to be able to detect when we've reached the end of traversal for the current branch,
+	// which will either be at the input node, or when there's nothing connected to the left of it.
+	bool bIsTraversingBranch = false;
+	UMovieGraphNode* LocalInputNode = GetInputNode();
+
+	VisitUpstreamNodes(GetOutputNode(), UMovieGraphConfig::FVisitNodesCallback::CreateLambda(
+		[&InBranchName, &bIsTraversingBranch, &ResultNodes, LocalInputNode, bExactMatch, InClass](const UMovieGraphNode* VisitedNode, const UMovieGraphPin* VisitedPin)
+		{
+			check(VisitedNode);
+			constexpr bool bFollowRerouteNodes = true;
+
+			// If we return true from this function, then we will follow that pin. Unfortunately it only visits the first
+			// child (ie: not the output node), so we have to check if the downstream pin was their desired pin
+			if (!bIsTraversingBranch)
+			{
+				UMovieGraphPin* FirstOutputPin = VisitedNode->GetFirstConnectedOutputPin();
+				if (FirstOutputPin)
+				{
+					for (UMovieGraphEdge* Edge : FirstOutputPin->Edges)
+					{
+						UMovieGraphPin* OtherPin = Edge->GetOtherPin(FirstOutputPin, bFollowRerouteNodes);
+						if (OtherPin)
+						{
+							const bool bBranchMatch = OtherPin->Properties.Label == InBranchName;
+							if (bBranchMatch)
+							{
+								bIsTraversingBranch = true;
+							}
+						}
+					}
+				}
+			}
+
+			if (bIsTraversingBranch)
+			{
+
+				// If there's nothing to our left on this node (or we're the Input node) then we've fully explored this branch
+				if (VisitedNode->GetInputPins().Num() == 0 || VisitedNode == LocalInputNode)
+				{
+					bIsTraversingBranch = false;
+				}
+				else
+				{
+					// If we started traversing the branch, we need to keep adding all the nodes we find, even if they're not directly connected
+					bool bMatches = bExactMatch ? VisitedNode->GetClass() == InClass : VisitedNode->IsA(InClass);
+					if (bMatches)
+					{
+						ResultNodes.AddUnique(CastChecked<UMovieGraphSettingNode>(const_cast<UMovieGraphNode*>(VisitedNode)));
+					}
+				}
+			}
+
+			// We return true here anywhere on the correct branch so that we keep going.
+			return bIsTraversingBranch; 
+		}));
+
+	return ResultNodes;
+}
+
+UMovieGraphSettingNode* UMovieGraphConfig::GetSettingForTag(const FString& ScriptTag, UClass* OptionalClassFilter, const FName OptionalBranchFilter, bool bExactMatch) const
+{
+	TArray<UMovieGraphSettingNode*> AllSettings = GetSettingsForTag(ScriptTag, OptionalClassFilter, OptionalBranchFilter, bExactMatch);
+	if (AllSettings.Num() > 0)
+	{
+		return AllSettings[0];
+	}
+
+	return nullptr;
+}
+
+TArray<UMovieGraphSettingNode*> UMovieGraphConfig::GetSettingsForTag(const FString& ScriptTag, UClass* OptionalClassFilter, const FName OptionalBranchFilter, bool bExactMatch) const
+{
+	TArray<UMovieGraphSettingNode*> AllSettingsToConsider;
+	if(OptionalBranchFilter != NAME_None)
+	{
+		UClass* ClassFilter = OptionalClassFilter ? OptionalClassFilter : UMovieGraphSettingNode::StaticClass();
+		AllSettingsToConsider = GetSettingsForBranch(ClassFilter, OptionalBranchFilter, bExactMatch);
+	}
+	else
+	{
+		AllSettingsToConsider.Reserve(GetNodes().Num());
+		for(const TObjectPtr<UMovieGraphNode>& Node : GetNodes())
+		{
+			if (UMovieGraphSettingNode* AsSetting = Cast<UMovieGraphSettingNode>(Node))
+			{
+				AllSettingsToConsider.Add(AsSetting);
+			}
+		}
+	}
+			
+	
+	TArray<UMovieGraphSettingNode*> AllSettingsWithTag;
+	for (UMovieGraphSettingNode* SettingNode : AllSettingsToConsider)
+	{
+		// We re-run the class filter here because if they weren't using branch filters, they didn't get class-filtered yet.
+		const bool bHasClassFilter = OptionalClassFilter != nullptr;
+		const bool bPassesClassFilter = bExactMatch ? SettingNode->GetClass() == OptionalClassFilter : SettingNode->IsA(OptionalClassFilter);
+
+		const bool bConsider = bHasClassFilter ? bPassesClassFilter : true;
+		if (bConsider)
+		{
+			if (SettingNode->ScriptTags.Contains(ScriptTag))
+			{
+				AllSettingsWithTag.Add(SettingNode);
+			}
+		}
+	}
+
+	return AllSettingsWithTag;
+}
+
 void UMovieGraphConfig::GetOutputDirectory(FString& OutOutputDirectory) const
 {
 	check (OutputNode);
diff --git a/Engine/Plugins/MovieScene/MovieRenderPipeline/Source/MovieRenderPipelineCore/Public/Graph/MovieGraphConfig.h b/Engine/Plugins/MovieScene/MovieRenderPipeline/Source/MovieRenderPipelineCore/Public/Graph/MovieGraphConfig.h
--- a/Engine/Plugins/MovieScene/MovieRenderPipeline/Source/MovieRenderPipelineCore/Public/Graph/MovieGraphConfig.h
+++ b/Engine/Plugins/MovieScene/MovieRenderPipeline/Source/MovieRenderPipelineCore/Public/Graph/MovieGraphConfig.h
@@ -657,8 +657,7 @@ public:
 	*/
 	void AddNode(UMovieGraphNode* InNode);
 
-	/** 
-	* Removes the specified node from the graph, disconnecting connected edges as it goes. */
+	/** Removes the specified node from the graph, disconnecting connected edges as it goes. */
 	UFUNCTION(BlueprintCallable, Category = "Movie Graph")
 	bool RemoveNode(UMovieGraphNode* InNode);
 
@@ -676,6 +675,81 @@ public:
 
 	const TArray<TObjectPtr<UMovieGraphNode>>& GetNodes() const { return AllNodes; }
 
+	/** Returns an array of the branch names for the OutputNode on this Graph. */
+	UFUNCTION(BlueprintCallable, Category = "Movie Graph")
+	const TArray<FName> GetBranchNames() const;
+
+
+	/**
+	* Finds a setting node (by type) for the given branch name (see GetBranchNames()). Returns the first result
+	* of that type found, even if there are multiple, following traversal order (ie: right to left)
+	* To prevent accidentally editing unrelated assets, does not dive into sub-graphs (but can continue
+	* traversal beyond them), so returned results should only exist in current asset.
+	* Does not contain the Input or Output nodes (as they are not settings), see GetInputNode(), GetOutputNode())
+	*
+	* @param InClass - The node class type to search for.
+	* @param InBranchName - The branch to search for these nodes on.
+	* @param bExactMatch - Defaults false. If true, will not consider inherited classes of InClass.
+	* 					   (ie: If requesting a base class and only child classes exist, they will not count when using exact matching).
+	* @return First node found (if any) which are of the specified InClass type on the given branch.
+	*/
+	UFUNCTION(BlueprintCallable, meta = (DeterminesOutputType = "InClass"), Category = "Movie Graph")
+	UMovieGraphSettingNode* GetSettingForBranch(UClass* InClass, const FName& InBranchName, bool bExactMatch = false) const;
+
+	/**
+	* Finds all setting nodes (by type) for the given branch name (see GetBranchNames()). Returns all results of that type
+	* following traversal order (ie: right to left).
+	* To prevent accidentally editing unrelated assets, does not dive into sub-graphs (but can continue
+	* traversal beyond them), so returned results should only exist in current asset.
+	* Does not contain the Input or Output nodes (as they are not settings), see GetInputNode(), GetOutputNode())
+	*
+	* @param InClass - The node class type to search for.
+	* @param InBranchName - The branch to search for these nodes on.
+	* @param bExactMatch - Defaults false. If true, will not consider inherited classes of InClass.
+	* 					   (ie: If requesting a base class and only child classes exist, they will not count when using exact matching).
+	* @return All nodes found (if any) which are of the specified InClass type on the given branch.
+	*/
+	UFUNCTION(BlueprintCallable, meta = (DeterminesOutputType = "InClass"), Category = "Movie Graph")
+	TArray<UMovieGraphSettingNode*> GetSettingsForBranch(UClass* InClass, const FName& InBranchName, bool bExactMatch = false) const;
+
+	/**
+	* Finds a setting node by ScriptTag field. THIS SEARCH IS CASE SENSITIVE. Returns the first result
+	* of that tag found, even if there are multiple, following traversal order (ie: right to left).
+	* If OptionalClassFilter is specified, only matches against nodes that have the right class type and
+	* contain the correct tag.
+	* To prevent accidentally editing unrelated assets, does not dive into sub-graphs (but can continue
+	* traversal beyond them), so returned results should only exist in current asset.
+	* Does not contain the Input or Output nodes (as they are not settings), see GetInputNode(), GetOutputNode())
+	*
+	* @param ScriptTag - Case sensitive name into the "Script Tags" property on Nodes
+	* @param OptionalClassFilter - Defaults None. If specified, will only return a node of this class that matches the tag.
+	* @param OptionalBranchFilter - Defaults None. If specified, will only search the specified branch, otherwise searches all branches.
+	* @param bExactMatch - Defaults false. If true, will not consider inherited classes to pass the OptionalClassFilter.
+	* 					   (ie: If requesting a base class and only child classes exist, they will not count when using exact matching).
+	* @return The first node found (if any) which contains the case-sensitive tag and is of the correct class type (if using OptionalClassFilter)
+	*/
+	UFUNCTION(BlueprintCallable, meta = (DeterminesOutputType = "OptionalClassFilter"), Category = "Movie Graph")
+	UMovieGraphSettingNode* GetSettingForTag(const FString& ScriptTag, UClass* OptionalClassFilter = nullptr, const FName OptionalBranchFilter = NAME_None, bool bExactMatch = false) const;
+
+	/**
+	* Finds setting nodes by ScriptTag field. THIS SEARCH IS CASE SENSITIVE. Returns all results
+	* of that tag found, following traversal order (ie: right to left).
+	* If OptionalClassFilter is specified, only matches against nodes that have the right class type and
+	* contain the correct tag.
+	* To prevent accidentally editing unrelated assets, does not dive into sub-graphs (but can continue
+	* traversal beyond them), so returned results should only exist in current asset.
+	* Does not contain the Input or Output nodes (as they are not settings), see GetInputNode(), GetOutputNode())
+	*
+	* @param ScriptTag - Case sensitive name into the "Script Tags" property on Nodes
+	* @param OptionalClassFilter - Defaults None. If specified, will only return a node of this class that matches the tag.
+	* @param OptionalBranchFilter - Defaults None. If specified, will only search the specified branch, otherwise searches all branches.
+	* @param bExactMatch - Defaults false. If true, will not consider inherited classes to pass the OptionalClassFilter.
+	* 					   (ie: If requesting a base class and only child classes exist, they will not count when using exact matching).
+	* @return All nodes found (if any) which contain the case-sensitive tag and are of the correct class type (if using OptionalClassFilter)
+	*/
+	UFUNCTION(BlueprintCallable, meta = (DeterminesOutputType = "OptionalClassFilter"), Category = "Movie Graph")
+	TArray<UMovieGraphSettingNode*> GetSettingsForTag(const FString& ScriptTag, UClass* OptionalClassFilter = nullptr, const FName OptionalBranchFilter = NAME_None, bool bExactMatch = false) const;
+
 	/**
 	 * Adds a new variable member with default values to the graph. The new variable will have a base name of
 	 * "Variable" unless specified in InCustomBaseName. Returns the new variable on success, else nullptr.
