MoviePipeline: Added a new Settings dialog for Quick Render. This is available in the Quick Render drop-down menu.

The "mode" selector in the dialog is currently disabled and will be enabled in a subsequent CL.

#virtualized

diff --git a/Engine/Plugins/MovieScene/MovieRenderPipeline/Source/MovieRenderPipelineCore/Private/Graph/MovieGraphQuickRenderSettings.cpp b/Engine/Plugins/MovieScene/MovieRenderPipeline/Source/MovieRenderPipelineCore/Private/Graph/MovieGraphQuickRenderSettings.cpp
new file mode 100644
--- /dev/null
+++ b/Engine/Plugins/MovieScene/MovieRenderPipeline/Source/MovieRenderPipelineCore/Private/Graph/MovieGraphQuickRenderSettings.cpp
@@ -0,0 +1,102 @@
+﻿// Copyright Epic Games, Inc. All Rights Reserved.
+
+#include "Graph/MovieGraphQuickRenderSettings.h"
+
+#include "Misc/CoreDelegates.h"
+#include "MoviePipelineUtils.h"
+#include "MovieRenderPipelineCoreModule.h"
+#include "UObject/Package.h"
+
+#if WITH_EDITOR
+#include "PackageHelperFunctions.h"
+#include "Settings/EditorLoadingSavingSettings.h"
+#endif
+
+UMovieGraphQuickRenderSettings::UMovieGraphQuickRenderSettings()
+	: GraphPreset(FSoftObjectPath(TEXT("/MovieRenderPipeline/DefaultQuickRenderGraph.DefaultQuickRenderGraph")))
+{
+	
+}
+
+UMovieGraphQuickRenderSettings* UMovieGraphQuickRenderSettings::GetSavedQuickRenderSettings()
+{
+	UMovieGraphQuickRenderSettings* QuickRenderSettings = nullptr;
+	
+	if (const UPackage* QuickRenderSettingsPackage = LoadPackage(nullptr, *QuickRenderSettingsPackagePath, LOAD_None))
+	{
+		QuickRenderSettings = Cast<UMovieGraphQuickRenderSettings>(FindObjectWithOuter(QuickRenderSettingsPackage, StaticClass()));
+	}
+	else
+	{
+		UPackage* NewSettingsPackage = CreatePackage(*QuickRenderSettingsPackagePath);
+		QuickRenderSettings = NewObject<UMovieGraphQuickRenderSettings>(NewSettingsPackage);
+	}
+
+	return QuickRenderSettings;
+}
+
+void UMovieGraphQuickRenderSettings::RefreshVariableAssignments(UMovieGraphQuickRenderSettings* InSettings)
+{
+	if (!InSettings)
+	{
+		return;
+	}
+	
+	MoviePipeline::RefreshVariableAssignments(InSettings->GraphPreset.LoadSynchronous(), InSettings->GraphVariableAssignments, InSettings);
+}
+
+#if WITH_EDITOR
+void UMovieGraphQuickRenderSettings::SaveSettings(const UMovieGraphQuickRenderSettings* InSettings)
+{
+	// Settings are saved to a uasset rather than using the ini configuration system. The ini configuration system isn't flexible enough for the
+	// needs of Quick Render (for example, it cannot store the variable assignments).
+	
+	const FString PackageFileName = FPackageName::LongPackageNameToFilename(*QuickRenderSettingsPackagePath, FPackageName::GetAssetPackageExtension());
+
+	// Duplicate the settings into a new package (or the existing one)
+	UPackage* NewPackage = CreatePackage(*QuickRenderSettingsPackagePath);
+	UMovieGraphQuickRenderSettings* DuplicatedSettings = CastChecked<UMovieGraphQuickRenderSettings>(StaticDuplicateObject(InSettings, NewPackage));
+	DuplicatedSettings->SetFlags(RF_Public | RF_Transactional | RF_Standalone);
+
+	// Save the settings out to disk. Turn off the behavior that auto-adds new files to source control.
+	UEditorLoadingSavingSettings* SaveSettings = GetMutableDefault<UEditorLoadingSavingSettings>();
+	const uint32 bSCCAutoAddNewFiles = SaveSettings->bSCCAutoAddNewFiles;
+	SaveSettings->bSCCAutoAddNewFiles = 0;
+	const bool bSuccess = SavePackageHelper(NewPackage, *PackageFileName);
+	SaveSettings->bSCCAutoAddNewFiles = bSCCAutoAddNewFiles;
+	
+	if (!bSuccess)
+	{
+		// SavePackageHelper() will emit warnings if the save was unsuccessful, but log a separate warning for movie pipeline in case warnings are
+		// being specifically filtered for LogMovieRenderPipeline.
+		UE_LOG(LogMovieRenderPipeline, Warning, TEXT("Unable to save Quick Render settings. Could not save to destination file [%s]."), *PackageFileName);
+	}
+}
+#endif	// WITH_EDITOR
+
+#if WITH_EDITOR
+void UMovieGraphQuickRenderSettings::PostEditChangeProperty(FPropertyChangedEvent& PropertyChangedEvent)
+{
+	if (OnEnginePreExitHandle.IsValid())
+	{
+		FCoreDelegates::OnEnginePreExit.Remove(OnEnginePreExitHandle);
+	}
+
+	// Queue up saving the settings to a uasset file in the Saved directory. Do this only when the editor is closing, otherwise the package-saving
+	// dialog will briefly pop up, which looks glitchy.
+	OnEnginePreExitHandle = FCoreDelegates::OnEnginePreExit.AddLambda([this]()
+	{
+		SaveSettings(this);
+	});
+	
+	if (PropertyChangedEvent.GetPropertyName() == GET_MEMBER_NAME_CHECKED(UMovieGraphQuickRenderSettings, GraphPreset))
+	{
+		// Refresh the variable assignments if the graph preset changes
+		RefreshVariableAssignments(this);
+		
+		OnGraphChangedDelegate.Broadcast();
+	}
+	
+	Super::PostEditChangeProperty(PropertyChangedEvent);
+}
+#endif	// WITH_EDITOR
\ No newline at end of file
diff --git a/Engine/Plugins/MovieScene/MovieRenderPipeline/Source/MovieRenderPipelineCore/Public/Graph/MovieGraphQuickRenderSettings.h b/Engine/Plugins/MovieScene/MovieRenderPipeline/Source/MovieRenderPipelineCore/Public/Graph/MovieGraphQuickRenderSettings.h
--- a/Engine/Plugins/MovieScene/MovieRenderPipeline/Source/MovieRenderPipelineCore/Public/Graph/MovieGraphQuickRenderSettings.h
+++ b/Engine/Plugins/MovieScene/MovieRenderPipeline/Source/MovieRenderPipelineCore/Public/Graph/MovieGraphQuickRenderSettings.h
@@ -1,4 +1,4 @@
-﻿// Copyright Epic Games, Inc. All Rights Reserved.
+// Copyright Epic Games, Inc. All Rights Reserved.
 
 #pragma once
 
@@ -69,60 +69,113 @@ enum class EMovieGraphQuickRenderViewportLookFlags : uint8
 };
 ENUM_CLASS_FLAGS(EMovieGraphQuickRenderViewportLookFlags)
 
-/** Data that sets up user-configurable aspects of a quick render. */
-USTRUCT(BlueprintType)
-struct MOVIERENDERPIPELINECORE_API FMovieGraphQuickRenderSettings
+/** TODO: Docs */
+UENUM(BlueprintType)
+enum class EMovieGraphQuickRenderSettingsCategory : uint8
+{
+	Camera,
+	Sequence,
+	Viewport
+};
+
+#if WITH_EDITOR
+DECLARE_MULTICAST_DELEGATE(FOnMovieGraphQuickRenderGraphChanged);
+#endif
+
+/** Data that specifies user-configurable aspects of a quick render. */
+UCLASS(BlueprintType)
+class MOVIERENDERPIPELINECORE_API UMovieGraphQuickRenderSettings : public UObject
 {
 	GENERATED_BODY()
 
+public:
+	UMovieGraphQuickRenderSettings();
+	
+	/** Gets the Quick Render settings that have been persisted to disk (Saved/MovieRenderPipeline/QuickRenderSettings.uasset). */
+	UFUNCTION(BlueprintCallable, Category="Quick Render")
+	static UMovieGraphQuickRenderSettings* GetSavedQuickRenderSettings();
+
+	/**
+	 * Within the settings object provided, refreshes the variable assignments in the settings to match the setting's GraphPreset. Note that
+	 * changing the graph preset within the editor will do this automatically, but if changing the graph outside of the editor, this may need to be
+	 * run manually.
+	 */
+	UFUNCTION(BlueprintCallable, Category="Quick Render")
+	static void RefreshVariableAssignments(UMovieGraphQuickRenderSettings* InSettings);
+
+private:
+	inline static const FString QuickRenderSettingsPackagePath = TEXT("/Temp/MovieRenderPipeline/QuickRenderSettings");
+	inline static FDelegateHandle OnEnginePreExitHandle;
+
+private:
+#if WITH_EDITOR
+	/** Saves the settings to a uasset file (Saved/MovieRenderPipeline/QuickRenderSettings.uasset). The file will be created if it doesn't exist. */
+	static void SaveSettings(const UMovieGraphQuickRenderSettings* InSettings);
+#endif
+
+#if WITH_EDITOR
+	//~ Begin UObject Interface
+	virtual void PostEditChangeProperty(FPropertyChangedEvent& PropertyChangedEvent) override;
+	//~ End UObject Interface
+#endif
+
+public:
+#if WITH_EDITOR
+	FOnMovieGraphQuickRenderGraphChanged OnGraphChangedDelegate;
+#endif
+	
+	/** The graph preset that is used to configure the render. If not specified, the quick render default graph will be used. */
+	UPROPERTY(EditAnywhere, BlueprintReadWrite, Category="Configuration", DisplayName="Movie Render Graph Config")
+	TSoftObjectPtr<UMovieGraphConfig> GraphPreset;
+
+	/** The action that quick render should perform after a render is finished. */
+	UPROPERTY(EditAnywhere, BlueprintReadWrite, Category="Configuration", DisplayName="After Render")
+	EMovieGraphQuickRenderPostRenderActionType PostRenderBehavior = EMovieGraphQuickRenderPostRenderActionType::PlayRenderOutput;
+
+	/** Enables/disables the ViewportLookFlags property from taking effect. */
+	UPROPERTY(EditAnywhere, BlueprintReadWrite, Category="Quick Render", meta=(InlineEditConditionToggle))
+	bool bOverride_ViewportLookFlags = false;
+
+	// Note: ViewportLookFlags is an int32 instead of the enum to satisfy the requirement of the Bitmask/BitmaskEnum metadata; this can be changed back to the enum once
+	// a proper customization per the design is done.
+	
+	/** The properties of the viewport that should be applied to the quick render. These will override any equivalent properties specified in the graph. */
+	UPROPERTY(EditAnywhere, BlueprintReadWrite, Category="Quick Render", DisplayName="Apply Viewport Look", meta=(EditCondition="bOverride_ViewportLookFlags", Bitmask, BitmaskEnum="/Script/MovieRenderPipelineCore.EMovieGraphQuickRenderViewportLookFlags"))
+	int32 ViewportLookFlags = static_cast<int32>(EMovieGraphQuickRenderViewportLookFlags::None);
+
 	/** How to determine what the start/end frame should be. */
-	UPROPERTY(Config)
+	UPROPERTY(EditAnywhere, BlueprintReadWrite, Category="Quick Render", DisplayName="Sequencer Frame Range")
 	EMovieGraphQuickRenderFrameRangeType FrameRangeType = EMovieGraphQuickRenderFrameRangeType::PlaybackRange;
 
 	/** The frame in the level sequence that the render will begin on. */
-	UPROPERTY(Config)
+	UPROPERTY(BlueprintReadWrite, Category="Quick Render")
 	int32 CustomStartFrame = 0;
 
 	/** The frame in the level sequence that the render will end on. */
-	UPROPERTY(Config)
+	UPROPERTY(BlueprintReadWrite, Category="Quick Render")
 	int32 CustomEndFrame = 0;
 
-	/** The properties of the viewport that should be applied to the quick render. These will override any equivalent properties specified in the graph. */
-	UPROPERTY(Config)
-	EMovieGraphQuickRenderViewportLookFlags ViewportLookFlags = EMovieGraphQuickRenderViewportLookFlags::None;
+	/** Values that are set on the variables contained within the graph preset. See note on RefreshVariableAssignments(). */
+	UPROPERTY(EditAnywhere, BlueprintReadWrite, Category="Primary Graph Variables")
+	TArray<TObjectPtr<UMovieJobVariableAssignmentContainer>> GraphVariableAssignments;
 
-	// TODO: This should probably be deleted, since it will be controlled by variable assignments
-	/** The directory that render outputs are saved to. */
-	UPROPERTY(Config)
-	FDirectoryPath OutputDirectory;
+	// 
+	// ---- The settings below should generally not need to be controlled outside of the UI. ----
+	//
 
-	/** The graph preset that is used to configure the render. If not specified, the quick render default graph will be used. */
-	UPROPERTY(Config)
-	TSoftObjectPtr<UMovieGraphConfig> GraphPreset = nullptr;
+	/** TODO: Docs */
+	UPROPERTY()
+	EMovieGraphQuickRenderSettingsCategory SettingsCategory;
 
 	/** What mode the quick render mode button is using. */
-	UPROPERTY(Config)
+	UPROPERTY()
 	EMovieGraphQuickRenderButtonMode QuickAccessButtonMode = EMovieGraphQuickRenderButtonMode::NormalMovieRenderQueue;
 
 	/** The mode that Quick Render will use when a render begins. Generally dictates the level sequence and camera(s) that will be used. */
-	UPROPERTY(Config)
+	UPROPERTY()
 	EMovieGraphQuickRenderMode QuickRenderMode = EMovieGraphQuickRenderMode::CurrentSequence;
 
 	/** Whether the settings dialog should be shown before a quick render is started. */
-	UPROPERTY(Config)
-	bool bShowSettingsBeforeQuickRender = false;
-
-	/** The action that quick render should perform after a render is finished. */
-	UPROPERTY(Config)
-	EMovieGraphQuickRenderPostRenderActionType PostRenderBehavior = EMovieGraphQuickRenderPostRenderActionType::DoNothing;
-
-	/** Values that are set on the variables contained within the graph preset. */
-	// TODO: This can't be stored as config... how to fix this?
 	UPROPERTY()
-	TArray<TObjectPtr<UMovieJobVariableAssignmentContainer>> GraphVariableAssignments;
-
-	// TODO: This should probably be deleted, since it will be controlled by variable assignments
-	/** The resolution to render at. */
-	UPROPERTY(Config)
-	FMovieGraphNamedResolution Resolution = FMovieGraphNamedResolution(FName(), FIntPoint(-1, -1), FString());
+	bool bShowSettingsBeforeQuickRender = false;
 };
\ No newline at end of file
diff --git a/Engine/Plugins/MovieScene/MovieRenderPipeline/Source/MovieRenderPipelineCore/Public/Graph/MovieRenderGraphEditorSettings.h b/Engine/Plugins/MovieScene/MovieRenderPipeline/Source/MovieRenderPipelineCore/Public/Graph/MovieRenderGraphEditorSettings.h
--- a/Engine/Plugins/MovieScene/MovieRenderPipeline/Source/MovieRenderPipelineCore/Public/Graph/MovieRenderGraphEditorSettings.h
+++ b/Engine/Plugins/MovieScene/MovieRenderPipeline/Source/MovieRenderPipelineCore/Public/Graph/MovieRenderGraphEditorSettings.h
@@ -21,8 +21,4 @@ public:
 	/** If PostRenderBehavior is set to PlayRenderOutput, these settings are used to determine how to play media. */
 	UPROPERTY(Config, EditAnywhere, Category="Play Render Output", meta=(ShowOnlyInnerProperties))
 	FMovieGraphPostRenderSettings PostRenderSettings;
-
-	/** Settings that are used to generate a quick render. */
-	UPROPERTY(Config)
-	FMovieGraphQuickRenderSettings QuickRenderSettings;
 };
diff --git a/Engine/Plugins/MovieScene/MovieRenderPipeline/Source/MovieRenderPipelineEditor/Private/Customizations/Graph/MovieGraphCustomizationUtils.cpp b/Engine/Plugins/MovieScene/MovieRenderPipeline/Source/MovieRenderPipelineEditor/Private/Customizations/Graph/MovieGraphCustomizationUtils.cpp
new file mode 100644
--- /dev/null
+++ b/Engine/Plugins/MovieScene/MovieRenderPipeline/Source/MovieRenderPipelineEditor/Private/Customizations/Graph/MovieGraphCustomizationUtils.cpp
@@ -0,0 +1,41 @@
+﻿// Copyright Epic Games, Inc. All Rights Reserved.
+
+#include "MovieGraphCustomizationUtils.h"
+
+#include "DetailCategoryBuilder.h"
+#include "DetailLayoutBuilder.h"
+#include "Graph/MovieGraphConfig.h"
+#include "IDetailGroup.h"
+#include "MovieJobVariableAssignmentContainer.h"
+
+namespace UE::MovieRenderPipelineEditor::Private
+{
+	void AddVariableAssignments(TArray<TObjectPtr<UMovieJobVariableAssignmentContainer>>& InVariableAssignments, IDetailCategoryBuilder& InCategory, IDetailLayoutBuilder* InDetailBuilder)
+	{
+		// Add a sub-category for each graph (including subgraphs). Each entry in the array represents the assignments for one graph.
+		for (TObjectPtr<UMovieJobVariableAssignmentContainer>& VariableAssignment : InVariableAssignments)
+		{
+			// Skip if the graph associated with this container has no variables in it
+			if (VariableAssignment->GetNumAssignments() <= 0)
+			{
+				continue;
+			}
+
+			// If the graph can be found, display its variable assignments under its own category (group)
+			TSoftObjectPtr<UMovieGraphConfig> SoftGraphConfig = VariableAssignment->GetGraphConfig();
+			if (const UMovieGraphConfig* GraphConfig = SoftGraphConfig.Get())
+			{
+				constexpr bool bForAdvanced = false;
+				constexpr bool bStartExpanded = true;
+				IDetailGroup& GraphGroup = InCategory.AddGroup(GraphConfig->GetFName(), FText::FromString(GraphConfig->GetName()), bForAdvanced, bStartExpanded);
+
+				// "Value" is private so we can't use GET_MEMBER_NAME_CHECKED unfortunately
+				TSharedPtr<IPropertyHandle> ValueProperty = InDetailBuilder->AddObjectPropertyData({VariableAssignment}, FName("Value"));
+				GraphGroup.AddPropertyRow(ValueProperty.ToSharedRef());
+
+				// Un-hide the category if it's currently visible
+				InCategory.SetCategoryVisibility(true);
+			}
+		}
+	}
+}
diff --git a/Engine/Plugins/MovieScene/MovieRenderPipeline/Source/MovieRenderPipelineEditor/Private/Customizations/Graph/MovieGraphCustomizationUtils.h b/Engine/Plugins/MovieScene/MovieRenderPipeline/Source/MovieRenderPipelineEditor/Private/Customizations/Graph/MovieGraphCustomizationUtils.h
new file mode 100644
--- /dev/null
+++ b/Engine/Plugins/MovieScene/MovieRenderPipeline/Source/MovieRenderPipelineEditor/Private/Customizations/Graph/MovieGraphCustomizationUtils.h
@@ -0,0 +1,17 @@
+﻿// Copyright Epic Games, Inc. All Rights Reserved.
+
+#pragma once
+
+#include "CoreMinimal.h"
+#include "DetailCategoryBuilder.h"
+
+class UMovieJobVariableAssignmentContainer;
+
+namespace UE::MovieRenderPipelineEditor::Private
+{
+	/**
+	 * Adds InVariableAssignments to the given InCategory. For each graph that is used within InVariableAssignments, a dedicated sub-category (named
+	 * the name of the graph) is added under the given category.
+	 */
+	void AddVariableAssignments(TArray<TObjectPtr<UMovieJobVariableAssignmentContainer>>& InVariableAssignments, IDetailCategoryBuilder& InCategory, IDetailLayoutBuilder* InDetailBuilder);
+}
diff --git a/Engine/Plugins/MovieScene/MovieRenderPipeline/Source/MovieRenderPipelineEditor/Private/Customizations/Graph/MovieGraphQuickRenderSettingsCustomization.h b/Engine/Plugins/MovieScene/MovieRenderPipeline/Source/MovieRenderPipelineEditor/Private/Customizations/Graph/MovieGraphQuickRenderSettingsCustomization.h
new file mode 100644
--- /dev/null
+++ b/Engine/Plugins/MovieScene/MovieRenderPipeline/Source/MovieRenderPipelineEditor/Private/Customizations/Graph/MovieGraphQuickRenderSettingsCustomization.h
@@ -0,0 +1,122 @@
+﻿// Copyright Epic Games, Inc. All Rights Reserved.
+
+#pragma once
+
+#include "DetailCategoryBuilder.h"
+#include "DetailLayoutBuilder.h"
+#include "Graph/MovieGraphQuickRenderSettings.h"
+#include "IDetailCustomization.h"
+#include "IDetailGroup.h"
+#include "MovieGraphCustomizationUtils.h"
+#include "MoviePipelineQueue.h"
+#include "MoviePipelineUtils.h"
+#include "PropertyHandle.h"
+#include "UObject/Package.h"
+
+#define LOCTEXT_NAMESPACE "MoviePipelineQuickRenderSettings"
+
+/** Customize how properties for QuickRenderSettings appear in the details panel. */
+class FMovieGraphQuickRenderSettingsCustomization : public IDetailCustomization
+{
+public:
+	static TSharedRef<IDetailCustomization> MakeInstance()
+	{
+		return MakeShared<FMovieGraphQuickRenderSettingsCustomization>();
+	}
+
+	virtual ~FMovieGraphQuickRenderSettingsCustomization() override
+	{
+	}
+
+	virtual void PendingDelete() override
+	{
+		// Unregister delegates. It's important to do this in PendingDelete() vs the destructor because the destructor is not called before the next
+		// details panel is created (via ForceRefreshDetails()), leading to an exponential increase in the number of delegates registered.
+		
+		UPackage::PackageSavedWithContextEvent.RemoveAll(this);
+
+		if (QuickRenderSettings.IsValid())
+		{
+			QuickRenderSettings->OnGraphChangedDelegate.RemoveAll(this);
+		}
+	}
+
+protected:
+	//~ Begin IDetailCustomization interface
+	virtual void CustomizeDetails(const TSharedPtr<IDetailLayoutBuilder>& InDetailBuilder) override
+	{
+		DetailBuilder = InDetailBuilder.Get();
+		CustomizeDetails(*InDetailBuilder);
+	}
+
+	void RefreshLayout(const FString&, UPackage*, FObjectPostSaveContext) const
+	{
+		DetailBuilder->ForceRefreshDetails();
+
+		// Refresh the variable assignments (since the variables within the graph may have changed, thus the variable assignments need to be updated)
+		MoviePipeline::RefreshVariableAssignments(QuickRenderSettings->GraphPreset.LoadSynchronous(), QuickRenderSettings->GraphVariableAssignments, QuickRenderSettings.Get());
+	}
+
+	void RefreshLayout() const
+	{
+		// Note: The settings object will take care of refreshing variable assignments on itself in the case where the graph preset asset is updated to be a different asset
+		DetailBuilder->ForceRefreshDetails();
+	}
+
+	virtual void CustomizeDetails(IDetailLayoutBuilder& InDetailBuilder) override
+	{
+		// Refresh the customization every time a save happens. Use this opportunity to update the variables in the UI. We could update the UI before
+		// a save occurs, but this would be very difficult to get right when multiple subgraphs are involved.
+		UPackage::PackageSavedWithContextEvent.AddSP(this, &FMovieGraphQuickRenderSettingsCustomization::RefreshLayout);
+		
+		TArray<TWeakObjectPtr<UObject>> ObjectsBeingCustomized;
+		InDetailBuilder.GetObjectsBeingCustomized(ObjectsBeingCustomized);
+
+		if (ObjectsBeingCustomized.IsEmpty() || (ObjectsBeingCustomized.Num() > 1))
+		{
+			return;
+		}
+
+		QuickRenderSettings = Cast<UMovieGraphQuickRenderSettings>(ObjectsBeingCustomized[0]);
+		if (!QuickRenderSettings.IsValid())
+		{
+			return;
+		}
+
+		// Hide the original assignments property (since it presents an asset picker).
+		const TSharedRef<IPropertyHandle> GraphVariableAssignmentsProperty =
+			InDetailBuilder.GetProperty(TEXT("GraphVariableAssignments"), UMovieGraphQuickRenderSettings::StaticClass());
+		InDetailBuilder.HideProperty(GraphVariableAssignmentsProperty);
+
+		// Refresh the UI if the graph preset changes (so the new variable assignments are displayed).
+		QuickRenderSettings->OnGraphChangedDelegate.AddSP(this, &FMovieGraphQuickRenderSettingsCustomization::RefreshLayout);
+
+		// Set up the category for variable assignments.
+		IDetailCategoryBuilder& PrimaryGraphVariablesCategory = InDetailBuilder.EditCategory(
+			"PrimaryGraphVariables", LOCTEXT("PrimaryGraphVariablesCategory", "Primary Graph Variables"));
+
+		// Get the other categories for sort-order purposes.
+		IDetailCategoryBuilder& ConfigurationCategory = InDetailBuilder.EditCategory("Configuration");
+		IDetailCategoryBuilder& QuickRenderCategory = InDetailBuilder.EditCategory("Quick Render");
+
+		// Set the variable assignments category as hidden by default, and add the variable assignments. Individual categories will be made visible if variables are added under them.
+		PrimaryGraphVariablesCategory.SetCategoryVisibility(false);
+		UE::MovieRenderPipelineEditor::Private::AddVariableAssignments(QuickRenderSettings->GraphVariableAssignments, PrimaryGraphVariablesCategory, DetailBuilder);
+
+		// Give the categories a specific ordering.
+		int32 SortOrder = 0;
+		ConfigurationCategory.SetSortOrder(SortOrder);
+		QuickRenderCategory.SetSortOrder(++SortOrder);
+		PrimaryGraphVariablesCategory.SetSortOrder(++SortOrder);
+	}
+	//~ End IDetailCustomization interface
+
+private:
+	/** The details builder associated with the customization. */
+	IDetailLayoutBuilder* DetailBuilder = nullptr;
+
+	/** The quick render settings being displayed. */
+	TWeakObjectPtr<UMovieGraphQuickRenderSettings> QuickRenderSettings = nullptr;
+};
+
+#undef LOCTEXT_NAMESPACE
\ No newline at end of file
diff --git a/Engine/Plugins/MovieScene/MovieRenderPipeline/Source/MovieRenderPipelineEditor/Private/Customizations/JobCustomization.h b/Engine/Plugins/MovieScene/MovieRenderPipeline/Source/MovieRenderPipelineEditor/Private/Customizations/JobCustomization.h
--- a/Engine/Plugins/MovieScene/MovieRenderPipeline/Source/MovieRenderPipelineEditor/Private/Customizations/JobCustomization.h
+++ b/Engine/Plugins/MovieScene/MovieRenderPipeline/Source/MovieRenderPipelineEditor/Private/Customizations/JobCustomization.h
@@ -4,6 +4,7 @@
 
 #include "DetailCategoryBuilder.h"
 #include "DetailLayoutBuilder.h"
+#include "Graph/MovieGraphCustomizationUtils.h"
 #include "IDetailCustomization.h"
 #include "IDetailGroup.h"
 #include "MoviePipelineQueue.h"
@@ -143,12 +144,12 @@ protected:
 		
 		if (bIsShot)
 		{
-			AddVariableAssignments(SelectedShot->GetGraphVariableAssignments(), ShotGraphVariablesCategory);
-			AddVariableAssignments(SelectedShot->GetPrimaryGraphVariableAssignments(), PrimaryGraphVariablesShotOverridesCategory);
+			UE::MovieRenderPipelineEditor::Private::AddVariableAssignments(SelectedShot->GetGraphVariableAssignments(), ShotGraphVariablesCategory, DetailBuilder);
+			UE::MovieRenderPipelineEditor::Private::AddVariableAssignments(SelectedShot->GetPrimaryGraphVariableAssignments(), PrimaryGraphVariablesShotOverridesCategory, DetailBuilder);
 		}
 		else
 		{
-			AddVariableAssignments(SelectedJob->GetGraphVariableAssignments(), PrimaryGraphVariablesCategory);
+			UE::MovieRenderPipelineEditor::Private::AddVariableAssignments(SelectedJob->GetGraphVariableAssignments(), PrimaryGraphVariablesCategory, DetailBuilder);
 		}
 
 		IDetailCategoryBuilder& MovieRenderPipelineCategory = InDetailBuilder.EditCategory("Movie Render Pipeline");
@@ -167,36 +168,6 @@ protected:
 	}
 	//~ End IDetailCustomization interface
 
-private:
-	void AddVariableAssignments(TArray<TObjectPtr<UMovieJobVariableAssignmentContainer>>& VariableAssignments, IDetailCategoryBuilder& InCategory) const
-	{
-		// Add a sub-category for each graph (including subgraphs). Each entry in the array represents the assignments for one graph.
-		for (TObjectPtr<UMovieJobVariableAssignmentContainer>& VariableAssignment : VariableAssignments)
-		{
-			// Skip if the graph associated with this container has no variables in it
-			if (VariableAssignment->GetNumAssignments() <= 0)
-			{
-				continue;
-			}
-
-			// If the graph can be found, display its variable assignments under its own category (group)
-			TSoftObjectPtr<UMovieGraphConfig> SoftGraphConfig = VariableAssignment->GetGraphConfig();
-			if (const UMovieGraphConfig* GraphConfig = SoftGraphConfig.Get())
-			{
-				constexpr bool bForAdvanced = false;
-				constexpr bool bStartExpanded = true;
-				IDetailGroup& GraphGroup = InCategory.AddGroup(GraphConfig->GetFName(), FText::FromString(GraphConfig->GetName()), bForAdvanced, bStartExpanded);
-
-				// "Value" is private so we can't use GET_MEMBER_NAME_CHECKED unfortunately
-				TSharedPtr<IPropertyHandle> ValueProperty = DetailBuilder->AddObjectPropertyData({VariableAssignment}, FName("Value"));
-				GraphGroup.AddPropertyRow(ValueProperty.ToSharedRef());
-
-				// Un-hide the category if it's currently visible
-				InCategory.SetCategoryVisibility(true);
-			}
-		}
-	}
-
 private:
 	/** The details builder associated with the customization. */
 	IDetailLayoutBuilder* DetailBuilder = nullptr;
diff --git a/Engine/Plugins/MovieScene/MovieRenderPipeline/Source/MovieRenderPipelineEditor/Private/Graph/MovieGraphQuickRender.cpp b/Engine/Plugins/MovieScene/MovieRenderPipeline/Source/MovieRenderPipelineEditor/Private/Graph/MovieGraphQuickRender.cpp
--- a/Engine/Plugins/MovieScene/MovieRenderPipeline/Source/MovieRenderPipelineEditor/Private/Graph/MovieGraphQuickRender.cpp
+++ b/Engine/Plugins/MovieScene/MovieRenderPipeline/Source/MovieRenderPipelineEditor/Private/Graph/MovieGraphQuickRender.cpp
@@ -40,7 +40,7 @@ static TAutoConsoleVariable<bool> CVarMoviePipelineEnableQuickRender(
 	false,
 	TEXT("Enables or disables whether the Quick Render feature is exposed in the UI (defaults to off)."));
 
-void FMovieGraphQuickRender::BeginQuickRender(const FMovieGraphQuickRenderSettings& InQuickRenderSettings)
+void FMovieGraphQuickRender::BeginQuickRender(const UMovieGraphQuickRenderSettings& InQuickRenderSettings)
 {
 	UMoviePipelineQueueSubsystem* Subsystem = GEditor->GetEditorSubsystem<UMoviePipelineQueueSubsystem>();
 	check(Subsystem);
@@ -136,7 +136,7 @@ void FMovieGraphQuickRender::BeginQuickRender(const FMovieGraphQuickRenderSettin
 	UMoviePipelinePIEExecutor* Executor = Cast<UMoviePipelinePIEExecutor>(Subsystem->RenderQueueInstanceWithExecutor(TemporaryQueue.Get(), ExecutorClass));
 
 	// View the job's rendered frame(s)/video(s) if requested
-	Executor->OnIndividualJobWorkFinished().AddLambda([InQuickRenderSettings](FMoviePipelineOutputData OutputData)
+	Executor->OnIndividualJobWorkFinished().AddLambda([&InQuickRenderSettings](FMoviePipelineOutputData OutputData)
 	{
 		if (InQuickRenderSettings.PostRenderBehavior == EMovieGraphQuickRenderPostRenderActionType::PlayRenderOutput)
 		{
@@ -189,19 +189,13 @@ UMovieGraphConfig* FMovieGraphQuickRender::GetQuickRenderGraph(UMovieGraphConfig
 	return TemporaryGraph.Get();
 }
 
-void FMovieGraphQuickRender::ApplyQuickRenderUpdatesToDuplicateGraph(const FMovieGraphQuickRenderSettings& InQuickRenderSettings, UWorld* InEditorWorld)
+void FMovieGraphQuickRender::ApplyQuickRenderUpdatesToDuplicateGraph(const UMovieGraphQuickRenderSettings& InQuickRenderSettings, UWorld* InEditorWorld)
 {
 	check(TemporaryGraph.IsValid());
 
 	// Global Output Settings
 	{
 		UMovieGraphGlobalOutputSettingNode* GlobalOutputSettings = Cast<UMovieGraphGlobalOutputSettingNode>(InjectNodeIntoGlobals(UMovieGraphGlobalOutputSettingNode::StaticClass()));
-		
-		if (!InQuickRenderSettings.OutputDirectory.Path.IsEmpty())
-		{
-			GlobalOutputSettings->bOverride_OutputDirectory = true;
-			GlobalOutputSettings->OutputDirectory = InQuickRenderSettings.OutputDirectory;
-		}
 
 		// If the frame range type is StartEnd, the frame range doesn't need to be modified
 		if ((InQuickRenderSettings.FrameRangeType == EMovieGraphQuickRenderFrameRangeType::SelectionRange) || (InQuickRenderSettings.FrameRangeType == EMovieGraphQuickRenderFrameRangeType::Custom))
@@ -219,12 +213,6 @@ void FMovieGraphQuickRender::ApplyQuickRenderUpdatesToDuplicateGraph(const FMovi
 			GlobalOutputSettings->bOverride_CustomPlaybackRangeEndFrame = true;
 			GlobalOutputSettings->CustomPlaybackRangeEndFrame = EndFrame;
 		}
-
-		if (InQuickRenderSettings.Resolution.IsValid())
-		{
-			GlobalOutputSettings->bOverride_OutputResolution = true;
-			GlobalOutputSettings->OutputResolution = InQuickRenderSettings.Resolution;
-		}
 	}
 
 	// Some updates are only done if there's an Apply Viewport Look node present
diff --git a/Engine/Plugins/MovieScene/MovieRenderPipeline/Source/MovieRenderPipelineEditor/Private/MoviePipelineQuickRenderMenu.cpp b/Engine/Plugins/MovieScene/MovieRenderPipeline/Source/MovieRenderPipelineEditor/Private/MoviePipelineQuickRenderMenu.cpp
--- a/Engine/Plugins/MovieScene/MovieRenderPipeline/Source/MovieRenderPipelineEditor/Private/MoviePipelineQuickRenderMenu.cpp
+++ b/Engine/Plugins/MovieScene/MovieRenderPipeline/Source/MovieRenderPipelineEditor/Private/MoviePipelineQuickRenderMenu.cpp
@@ -1,21 +1,25 @@
-﻿// Copyright Epic Games, Inc. All Rights Reserved.
+// Copyright Epic Games, Inc. All Rights Reserved.
 
 #include "MoviePipelineQuickRenderMenu.h"
 
+#include "Customizations/Graph/MovieGraphQuickRenderSettingsCustomization.h"
 #include "Editor.h"
 #include "Framework/Application/SlateApplication.h"
 #include "Framework/Docking/TabManager.h"
 #include "Graph/MovieGraphQuickRender.h"
 #include "Graph/MovieRenderGraphEditorSettings.h"
 #include "HAL/IConsoleManager.h"
-#include "MoviePipelineQueueSubsystem.h"
 #include "MoviePipelineUtils.h"
 #include "MovieRenderPipelineCoreModule.h"
 #include "MovieRenderPipelineEditorModule.h"
 #include "MovieRenderPipelineEditorUtils.h"
 #include "MovieRenderPipelineSettings.h"
+#include "PropertyEditorModule.h"
+#include "SActionButton.h"
+#include "Styling/ToolBarStyle.h"
 #include "ToolMenu.h"
 #include "ToolMenus.h"
+#include "Widgets/Input/SComboBox.h"
 #include "Widgets/Text/SRichTextBlock.h"
 
 class UMovieGraphVariable;
@@ -86,21 +90,26 @@ void FMoviePipelineQuickRenderMenu::RemoveQuickRenderButtonToolMenu(UToolMenu* I
 	}
 }
 
-void FMoviePipelineQuickRenderMenu::QuickRenderButtonPressed()
+void FMoviePipelineQuickRenderMenu::LoadQuickRenderSettings()
 {
-	const UMovieRenderGraphEditorSettings* EditorSettings = GetDefault<UMovieRenderGraphEditorSettings>();
+	QuickRenderSettings = TStrongObjectPtr(UMovieGraphQuickRenderSettings::GetSavedQuickRenderSettings());
+}
 
-	FDirectoryPath OutputDirectory = FDirectoryPath();
-	OutputDirectory.Path = FString(TEXT("{project_dir}/Saved/MovieRenders/QuickRender"));
-	
-	// TODO: These should come from the Quick Render Settings dialog (not finished yet)
-	FMovieGraphQuickRenderSettings QuickRenderSettings;
-	QuickRenderSettings.PostRenderBehavior = EMovieGraphQuickRenderPostRenderActionType::PlayRenderOutput;
-	QuickRenderSettings.OutputDirectory = OutputDirectory;
+void FMoviePipelineQuickRenderMenu::QuickRenderButtonPressed()
+{
+	LoadQuickRenderSettings();
 
-	if (EditorSettings->QuickRenderSettings.QuickAccessButtonMode == EMovieGraphQuickRenderButtonMode::QuickRender)
+	if (QuickRenderSettings->QuickAccessButtonMode == EMovieGraphQuickRenderButtonMode::QuickRender)
 	{
-		FMovieGraphQuickRender::BeginQuickRender(QuickRenderSettings);
+		// The settings window will have its own Render button, so don't perform the render if showing the settings first
+		if (QuickRenderSettings->bShowSettingsBeforeQuickRender)
+		{
+			OpenQuickRenderSettingsWindow(FToolMenuContext());
+		}
+		else
+		{
+			FMovieGraphQuickRender::BeginQuickRender(*QuickRenderSettings);
+		}
 	}
 	else
 	{
@@ -115,6 +124,8 @@ TSharedRef<SWidget> FMoviePipelineQuickRenderMenu::GenerateQuickRenderOptionsMen
 {
 	static const FName MenuName("MoviePipeline.QuickRenderOptionsMenu");
 
+	LoadQuickRenderSettings();
+
 	if (!UToolMenus::Get()->IsMenuRegistered(MenuName))
 	{
 		UToolMenu* Menu = UToolMenus::Get()->RegisterMenu(MenuName);
@@ -131,20 +142,60 @@ TSharedRef<SWidget> FMoviePipelineQuickRenderMenu::GenerateQuickRenderOptionsMen
 
 void FMoviePipelineQuickRenderMenu::OpenQuickRenderSettingsWindow(const FToolMenuContext& InToolMenuContext)
 {
+	LoadQuickRenderSettings();
+	
 	if (const TSharedPtr<SWindow> SettingsWindow = WeakQuickRenderSettingsWindow.Pin())
 	{
 		SettingsWindow->DrawAttention(FWindowDrawAttentionParameters());
 		return;
 	}
 
+	// Init the settings categories
+	QuickRenderSettingsCategories.Empty();
+	const UEnum* SettingsCategoryEnum = StaticEnum<EMovieGraphQuickRenderSettingsCategory>();
+	for (int32 EnumIndex = 0; EnumIndex < SettingsCategoryEnum->NumEnums() - 1; ++EnumIndex)	// -1 to avoid MAX
+	{
+		QuickRenderSettingsCategories.Add(MakeShared<FName>(SettingsCategoryEnum->GetNameByIndex(EnumIndex)));
+	}
+	
+	// Update the graph preset's variable assignments in case the graph was updated since the last time the window opened
+	MoviePipeline::RefreshVariableAssignments(QuickRenderSettings->GraphPreset.LoadSynchronous(), QuickRenderSettings->GraphVariableAssignments, QuickRenderSettings.Get());
+
+	FPropertyEditorModule& PropertyEditorModule = FModuleManager::Get().LoadModuleChecked<FPropertyEditorModule>("PropertyEditor");
+	FDetailsViewArgs DetailsViewArgs;
+	DetailsViewArgs.bShowPropertyMatrixButton = false;
+	DetailsViewArgs.bCustomNameAreaLocation = false;
+	DetailsViewArgs.NameAreaSettings = FDetailsViewArgs::HideNameArea;
+	DetailsViewArgs.bAllowSearch = false;
+	DetailsViewArgs.ViewIdentifier = "MoviePipelineQuickRenderSettings";
+	DetailsViewArgs.bLockable = false;
+
+	// Create the details panel and display the quick render settings
+	const TSharedPtr<IDetailsView> DetailsView = PropertyEditorModule.CreateDetailView(DetailsViewArgs);
+	DetailsView->RegisterInstancedCustomPropertyLayout(
+		UMovieGraphQuickRenderSettings::StaticClass(),
+		FOnGetDetailCustomizationInstance::CreateStatic(&FMovieGraphQuickRenderSettingsCustomization::MakeInstance));
+	DetailsView->SetObject(QuickRenderSettings.Get());
+	
 	const TSharedRef<SWindow> QuickRenderSettingsWindow = SNew(SWindow)
 		.Title(LOCTEXT("QuickRenderSettingsWindow_Title", "Quick Render"))
 		.SupportsMaximize(false)
 		.ClientSize(FVector2D(600.f, 510.f))
 		.Content()
 		[
-			SNew(STextBlock)
-			.Text(FText::FromString(TEXT("Settings Window")))
+			SNew(SVerticalBox)
+			
+			+ SVerticalBox::Slot()
+			.AutoHeight()
+			.HAlign(HAlign_Fill)
+			[
+				GenerateQuickRenderSettingsWindowToolbar()
+			]
+
+			+ SVerticalBox::Slot()
+			[
+				DetailsView.ToSharedRef()
+			]
 		];
 
 	if (const TSharedPtr<SWindow> RootWindow = FGlobalTabmanager::Get()->GetRootWindow())
@@ -159,6 +210,67 @@ void FMoviePipelineQuickRenderMenu::OpenQuickRenderSettingsWindow(const FToolMen
 	WeakQuickRenderSettingsWindow = QuickRenderSettingsWindow;
 }
 
+TSharedRef<SWidget> FMoviePipelineQuickRenderMenu::GenerateQuickRenderSettingsWindowToolbar()
+{
+	FToolBarBuilder ToolbarBuilder(TSharedPtr<const FUICommandList>(), FMultiBoxCustomization::None);
+	ToolbarBuilder.BeginSection("Quick Render");
+	ToolbarBuilder.AddWidget(
+		SNew(SComboBox<TSharedPtr<FName>>)
+		.OptionsSource(&FMoviePipelineQuickRenderMenu::QuickRenderSettingsCategories)
+		.IsEnabled(false)	// TODO: Enable once settings categories are implemented
+		.OnGenerateWidget_Lambda([](TSharedPtr<FName> InItem)
+		{
+			return
+				SNew(STextBlock)
+				.Text_Lambda([InItem]()
+				{
+					const int64 SettingsCategory = StaticEnum<EMovieGraphQuickRenderSettingsCategory>()->GetValueByName(*InItem);
+
+					return StaticEnum<EMovieGraphQuickRenderSettingsCategory>()->GetDisplayNameTextByValue(SettingsCategory);
+				});
+		})
+		.OnSelectionChanged_Lambda([](TSharedPtr<FName> InSelectedItem, ESelectInfo::Type)
+		{
+			QuickRenderSettings->SettingsCategory = static_cast<EMovieGraphQuickRenderSettingsCategory>(
+				StaticEnum<EMovieGraphQuickRenderSettingsCategory>()->GetValueByName(*InSelectedItem));
+		})
+		[
+			SNew(STextBlock)
+			.Text_Lambda([]()
+			{
+				return StaticEnum<EMovieGraphQuickRenderSettingsCategory>()->GetDisplayNameTextByValue(static_cast<int64>(QuickRenderSettings->SettingsCategory));
+			})
+		]);
+
+	ToolbarBuilder.AddWidget(
+		SNew(SActionButton)
+		.Text(LOCTEXT("QuickRenderSettingsWindow_ApplyModeButtonText", "Apply"))
+		.Icon(FAppStyle::Get().GetBrush("MaterialEditor.Apply"))
+		.ActionButtonType(EActionButtonType::Simple)
+		.IsEnabled(false)	// TODO: Enable once settings categories are implemented
+		.OnClicked_Lambda([]()
+		{
+			// TODO: Apply the mode setting
+			return FReply::Handled();
+		}));
+
+	ToolbarBuilder.AddWidget(
+		SNew(SActionButton)
+		.Text(LOCTEXT("QuickRenderSettingsWindow_QuickRenderButtonText", "Quick Render"))
+		.Icon(FAppStyle::Get().GetBrush("LevelEditor.OpenCinematic"))
+		.ActionButtonType(EActionButtonType::Primary)
+		.OnClicked_Lambda([]()
+		{
+			FMovieGraphQuickRender::BeginQuickRender(*QuickRenderSettings);
+			return FReply::Handled();
+		}),
+		NAME_None, false, HAlign_Right);
+
+	ToolbarBuilder.EndSection();
+
+	return ToolbarBuilder.MakeWidget();
+}
+
 void FMoviePipelineQuickRenderMenu::GenerateModesMenuSection(UToolMenu* InMenu)
 {
 	FToolMenuSection& Section = InMenu->AddSection("MoviePipelineQuickRenderModes", LOCTEXT("QuickRenderModesSection", "Modes"));
@@ -168,15 +280,13 @@ void FMoviePipelineQuickRenderMenu::GenerateModesMenuSection(UToolMenu* InMenu)
 		FToolUIAction NewAction;
 		NewAction.ExecuteAction = FToolMenuExecuteAction::CreateLambda([CheckedValue](const FToolMenuContext&)
 		{
-			UMovieRenderGraphEditorSettings* EditorSettings = GetMutableDefault<UMovieRenderGraphEditorSettings>();
-			EditorSettings->QuickRenderSettings.QuickAccessButtonMode = CheckedValue;
-			EditorSettings->PostEditChange();
-			EditorSettings->SaveConfig();
+			QuickRenderSettings->QuickAccessButtonMode = CheckedValue;
+			QuickRenderSettings->PostEditChange();
 		});
 		NewAction.CanExecuteAction = FToolMenuCanExecuteAction::CreateLambda([](const FToolMenuContext&) { return true; });
 		NewAction.GetActionCheckState = FToolMenuGetActionCheckState::CreateLambda([CheckedValue](const FToolMenuContext&)
 		{
-			return (GetDefault<UMovieRenderGraphEditorSettings>()->QuickRenderSettings.QuickAccessButtonMode == CheckedValue)
+			return (QuickRenderSettings->QuickAccessButtonMode == CheckedValue)
 				? ECheckBoxState::Checked
 				: ECheckBoxState::Unchecked;
 		});
@@ -207,15 +317,13 @@ void FMoviePipelineQuickRenderMenu::GenerateQuickRenderMenuSection(UToolMenu* In
 		FToolUIAction NewAction;
 		NewAction.ExecuteAction = FToolMenuExecuteAction::CreateLambda([CheckedValue](const FToolMenuContext&)
 		{
-			UMovieRenderGraphEditorSettings* EditorSettings = GetMutableDefault<UMovieRenderGraphEditorSettings>();
-			EditorSettings->QuickRenderSettings.QuickRenderMode = CheckedValue;
-			EditorSettings->PostEditChange();
-			EditorSettings->SaveConfig();
+			QuickRenderSettings->QuickRenderMode = CheckedValue;
+			QuickRenderSettings->PostEditChange();
 		});
 		NewAction.CanExecuteAction = FToolMenuCanExecuteAction::CreateLambda([](const FToolMenuContext&) { return true; });
 		NewAction.GetActionCheckState = FToolMenuGetActionCheckState::CreateLambda([CheckedValue](const FToolMenuContext&)
 		{
-			return (GetDefault<UMovieRenderGraphEditorSettings>()->QuickRenderSettings.QuickRenderMode == CheckedValue)
+			return (QuickRenderSettings->QuickRenderMode == CheckedValue)
 				? ECheckBoxState::Checked
 				: ECheckBoxState::Unchecked;
 		});
@@ -271,36 +379,39 @@ void FMoviePipelineQuickRenderMenu::GenerateQuickRenderConfigurationMenuSection(
 			.TextStyle(FAppStyle::Get(), "NormalText.Subdued")
 			.Text_Lambda([]()
 			{
-				const UMovieRenderGraphEditorSettings* EditorSettings = GetDefault<UMovieRenderGraphEditorSettings>();
-				const FMovieGraphQuickRenderSettings& QuickRenderSettings = EditorSettings->QuickRenderSettings;
-
-				const FText GraphName = QuickRenderSettings.GraphPreset.IsValid()
-					? FText::FromString(QuickRenderSettings.GraphPreset->GetName())
+				// Generate "Configuration" text
+				const FText GraphName = !QuickRenderSettings->GraphPreset.IsNull()
+					? FText::FromString(QuickRenderSettings->GraphPreset.GetAssetName())
 					: LOCTEXT("QuickRenderConfigMenu_InvalidGraph", "Invalid");
 				
-				const FText AfterRender = (QuickRenderSettings.PostRenderBehavior == EMovieGraphQuickRenderPostRenderActionType::PlayRenderOutput)
+				// Generate "After Render" text
+				const FText AfterRender = (QuickRenderSettings->PostRenderBehavior == EMovieGraphQuickRenderPostRenderActionType::PlayRenderOutput)
 					? LOCTEXT("QuickRenderConfigMenu_PlayRenderOutput", "Play Render Output")
 					: LOCTEXT("QuickRenderConfigMenu_DoNothing", "Do Nothing");
 				
+				// Generate "Viewport Look" text
+				const EMovieGraphQuickRenderViewportLookFlags ViewportLookFlags = static_cast<EMovieGraphQuickRenderViewportLookFlags>(QuickRenderSettings->ViewportLookFlags);
 				TArray<FText> ViewportFlags;
-				if (EnumHasAnyFlags(QuickRenderSettings.ViewportLookFlags, EMovieGraphQuickRenderViewportLookFlags::Ocio))
+				if (EnumHasAnyFlags(ViewportLookFlags, EMovieGraphQuickRenderViewportLookFlags::Ocio))
 				{
 					ViewportFlags.Add(LOCTEXT("QuickRenderConfigMenu_OCIO", "OCIO"));
 				}
-				if (EnumHasAnyFlags(QuickRenderSettings.ViewportLookFlags, EMovieGraphQuickRenderViewportLookFlags::ShowFlags))
+				if (EnumHasAnyFlags(ViewportLookFlags, EMovieGraphQuickRenderViewportLookFlags::ShowFlags))
 				{
 					ViewportFlags.Add(LOCTEXT("QuickRenderConfigMenu_ShowFlags", "Show Flags"));
 				}
-				if (EnumHasAnyFlags(QuickRenderSettings.ViewportLookFlags, EMovieGraphQuickRenderViewportLookFlags::ViewMode))
+				if (EnumHasAnyFlags(ViewportLookFlags, EMovieGraphQuickRenderViewportLookFlags::ViewMode))
 				{
 					ViewportFlags.Add(LOCTEXT("QuickRenderConfigMenu_ViewMode", "View Mode"));
 				}
-				if (QuickRenderSettings.ViewportLookFlags == EMovieGraphQuickRenderViewportLookFlags::None)
+				if (ViewportLookFlags == EMovieGraphQuickRenderViewportLookFlags::None)
 				{
 					ViewportFlags.Add(LOCTEXT("QuickRenderConfigMenu_DontApply", "Don't Apply"));
 				}
+				
+				// Generate "Sequencer Frame Range" text
+				const FText SequencerFrameRange = StaticEnum<EMovieGraphQuickRenderFrameRangeType>()->GetDisplayNameTextByValue(static_cast<int64>(QuickRenderSettings->FrameRangeType));
 
-				// TODO: "Sequencer Frame Range" needs to be sourced from data
 				// TODO: Add in graph variables summary as well (although this might make a summary that's too long in some cases)
 				FText ConfigurationSummaryText = FText::Format(
 					LOCTEXT("QuickRenderConfigurationSummary", "Configuration: <RichTextBlock.BoldHighlight>{0}</>\nAfter Render: {1}\nApply Viewport Look: {2}\nSequencer Frame Range: {3}"),
@@ -308,7 +419,7 @@ void FMoviePipelineQuickRenderMenu::GenerateQuickRenderConfigurationMenuSection(
 						GraphName,
 						AfterRender,
 						FText::Join(LOCTEXT("CommaDelim", ", "), ViewportFlags),
-						FText::FromString(FString(TEXT("Section"))),
+						SequencerFrameRange,
 					});
 
 				return ConfigurationSummaryText;
@@ -343,15 +454,13 @@ void FMoviePipelineQuickRenderMenu::GenerateSettingsMenuSection(UToolMenu* InMen
 	FToolUIAction ShowSettingsAction;
         ShowSettingsAction.ExecuteAction = FToolMenuExecuteAction::CreateLambda([](const FToolMenuContext&)
         {
-            UMovieRenderGraphEditorSettings* EditorSettings = GetMutableDefault<UMovieRenderGraphEditorSettings>();
-			EditorSettings->QuickRenderSettings.bShowSettingsBeforeQuickRender = !EditorSettings->QuickRenderSettings.bShowSettingsBeforeQuickRender;
-			EditorSettings->PostEditChange();
-			EditorSettings->SaveConfig();
+			QuickRenderSettings->bShowSettingsBeforeQuickRender = !QuickRenderSettings->bShowSettingsBeforeQuickRender;
+        	QuickRenderSettings->PostEditChange();
         });
         ShowSettingsAction.CanExecuteAction = FToolMenuCanExecuteAction::CreateLambda([](const FToolMenuContext&) { return true; });
         ShowSettingsAction.GetActionCheckState = FToolMenuGetActionCheckState::CreateLambda([](const FToolMenuContext&)
         {
-            return GetDefault<UMovieRenderGraphEditorSettings>()->QuickRenderSettings.bShowSettingsBeforeQuickRender
+            return QuickRenderSettings->bShowSettingsBeforeQuickRender
 				? ECheckBoxState::Checked
 				: ECheckBoxState::Unchecked;
         });
diff --git a/Engine/Plugins/MovieScene/MovieRenderPipeline/Source/MovieRenderPipelineEditor/Private/MoviePipelineQuickRenderMenu.h b/Engine/Plugins/MovieScene/MovieRenderPipeline/Source/MovieRenderPipelineEditor/Private/MoviePipelineQuickRenderMenu.h
--- a/Engine/Plugins/MovieScene/MovieRenderPipeline/Source/MovieRenderPipelineEditor/Private/MoviePipelineQuickRenderMenu.h
+++ b/Engine/Plugins/MovieScene/MovieRenderPipeline/Source/MovieRenderPipelineEditor/Private/MoviePipelineQuickRenderMenu.h
@@ -3,6 +3,7 @@
 #pragma once
 
 #include "CoreMinimal.h"
+#include "Graph/MovieGraphQuickRenderSettings.h"
 
 struct FToolMenuContext;
 class SWidget;
@@ -25,10 +26,13 @@ public:
 	static void RemoveQuickRenderButtonToolMenu(UToolMenu* InMenu);
 
 private:
+	static void LoadQuickRenderSettings();
+	
 	static void QuickRenderButtonPressed();
 	
 	static TSharedRef<SWidget> GenerateQuickRenderOptionsMenu();
 	static void OpenQuickRenderSettingsWindow(const FToolMenuContext& InToolMenuContext);
+	static TSharedRef<SWidget> GenerateQuickRenderSettingsWindowToolbar();
 
 	static void GenerateModesMenuSection(UToolMenu* InMenu);
 	static void GenerateQuickRenderMenuSection(UToolMenu* InMenu);
@@ -39,4 +43,8 @@ private:
 private:
 	inline static const FName QuickRenderToolbarId = TEXT("MovieRenderGraphQuickRender");
 	inline static TWeakPtr<SWindow> WeakQuickRenderSettingsWindow;
+
+	inline static TStrongObjectPtr<UMovieGraphQuickRenderSettings> QuickRenderSettings = nullptr;
+
+	inline static TArray<TSharedPtr<FName>> QuickRenderSettingsCategories;
 };
diff --git a/Engine/Plugins/MovieScene/MovieRenderPipeline/Source/MovieRenderPipelineEditor/Public/Graph/MovieGraphQuickRender.h b/Engine/Plugins/MovieScene/MovieRenderPipeline/Source/MovieRenderPipelineEditor/Public/Graph/MovieGraphQuickRender.h
--- a/Engine/Plugins/MovieScene/MovieRenderPipeline/Source/MovieRenderPipelineEditor/Public/Graph/MovieGraphQuickRender.h
+++ b/Engine/Plugins/MovieScene/MovieRenderPipeline/Source/MovieRenderPipelineEditor/Public/Graph/MovieGraphQuickRender.h
@@ -20,9 +20,7 @@ class MOVIERENDERPIPELINEEDITOR_API FMovieGraphQuickRender
 {
 public:
 	/* Begins a quick render using the provided settings. */
-	static void BeginQuickRender(const FMovieGraphQuickRenderSettings& InQuickRenderSettings);
-
-	// TODO: Add python-exposed helper to refresh the variable assignments on the context?
+	static void BeginQuickRender(const UMovieGraphQuickRenderSettings& InQuickRenderSettings);
 
 private:
 	/**
@@ -32,7 +30,7 @@ private:
 	static UMovieGraphConfig* GetQuickRenderGraph(UMovieGraphConfig* InUserSpecifiedGraph, TMap<TObjectPtr<UMovieGraphConfig>, TObjectPtr<UMovieGraphConfig>>& OutOriginalToDupeMap);
 
 	/** Applies all necessary updates to the duplicated graph to prepare it for a quick render (adjusting start/end frames, etc). */
-	static void ApplyQuickRenderUpdatesToDuplicateGraph(const FMovieGraphQuickRenderSettings& InQuickRenderSettings, UWorld* InEditorWorld);
+	static void ApplyQuickRenderUpdatesToDuplicateGraph(const UMovieGraphQuickRenderSettings& InQuickRenderSettings, UWorld* InEditorWorld);
 
 	/** Injects a new node in the furthest-downstream position within the Globals branch in the duplicated graph. */
 	static UMovieGraphNode* InjectNodeIntoGlobals(const TSubclassOf<UMovieGraphNode> NodeType);
