MovieGraph: Add support for Custom Playback Ranges to be optionally overridden as a variable. Previously in MRG custom playback range was an integer, which could be exposed as a property on a node, and then a Graph Variable could be created and plugged into it. This could then be exposed on the job, and the user could override it from the job itself. There were shortcomings to this process though - once you had exposed something as a pin and plugged in a variable, if the job did not overwrite it, then the default value from the variable would be used (which would often be zero). Since zero is a valid value for a level sequence's playback range, the system could not determine between user intent (to override it to zero) or being unset.

To solve this, Custom Playback Range was switched to a struct with an enum, where the enum allows between using the Sequence Default, and a custom value, and only if it's set to the custom value does it respect it. Because the struct is exposed as a monolithic unit to variables, when overwriting it you now have the choice between setting it to a specific value (even if that is zero), or, to Sequencer Default meaning that you don't actually want to change it.

diff --git a/Engine/Plugins/MovieScene/MovieRenderPipeline/Source/MovieRenderPipelineCore/Private/Graph/MovieGraphConfig.cpp b/Engine/Plugins/MovieScene/MovieRenderPipeline/Source/MovieRenderPipelineCore/Private/Graph/MovieGraphConfig.cpp
--- a/Engine/Plugins/MovieScene/MovieRenderPipeline/Source/MovieRenderPipelineCore/Private/Graph/MovieGraphConfig.cpp
+++ b/Engine/Plugins/MovieScene/MovieRenderPipeline/Source/MovieRenderPipelineCore/Private/Graph/MovieGraphConfig.cpp
@@ -2158,6 +2158,20 @@ UMovieGraphEvaluatedConfig* UMovieGraphConfig::CreateFlattenedGraph(const FMovie
 		}		
 	}
 
+	// Now that we have traversed the graph and flattened it, we need to give nodes a chance to perform a Post-Load type upgrade to the flattened graph.
+	// This is because their scripting pipeline may have modified the source graph this flattened graph was created from to set legacy properties,
+	// so we need to give the flattened graphs a chance to migrate that data and warn the user of the need to update their scripts.
+	for (const TPair<FName, FMovieGraphEvaluatedBranchConfig>& Pair : NewContext->BranchConfigMapping)
+	{
+		for (UMovieGraphNode* Node : Pair.Value.GetNodes())
+		{
+			if (UMovieGraphSettingNode* Setting = Cast<UMovieGraphSettingNode>(Node))
+			{
+				Setting->PostFlatten();	
+			}
+		}
+	}
+	
 	bool bHasRenderLayerNode = false;
 
 	for (const TPair<FName, FMovieGraphEvaluatedBranchConfig>& Pair : NewContext->BranchConfigMapping)
diff --git a/Engine/Plugins/MovieScene/MovieRenderPipeline/Source/MovieRenderPipelineCore/Private/Graph/MovieGraphNode.cpp b/Engine/Plugins/MovieScene/MovieRenderPipeline/Source/MovieRenderPipelineCore/Private/Graph/MovieGraphNode.cpp
--- a/Engine/Plugins/MovieScene/MovieRenderPipeline/Source/MovieRenderPipelineCore/Private/Graph/MovieGraphNode.cpp
+++ b/Engine/Plugins/MovieScene/MovieRenderPipeline/Source/MovieRenderPipelineCore/Private/Graph/MovieGraphNode.cpp
@@ -121,6 +121,26 @@ TArray<FMovieGraphPinProperties> UMovieGraphNode::GetExposedPinProperties() cons
 
 void UMovieGraphNode::TogglePromotePropertyToPin(const FName& PropertyName)
 {
+	bool bIsAlreadyExposedProperty = ExposedPropertyInfo.ContainsByPredicate([PropertyName](const FMovieGraphPropertyInfo& ExposedProperty)
+	{
+		return ExposedProperty.Name == PropertyName;
+	});
+	
+	// Both Promote and Unpromote will call Modify() and UpdatePins() on their own.
+	if(bIsAlreadyExposedProperty)
+	{
+		UnpromotePropertyFromPin(PropertyName);
+	}
+	else
+	{
+		PromotePropertyToPin(PropertyName);
+	}
+}
+
+void UMovieGraphNode::PromotePropertyToPin(const FName& PropertyName)
+{
+	Modify();
+	
 	// Ensure this is a property that can be promoted
 	bool bIsPromotableProperty = false;
 	FMovieGraphPropertyInfo OverrideablePropertyInfo;
@@ -133,28 +153,41 @@ void UMovieGraphNode::TogglePromotePropertyToPin(const FName& PropertyName)
 			break;
 		}
 	}
-
-	// Add this property to promoted properties if it hasn't already been promoted, otherwise remove it
+	
 	if (bIsPromotableProperty)
 	{
-		const bool bFoundExposedPropertyInfo = ExposedPropertyInfo.ContainsByPredicate([&OverrideablePropertyInfo](const FMovieGraphPropertyInfo& ExposedInfo)
-		{
-			return OverrideablePropertyInfo.IsSamePropertyAs(ExposedInfo);
-		});
-
-		Modify();
+		ExposedPropertyInfo.AddUnique(MoveTemp(OverrideablePropertyInfo));
+	}
+	else
+	{
+		// They requested a property that didn't exist, throw a kismet exception so scripts can know.
+		FFrame::KismetExecutionMessage(
+			*FString::Printf(
+				TEXT("%hs: Could not find a property with the name '%s' to promote to pin."), __FUNCTION__, *PropertyName.ToString()),
+			ELogVerbosity::Error);
+	}
+	
+	UpdatePins();
+}
 
-		if (!bFoundExposedPropertyInfo)
-		{
-			ExposedPropertyInfo.Add(MoveTemp(OverrideablePropertyInfo));
-		}
-		else
+void UMovieGraphNode::UnpromotePropertyFromPin(const FName& PropertyName)
+{
+	Modify();
+	
+	bool bIsExistingProperty = ExposedPropertyInfo.ContainsByPredicate([PropertyName](const FMovieGraphPropertyInfo& ExposedProperty)
+	{
+		return ExposedProperty.Name == PropertyName;
+	});
+	
+	if(bIsExistingProperty)
+	{
+		// Remove the requested pin from our ExposedPropertyInfo array. We use ExposedPropertyInfo here (and not GetOverrideablePropertyInfo)
+		// because this reflects the actual currently exposed pins, even if they're no longer something you could override. This allows you to
+		// unpromote an old property.
+		ExposedPropertyInfo.RemoveAll([&PropertyName](const FMovieGraphPropertyInfo& ExposedInfo)
 		{
-			ExposedPropertyInfo.RemoveAll([&OverrideablePropertyInfo](const FMovieGraphPropertyInfo& ExposedInfo)
-			{
-				return OverrideablePropertyInfo.IsSamePropertyAs(ExposedInfo);
-			});
-		}
+			return ExposedInfo.Name == PropertyName;
+		});
 	}
 	else
 	{
@@ -164,7 +197,7 @@ void UMovieGraphNode::TogglePromotePropertyToPin(const FName& PropertyName)
 				TEXT("%hs: Could not find a property with the name '%s' to promote to pin."), __FUNCTION__, *PropertyName.ToString()),
 			ELogVerbosity::Error);
 	}
-
+	
 	UpdatePins();
 }
 
diff --git a/Engine/Plugins/MovieScene/MovieRenderPipeline/Source/MovieRenderPipelineCore/Private/Graph/MovieGraphSequenceDataSource.cpp b/Engine/Plugins/MovieScene/MovieRenderPipeline/Source/MovieRenderPipelineCore/Private/Graph/MovieGraphSequenceDataSource.cpp
--- a/Engine/Plugins/MovieScene/MovieRenderPipeline/Source/MovieRenderPipelineCore/Private/Graph/MovieGraphSequenceDataSource.cpp
+++ b/Engine/Plugins/MovieScene/MovieRenderPipeline/Source/MovieRenderPipelineCore/Private/Graph/MovieGraphSequenceDataSource.cpp
@@ -81,13 +81,13 @@ void UMovieGraphSequenceDataSource::OverrideSequencePlaybackRangeFromGlobalOutpu
 	FFrameNumber StartFrameTickResolution = CurrentPlaybackRange.GetLowerBound().GetValue();
 	FFrameNumber EndFrameTickResolution = CurrentPlaybackRange.GetUpperBound().GetValue();
 	
-	if (OutputSetting->bOverride_CustomPlaybackRangeStartFrame)
+	if (OutputSetting->bOverride_CustomPlaybackRangeStart && OutputSetting->CustomPlaybackRangeStart.Type == EMovieGraphSequenceRangeType::Custom)
 	{
-		StartFrameTickResolution = FFrameRate::TransformTime(FFrameTime(FFrameNumber(OutputSetting->CustomPlaybackRangeStartFrame)), InSequence->GetMovieScene()->GetDisplayRate(), InSequence->GetMovieScene()->GetTickResolution()).FloorToFrame();
+		StartFrameTickResolution = FFrameRate::TransformTime(FFrameTime(FFrameNumber(OutputSetting->CustomPlaybackRangeStart.Value)), InSequence->GetMovieScene()->GetDisplayRate(), InSequence->GetMovieScene()->GetTickResolution()).FloorToFrame();
 	}
-	if (OutputSetting->bOverride_CustomPlaybackRangeEndFrame)
+	if (OutputSetting->bOverride_CustomPlaybackRangeEnd && OutputSetting->CustomPlaybackRangeEnd.Type == EMovieGraphSequenceRangeType::Custom)
 	{
-		EndFrameTickResolution = FFrameRate::TransformTime(FFrameTime(FFrameNumber(OutputSetting->CustomPlaybackRangeEndFrame)), InSequence->GetMovieScene()->GetDisplayRate(), InSequence->GetMovieScene()->GetTickResolution()).CeilToFrame();
+		EndFrameTickResolution = FFrameRate::TransformTime(FFrameTime(FFrameNumber(OutputSetting->CustomPlaybackRangeEnd.Value)), InSequence->GetMovieScene()->GetDisplayRate(), InSequence->GetMovieScene()->GetTickResolution()).CeilToFrame();
 	}
 
 	TRange<FFrameNumber> NewPlaybackRange = TRange<FFrameNumber>(StartFrameTickResolution, EndFrameTickResolution);
diff --git a/Engine/Plugins/MovieScene/MovieRenderPipeline/Source/MovieRenderPipelineCore/Private/Graph/Nodes/MovieGraphGlobalOutputSettingNode.cpp b/Engine/Plugins/MovieScene/MovieRenderPipeline/Source/MovieRenderPipelineCore/Private/Graph/Nodes/MovieGraphGlobalOutputSettingNode.cpp
--- a/Engine/Plugins/MovieScene/MovieRenderPipeline/Source/MovieRenderPipelineCore/Private/Graph/Nodes/MovieGraphGlobalOutputSettingNode.cpp
+++ b/Engine/Plugins/MovieScene/MovieRenderPipeline/Source/MovieRenderPipelineCore/Private/Graph/Nodes/MovieGraphGlobalOutputSettingNode.cpp
@@ -7,6 +7,7 @@
 #include "Graph/MovieGraphBlueprintLibrary.h"
 #include "MoviePipelineTelemetry.h"
 #include "Styling/AppStyle.h"
+#include "MovieRenderPipelineCoreModule.h"
 
 UMovieGraphGlobalOutputSettingNode::UMovieGraphGlobalOutputSettingNode()
 	: OutputFrameRate(FFrameRate(24, 1))
@@ -14,10 +15,10 @@ UMovieGraphGlobalOutputSettingNode::UMovieGraphGlobalOutputSettingNode()
 	, ZeroPadFrameNumbers(4)
 	, FrameNumberOffset(0)
 	, HandleFrameCount(0)
-	, CustomPlaybackRangeStartFrame(0)
-	, CustomPlaybackRangeEndFrame(0)
 	, bDropFrameTimecode(true)	// Defaults to true because most 29.97 FPS content uses this
 	, bFlushDiskWritesPerShot(false)
+	, CustomPlaybackRangeStartFrame(0)
+	, CustomPlaybackRangeEndFrame(0)
 {
 	OutputDirectory.Path = TEXT("{project_dir}/Saved/MovieRenders/");
 }
@@ -88,3 +89,62 @@ EMovieGraphBranchRestriction UMovieGraphGlobalOutputSettingNode::GetBranchRestri
 }
 #endif // WITH_EDITOR
 
+void UMovieGraphGlobalOutputSettingNode::PostLoad()
+{
+	Super::PostLoad();
+
+	// We don't emit a warning here because old assets that are just
+	// upgrading will update the properties and be fixed on next save.
+	// The ability to warn when running ApplyPostLoadPropertyConversions 
+	// is meant to catch scripting which changes it after load, but before render.
+	constexpr bool bEmitWarning = false;
+	ApplyPostLoadPropertyConversions(bEmitWarning);
+}
+
+void UMovieGraphGlobalOutputSettingNode::PostFlatten()
+{
+	Super::PostFlatten();
+	const bool bEmitWarning = true;
+	ApplyPostLoadPropertyConversions(bEmitWarning);
+}
+
+void UMovieGraphGlobalOutputSettingNode::ApplyPostLoadPropertyConversions(bool bEmitWarning)
+{
+PRAGMA_DISABLE_DEPRECATION_WARNINGS
+	// If they had previously stated that they wanted to use a custom Start Frame
+	// then we transfer that override to the new Playback Start, and then clear
+	// the override flag so that the upgrade doesn't get performed again on next load.
+	if (bOverride_CustomPlaybackRangeStartFrame)
+	{
+		bOverride_CustomPlaybackRangeStart = true;
+		CustomPlaybackRangeStart.Type = EMovieGraphSequenceRangeType::Custom;
+		CustomPlaybackRangeStart.Value = CustomPlaybackRangeStartFrame;
+
+		// Clear the original override the user had so that we don't upgrade again next load.
+		bOverride_CustomPlaybackRangeStartFrame = false;
+		CustomPlaybackRangeStartFrame = 0;
+
+		if (bEmitWarning)
+		{
+			UE_LOG(LogMovieRenderPipeline, Warning, TEXT("CustomPlaybackRangeStartFrame is deprecated, but it was changed after the asset was loaded. Please update your scripts/pipeline to use CustomPlaybackRangeStart instead!"));
+		}
+	}
+
+	if (bOverride_CustomPlaybackRangeEndFrame)
+	{
+		bOverride_CustomPlaybackRangeEnd = true;
+		CustomPlaybackRangeEnd.Type = EMovieGraphSequenceRangeType::Custom;
+		CustomPlaybackRangeEnd.Value = CustomPlaybackRangeEndFrame;
+
+		// Clear the original override the user had so that we don't upgrade again next load.
+		bOverride_CustomPlaybackRangeEndFrame = false;
+		CustomPlaybackRangeEndFrame = 0;
+
+		if (bEmitWarning)
+		{
+			UE_LOG(LogMovieRenderPipeline, Warning, TEXT("CustomPlaybackRangeEndFrame is deprecated, but it was changed after the asset was loaded. Please update your scripts/pipeline to use CustomPlaybackRangeStart instead!"));
+		}
+	}
+PRAGMA_ENABLE_DEPRECATION_WARNINGS
+}
+
diff --git a/Engine/Plugins/MovieScene/MovieRenderPipeline/Source/MovieRenderPipelineCore/Public/Graph/MovieGraphNode.h b/Engine/Plugins/MovieScene/MovieRenderPipeline/Source/MovieRenderPipelineCore/Public/Graph/MovieGraphNode.h
--- a/Engine/Plugins/MovieScene/MovieRenderPipeline/Source/MovieRenderPipelineCore/Public/Graph/MovieGraphNode.h
+++ b/Engine/Plugins/MovieScene/MovieRenderPipeline/Source/MovieRenderPipelineCore/Public/Graph/MovieGraphNode.h
@@ -238,7 +238,15 @@ public:
 	/** Toggles the promotion of the property with the given name to a pin on the node. */
 	UFUNCTION(BlueprintCallable, Category = "Movie Graph")
 	virtual void TogglePromotePropertyToPin(const FName& PropertyName);
-
+	
+	/** Promotes the given PropertyName to a pinned property, if it is found in GetOverrideablePropertyInfo(). Throws a kismet (scripting) exception if it isn't a overridable property. */
+	UFUNCTION(BlueprintCallable, Category = "Movie Graph")
+	virtual void PromotePropertyToPin(const FName& PropertyName);
+	
+	/** Unpromotes a given PropertyName from being pinned. Throws a kismet (scripting) exception if it's not already exposed. See GetExposedProperties() */
+	UFUNCTION(BlueprintCallable, Category = "Movie Graph")
+	virtual void UnpromotePropertyFromPin(const FName& PropertyName);
+	
 	/**
 	 * Determines if this node type can be added to the graph interactively by a user or via the API when constructing a graph.
 	 * @return true if the object can be added via the API, false otherwise
@@ -445,6 +453,21 @@ public:
 	 * The need for this should be exceedingly uncommon, so only use if absolutely necessary.
 	 */
 	virtual void PrepareForFlattening(const UMovieGraphSettingNode* InSourceNode) { }
+
+
+	/**
+	* When the graph asset was loaded, any deprecated properties were converted to their "real" representations
+	* in PostLoad. However, scripting can edit these assets after load - either an editor script, or
+	* an ExecuteScript node, and these custom scripts may still be (re) setting the deprecated property, at which
+	* point the runtime code will ignore it.
+	* 
+	* To catch these scenarios, after the graph is flattened, we call PostFlatten on each node in the graph which
+	* gives it one more chance to do the conversion from deprecated properties to real ones, and warn the user that
+	* they need to update their scripts.
+	* 
+	* This is called on the flattened graph, so only one "merged" copy of the nodes exist at this point.
+	*/
+	virtual void PostFlatten() {}
 	
 	/*
 	* This is called either on the CDO, or on a "flattened" instance of the node every frame when
diff --git a/Engine/Plugins/MovieScene/MovieRenderPipeline/Source/MovieRenderPipelineCore/Public/Graph/Nodes/MovieGraphGlobalOutputSettingNode.h b/Engine/Plugins/MovieScene/MovieRenderPipeline/Source/MovieRenderPipelineCore/Public/Graph/Nodes/MovieGraphGlobalOutputSettingNode.h
--- a/Engine/Plugins/MovieScene/MovieRenderPipeline/Source/MovieRenderPipelineCore/Public/Graph/Nodes/MovieGraphGlobalOutputSettingNode.h
+++ b/Engine/Plugins/MovieScene/MovieRenderPipeline/Source/MovieRenderPipelineCore/Public/Graph/Nodes/MovieGraphGlobalOutputSettingNode.h
@@ -30,17 +30,56 @@ struct FMovieGraphVersioningSettings
 	int32 VersionNumber = 1;
 };
 
+UENUM(BlueprintType)
+enum class EMovieGraphSequenceRangeType : uint8
+{
+	/** Use the Playback Range value from the Level Sequence (without overriding it.) */
+	SequenceDefault,
+	/** Override the Playback Range value from the Level Sequence and instead override it to use a custom Frame Number. */
+	Custom
+};
+
+USTRUCT(BlueprintType)
+struct FMovieGraphSequencePlaybackRangeBound
+{
+	GENERATED_BODY()
+
+public:
+	FMovieGraphSequencePlaybackRangeBound()
+		: Type(EMovieGraphSequenceRangeType::Custom)
+		, Value(0)
+	{}
+			
+	/** 
+	* By default the render will use the Playback Range Start/End frame to determine what to render. Set this to Custom to
+	* override the Playback Range bound value, and instead use the Value below as the Start or End frame. 
+	*/
+	UPROPERTY(EditAnywhere, BlueprintReadWrite, Category = "Value")
+	EMovieGraphSequenceRangeType Type;
+
+	/**
+	* If Type is set to Custom, then this value is used for the Playback Range Start or End frame. Value is considered frames at the original Sequence
+	* frame rate (it is applied before Frame Rate Override.)
+	*/
+	UPROPERTY(EditAnywhere, BlueprintReadWrite, Category = "Value", meta = (EditCondition = "Type == EMovieGraphSequenceRangeType::Custom"))
+	int32 Value;
+};
+
 UCLASS()
 class MOVIERENDERPIPELINECORE_API UMovieGraphGlobalOutputSettingNode : public UMovieGraphSettingNode
 {
 	GENERATED_BODY()
 public:
 	UMovieGraphGlobalOutputSettingNode();
+	
+	// UObject Interface
+	virtual void PostLoad() override;
+	// ~UObject Interface
 
 	// UMovieGraphSettingNode Interface
 	virtual void GetFormatResolveArgs(FMovieGraphResolveArgs& OutMergedFormatArgs, const FMovieGraphRenderDataIdentifier& InRenderDataIdentifier) const override;
 	virtual void UpdateTelemetry(FMoviePipelineShotRenderTelemetry* InTelemetry) const override;
-
+	virtual void PostFlatten() override;
 #if WITH_EDITOR
 	virtual FText GetNodeTitle(const bool bGetDescriptive) const override;
 	virtual FText GetMenuCategory() const override;
@@ -49,6 +88,15 @@ public:
 	virtual EMovieGraphBranchRestriction GetBranchRestriction() const override;
 #endif
 	// ~UMovieGraphSettingNode Interface
+	
+protected:
+	/*
+	* This is called from PostLoad and when the UMovieGraphPipeline is initialized to convert any legacy properties. We do it this way to preserve
+	* existing code that may be configuring assets after loading them but before rendering them.
+	*
+	* @param bEmitWarning - If true, a warning will be printed that the conversion took place and the users need to update their scripts.
+	*/	
+	virtual void ApplyPostLoadPropertyConversions(bool bEmitWarning);
 
 public:
 	UPROPERTY(EditAnywhere, BlueprintReadWrite, Category = Overrides, meta = (InlineEditConditionToggle))
@@ -72,12 +120,20 @@ public:
 	UPROPERTY(EditAnywhere, BlueprintReadWrite, Category = Overrides, meta = (InlineEditConditionToggle))
 	uint8 bOverride_HandleFrameCount : 1;
 
-	UPROPERTY(EditAnywhere, BlueprintReadWrite, Category = Overrides, meta = (InlineEditConditionToggle))
+	UE_DEPRECATED(5.6, "Use bOverride_CustomPlaybackRangeStart instead.")
+	UPROPERTY(BlueprintReadWrite, Category = Overrides, meta = (DeprecatedProperty, DeprecationMessage = "Use bOverride_CustomPlaybackRangeStart instead."))
 	uint8 bOverride_CustomPlaybackRangeStartFrame : 1;
 
-	UPROPERTY(EditAnywhere, BlueprintReadWrite, Category = Overrides, meta = (InlineEditConditionToggle))
+	UE_DEPRECATED(5.6, "Use bOverride_CustomPlaybackRangeEnd instead.")
+	UPROPERTY(BlueprintReadWrite, Category = Overrides, meta = (DeprecatedProperty, DeprecationMessage = "Use bOverride_CustomPlaybackRangeEnd instead."))
 	uint8 bOverride_CustomPlaybackRangeEndFrame : 1;
 
+	UPROPERTY(EditAnywhere, BlueprintReadWrite, Category = Overrides, meta = (InlineEditConditionToggle))
+	uint8 bOverride_CustomPlaybackRangeStart : 1;
+
+	UPROPERTY(EditAnywhere, BlueprintReadWrite, Category = Overrides, meta = (InlineEditConditionToggle))
+	uint8 bOverride_CustomPlaybackRangeEnd : 1;
+
 	UPROPERTY(EditAnywhere, BlueprintReadWrite, Category = Overrides, meta = (InlineEditConditionToggle))
 	uint8 bOverride_CustomTimecodeStart : 1;
 
@@ -126,11 +182,20 @@ public:
 	UPROPERTY(EditAnywhere, BlueprintReadWrite, Category = "Frames", meta = (UIMin = 0, ClampMin = 0, EditCondition = "bOverride_HandleFrameCount"))
 	int32 HandleFrameCount;
 
-	UPROPERTY(EditAnywhere, BlueprintReadWrite, Category = "Frames", meta = (EditCondition = "bOverride_CustomPlaybackRangeStartFrame"))
-	int32 CustomPlaybackRangeStartFrame;
+	/**
+	* If overwritten, and the Type is set to "Custom", then the Value field will override the Sequence's Playback Range Start when rendering.
+	* Values are expected to be in the Sequence's original Frame Rate (the custom range is applied before Frame Rate Override).
+	*/
+	UPROPERTY(EditAnywhere, BlueprintReadWrite, Category = "Frames", meta = (EditCondition = "bOverride_CustomPlaybackRangeStart"))
+	FMovieGraphSequencePlaybackRangeBound CustomPlaybackRangeStart;
 	
-	UPROPERTY(EditAnywhere, BlueprintReadWrite, Category = "Frames", meta = (EditCondition = "bOverride_CustomPlaybackRangeEndFrame"))
-	int32 CustomPlaybackRangeEndFrame;
+
+	/**
+	* If overwritten, and the Type is set to "Custom", then the Value field will override the Sequence's Playback Range End when rendering.
+	* Values are expected to be in the Sequence's original Frame Rate (the custom range is applied before Frame Rate Override).
+	*/
+	UPROPERTY(EditAnywhere, BlueprintReadWrite, Category = "Frames", meta = (EditCondition = "bOverride_CustomPlaybackRangeEnd"))
+	FMovieGraphSequencePlaybackRangeBound CustomPlaybackRangeEnd;
 
 	/** Start the timecode at a specific value, rather than the value coming from the Level Sequence. Only applicable to output formats that support timecode. */
 	UPROPERTY(EditAnywhere, BlueprintReadWrite, Category = "Timecode", meta = (EditCondition = "bOverride_CustomTimecodeStart"))
@@ -152,4 +217,16 @@ public:
 	*/
 	UPROPERTY(EditAnywhere, BlueprintReadWrite, Category = "Scripting", meta = (EditCondition = "bOverride_bFlushDiskWritesPerShot"))
 	bool bFlushDiskWritesPerShot;
+
+// These are deprecated but not moved to WITH_EDITORONLY_DATA because that would break existing plugins when they compiled for shipping. 	
+public:
+
+	UE_DEPRECATED(5.6, "Use CustomPlaybackRangeStart with Type set to Custom and Value set to the desired value instead.")
+	UPROPERTY(BlueprintReadWrite, Category = "Frames", meta = (DeprecatedFunction, DeprecationMessage = "Use CustomPlaybackRangeStart with Type set to Custom and Value set to the desired value instead."))
+	int32 CustomPlaybackRangeStartFrame;
+	
+	UE_DEPRECATED(5.6, "Use CustomPlaybackRangeEnd with Type set to Custom and Value set to the desired value instead.")
+	UPROPERTY(BlueprintReadWrite, Category = "Frames", meta = (DeprecatedFunction, DeprecationMessage = "Use CustomPlaybackRangeEnd with Type set to Custom and Value set to the desired value instead."))
+	int32 CustomPlaybackRangeEndFrame;
+
 };
diff --git a/Engine/Plugins/MovieScene/MovieRenderPipeline/Source/MovieRenderPipelineEditor/Private/Graph/MovieEdGraphNode.cpp b/Engine/Plugins/MovieScene/MovieRenderPipeline/Source/MovieRenderPipelineEditor/Private/Graph/MovieEdGraphNode.cpp
--- a/Engine/Plugins/MovieScene/MovieRenderPipeline/Source/MovieRenderPipelineEditor/Private/Graph/MovieEdGraphNode.cpp
+++ b/Engine/Plugins/MovieScene/MovieRenderPipeline/Source/MovieRenderPipelineEditor/Private/Graph/MovieEdGraphNode.cpp
@@ -355,6 +355,53 @@ void UMoviePipelineEdGraphNode::GetPropertyPromotionContextMenuActions(UToolMenu
 				FCanExecuteAction::CreateLambda([]() { return false; }))
 		);
 	}
+
+	// Show legacy/deprecated properties as pins so they can be unchecked (to remove the pin), as they no longer show up in the OverrideablePropertyInfo
+	// Unfortunately, the only way we know a legacy property is if there's already a pin exposed for it, but it's no longer in the OverrideablePropertyInfo
+	// list. Once unchecked, the pin cannot be re-added anymore.
+	TArray<FMovieGraphPropertyInfo> DeprecatedProperties;
+	for (const FMovieGraphPropertyInfo& ExposedProperty : RuntimeNode->GetExposedProperties())
+	{
+		// For each of our actually exposed properties, check them against the properties that are overrideable.
+		// If we find a match, then we know that the exposed property is still overrideable, thus not a deprecated one.
+		bool bIsRealProperty = OverrideablePropertyInfo.ContainsByPredicate([ExposedProperty](const FMovieGraphPropertyInfo& OverrideableProperty)
+			{
+				return ExposedProperty == OverrideableProperty;
+			});
+
+		if (!bIsRealProperty)
+		{
+			// This adds a copy but the only thing we need from it below are the name and context menu name.
+			DeprecatedProperties.Add(ExposedProperty);
+		}
+	}
+
+	if (DeprecatedProperties.Num() > 0)
+	{
+		FToolMenuSection& DeprecatedPropertiesSection = Menu->AddSection("DeprecatedProperties", LOCTEXT("DeprecatedProperties", "Deprecated Properties"));
+		for (const FMovieGraphPropertyInfo& PropertyInfo : DeprecatedProperties)
+		{
+			DeprecatedPropertiesSection.AddMenuEntry(
+				PropertyInfo.Name,
+				PropertyInfo.ContextMenuName.IsEmpty() ? FText::FromName(PropertyInfo.Name) : PropertyInfo.ContextMenuName,
+				LOCTEXT("UnPromoteLegacyProperty", "Remove this deprecated pin (which cannot be re-added once removed due to being deprecated)"),
+				FSlateIcon(),
+				FUIAction(
+					// This can still call Toggle because toggle will see that it is a exposed property and allow untoggling it, even if
+					// it's no longer a overrideable property (because it's deprecated).
+					FExecuteAction::CreateUObject(this, &UMoviePipelineEdGraphNode::TogglePromotePropertyToPin, PropertyInfo.Name),
+					FCanExecuteAction(),
+					FGetActionCheckState::CreateLambda([this, PropertyInfo]()
+						{
+							const TArray<FMovieGraphPropertyInfo>& ExposedProperties = RuntimeNode->GetExposedProperties();
+							return ExposedProperties.Contains(PropertyInfo) ? ECheckBoxState::Checked : ECheckBoxState::Unchecked;
+						})),
+				EUserInterfaceActionType::ToggleButton
+			);
+
+		}
+	}
+
 }
 
 void UMoviePipelineEdGraphNode::PromotePropertyToVariable(const FMovieGraphPropertyInfo& TargetProperty) const
@@ -806,4 +853,4 @@ void UMoviePipelineEdGraphNodeBase::GetNodeContextMenuActions(UToolMenu* Menu, U
 	}
 }
 
-#undef LOCTEXT_NAMESPACE
\ No newline at end of file
+#undef LOCTEXT_NAMESPACE
diff --git a/Engine/Plugins/MovieScene/MovieRenderPipeline/Source/MovieRenderPipelineEditor/Private/Graph/MovieGraphQuickRender.cpp b/Engine/Plugins/MovieScene/MovieRenderPipeline/Source/MovieRenderPipelineEditor/Private/Graph/MovieGraphQuickRender.cpp
--- a/Engine/Plugins/MovieScene/MovieRenderPipeline/Source/MovieRenderPipelineEditor/Private/Graph/MovieGraphQuickRender.cpp
+++ b/Engine/Plugins/MovieScene/MovieRenderPipeline/Source/MovieRenderPipelineEditor/Private/Graph/MovieGraphQuickRender.cpp
@@ -206,12 +206,13 @@ void FMovieGraphQuickRender::ApplyQuickRenderUpdatesToDuplicateGraph(const UMovi
 			const int32 EndFrame = (InQuickRenderSettings.FrameRangeType == EMovieGraphQuickRenderFrameRangeType::SelectionRange)
 				? ULevelSequenceEditorBlueprintLibrary::GetSelectionRangeEnd()
 				: InQuickRenderSettings.CustomEndFrame;
-		
+PRAGMA_DISABLE_DEPRECATION_WARNINGS
 			GlobalOutputSettings->bOverride_CustomPlaybackRangeStartFrame = true;
 			GlobalOutputSettings->CustomPlaybackRangeStartFrame = StartFrame;
 
 			GlobalOutputSettings->bOverride_CustomPlaybackRangeEndFrame = true;
 			GlobalOutputSettings->CustomPlaybackRangeEndFrame = EndFrame;
+PRAGMA_ENABLE_DEPRECATION_WARNINGS
 		}
 	}
 
