[MovieRender]
- Avoid loading assets from editor at boot and use best practices to include them at cook time instead
- 6s -> 5s of total startup module time

diff --git a/Engine/Plugins/MovieScene/MovieRenderPipeline/Source/MovieRenderPipelineCore/Private/MovieRenderPipelineCoreModule.cpp b/Engine/Plugins/MovieScene/MovieRenderPipeline/Source/MovieRenderPipelineCore/Private/MovieRenderPipelineCoreModule.cpp
--- a/Engine/Plugins/MovieScene/MovieRenderPipeline/Source/MovieRenderPipelineCore/Private/MovieRenderPipelineCoreModule.cpp
+++ b/Engine/Plugins/MovieScene/MovieRenderPipeline/Source/MovieRenderPipelineCore/Private/MovieRenderPipelineCoreModule.cpp
@@ -10,24 +10,38 @@
 #include "MoviePipeline.h"
 #include "Graph/MovieGraphPipeline.h"
 #include "Graph/Nodes/MovieGraphBurnInNode.h"
+#include "UObject/ICookInfo.h"
 
 FName IMoviePipelineBurnInExtension::ModularFeatureName = "ModularFeature_MoviePipelineBurnInExt";
 
 void FMovieRenderPipelineCoreModule::StartupModule()
 {
 #if WITH_EDITOR
-	// Hack to ensure these assets (which are referenced only by code) get packaged
+	if (IsRunningCookCommandlet())
 	{
-		TArray<FString> Assets;
-		Assets.Add(UMoviePipeline::DefaultDebugWidgetAsset);
-		Assets.Add(UMovieGraphPipeline::DefaultPreviewWidgetAsset);
-		Assets.Add(UMovieGraphBurnInNode::DefaultBurnInWidgetAsset);
+		UE::Cook::FDelegates::ModifyCook.AddLambda(
+			[](UE::Cook::ICookInfo& CookInfo, TArray<UE::Cook::FPackageCookRule>& InOutPackageCookRules)
+			{
+				// Ensure these assets (which are referenced only by code) get packaged
+				const FString* Assets[] =
+				{
+					&UMoviePipeline::DefaultDebugWidgetAsset,
+					&UMovieGraphPipeline::DefaultPreviewWidgetAsset,
+					&UMovieGraphBurnInNode::DefaultBurnInWidgetAsset
+				};
 
-		for (const FString& Asset : Assets)
-		{
-			TSoftObjectPtr<UObject> AssetRef = TSoftObjectPtr<UObject>(FSoftObjectPath(Asset));
-			AssetRef.LoadSynchronous();
-		}
+				for (const FString* Asset : Assets)
+				{
+					InOutPackageCookRules.Add(
+						UE::Cook::FPackageCookRule{
+							.PackageName = FName(FSoftObjectPath(*Asset).GetLongPackageName()),
+							.InstigatorName = FName("FMovieRenderPipelineCoreModule"),
+							.CookRule = UE::Cook::EPackageCookRule::AddToCook
+						}
+					);
+				}
+			}
+		);
 	}
 #endif
 
diff --git a/Engine/Plugins/MovieScene/MovieRenderPipeline/Source/MovieRenderPipelineRenderPasses/Private/MovieRenderPipelineRenderPassesModule.cpp b/Engine/Plugins/MovieScene/MovieRenderPipeline/Source/MovieRenderPipelineRenderPasses/Private/MovieRenderPipelineRenderPassesModule.cpp
--- a/Engine/Plugins/MovieScene/MovieRenderPipeline/Source/MovieRenderPipelineRenderPasses/Private/MovieRenderPipelineRenderPassesModule.cpp
+++ b/Engine/Plugins/MovieScene/MovieRenderPipeline/Source/MovieRenderPipelineRenderPasses/Private/MovieRenderPipelineRenderPassesModule.cpp
@@ -4,6 +4,7 @@
 #include "MovieRenderPipelineDataTypes.h"
 #include "MovieRenderPipelineCoreModule.h"
 #include "MoviePipelineDeferredPasses.h"
+#include "UObject/ICookInfo.h"
 
 class FMovieRenderPipelineRenderPassesModule : public IModuleInterface
 {
@@ -11,19 +12,31 @@ public:
 	virtual void StartupModule() override
 	{
 #if WITH_EDITOR
-		// Hack to ensure these assets (which are referenced only by code) get packaged
+		if (IsRunningCookCommandlet())
 		{
-			TArray<FString> Assets;
-			Assets.Add(UMoviePipelineDeferredPassBase::StencilLayerMaterialAsset);
-			Assets.Add(UMoviePipelineDeferredPassBase::DefaultDepthAsset);
-			Assets.Add(UMoviePipelineDeferredPassBase::DefaultMotionVectorsAsset);
-
-			for (const FString& Asset : Assets)
-			{
-				TSoftObjectPtr<UObject> AssetRef = TSoftObjectPtr<UObject>(FSoftObjectPath(Asset));
-				AssetRef.LoadSynchronous();
-			}
+			UE::Cook::FDelegates::ModifyCook.AddLambda(
+				[](UE::Cook::ICookInfo& CookInfo, TArray<UE::Cook::FPackageCookRule>& InOutPackageCookRules)
+				{
+					// Ensure these assets (which are referenced only by code) get packaged
+					const FString* Assets[] =
+					{
+						&UMoviePipelineDeferredPassBase::StencilLayerMaterialAsset,
+						&UMoviePipelineDeferredPassBase::DefaultDepthAsset,
+						&UMoviePipelineDeferredPassBase::DefaultMotionVectorsAsset
+					};
 
+					for (const FString* Asset : Assets)
+					{
+						InOutPackageCookRules.Add(
+							UE::Cook::FPackageCookRule{
+								.PackageName = FName(FSoftObjectPath(*Asset).GetLongPackageName()),
+								.InstigatorName = FName("FMovieRenderPipelineRenderPassesModule"),
+								.CookRule = UE::Cook::EPackageCookRule::AddToCook
+							}
+						);
+					}
+				}
+			);
 		}
 #endif
 	}
diff --git a/Engine/Plugins/MovieScene/MovieRenderPipeline/Source/MovieRenderPipelineSettings/Private/MovieRenderPipelineSettingsModule.cpp b/Engine/Plugins/MovieScene/MovieRenderPipeline/Source/MovieRenderPipelineSettings/Private/MovieRenderPipelineSettingsModule.cpp
--- a/Engine/Plugins/MovieScene/MovieRenderPipeline/Source/MovieRenderPipelineSettings/Private/MovieRenderPipelineSettingsModule.cpp
+++ b/Engine/Plugins/MovieScene/MovieRenderPipeline/Source/MovieRenderPipelineSettings/Private/MovieRenderPipelineSettingsModule.cpp
@@ -2,22 +2,36 @@
 
 #include "Modules/ModuleManager.h"
 #include "MoviePipelineBurnInSetting.h"
+#include "UObject/ICookInfo.h"
 
 class FMovieRenderPipelineSettingsModule : public IModuleInterface
 {
 	virtual void StartupModule() override
 	{
 #if WITH_EDITOR
-		// Hack to ensure these assets (which are referenced only by code) get packaged
+		if (IsRunningCookCommandlet())
 		{
-			TArray<FString> Assets;
-			Assets.Add(UMoviePipelineBurnInSetting::DefaultBurnInWidgetAsset);
+			UE::Cook::FDelegates::ModifyCook.AddLambda(
+				[](UE::Cook::ICookInfo& CookInfo, TArray<UE::Cook::FPackageCookRule>& InOutPackageCookRules)
+				{
+					// Ensure these assets (which are referenced only by code) get packaged
+					const FString* Assets[] =
+					{
+						&UMoviePipelineBurnInSetting::DefaultBurnInWidgetAsset
+					};
 
-			for (const FString& Asset : Assets)
-			{
-				TSoftObjectPtr<UObject> AssetRef = TSoftObjectPtr<UObject>(FSoftObjectPath(Asset));
-				AssetRef.LoadSynchronous();
-			}
+					for (FString const* Asset : Assets)
+					{
+						InOutPackageCookRules.Add(
+							UE::Cook::FPackageCookRule{
+								.PackageName = FName(FSoftObjectPath(*Asset).GetLongPackageName()),
+								.InstigatorName = FName("FMovieRenderPipelineSettingsModule"),
+								.CookRule = UE::Cook::EPackageCookRule::AddToCook
+							}
+						);
+					}
+				}
+			);
 		}
 #endif
 	}
