MoviePipeline: Updated the "Apply Viewport Look" node to have the capability of turning on/off certain aspects of the viewport look. The render pipeline and Quick Render now respect these settings.

Also updated the default Quick Render graph to exclude the Apply Viewport Look node. It is now injected by [REDACTED_PERSON], and is hidden from the UI so it cannot be created by users.

#virtualized

diff --git a/Engine/Plugins/MovieScene/MovieRenderPipeline/Source/MovieRenderPipelineCore/Public/Graph/Nodes/MovieGraphApplyViewportLookNode.h b/Engine/Plugins/MovieScene/MovieRenderPipeline/Source/MovieRenderPipelineCore/Public/Graph/Nodes/MovieGraphApplyViewportLookNode.h
--- a/Engine/Plugins/MovieScene/MovieRenderPipeline/Source/MovieRenderPipelineCore/Public/Graph/Nodes/MovieGraphApplyViewportLookNode.h
+++ b/Engine/Plugins/MovieScene/MovieRenderPipeline/Source/MovieRenderPipelineCore/Public/Graph/Nodes/MovieGraphApplyViewportLookNode.h
@@ -10,8 +10,11 @@
 class FEditorViewportClient;
 class FSceneView;
 
+// Note: This node is marked as HideDropdown so it does not appear in the node creation menu within the graph. It is meant to
+// be created by Quick Render only. This node will not properly provide all of its functionality outside of Quick Render.
+
 /** A node which applies the look of the viewport to the render (show flags, view mode, OCIO, etc). */
-UCLASS()
+UCLASS(HideDropDown, NotBlueprintType)
 class MOVIERENDERPIPELINECORE_API UMovieGraphApplyViewportLookNode : public UMovieGraphSettingNode
 {
 	GENERATED_BODY()
@@ -23,8 +26,6 @@ public:
 	virtual EMovieGraphBranchRestriction GetBranchRestriction() const override;
 	//~ End UMovieGraphNode interface
 
-	// TODO: Expose a property which allows a selection of the following to be applied: OCIO, Show Flags, View Mode
-
 #if WITH_EDITOR
 	/** Gets the viewport client for the currently active viewport. */
 	static FEditorViewportClient* GetViewportClient();
@@ -42,4 +43,26 @@ public:
 	virtual FSlateIcon GetIconAndTint(FLinearColor& OutColor) const override;
 	//~ End UMovieGraphNode interface
 #endif
+
+public:
+	UPROPERTY(EditAnywhere, BlueprintReadWrite, Category = Overrides, meta = (InlineEditConditionToggle))
+	uint8 bOverride_bOcio : 1;
+
+	UPROPERTY(EditAnywhere, BlueprintReadWrite, Category = Overrides, meta = (InlineEditConditionToggle))
+	uint8 bOverride_bShowFlags : 1;
+
+	UPROPERTY(EditAnywhere, BlueprintReadWrite, Category = Overrides, meta = (InlineEditConditionToggle))
+	uint8 bOverride_bViewMode : 1;
+	
+	/** Set to true to apply the viewport's OCIO settings to the render. Only render nodes that have "Allow OCIO" turned on will be impacted. */
+	UPROPERTY(EditAnywhere, BlueprintReadWrite, Category="Viewport Look", DisplayName="OCIO", meta=(EditCondition="bOverride_bOcio"))
+	bool bOcio;
+
+	/** Set to true to apply the viewport's show flags to the render. */
+	UPROPERTY(EditAnywhere, BlueprintReadWrite, Category="Viewport Look", meta=(EditCondition="bOverride_bShowFlags"))
+	bool bShowFlags;
+
+	/** Set to true to apply the viewport's view mode to the render. */
+	UPROPERTY(EditAnywhere, BlueprintReadWrite, Category="Viewport Look", meta=(EditCondition="bOverride_bViewMode"))
+	bool bViewMode;
 };
\ No newline at end of file
diff --git a/Engine/Plugins/MovieScene/MovieRenderPipeline/Source/MovieRenderPipelineEditor/Private/Graph/MovieGraphQuickRender.cpp b/Engine/Plugins/MovieScene/MovieRenderPipeline/Source/MovieRenderPipelineEditor/Private/Graph/MovieGraphQuickRender.cpp
--- a/Engine/Plugins/MovieScene/MovieRenderPipeline/Source/MovieRenderPipelineEditor/Private/Graph/MovieGraphQuickRender.cpp
+++ b/Engine/Plugins/MovieScene/MovieRenderPipeline/Source/MovieRenderPipelineEditor/Private/Graph/MovieGraphQuickRender.cpp
@@ -108,8 +108,6 @@ void FMovieGraphQuickRender::BeginQuickRender(const UMovieGraphQuickRenderSettin
 	}
 	NewJob->SetGraphPreset(QuickRenderGraph);
 
-	ApplyQuickRenderUpdatesToDuplicateGraph(InQuickRenderSettings, EditorWorld);
-
 	// Apply any job-level variable overrides, if they were specified
 	for (const TObjectPtr<UMovieJobVariableAssignmentContainer>& IncomingVariableAssignment : InQuickRenderSettings.GraphVariableAssignments)
 	{
@@ -132,6 +130,10 @@ void FMovieGraphQuickRender::BeginQuickRender(const UMovieGraphQuickRenderSettin
 		}
 	}
 
+	// Update the graph as needed in order to properly apply the settings specified by the user. This will perform a graph evaluation, so the job
+	// should be fully initialized at this point.
+	ApplyQuickRenderUpdatesToDuplicateGraph(InQuickRenderSettings, EditorWorld, NewJob);
+
 	// Do the render
 	UMoviePipelinePIEExecutor* Executor = Cast<UMoviePipelinePIEExecutor>(Subsystem->RenderQueueInstanceWithExecutor(TemporaryQueue.Get(), ExecutorClass));
 
@@ -189,13 +191,13 @@ UMovieGraphConfig* FMovieGraphQuickRender::GetQuickRenderGraph(UMovieGraphConfig
 	return TemporaryGraph.Get();
 }
 
-void FMovieGraphQuickRender::ApplyQuickRenderUpdatesToDuplicateGraph(const UMovieGraphQuickRenderSettings& InQuickRenderSettings, UWorld* InEditorWorld)
+void FMovieGraphQuickRender::ApplyQuickRenderUpdatesToDuplicateGraph(const UMovieGraphQuickRenderSettings& InQuickRenderSettings, UWorld* InEditorWorld, UMoviePipelineExecutorJob* InJob)
 {
 	check(TemporaryGraph.IsValid());
 
 	// Global Output Settings
 	{
-		UMovieGraphGlobalOutputSettingNode* GlobalOutputSettings = Cast<UMovieGraphGlobalOutputSettingNode>(InjectNodeIntoGlobals(UMovieGraphGlobalOutputSettingNode::StaticClass()));
+		UMovieGraphGlobalOutputSettingNode* GlobalOutputSettings = Cast<UMovieGraphGlobalOutputSettingNode>(InjectNodeIntoBranch(UMovieGraphGlobalOutputSettingNode::StaticClass(), UMovieGraphNode::GlobalsPinName));
 
 		// If the frame range type is StartEnd, the frame range doesn't need to be modified
 		if ((InQuickRenderSettings.FrameRangeType == EMovieGraphQuickRenderFrameRangeType::SelectionRange) || (InQuickRenderSettings.FrameRangeType == EMovieGraphQuickRenderFrameRangeType::Custom))
@@ -216,158 +218,230 @@ PRAGMA_ENABLE_DEPRECATION_WARNINGS
 		}
 	}
 
-	// Some updates are only done if there's an Apply Viewport Look node present
-	const TObjectPtr<UMovieGraphNode>* ViewportNode = TemporaryGraph->GetNodes().FindByPredicate([](const TObjectPtr<UMovieGraphNode>& InNode)
+	// Inject an Apply Viewport Look node if any viewport look flags were specified. Logic outside of Quick Render will for look this node in order to
+	// determine how to apply some of its settings
+	const EMovieGraphQuickRenderViewportLookFlags LookFlags = static_cast<EMovieGraphQuickRenderViewportLookFlags>(InQuickRenderSettings.ViewportLookFlags);
+	if (LookFlags != EMovieGraphQuickRenderViewportLookFlags::None)
 	{
-		return InNode && InNode->IsA<UMovieGraphApplyViewportLookNode>();
-	});
+		UMovieGraphApplyViewportLookNode* ApplyViewportLookNode = Cast<UMovieGraphApplyViewportLookNode>(InjectNodeIntoBranch(UMovieGraphApplyViewportLookNode::StaticClass(), UMovieGraphNode::GlobalsPinName));
+		check(ApplyViewportLookNode);
+
+#define APPLY_VIEWPORT_LOOK_FLAG(EnumName, FlagName) \
+		if (EnumHasAnyFlags(LookFlags, EnumName)) \
+		{ \
+			ApplyViewportLookNode->bOverride_b##FlagName = true; \
+			ApplyViewportLookNode->b##FlagName = true; \
+		} \
+
+		// Transfer the flags from the render settings to the Apply Viewport Look node
+		APPLY_VIEWPORT_LOOK_FLAG(EMovieGraphQuickRenderViewportLookFlags::Ocio, Ocio);
+		APPLY_VIEWPORT_LOOK_FLAG(EMovieGraphQuickRenderViewportLookFlags::ShowFlags, ShowFlags);
+		APPLY_VIEWPORT_LOOK_FLAG(EMovieGraphQuickRenderViewportLookFlags::ViewMode, ViewMode);
+#undef APPLY_VIEWPORT_LOOK_FLAG
+	}
 
-	if (ViewportNode)
+	// Apply viewport/editor actor visibility
+	ApplyQuickRenderUpdatesToDuplicateGraph_ApplyEditorVisibility(InEditorWorld);
+
+	// Apply OCIO if it is activated on the viewport and enabled in the viewport look flags
+	if (InQuickRenderSettings.bOverride_ViewportLookFlags && EnumHasAnyFlags(LookFlags, EMovieGraphQuickRenderViewportLookFlags::Ocio))
 	{
-		// Apply viewport actor visibility
+#if WITH_OCIO
+		if (const FEditorViewportClient* ViewportClient = UMovieGraphApplyViewportLookNode::GetViewportClient())
 		{
-			TArray<FMovieGraphActorQueryEntry> VisibleActors;
-			TArray<FMovieGraphActorQueryEntry> HiddenActors;
-			
-			for (FActorIterator It(InEditorWorld); It; ++It)
+			if (const FOpenColorIODisplayConfiguration* OcioConfiguration = IOpenColorIOModule::Get().GetDisplayManager().GetDisplayConfiguration(ViewportClient))
 			{
-				AActor* Actor = *It;
-
-				if (!Actor)
+				if (OcioConfiguration->bIsEnabled)
 				{
-					continue;
-				}
-
-				if (Actor->IsHiddenEd())
-				{
-					FMovieGraphActorQueryEntry& HiddenActorEntry = HiddenActors.AddDefaulted_GetRef();
-					HiddenActorEntry.ActorToMatch = Actor;
-				}
-				else
-				{
-					FMovieGraphActorQueryEntry& VisibleActorEntry = VisibleActors.AddDefaulted_GetRef();
-					VisibleActorEntry.ActorToMatch = Actor;
+					ApplyQuickRenderUpdatesToDuplicateGraph_ApplyOcio(InJob, OcioConfiguration);
 				}
 			}
+		}
+#endif	// WITH_OCIO
+	}
+}
+
+void FMovieGraphQuickRender::ApplyQuickRenderUpdatesToDuplicateGraph_ApplyEditorVisibility(const UWorld* InEditorWorld)
+{
+	TArray<FMovieGraphActorQueryEntry> VisibleActors;
+	TArray<FMovieGraphActorQueryEntry> HiddenActors;
+	
+	for (FActorIterator It(InEditorWorld); It; ++It)
+	{
+		AActor* Actor = *It;
 
-			const UMovieGraphCollectionNode* VisibleActorsCollectionNode = Cast<UMovieGraphCollectionNode>(InjectNodeIntoGlobals(UMovieGraphCollectionNode::StaticClass()));
-			const UMovieGraphCollectionNode* HiddenActorsCollectionNode = Cast<UMovieGraphCollectionNode>(InjectNodeIntoGlobals(UMovieGraphCollectionNode::StaticClass()));
+		if (!Actor)
+		{
+			continue;
+		}
+
+		if (Actor->IsHiddenEd())
+		{
+			FMovieGraphActorQueryEntry& HiddenActorEntry = HiddenActors.AddDefaulted_GetRef();
+			HiddenActorEntry.ActorToMatch = Actor;
+		}
+		else
+		{
+			FMovieGraphActorQueryEntry& VisibleActorEntry = VisibleActors.AddDefaulted_GetRef();
+			VisibleActorEntry.ActorToMatch = Actor;
+		}
+	}
 
-			const FName VisibleActorsCollectionName(TEXT("__AUTOGEN_COLLECTION_VISIBLE_ACTORS"));
-			const FName HiddenActorsCollectionName(TEXT("__AUTOGEN_COLLECTION_HIDDEN_ACTORS"));
-			VisibleActorsCollectionNode->Collection->SetCollectionName(VisibleActorsCollectionName.ToString());
-			HiddenActorsCollectionNode->Collection->SetCollectionName(HiddenActorsCollectionName.ToString());
+	const UMovieGraphCollectionNode* VisibleActorsCollectionNode = Cast<UMovieGraphCollectionNode>(InjectNodeIntoBranch(UMovieGraphCollectionNode::StaticClass(), UMovieGraphNode::GlobalsPinName));
+	const UMovieGraphCollectionNode* HiddenActorsCollectionNode = Cast<UMovieGraphCollectionNode>(InjectNodeIntoBranch(UMovieGraphCollectionNode::StaticClass(), UMovieGraphNode::GlobalsPinName));
 
-			UMovieGraphConditionGroupQuery_Actor* VisibleActorsQuery =
-				Cast<UMovieGraphConditionGroupQuery_Actor>(VisibleActorsCollectionNode->Collection->AddConditionGroup()->AddQuery(UMovieGraphConditionGroupQuery_Actor::StaticClass()));
-			VisibleActorsQuery->ActorsAndComponentsToMatch = VisibleActors;
+	const FName VisibleActorsCollectionName(TEXT("__AUTOGEN_COLLECTION_VISIBLE_ACTORS"));
+	const FName HiddenActorsCollectionName(TEXT("__AUTOGEN_COLLECTION_HIDDEN_ACTORS"));
+	VisibleActorsCollectionNode->Collection->SetCollectionName(VisibleActorsCollectionName.ToString());
+	HiddenActorsCollectionNode->Collection->SetCollectionName(HiddenActorsCollectionName.ToString());
 
-			UMovieGraphConditionGroupQuery_Actor* HiddenActorsQuery =
-				Cast<UMovieGraphConditionGroupQuery_Actor>(HiddenActorsCollectionNode->Collection->AddConditionGroup()->AddQuery(UMovieGraphConditionGroupQuery_Actor::StaticClass()));
-			HiddenActorsQuery->ActorsAndComponentsToMatch = HiddenActors;
+	UMovieGraphConditionGroupQuery_Actor* VisibleActorsQuery =
+		Cast<UMovieGraphConditionGroupQuery_Actor>(VisibleActorsCollectionNode->Collection->AddConditionGroup()->AddQuery(UMovieGraphConditionGroupQuery_Actor::StaticClass()));
+	VisibleActorsQuery->ActorsAndComponentsToMatch = VisibleActors;
 
-			UMovieGraphModifierNode* VisibleActorsModifierNode = Cast<UMovieGraphModifierNode>(InjectNodeIntoGlobals(UMovieGraphModifierNode::StaticClass()));
-			UMovieGraphModifierNode* HiddenActorsModifierNode = Cast<UMovieGraphModifierNode>(InjectNodeIntoGlobals(UMovieGraphModifierNode::StaticClass()));
-			VisibleActorsModifierNode->ModifierName = TEXT("__AUTOGEN_MODIFIER_VISIBLE_ACTORS");
-			HiddenActorsModifierNode->ModifierName = TEXT("__AUTOGEN_MODIFIER_HIDDEN_ACTORS");
-			VisibleActorsModifierNode->AddCollection(VisibleActorsCollectionName);
-			HiddenActorsModifierNode->AddCollection(HiddenActorsCollectionName);
+	UMovieGraphConditionGroupQuery_Actor* HiddenActorsQuery =
+		Cast<UMovieGraphConditionGroupQuery_Actor>(HiddenActorsCollectionNode->Collection->AddConditionGroup()->AddQuery(UMovieGraphConditionGroupQuery_Actor::StaticClass()));
+	HiddenActorsQuery->ActorsAndComponentsToMatch = HiddenActors;
 
-			UMovieGraphRenderPropertyModifier* VisibleActorsModifier = Cast<UMovieGraphRenderPropertyModifier>(VisibleActorsModifierNode->GetModifier(UMovieGraphRenderPropertyModifier::StaticClass()));
-			UMovieGraphRenderPropertyModifier* HiddenActorsModifier = Cast<UMovieGraphRenderPropertyModifier>(HiddenActorsModifierNode->GetModifier(UMovieGraphRenderPropertyModifier::StaticClass()));
+	UMovieGraphModifierNode* VisibleActorsModifierNode = Cast<UMovieGraphModifierNode>(InjectNodeIntoBranch(UMovieGraphModifierNode::StaticClass(), UMovieGraphNode::GlobalsPinName));
+	UMovieGraphModifierNode* HiddenActorsModifierNode = Cast<UMovieGraphModifierNode>(InjectNodeIntoBranch(UMovieGraphModifierNode::StaticClass(), UMovieGraphNode::GlobalsPinName));
+	VisibleActorsModifierNode->ModifierName = TEXT("__AUTOGEN_MODIFIER_VISIBLE_ACTORS");
+	HiddenActorsModifierNode->ModifierName = TEXT("__AUTOGEN_MODIFIER_HIDDEN_ACTORS");
+	VisibleActorsModifierNode->AddCollection(VisibleActorsCollectionName);
+	HiddenActorsModifierNode->AddCollection(HiddenActorsCollectionName);
 
-			VisibleActorsModifier->bOverride_bIsHidden = true;
-			VisibleActorsModifier->bIsHidden = false;
-			HiddenActorsModifier->bOverride_bIsHidden = true;
-			HiddenActorsModifier->bIsHidden = true;
+	UMovieGraphRenderPropertyModifier* VisibleActorsModifier = Cast<UMovieGraphRenderPropertyModifier>(VisibleActorsModifierNode->GetModifier(UMovieGraphRenderPropertyModifier::StaticClass()));
+	UMovieGraphRenderPropertyModifier* HiddenActorsModifier = Cast<UMovieGraphRenderPropertyModifier>(HiddenActorsModifierNode->GetModifier(UMovieGraphRenderPropertyModifier::StaticClass()));
+
+	VisibleActorsModifier->bOverride_bIsHidden = true;
+	VisibleActorsModifier->bIsHidden = false;
+	HiddenActorsModifier->bOverride_bIsHidden = true;
+	HiddenActorsModifier->bIsHidden = true;
+}
+
+void FMovieGraphQuickRender::ApplyQuickRenderUpdatesToDuplicateGraph_ApplyOcio(UMoviePipelineExecutorJob* InJob, const FOpenColorIODisplayConfiguration* InOcioConfiguration)
+{
+	FMovieGraphTraversalContext Context;
+	Context.Job = InJob;
+	Context.RootGraph = TemporaryGraph.Get();
+	FString OutError;
+
+	// Evaluate the graph. This is done so we can tell which branch(es) have nodes we need to update in order to properly apply
+	// OCIO. For example, if the evaluated graph shows a deferred renderer and JPEG node on "Branch1", then update Branch1
+	// accordingly. But if "Branch2" has no deferred renderer, then there's nothing to do there.
+	//
+	// Ideally we would not have to evaluate the graph here, and we could get the evaluated graph from elsewhere within the pipeline.
+	// However, that would require exposing the inner workings of the pipeline to an outside API, which was deemed a bigger negative
+	// than taking the (very small) performance hit of doing a graph evaluation here.
+	UMovieGraphEvaluatedConfig* FlattenedConfig = TemporaryGraph->CreateFlattenedGraph(Context, OutError);
+
+	const TArray<FName> BranchNames = FlattenedConfig->GetBranchNames();
+
+	// For each branch, apply OCIO configuration.
+	for (const FName& BranchName : BranchNames)
+	{
+		constexpr bool bIncludeCDOs = false;
+		constexpr bool bExactMatch = false;
+		const TArray<UMovieGraphFileOutputNode*> FileOutputNodes =
+			FlattenedConfig->GetSettingsForBranch<UMovieGraphFileOutputNode>(BranchName, bIncludeCDOs, bExactMatch);
+		const TArray<UMovieGraphImagePassBaseNode*> ImagePassNodes =
+			FlattenedConfig->GetSettingsForBranch<UMovieGraphImagePassBaseNode>(BranchName, bIncludeCDOs, bExactMatch);
+
+		// Update all file output nodes to have the viewport's OCIO configuration
+		for (UMovieGraphFileOutputNode* FileOutputNode : FileOutputNodes)
+		{
+			static const FName OcioOverrideName(TEXT("bOverride_OCIOConfiguration"));
+			static const FName OcioConfigName(TEXT("OCIOConfiguration"));
+
+			// The evaluated graph contains a node of this type, so insert a new node into the non-evaluated graph downstream to
+			// override the OCIO settings.
+			if (UMovieGraphNode* NewFileOutputNode = InjectNodeIntoBranch(FileOutputNode->GetClass(), BranchName))
+			{
+				// This is hacky, but OCIO properties are not uniformly inherited across image/video nodes. The properties are, however,
+				// named the same.
+				const FBoolProperty* OverrideProp = FindFProperty<FBoolProperty>(NewFileOutputNode->GetClass(), OcioOverrideName);
+				const FStructProperty* OcioProp = FindFProperty<FStructProperty>(NewFileOutputNode->GetClass(), OcioConfigName);
+				
+				if (OverrideProp && OcioProp)
+				{
+					OverrideProp->SetPropertyValue_InContainer(NewFileOutputNode, true);
+					OcioProp->SetValue_InContainer(NewFileOutputNode, InOcioConfiguration);
+				}
+			}
 		}
-		
-		// Apply OCIO if it is activated on the viewport
-#if WITH_OCIO
-		if (const FEditorViewportClient* ViewportClient = UMovieGraphApplyViewportLookNode::GetViewportClient())
+
+		// Also enable OCIO on the renderer nodes
+		for (UMovieGraphImagePassBaseNode* ImagePassNode : ImagePassNodes)
 		{
-			if (const FOpenColorIODisplayConfiguration* OCIOConfiguration = IOpenColorIOModule::Get().GetDisplayManager().GetDisplayConfiguration(ViewportClient))
+			static const FName AllowOcioName(TEXT("bAllowOCIO"));
+			static const FName DisableToneCurveOverrideName(TEXT("bOverride_bDisableToneCurve"));
+			static const FName DisableToneCurveName(TEXT("bDisableToneCurve"));
+
+			const FBoolProperty* AllowOcioProp = FindFProperty<FBoolProperty>(ImagePassNode->GetClass(), AllowOcioName);
+			if (!AllowOcioProp)
+			{
+				// This node doesn't support OCIO
+				continue;
+			}
+
+			// Skip nodes that have "Allow OCIO" explicitly turned off
+			bool AllowOcioValue;
+			AllowOcioProp->GetValue_InContainer(ImagePassNode, &AllowOcioValue);
+			if (!AllowOcioValue)
 			{
-				if (OCIOConfiguration->bIsEnabled)
+				UE_LOG(LogMovieRenderPipeline, Warning, TEXT("Quick Render was set to apply viewport OCIO, but the renderer node [%s] has Allow OCIO turned off. OCIO will not be applied in this node's render."), *ImagePassNode->GetName());
+				continue;
+			}
+
+			// Insert a new node of this type into the non-evaluated graph downstream to override the OCIO settings.
+			if (UMovieGraphNode* NewNode = InjectNodeIntoBranch(ImagePassNode->GetClass(), BranchName))
+			{
+				// Like the file output nodes, OCIO properties are not uniformly inherited across renderer nodes. The properties are, however,
+				// named the same.
+				UMovieGraphImagePassBaseNode* NewRendererNode = Cast<UMovieGraphImagePassBaseNode>(NewNode);
+
+				const FBoolProperty* DisableToneCurveOverrideProp = FindFProperty<FBoolProperty>(NewRendererNode->GetClass(), DisableToneCurveOverrideName);
+				const FBoolProperty* DisableToneCurveProp = FindFProperty<FBoolProperty>(NewRendererNode->GetClass(), DisableToneCurveName);
+
+				// At this point "Allow OCIO" is turned on, so we just need to make sure the tone curve is disabled.
+				if (DisableToneCurveOverrideProp && DisableToneCurveProp)
 				{
-					TArray<TObjectPtr<UMovieGraphNode>> FileOutputNodes = TemporaryGraph->GetNodes().FilterByPredicate([](const TObjectPtr<UMovieGraphNode>& InNode)
-					{
-						return InNode && InNode->IsA<UMovieGraphFileOutputNode>();
-					});
-
-					TArray<TObjectPtr<UMovieGraphNode>> RendererNodes = TemporaryGraph->GetNodes().FilterByPredicate([](const TObjectPtr<UMovieGraphNode>& InNode)
-					{
-						return InNode && InNode->IsA<UMovieGraphDeferredRenderPassNode>();
-					});
-
-					// Specify the OCIO configuration on all nodes that support it
-					for (const TObjectPtr<UMovieGraphNode>& OutputNode : FileOutputNodes)
-					{
-						static const FName OCIOOverrideName(TEXT("bOverride_OCIOConfiguration"));
-						static const FName OCIOConfigName(TEXT("OCIOConfiguration"));
-						
-						// This is hacky, but OCIO properties are not uniformly inherited across image/video nodes. The properties are, however,
-						// named the same.
-						const FBoolProperty* OverrideProp = FindFProperty<FBoolProperty>(OutputNode->GetClass(), FName(OCIOOverrideName));
-						const FStructProperty* OCIOProp = FindFProperty<FStructProperty>(OutputNode->GetClass(), FName(OCIOConfigName));
-						
-						if (OverrideProp && OCIOProp)
-						{
-							OverrideProp->SetPropertyValue_InContainer(OutputNode.Get(), true);
-							OCIOProp->SetValue_InContainer(OutputNode.Get(), OCIOConfiguration);
-						}
-					}
-					
-					// Also enable OCIO on the renderer nodes
-					for (const TObjectPtr<UMovieGraphNode>& RendererNode : RendererNodes)
-					{
-						UMovieGraphDeferredRenderPassNode* DeferredRenderNode = Cast<UMovieGraphDeferredRenderPassNode>(RendererNode.Get());
-						if (!DeferredRenderNode)
-						{
-							continue;
-						}
-
-						if (!DeferredRenderNode->bAllowOCIO)
-						{
-							UE_LOG(LogMovieRenderPipeline, Warning, TEXT("Quick Render was set to apply Viewport OCIO, but the renderer node [%s] has Allow OCIO turned off. OCIO will not be applied in this node's render."), *DeferredRenderNode->GetName());
-							continue;
-						}
-						
-						DeferredRenderNode->bOverride_bAllowOCIO = true;
-						DeferredRenderNode->bOverride_bDisableToneCurve = true;
-						DeferredRenderNode->bAllowOCIO = true;
-						DeferredRenderNode->bDisableToneCurve = true;
-					}
+					DisableToneCurveOverrideProp->SetPropertyValue_InContainer(NewRendererNode, true);
+					DisableToneCurveProp->SetPropertyValue_InContainer(NewRendererNode, true);
 				}
 			}
 		}
-#endif	// WITH_OCIO
 	}
 }
 
-UMovieGraphNode* FMovieGraphQuickRender::InjectNodeIntoGlobals(const TSubclassOf<UMovieGraphNode> NodeType)
+UMovieGraphNode* FMovieGraphQuickRender::InjectNodeIntoBranch(const TSubclassOf<UMovieGraphNode> NodeType, const FName& InBranchName)
 {
 	check(TemporaryGraph.IsValid());
 
 	UMovieGraphNode* NewNode = TemporaryGraph->CreateNodeByClass(NodeType);
-
 	UMovieGraphNode* OutputNode = TemporaryGraph->GetOutputNode();
-	const UMovieGraphPin* GlobalsOutputPin = OutputNode->GetInputPin(UMovieGraphNode::GlobalsPinName);
 
-	// Get the node on the Globals branch that is most downstream (if any)
-	UMovieGraphNode* MostDownstreamGlobalsNode = nullptr;
-	if (const UMovieGraphPin* ConnectedPin = GlobalsOutputPin->GetFirstConnectedPin())
+	const UMovieGraphPin* BranchPin = OutputNode->GetInputPin(InBranchName);
+	if (!BranchPin)
+	{
+		UE_LOG(LogMovieRenderPipeline, Warning, TEXT("Quick Render: Could not find branch [%s] to inject a node override into. The render may look different than expected."), *InBranchName.ToString());
+		return nullptr;
+	}
+
+	// Get the node on the branch that is most downstream (if any)
+	UMovieGraphNode* MostDownstreamNode = nullptr;
+	if (const UMovieGraphPin* ConnectedPin = BranchPin->GetFirstConnectedPin())
 	{
-		MostDownstreamGlobalsNode = ConnectedPin->Node;
+		MostDownstreamNode = ConnectedPin->Node;
 	}
 
-	// Add in the new node to Globals, downstream from the existing node that was previously furthest downstream
-	TemporaryGraph->AddLabeledEdge(NewNode, FName(), OutputNode, UMovieGraphNode::GlobalsPinName);
+	// Add in the new node to the branch, downstream from the existing node that was previously furthest downstream
+	TemporaryGraph->AddLabeledEdge(NewNode, FName(), OutputNode, InBranchName);
 
 	// Re-connect the previously connected node upstream of the new node, if there was one
-	if (MostDownstreamGlobalsNode)
+	if (MostDownstreamNode)
 	{
-		TemporaryGraph->AddLabeledEdge(MostDownstreamGlobalsNode, FName(), NewNode, FName());
+		TemporaryGraph->AddLabeledEdge(MostDownstreamNode, FName(), NewNode, FName());
 	}
 
 	return NewNode;
diff --git a/Engine/Plugins/MovieScene/MovieRenderPipeline/Source/MovieRenderPipelineEditor/Public/Graph/MovieGraphQuickRender.h b/Engine/Plugins/MovieScene/MovieRenderPipeline/Source/MovieRenderPipelineEditor/Public/Graph/MovieGraphQuickRender.h
--- a/Engine/Plugins/MovieScene/MovieRenderPipeline/Source/MovieRenderPipelineEditor/Public/Graph/MovieGraphQuickRender.h
+++ b/Engine/Plugins/MovieScene/MovieRenderPipeline/Source/MovieRenderPipelineEditor/Public/Graph/MovieGraphQuickRender.h
@@ -30,10 +30,16 @@ private:
 	static UMovieGraphConfig* GetQuickRenderGraph(UMovieGraphConfig* InUserSpecifiedGraph, TMap<TObjectPtr<UMovieGraphConfig>, TObjectPtr<UMovieGraphConfig>>& OutOriginalToDupeMap);
 
 	/** Applies all necessary updates to the duplicated graph to prepare it for a quick render (adjusting start/end frames, etc). */
-	static void ApplyQuickRenderUpdatesToDuplicateGraph(const UMovieGraphQuickRenderSettings& InQuickRenderSettings, UWorld* InEditorWorld);
+	static void ApplyQuickRenderUpdatesToDuplicateGraph(const UMovieGraphQuickRenderSettings& InQuickRenderSettings, UWorld* InEditorWorld, UMoviePipelineExecutorJob* InJob);
 
-	/** Injects a new node in the furthest-downstream position within the Globals branch in the duplicated graph. */
-	static UMovieGraphNode* InjectNodeIntoGlobals(const TSubclassOf<UMovieGraphNode> NodeType);
+	/** Helper for ApplyQuickRenderUpdatesToDuplicateGraph() that applies actor editor visibility (by injecting a visibility modifier into the graph). */
+	static void ApplyQuickRenderUpdatesToDuplicateGraph_ApplyEditorVisibility(const UWorld* InEditorWorld);
+
+	/** Helper for ApplyQuickRenderUpdatesToDuplicateGraph() that applies viewport OCIO (by injecting overrides to file output and renderer nodes into the graph). */
+	static void ApplyQuickRenderUpdatesToDuplicateGraph_ApplyOcio(UMoviePipelineExecutorJob* InJob, const FOpenColorIODisplayConfiguration* InOcioConfiguration);
+
+	/** Injects a new node in the furthest-downstream position within the specified branch in the duplicated graph. */
+	static UMovieGraphNode* InjectNodeIntoBranch(const TSubclassOf<UMovieGraphNode> NodeType, const FName& InBranchName);
 
 private:
 	/** The temporary queue that is used by quick render. */
diff --git a/Engine/Plugins/MovieScene/MovieRenderPipeline/Source/MovieRenderPipelineRenderPasses/Private/Graph/Renderers/MovieGraphImagePassBase.cpp b/Engine/Plugins/MovieScene/MovieRenderPipeline/Source/MovieRenderPipelineRenderPasses/Private/Graph/Renderers/MovieGraphImagePassBase.cpp
--- a/Engine/Plugins/MovieScene/MovieRenderPipeline/Source/MovieRenderPipelineRenderPasses/Private/Graph/Renderers/MovieGraphImagePassBase.cpp
+++ b/Engine/Plugins/MovieScene/MovieRenderPipeline/Source/MovieRenderPipelineRenderPasses/Private/Graph/Renderers/MovieGraphImagePassBase.cpp
@@ -168,8 +168,15 @@ TSharedRef<FSceneViewFamilyContext> FMovieGraphImagePassBase::CreateSceneViewFam
 			
 			if (ViewportLookNode->GetViewportInfo(ViewportShowFlags, ViewportViewModeIndex))
 			{
-				ViewModeIndex = ViewportViewModeIndex;
-				ShowFlags = ViewportShowFlags;
+				if (ViewportLookNode->bOverride_bViewMode && ViewportLookNode->bViewMode)
+				{
+					ViewModeIndex = ViewportViewModeIndex;
+				}
+
+				if (ViewportLookNode->bOverride_bShowFlags && ViewportLookNode->bShowFlags)
+				{
+					ShowFlags = ViewportShowFlags;
+				}
 			}
 		}
 	}
