MoviePipeline: Implemented the "Current Viewport" mode in Quick Render.

Also fixed a small hidden state issue with volumetrics in the visibility modifier.

diff --git a/Engine/Plugins/MovieScene/MovieRenderPipeline/Source/MovieRenderPipelineCore/Private/Graph/MovieGraphRenderLayerSubsystem.cpp b/Engine/Plugins/MovieScene/MovieRenderPipeline/Source/MovieRenderPipelineCore/Private/Graph/MovieGraphRenderLayerSubsystem.cpp
--- a/Engine/Plugins/MovieScene/MovieRenderPipeline/Source/MovieRenderPipelineCore/Private/Graph/MovieGraphRenderLayerSubsystem.cpp
+++ b/Engine/Plugins/MovieScene/MovieRenderPipeline/Source/MovieRenderPipelineCore/Private/Graph/MovieGraphRenderLayerSubsystem.cpp
@@ -323,7 +323,7 @@ void UMovieGraphRenderPropertyModifier::SetActorAndComponentState(const TArray<F
 
 		if (bOverride_bIsHidden)
 		{
-			VolumetricsComponent->SetRenderInMainPass(NewComponentState.bIsHidden);
+			VolumetricsComponent->SetRenderInMainPass(!NewComponentState.bIsHidden);
 		}
 
 		if (bOverride_bAffectIndirectLightingWhileHidden)
diff --git a/Engine/Plugins/MovieScene/MovieRenderPipeline/Source/MovieRenderPipelineCore/Private/Graph/Nodes/MovieGraphApplyViewportLookNode.cpp b/Engine/Plugins/MovieScene/MovieRenderPipeline/Source/MovieRenderPipelineCore/Private/Graph/Nodes/MovieGraphApplyViewportLookNode.cpp
--- a/Engine/Plugins/MovieScene/MovieRenderPipeline/Source/MovieRenderPipelineCore/Private/Graph/Nodes/MovieGraphApplyViewportLookNode.cpp
+++ b/Engine/Plugins/MovieScene/MovieRenderPipeline/Source/MovieRenderPipelineCore/Private/Graph/Nodes/MovieGraphApplyViewportLookNode.cpp
@@ -10,6 +10,7 @@
 #if WITH_EDITOR
 #include "Editor.h"
 #include "EditorViewportClient.h"
+#include "LevelEditorViewport.h"
 #endif
 
 EMovieGraphBranchRestriction UMovieGraphApplyViewportLookNode::GetBranchRestriction() const
@@ -20,9 +21,16 @@ EMovieGraphBranchRestriction UMovieGraphApplyViewportLookNode::GetBranchRestrict
 #if WITH_EDITOR
 FEditorViewportClient* UMovieGraphApplyViewportLookNode::GetViewportClient()
 {
+	// Get the currently active *perspective* viewport that's in use. Only perspective viewports are currently supported in Quick Render.
 	if (const FViewport* ActiveViewport = GEditor->GetActiveViewport())
 	{
-		return static_cast<FEditorViewportClient*>(ActiveViewport->GetClient());
+		for (FEditorViewportClient* LevelViewportClient : GEditor->GetLevelViewportClients())
+		{
+			if (LevelViewportClient->IsPerspective() && (LevelViewportClient->Viewport == ActiveViewport))
+			{
+				return LevelViewportClient;
+			}
+		}
 	}
 
 	return nullptr;
diff --git a/Engine/Plugins/MovieScene/MovieRenderPipeline/Source/MovieRenderPipelineEditor/Private/Graph/MovieGraphQuickRender.cpp b/Engine/Plugins/MovieScene/MovieRenderPipeline/Source/MovieRenderPipelineEditor/Private/Graph/MovieGraphQuickRender.cpp
--- a/Engine/Plugins/MovieScene/MovieRenderPipeline/Source/MovieRenderPipelineEditor/Private/Graph/MovieGraphQuickRender.cpp
+++ b/Engine/Plugins/MovieScene/MovieRenderPipeline/Source/MovieRenderPipelineEditor/Private/Graph/MovieGraphQuickRender.cpp
@@ -2,9 +2,11 @@
 
 #include "Graph/MovieGraphQuickRender.h"
 
+#include "Camera/CameraActor.h"
 #include "Editor.h"
 #include "EngineUtils.h"
 #include "LevelSequence.h"
+#include "LevelSequenceActor.h"
 #include "LevelSequenceEditorBlueprintLibrary.h"
 #include "Misc/MessageDialog.h"
 #include "MoviePipelinePIEExecutor.h"
@@ -13,7 +15,12 @@
 #include "MoviePipelineUtils.h"
 #include "MovieRenderPipelineCoreModule.h"
 #include "MovieRenderPipelineSettings.h"
+#include "MovieScene.h"
 #include "ObjectEditorUtils.h"
+#include "Sections/MovieSceneCameraCutSection.h"
+#include "SequencerUtilities.h"
+#include "Tracks/MovieSceneCameraCutTrack.h"
+
 #include "Graph/MovieGraphRenderLayerSubsystem.h"
 #include "Graph/MovieRenderGraphEditorSettings.h"
 #include "Graph/Nodes/MovieGraphApplyViewportLookNode.h"
@@ -54,25 +61,15 @@ void FMovieGraphQuickRender::BeginQuickRender(const UMovieGraphQuickRenderSettin
 	// Always use the PIE executor for quick render, rather than the executor specified in the project settings for normal MRG.
 	const TSubclassOf<UMoviePipelineExecutorBase> ExecutorClass = UMoviePipelinePIEExecutor::StaticClass();
 
-	// Get the level sequence that Sequencer is currently displaying
-	const ULevelSequence* CurrentLevelSequence = ULevelSequenceEditorBlueprintLibrary::GetCurrentLevelSequence();
-	if (!CurrentLevelSequence)
+	// Get/create the level sequence that's appropriate for the current Quick Render mode
+	SetUpLevelSequenceForRender(InQuickRenderSettings);
+	if (!TemporaryLevelSequence.IsValid())
 	{
-		UE_LOG(LogMovieRenderPipeline, Error, TEXT("Unable to Quick Render: No level sequence currently active to render from."));
 		return;
 	}
 
 	// Find the editor world
-	UWorld* EditorWorld = nullptr;
-	for (const FWorldContext& Context : GEngine->GetWorldContexts())
-	{
-		if (Context.WorldType == EWorldType::Editor)
-		{
-			EditorWorld = Context.World();
-			break;
-		}
-	}
-	
+	UWorld* EditorWorld = GetWorldOfType(EWorldType::Editor);
 	if (!EditorWorld)
 	{
 		UE_LOG(LogMovieRenderPipeline, Error, TEXT("Unable to Quick Render: Could not find a currently active editor world."));
@@ -84,7 +81,7 @@ void FMovieGraphQuickRender::BeginQuickRender(const UMovieGraphQuickRenderSettin
 
 	// Add a temp job to the queue, and give it the editor's level and the current level sequence being edited in Sequencer
 	UMoviePipelineExecutorJob* NewJob = TemporaryQueue->AllocateNewJob(UMoviePipelineExecutorJob::StaticClass());
-	NewJob->SetSequence(FSoftObjectPath(CurrentLevelSequence));
+	NewJob->SetSequence(FSoftObjectPath(TemporaryLevelSequence.Get()));
 	NewJob->Map = FSoftObjectPath(EditorWorld);
 	NewJob->JobName = FString(TEXT("QuickRender"));
 
@@ -134,9 +131,26 @@ void FMovieGraphQuickRender::BeginQuickRender(const UMovieGraphQuickRenderSettin
 	// should be fully initialized at this point.
 	ApplyQuickRenderUpdatesToDuplicateGraph(InQuickRenderSettings, EditorWorld, NewJob);
 
+	// Prevent all existing level sequences in the world from auto-playing. We need the level sequence from Quick Render to be in control. 
+	for (auto It = TActorIterator<ALevelSequenceActor>(EditorWorld); It; ++It)
+	{
+		const ALevelSequenceActor* LevelSequenceActor = *It;
+		if (LevelSequenceActor->PlaybackSettings.bAutoPlay)
+		{
+			(*It)->PlaybackSettings.bAutoPlay = false;
+			ModifiedLevelSequenceActors.Add(MakeWeakObjectPtr(*It));
+		}
+	}
+
 	// Do the render
 	UMoviePipelinePIEExecutor* Executor = Cast<UMoviePipelinePIEExecutor>(Subsystem->RenderQueueInstanceWithExecutor(TemporaryQueue.Get(), ExecutorClass));
 
+	// Do any needed setup in the PIE world before rendering starts
+	Executor->OnIndividualJobStarted().AddLambda([&InQuickRenderSettings](UMoviePipelineExecutorJob* InJob)
+	{
+		PerformPreRenderSetup(InQuickRenderSettings);
+	});
+
 	// View the job's rendered frame(s)/video(s) if requested
 	Executor->OnIndividualJobWorkFinished().AddLambda([&InQuickRenderSettings](FMoviePipelineOutputData OutputData)
 	{
@@ -151,9 +165,19 @@ void FMovieGraphQuickRender::BeginQuickRender(const UMovieGraphQuickRenderSettin
 	// Clean up
 	Executor->OnExecutorFinished().AddLambda([](UMoviePipelineExecutorBase* InExecutor, bool bSuccess)
 	{
-		// The temp queue and graph are part of the transient package, so manually release them when the render is finished
+		// The temp objects are part of the transient package, so manually release them when the render is finished
 		TemporaryQueue.Reset();
 		TemporaryGraph.Reset();
+		TemporaryLevelSequence.Reset();
+
+		// Restore the modified level sequences to the way they were before the render
+		for (const TWeakObjectPtr<ALevelSequenceActor>& LevelSequenceActor : ModifiedLevelSequenceActors)
+		{
+			if (LevelSequenceActor.IsValid())
+			{
+				LevelSequenceActor->PlaybackSettings.bAutoPlay = true;
+			}
+		}
 	});
 }
 
@@ -447,4 +471,91 @@ UMovieGraphNode* FMovieGraphQuickRender::InjectNodeIntoBranch(const TSubclassOf<
 	return NewNode;
 }
 
+void FMovieGraphQuickRender::SetUpLevelSequenceForRender(const UMovieGraphQuickRenderSettings& InQuickRenderSettings)
+{
+	if (InQuickRenderSettings.QuickRenderMode == EMovieGraphQuickRenderMode::CurrentViewport)
+	{
+		// Create the new, temporary level sequence
+		ULevelSequence* QuickRenderSequence = NewObject<ULevelSequence>(GetTransientPackage(), TEXT("QuickRenderSequence"), RF_Transient);
+		QuickRenderSequence->Initialize();
+
+		TemporaryLevelSequence.Reset(QuickRenderSequence);
+		return;
+	}
+
+	// Use the level sequence that's currently active
+	ULevelSequence* CurrentLevelSequence = ULevelSequenceEditorBlueprintLibrary::GetCurrentLevelSequence();
+	if (!CurrentLevelSequence)
+	{
+		UE_LOG(LogMovieRenderPipeline, Error, TEXT("Unable to Quick Render: No level sequence currently active to render from ('Current Sequence' mode is active)."));
+		return;
+	}
+	
+	TemporaryLevelSequence.Reset(CurrentLevelSequence);
+}
+
+UWorld* FMovieGraphQuickRender::GetWorldOfType(EWorldType::Type WorldType)
+{
+	UWorld* World = nullptr;
+	for (const FWorldContext& Context : GEngine->GetWorldContexts())
+	{
+		if (Context.WorldType == WorldType)
+		{
+			World = Context.World();
+			break;
+		}
+	}
+
+	return World;
+}
+
+void FMovieGraphQuickRender::PerformPreRenderSetup(const UMovieGraphQuickRenderSettings& InQuickRenderSettings)
+{
+	if (InQuickRenderSettings.QuickRenderMode == EMovieGraphQuickRenderMode::CurrentViewport)
+	{
+		// Get the PIE world, this is where we need to target the changes (not the editor world). Otherwise we'll have to deal with deleting actors
+		// after the render finishes, which isn't ideal (and this will also create entries in the undo stack, also not ideal).
+		UWorld* PieWorld = GetWorldOfType(EWorldType::PIE);
+		
+		// Add a new camera that mimics the viewport's camera to the level sequence. Use a regular camera actor here rather than a cine camera. Using
+		// a cine camera massively complicates things and creates a situation where it's difficult to match what the viewport looks like (eg, exposure).
+		FActorSpawnParameters CameraSpawnParams;
+		CameraSpawnParams.OverridePackage = GetTransientPackage();
+		CameraSpawnParams.bCreateActorPackage = false;
+		CameraSpawnParams.OverrideLevel = PieWorld->GetCurrentLevel();
+		const FEditorViewportClient* ViewportClient = UMovieGraphApplyViewportLookNode::GetViewportClient();
+		ACameraActor* NewCamera = Cast<ACameraActor>(PieWorld->SpawnActor<ACameraActor>(CameraSpawnParams));
+
+		// Update the camera to look like the viewport's camera
+		NewCamera->SetActorLocation(ViewportClient->GetViewLocation(), false);
+		NewCamera->SetActorRotation(ViewportClient->GetViewRotation());
+		NewCamera->GetCameraComponent()->SetFieldOfView(ViewportClient->ViewFOV);
+
+		// Add a new camera cut track. We only need 1 frame in Viewport mode, so use a range of [0, 1]
+		UMovieScene* MovieScene = TemporaryLevelSequence->GetMovieScene();
+		const FFrameTime EndOfRange = FFrameRate::TransformTime(FFrameTime(1), MovieScene->GetDisplayRate(), MovieScene->GetTickResolution());
+		MovieScene->SetPlaybackRange(TRange<FFrameNumber>::Inclusive(FFrameNumber(0), EndOfRange.GetFrame()));
+		UMovieSceneCameraCutTrack* CameraCutTrack = CastChecked<UMovieSceneCameraCutTrack>(MovieScene->AddCameraCutTrack(UMovieSceneCameraCutTrack::StaticClass()));
+		
+		UE::Sequencer::FCreateBindingParams CreateBindingParams;
+		CreateBindingParams.bSpawnable = true;
+		CreateBindingParams.bAllowCustomBinding = true;
+		CreateBindingParams.BindingNameOverride = NewCamera->GetActorLabel();
+
+		// Create the binding for the camera. This is wrapped in a dummy transaction so we can cancel it later on to prevent the binding creation from
+		// making an entry on the undo stack.
+#if WITH_EDITOR
+		const int32 TransactionId = GEditor->BeginTransaction(FText::GetEmpty());
+#endif
+		const FGuid CameraBinding = FSequencerUtilities::CreateOrReplaceBinding(nullptr, MovieScene->GetTypedOuter<UMovieSceneSequence>(), NewCamera, CreateBindingParams);
+#if WITH_EDITOR
+		GEditor->CancelTransaction(TransactionId);
+#endif
+
+		// Add the new camera to the camera cut track.
+		UMovieSceneCameraCutSection* CameraCutSection = CameraCutTrack->AddNewCameraCut(UE::MovieScene::FRelativeObjectBindingID(CameraBinding), FFrameNumber(0));
+		CameraCutSection->SetRange(MovieScene->GetPlaybackRange());
+	}
+}
+
 #undef LOCTEXT_NAMESPACE
diff --git a/Engine/Plugins/MovieScene/MovieRenderPipeline/Source/MovieRenderPipelineEditor/Public/Graph/MovieGraphQuickRender.h b/Engine/Plugins/MovieScene/MovieRenderPipeline/Source/MovieRenderPipelineEditor/Public/Graph/MovieGraphQuickRender.h
--- a/Engine/Plugins/MovieScene/MovieRenderPipeline/Source/MovieRenderPipelineEditor/Public/Graph/MovieGraphQuickRender.h
+++ b/Engine/Plugins/MovieScene/MovieRenderPipeline/Source/MovieRenderPipelineEditor/Public/Graph/MovieGraphQuickRender.h
@@ -7,6 +7,7 @@
 
 #include "Graph/MovieGraphDataTypes.h"
 #include "Graph/MovieGraphQuickRenderSettings.h"
+#include "LevelSequence.h"
 #include "MoviePipelineQueue.h"
 
 class FEditorViewportClient;
@@ -41,10 +42,28 @@ private:
 	/** Injects a new node in the furthest-downstream position within the specified branch in the duplicated graph. */
 	static UMovieGraphNode* InjectNodeIntoBranch(const TSubclassOf<UMovieGraphNode> NodeType, const FName& InBranchName);
 
+	/**
+	 * Populates TemporaryLevelSequence with the current level sequence, or sets up a new level sequence for rendering (depending on the mode that
+	 * Quick Render is using).
+	 */
+	static void SetUpLevelSequenceForRender(const UMovieGraphQuickRenderSettings& InQuickRenderSettings);
+
+	/** Gets the first world found of the specified type. */
+	static UWorld* GetWorldOfType(EWorldType::Type WorldType);
+
+	/** Does any setup needed in the PIE world prior to the render starting. */
+	static void PerformPreRenderSetup(const UMovieGraphQuickRenderSettings& InQuickRenderSettings);
+
 private:
 	/** The temporary queue that is used by quick render. */
 	inline static TStrongObjectPtr<UMoviePipelineQueue> TemporaryQueue = nullptr;
 
 	/** The temporary graph that is used by quick render. Usually a duplicate of a graph created outside of quick render, but will be modified by quick render. */
 	inline static TStrongObjectPtr<UMovieGraphConfig> TemporaryGraph = nullptr;
+
+	/** The level sequence that is being used by Quick Render. In some cases this may be made specifically for a quick render, in other cases it may be a level sequence that existed before. */
+	inline static TStrongObjectPtr<ULevelSequence> TemporaryLevelSequence = nullptr;
+
+	/** Level sequence actors which had their AutoPlay settings turned off prior to render, and need to be turned back on again after the render finishes. */
+	inline static TArray<TWeakObjectPtr<class ALevelSequenceActor>> ModifiedLevelSequenceActors;
 };
diff --git a/Engine/Plugins/MovieScene/MovieRenderPipeline/Source/MovieRenderPipelineRenderPasses/Private/Graph/Renderers/MovieGraphImagePassBase.cpp b/Engine/Plugins/MovieScene/MovieRenderPipeline/Source/MovieRenderPipelineRenderPasses/Private/Graph/Renderers/MovieGraphImagePassBase.cpp
--- a/Engine/Plugins/MovieScene/MovieRenderPipeline/Source/MovieRenderPipelineRenderPasses/Private/Graph/Renderers/MovieGraphImagePassBase.cpp
+++ b/Engine/Plugins/MovieScene/MovieRenderPipeline/Source/MovieRenderPipelineRenderPasses/Private/Graph/Renderers/MovieGraphImagePassBase.cpp
@@ -19,6 +19,10 @@
 #include "Engine/Engine.h"
 #include "ImageUtils.h"
 
+#if WITH_EDITOR
+#include "EditorViewportClient.h"
+#endif
+
 
 namespace UE::MovieGraph::Rendering
 {
@@ -153,8 +157,9 @@ TSharedRef<FSceneViewFamilyContext> FMovieGraphImagePassBase::CreateSceneViewFam
 	EViewModeIndex ViewModeIndex = InInitData.ViewModeIndex;
 	FEngineShowFlags ShowFlags = InInitData.ShowFlags;
 
-	// Viewport rendering mode may need to customize the show flags and view mode index
+	// Viewport rendering mode may need to customize the show flags, view mode index, and exposure
 	bool bIsViewportLook = false;
+	TOptional<FExposureSettings> ExposureSettings;
 #if WITH_EDITOR
 	if (InInitData.TimeData.EvaluatedConfig)
 	{
@@ -178,6 +183,11 @@ TSharedRef<FSceneViewFamilyContext> FMovieGraphImagePassBase::CreateSceneViewFam
 					ShowFlags = ViewportShowFlags;
 				}
 			}
+
+			if (const FEditorViewportClient* ViewportClient = UMovieGraphApplyViewportLookNode::GetViewportClient())
+			{
+				ExposureSettings = ViewportClient->ExposureSettings;
+			}
 		}
 	}
 #endif	// WITH_EDITOR
@@ -198,6 +208,11 @@ TSharedRef<FSceneViewFamilyContext> FMovieGraphImagePassBase::CreateSceneViewFam
 		.SetTime(FGameTime::CreateUndilated(InInitData.TimeData.WorldSeconds, InInitData.TimeData.FrameDeltaTime))
 		.SetRealtimeUpdate(true));
 
+	if (ExposureSettings.IsSet())
+	{
+		OutViewFamily->ExposureSettings = ExposureSettings.GetValue();
+	}
+
 	// Need to add the engine-wide view extensions, as rendering code may depend on them (ie: landscapes)
 	OutViewFamily->ViewExtensions.Append(GEngine->ViewExtensions->GatherActiveExtensions(FSceneViewExtensionContext(InInitData.World->Scene)));
 
@@ -212,8 +227,6 @@ TSharedRef<FSceneViewFamilyContext> FMovieGraphImagePassBase::CreateSceneViewFam
 	return OutViewFamily;
 }
 
-
-
 void FMovieGraphImagePassBase::ApplyMovieGraphOverridesToViewFamily(TSharedRef<FSceneViewFamilyContext> InOutFamily, const FViewFamilyInitData& InInitData) const
 {
 	// Used to specify if the Tone Curve is being applied or not to our Linear Output data
