MoviePipeline: Exposed some UMovieGraphConfig methods to Python, fixed some minor documentation typos, and fixed a bug in GetBranchNames().

diff --git a/Engine/Plugins/MovieScene/MovieRenderPipeline/Source/MovieRenderPipelineCore/Private/Graph/MovieGraphConfig.cpp b/Engine/Plugins/MovieScene/MovieRenderPipeline/Source/MovieRenderPipelineCore/Private/Graph/MovieGraphConfig.cpp
--- a/Engine/Plugins/MovieScene/MovieRenderPipeline/Source/MovieRenderPipelineCore/Private/Graph/MovieGraphConfig.cpp
+++ b/Engine/Plugins/MovieScene/MovieRenderPipeline/Source/MovieRenderPipelineCore/Private/Graph/MovieGraphConfig.cpp
@@ -1292,15 +1292,12 @@ void UMovieGraphConfig::RecurseUpGlobalsBranchToFindOutputDirectory(const UMovie
 	}
 };
 
-const TArray<FName> UMovieGraphConfig::GetBranchNames() const
+TArray<FName> UMovieGraphConfig::GetBranchNames() const
 {
 	TArray<FName> BranchNames;
-	for (UMovieGraphOutput* Output : GetOutputs())
+	for (const UMovieGraphOutput* Output : GetOutputs())
 	{
-		if (Output->IsA<UMovieGraphOutputNode>())
-		{
-			BranchNames.AddUnique(FName(Output->GetMemberName()));
-		}
+		BranchNames.AddUnique(FName(Output->GetMemberName()));
 	}
 	return BranchNames;
 }
diff --git a/Engine/Plugins/MovieScene/MovieRenderPipeline/Source/MovieRenderPipelineCore/Public/Graph/MovieGraphConfig.h b/Engine/Plugins/MovieScene/MovieRenderPipeline/Source/MovieRenderPipelineCore/Public/Graph/MovieGraphConfig.h
--- a/Engine/Plugins/MovieScene/MovieRenderPipeline/Source/MovieRenderPipelineCore/Public/Graph/MovieGraphConfig.h
+++ b/Engine/Plugins/MovieScene/MovieRenderPipeline/Source/MovieRenderPipelineCore/Public/Graph/MovieGraphConfig.h
@@ -626,24 +626,24 @@ public:
 	
 	
 	/**
-	* Convinence function which removes all Inbound (pins on the left side of a node) edges for the given node.
+	* Convenience function which removes all Inbound (pins on the left side of a node) edges for the given node.
 	*/
 	UFUNCTION(BlueprintCallable, Category = "Movie Graph")
 	bool RemoveAllInboundEdges(UMovieGraphNode* InNode);
 	/**
-	* Convinence function which removes all Outbound (pins on the right side of a node) edges for the given node.
+	* Convenience function which removes all Outbound (pins on the right side of a node) edges for the given node.
 	*/
 	UFUNCTION(BlueprintCallable, Category = "Movie Graph")
 	bool RemoveAllOutboundEdges(UMovieGraphNode* InNode);
 
 	/**
-	* Convinence function which removes all Inbound (pins on the left side of a node) edges connected to the given inbound pin by name, for the given node.
+	* Convenience function which removes all Inbound (pins on the left side of a node) edges connected to the given inbound pin by name, for the given node.
 	*/
 	UFUNCTION(BlueprintCallable, Category = "Movie Graph")
 	bool RemoveInboundEdges(UMovieGraphNode* InNode, const FName& InPinName);
 	
 	/**
-	* Convinence function which removes all Outobund (pins on the right side of a node) edges connected to the given outbound pin by name, for the given node.
+	* Convenience function which removes all Outbound (pins on the right side of a node) edges connected to the given outbound pin by name, for the given node.
 	*/
 	UFUNCTION(BlueprintCallable, Category = "Movie Graph")
 	bool RemoveOutboundEdges(UMovieGraphNode* InNode, const FName& InPinName);
@@ -661,7 +661,7 @@ public:
 	UFUNCTION(BlueprintCallable, Category = "Movie Graph")
 	bool RemoveNode(UMovieGraphNode* InNode);
 
-	/** Like RemoveNode but takes an entire array at once for convinence. */
+	/** Like RemoveNode but takes an entire array at once for convenience. */
 	UFUNCTION(BlueprintCallable, Category = "Movie Graph")
 	bool RemoveNodes(TArray<UMovieGraphNode*> InNodes);
 
@@ -677,15 +677,14 @@ public:
 
 	/** Returns an array of the branch names for the OutputNode on this Graph. */
 	UFUNCTION(BlueprintCallable, Category = "Movie Graph")
-	const TArray<FName> GetBranchNames() const;
-
+	TArray<FName> GetBranchNames() const;
 
 	/**
 	* Finds a setting node (by type) for the given branch name (see GetBranchNames()). Returns the first result
-	* of that type found, even if there are multiple, following traversal order (ie: right to left)
+	* of that type found, even if there are multiple, following traversal order (ie: right to left).
 	* To prevent accidentally editing unrelated assets, does not dive into sub-graphs (but can continue
 	* traversal beyond them), so returned results should only exist in current asset.
-	* Does not contain the Input or Output nodes (as they are not settings), see GetInputNode(), GetOutputNode())
+	* Does not contain the Input or Output nodes (as they are not settings), see GetInputNode() and GetOutputNode().
 	*
 	* @param InClass - The node class type to search for.
 	* @param InBranchName - The branch to search for these nodes on.
@@ -701,7 +700,7 @@ public:
 	* following traversal order (ie: right to left).
 	* To prevent accidentally editing unrelated assets, does not dive into sub-graphs (but can continue
 	* traversal beyond them), so returned results should only exist in current asset.
-	* Does not contain the Input or Output nodes (as they are not settings), see GetInputNode(), GetOutputNode())
+	* Does not contain the Input or Output nodes (as they are not settings), see GetInputNode() and GetOutputNode().
 	*
 	* @param InClass - The node class type to search for.
 	* @param InBranchName - The branch to search for these nodes on.
@@ -719,7 +718,7 @@ public:
 	* contain the correct tag.
 	* To prevent accidentally editing unrelated assets, does not dive into sub-graphs (but can continue
 	* traversal beyond them), so returned results should only exist in current asset.
-	* Does not contain the Input or Output nodes (as they are not settings), see GetInputNode(), GetOutputNode())
+	* Does not contain the Input or Output nodes (as they are not settings), see GetInputNode() and GetOutputNode().
 	*
 	* @param ScriptTag - Case sensitive name into the "Script Tags" property on Nodes
 	* @param OptionalClassFilter - Defaults None. If specified, will only return a node of this class that matches the tag.
@@ -776,6 +775,7 @@ public:
 	UMovieGraphOutput* AddOutput(const FText& InBaseName = FText::GetEmpty());
 
 	/** Gets the variable in the graph with the specified GUID, else nullptr if one could not be found. */
+	UFUNCTION(BlueprintCallable, Category="Movie Graph")
 	UMovieGraphVariable* GetVariableByGuid(const FGuid& InGuid) const;
 
 	/**
@@ -802,6 +802,7 @@ public:
 	bool DeleteMember(UMovieGraphMember* MemberToDelete);
 
 	/** Duplicates the provided variable. Returns the new variable on success, else nullptr. */
+	UFUNCTION(BlueprintCallable, Category="Movie Graph")
 	UMovieGraphVariable* DuplicateVariable(UMovieGraphVariable* InVariableToDuplicate);
 
 #if WITH_EDITOR
@@ -838,27 +839,32 @@ public:
 	 * Determines the name(s) of the branches downstream from FromNode, starting at FromPin. Optionally, subgraph nodes can halt graph traversal
 	 * if bStopAtSubgraph is set to true.
 	 */
+	UFUNCTION(BlueprintCallable, Category="Movie Graph")
 	TArray<FString> GetDownstreamBranchNames(UMovieGraphNode* FromNode, const UMovieGraphPin* FromPin, const bool bStopAtSubgraph = false) const;
 
 	/**
 	 * Determines the name(s) of the branches upstream from FromNode, starting at FromPin. Optionally, subgraph nodes can halt graph traversal
 	 * if bStopAtSubgraph is set to true.
 	 */
+	UFUNCTION(BlueprintCallable, Category="Movie Graph")
 	TArray<FString> GetUpstreamBranchNames(UMovieGraphNode* FromNode, const UMovieGraphPin* FromPin, const bool bStopAtSubgraph = false) const;
 
 	/** Get all subgraphs that this graph contains, recursively (ie, subgraphs of subgraphs are included, etc). */
+	UFUNCTION(BlueprintCallable, Category="Movie Graph")
 	void GetAllContainedSubgraphs(TSet<UMovieGraphConfig*>& OutSubgraphs) const;
 
 	/**
 	 * Walks the graph backward recursively from the output node searching for a UMovieGraphOutputSettings node. Traverses subgraphs as well.
 	 * If a node is not found with an override set, value is taken from the CDO of UMovieGraphOutputSettings.
 	 */
+	UFUNCTION(BlueprintCallable, Category="Movie Graph")
 	void GetOutputDirectory(FString& OutOutputDirectory) const;
 
 	/**
 	 * Moves one variable (InTargetVariable) before another variable (InBeforeVariable). Takes care of ensuring the variable's category is
 	 * set properly after the move.
 	 */
+	UFUNCTION(BlueprintCallable, Category="Movie Graph")
 	void MoveVariableBefore(UMovieGraphVariable* InTargetVariable, UMovieGraphVariable* InBeforeVariable);
 
 	/**
@@ -867,9 +873,11 @@ public:
 	 * Note that MoveVariableBefore() should be used in almost all cases unless there is very specific use case. This method will not take care of setting
 	 * the category for you after the move unlike MoveVariableBefore().
 	 */
+	UFUNCTION(BlueprintCallable, Category="Movie Graph")
 	void MoveVariableToIndex(UMovieGraphVariable* InTargetVariable, int32 NewIndex);
 
 	/** Moves one category (InCategoryToMove) and its variables before another category (InCategoryBefore). */
+	UFUNCTION(BlueprintCallable, Category="Movie Graph")
 	void MoveCategoryBefore(const FString& InCategoryToMove, const FString& InCategoryBefore);
 
 protected:
