MovieGraph: Add support for time dilation in renders. Time Dilation is cached at the start of each output frame so that all temporal sub-samples for that frame will use the same Time Dilation scale, which is required for the motion blur to line up within a given output frame.

When using Time Dilation the {frame_number} tokens will be forcibly made into {frame_number_rel} because the frame numbers given by Sequencer are no longer unique, ie: at a 0.5x timescale, Frame 0 would produce two frames for Sequencer's singular frame 0, and {frame_number} would assign both of them to Frame 0, which would cause the resulting files to overwrite each other on disk (resulting in data loss from a render).

In Movie Render Graph, relative frame numbers are now relative to the start of the level sequence/shot and not zero, meaning that if you have a shot that uses time dilation half way through it then it will create a consistent frame range even if your sequence didn't start at zero.

diff --git a/Engine/Plugins/MovieScene/MovieRenderPipeline/Source/MovieRenderPipelineCore/Private/Graph/MovieGraphBlueprintLibrary.cpp b/Engine/Plugins/MovieScene/MovieRenderPipeline/Source/MovieRenderPipelineCore/Private/Graph/MovieGraphBlueprintLibrary.cpp
--- a/Engine/Plugins/MovieScene/MovieRenderPipeline/Source/MovieRenderPipelineCore/Private/Graph/MovieGraphBlueprintLibrary.cpp
+++ b/Engine/Plugins/MovieScene/MovieRenderPipeline/Source/MovieRenderPipelineCore/Private/Graph/MovieGraphBlueprintLibrary.cpp
@@ -12,6 +12,7 @@
 #include "MoviePipelineBlueprintLibrary.h"
 #include "MoviePipelineUtils.h"
 #include "CineCameraActor.h"
+#include "MovieRenderPipelineCoreModule.h"
 
 #include "Misc/Paths.h"
 #include "HAL/FileManager.h"
@@ -49,6 +50,13 @@ FString UMovieGraphBlueprintLibrary::ResolveFilenameFormatArguments(const FStrin
 		FString FrameNumberRel = UE::MoviePipeline::GetPaddingFormatString(InParams.ZeroPadFrameNumberCount, InParams.RootFrameNumberRel + FrameNumberOffset); // Relative to 0
 		FString FrameNumberShotRel = UE::MoviePipeline::GetPaddingFormatString(InParams.ZeroPadFrameNumberCount, InParams.ShotFrameNumberRel + FrameNumberOffset); // Relative to 0 within the shot.
 
+		// If time dilation is in effect, RootFrameNumber and ShotFrameNumber will contain duplicates and the files will overwrite each other, 
+		// so we force them into relative mode and then warn users we did that (as their numbers will be out of sync with what they see in the level sequence now).
+		if (InFormatString.Contains(TEXT("{frame")) && InParams.bForceRelativeFrameNumbers && !InFormatString.Contains(TEXT("_rel}")))
+		{
+			UE_LOG(LogMovieRenderPipeline, Warning, TEXT("Time Dilation was used but output format does not use relative time, forcing relative numbers. Change {frame_number} to {frame_number_rel} (or the shot version) to remove this message."));
+		}
+
 		// Ensure they used relative frame numbers in the output so they get the right number of output frames.
 		if (InParams.bForceRelativeFrameNumbers)
 		{
diff --git a/Engine/Plugins/MovieScene/MovieRenderPipeline/Source/MovieRenderPipelineCore/Private/Graph/MovieGraphCoreTimeStep.cpp b/Engine/Plugins/MovieScene/MovieRenderPipeline/Source/MovieRenderPipelineCore/Private/Graph/MovieGraphCoreTimeStep.cpp
--- a/Engine/Plugins/MovieScene/MovieRenderPipeline/Source/MovieRenderPipelineCore/Private/Graph/MovieGraphCoreTimeStep.cpp
+++ b/Engine/Plugins/MovieScene/MovieRenderPipeline/Source/MovieRenderPipelineCore/Private/Graph/MovieGraphCoreTimeStep.cpp
@@ -30,10 +30,22 @@ void UMovieGraphCoreTimeStep::Shutdown()
 	// Clear out the evaluated config during shutdown. It's a strong object ptr, so references to it should be removed to ensure everything
 	// gets cleaned up properly.
 	CurrentFrameData.EvaluatedConfig.Reset();
-	
+	if (UWorld* World = GetWorld())
+	{
+		World->GetWorldSettings()->SetAllowTimeDilation(bPrevWorldAllowTimeDilation);
+	}
 	Super::Shutdown();
 }
 
+void UMovieGraphCoreTimeStep::Initialize()
+{
+	Super::Initialize();
+	if (UWorld* World = GetWorld())
+	{
+		bPrevWorldAllowTimeDilation = World->GetWorldSettings()->bAllowTimeDilation;
+	}
+}
+
 void UMovieGraphCoreTimeStep::TickProducingFrames()
 {
 	int32 CurrentShotIndex = GetOwningGraph()->GetCurrentShotIndex();
@@ -98,7 +110,6 @@ void UMovieGraphCoreTimeStep::TickProducingFrames()
 		// Sets up the render state, etc.
 		GetOwningGraph()->SetupShot(CurrentCameraCut);
 
-
 		// Generate render layers from the evaluated graph
 		GetOwningGraph()->CreateLayersInRenderLayerSubsystem(CurrentTimeStepData.EvaluatedConfig);
 
@@ -106,6 +117,10 @@ void UMovieGraphCoreTimeStep::TickProducingFrames()
 		// TODO: This should probably be done per-frame
 		GetOwningGraph()->UpdateLayerContentsInRenderLayerSubsystem(CurrentTimeStepData.EvaluatedConfig);
 
+		// Setting up the data source may cause a evaluation which can apply Time Dilation, so we should
+		// pick up the new time dilation before we calculate things.
+		UpdateFrameMetrics();
+
 		// We have three possible options we can go to after the initialization state. If they have no warm up frames,
 		// and don't want to emulate motion blur, we go directly into the Rendering State (which starts on Frame 0, then Frame 1, etc.)
 		// If they do have warm up frames and are not emulating motion blur, we "walk" the level sequence up to the first frame.
@@ -231,7 +246,6 @@ void UMovieGraphCoreTimeStep::TickProducingFrames()
 
 		if (bIncrementInternalCounters)
 		{
-			// ToDo: Apply time dilation by multiplying the calculated range by slow-mo duration.
 			// The upper bound is exclusive, so we initialize a new TRange with just the value to start inclusively.
 			CurrentFrameData.CurrentOutputFrameRange = TRange<FFrameTime>(EndOfPreviousFrame.GetValue(), EndOfPreviousFrame.GetValue() + CurrentFrameMetrics.FrameTimePerOutputFrame);
 		}
@@ -430,9 +444,21 @@ void UMovieGraphCoreTimeStep::TickProducingFrames()
  	double FrameDeltaTimeAsSeconds = CurrentFrameMetrics.TickResolution.AsSeconds(FrameDeltaTime);
 	CurrentTimeStepData.FrameDeltaTime = FrameDeltaTimeAsSeconds;
 	CurrentTimeStepData.WorldSeconds = CurrentTimeStepData.WorldSeconds + FrameDeltaTimeAsSeconds;
-	CurrentTimeStepData.WorldTimeDilation = 1.0; // Temp
+	CurrentTimeStepData.WorldTimeDilation = CurrentFrameMetrics.TimeDilation;
 	CurrentTimeStepData.FrameRate = CurrentFrameMetrics.FrameRate;
 
+	// Once a sequence starts using time dilation (due to non 1.0 timescale), we need to remember that for the rest
+	// of the render. If they're using {frame_number_shot} or {frame_number} tokens, when we first encounter time dilation
+	// the downstream system will force those tokens into their _rel versions. Sequences can drift in and out of using time
+	// dilation, ie: they could start at 0.5 and stay that way the whole time, or start at 1.0 then go to 0.5, or start at 0.5
+	// and go to 1.0. 
+	//
+	// To solve this, we track if time dilation has been used at any point within this level sequence. If it has, then all subsequent
+	// frames will need to use relative frame numbers to avoid overwriting each other.
+	const bool bTimeDilationUsed = !FMath::IsNearlyEqual(CurrentTimeStepData.WorldTimeDilation, 1.0f);
+	GetOwningGraph()->GetCustomEngineTimeStep()->SharedTimeStepData.bHasTimeDilationBeenUsed |= bTimeDilationUsed;
+	CurrentTimeStepData.bHasRelativeTimeBeenUsed = GetOwningGraph()->GetCustomEngineTimeStep()->SharedTimeStepData.bHasTimeDilationBeenUsed;
+
 	// The combination of shutter angle percentage, non-uniform render frame delta times and dividing by sample
 	// count produce the correct length for motion blur in all cases.
 	CurrentTimeStepData.MotionBlurFraction = CurrentFrameMetrics.MotionBlurAmount / CurrentFrameData.TemporalSampleCount;
@@ -652,6 +678,19 @@ void UMovieGraphCoreTimeStep::UpdateFrameMetrics()
 
 	FrameData.FrameTimePerOutputFrame = FFrameRate::TransformTime(FFrameTime(FFrameNumber(1)), FrameData.FrameRate, FrameData.TickResolution);
 
+	// Temporarily override time dilation so we can figure out what time dilation
+	// the engine wanted, then turn support back off so we can implement it.
+	GetWorld()->GetWorldSettings()->SetAllowTimeDilation(true);
+	float TimeDilation = GetWorld()->GetWorldSettings()->GetEffectiveTimeDilation();
+	GetWorld()->GetWorldSettings()->SetAllowTimeDilation(false);
+	
+	if(!FMath::IsNearlyEqual(TimeDilation, 1.0f))
+	{
+		FrameData.FrameTimePerOutputFrame = FrameData.FrameTimePerOutputFrame * TimeDilation;
+	}
+	
+	FrameData.TimeDilation = TimeDilation;
+
 	// Manually perform blending of the Post Process Volumes/Camera/Camera Modifiers to match what the renderer will do.
 	// This uses the primary camera specified by the PlayerCameraManager to get the motion blur amount so in the event of
 	// multi-camera rendering, all cameras will end up using the same motion blur amount defined by the primary camera).
@@ -865,4 +904,3 @@ int32 UMovieGraphCoreTimeStep::GetTemporalSampleCountFromConfig(UMovieGraphEvalu
 
 	return FMath::Max(SamplingMethod->TemporalSampleCount, 1);
 }
-
diff --git a/Engine/Plugins/MovieScene/MovieRenderPipeline/Source/MovieRenderPipelineCore/Private/Graph/MovieGraphFilenameResolveParams.cpp b/Engine/Plugins/MovieScene/MovieRenderPipeline/Source/MovieRenderPipelineCore/Private/Graph/MovieGraphFilenameResolveParams.cpp
--- a/Engine/Plugins/MovieScene/MovieRenderPipeline/Source/MovieRenderPipelineCore/Private/Graph/MovieGraphFilenameResolveParams.cpp
+++ b/Engine/Plugins/MovieScene/MovieRenderPipeline/Source/MovieRenderPipelineCore/Private/Graph/MovieGraphFilenameResolveParams.cpp
@@ -15,6 +15,8 @@ FMovieGraphFilenameResolveParams FMovieGraphFilenameResolveParams::MakeResolvePa
 {
     FMovieGraphFilenameResolveParams Params = FMovieGraphFilenameResolveParams();
 	
+	int32 RootFrameNumberRelOffset = 0;
+	int32 ShotFrameNumberRelOffset = 0;
 	if (ensureAlwaysMsgf(InPipeline, TEXT("InPipeline is not valid - ResolveParams will be created, but will be missing critical information.")))
 	{
 		Params.InitializationTime = InPipeline->GetInitializationTime();
@@ -26,6 +28,11 @@ FMovieGraphFilenameResolveParams FMovieGraphFilenameResolveParams::MakeResolvePa
 			const TObjectPtr<UMoviePipelineExecutorShot>& Shot = InPipeline->GetActiveShotList()[InTraversalContext.ShotIndex];
 			Params.Version = Shot->ShotInfo.VersionNumber;
 			Params.Shot = Shot;
+			
+			// Calculate the offset (in frames) that relative frame numbers need to add to themselves so that they are correctly
+			// offset by the starting frame of the root sequence/shot (to match new updated relative behavior).
+			RootFrameNumberRelOffset = FFrameRate::TransformTime(Shot->ShotInfo.InitialTimeInRoot, Shot->ShotInfo.CachedTickResolution, Shot->ShotInfo.CachedFrameRate).FloorToFrame().Value;
+			ShotFrameNumberRelOffset = FFrameRate::TransformTime(Shot->ShotInfo.InitialTimeInShot, Shot->ShotInfo.CachedTickResolution, Shot->ShotInfo.CachedFrameRate).FloorToFrame().Value;
 		}
 	}
         
@@ -33,8 +40,11 @@ FMovieGraphFilenameResolveParams FMovieGraphFilenameResolveParams::MakeResolvePa
 	
 	Params.RootFrameNumber = InTraversalContext.Time.RootFrameNumber.Value;
 	Params.ShotFrameNumber = InTraversalContext.Time.ShotFrameNumber.Value;
-	Params.RootFrameNumberRel = InTraversalContext.Time.OutputFrameNumber;
-	Params.ShotFrameNumberRel = InTraversalContext.Time.ShotOutputFrameNumber;
+
+	// Starting in MRG, relative file numbers are now relative to the first frame of the shot/sequence, and not
+	// to zero. To do this, we use our zero-relative numbers and offset them by the starting point of the shot/sequence.
+	Params.RootFrameNumberRel = InTraversalContext.Time.OutputFrameNumber + RootFrameNumberRelOffset;
+	Params.ShotFrameNumberRel = InTraversalContext.Time.ShotOutputFrameNumber + ShotFrameNumberRelOffset;
 	//Params.FileMetadata = ToDo: Track File Metadata
 
 	if (InEvaluatedConfig)
@@ -48,15 +58,8 @@ FMovieGraphFilenameResolveParams FMovieGraphFilenameResolveParams::MakeResolvePa
 		Params.EvaluatedConfig = InEvaluatedConfig;
 	}
 
-    // If time dilation is in effect, RootFrameNumber and ShotFrameNumber will contain duplicates and the files will overwrite each other, 
-    // so we force them into relative mode and then warn users we did that (as their numbers will jump from say 1001 -> 0000).
-    bool bForceRelativeFrameNumbers = false;
-    // if (FileNameFormatString.Contains(TEXT("{frame")) && InTraversalContext.Time.IsTimeDilated() && !FileNameFormatString.Contains(TEXT("_rel}")))
-    // {
-    // 	UE_LOG(LogMovieRenderPipeline, Warning, TEXT("Time Dilation was used but output format does not use relative time, forcing relative numbers. Change {frame_number} to {frame_number_rel} (or shot version) to remove this message."));
-    // 	bForceRelativeFrameNumbers = true;
-    // }
-    Params.bForceRelativeFrameNumbers = bForceRelativeFrameNumbers;
+	const bool bTimeDilationUsed = !FMath::IsNearlyEqual(InTraversalContext.Time.WorldTimeDilation, 1.f) || InTraversalContext.Time.bHasRelativeTimeBeenUsed;
+    Params.bForceRelativeFrameNumbers = bTimeDilationUsed;
 	Params.bEnsureAbsolutePath = true;
 	Params.FileNameFormatOverrides = InAdditionalFormatArgs;
 
diff --git a/Engine/Plugins/MovieScene/MovieRenderPipeline/Source/MovieRenderPipelineCore/Private/Graph/MovieGraphPipeline.cpp b/Engine/Plugins/MovieScene/MovieRenderPipeline/Source/MovieRenderPipelineCore/Private/Graph/MovieGraphPipeline.cpp
--- a/Engine/Plugins/MovieScene/MovieRenderPipeline/Source/MovieRenderPipelineCore/Private/Graph/MovieGraphPipeline.cpp
+++ b/Engine/Plugins/MovieScene/MovieRenderPipeline/Source/MovieRenderPipelineCore/Private/Graph/MovieGraphPipeline.cpp
@@ -524,6 +524,8 @@ void UMovieGraphPipeline::BuildShotListFromDataSource()
 		ExpandShot(Shot, OutputNode->HandleFrameCount, bExpandForTemporalSubSample, bPrePass, FinalFrameRate, TickResolution, Shot->ShotInfo.NumEngineWarmUpFramesRemaining);
 
 		Shot->ShotInfo.CurrentTimeInRoot = Shot->ShotInfo.TotalOutputRangeRoot.GetLowerBoundValue();
+		Shot->ShotInfo.InitialTimeInRoot = Shot->ShotInfo.CurrentTimeInRoot;
+		Shot->ShotInfo.InitialTimeInShot = Shot->ShotInfo.CurrentTimeInRoot * Shot->ShotInfo.OuterToInnerTransform;
 		Shot->ShotInfo.CalculateWorkMetrics();
 		Shot->ShotInfo.VersionNumber = ResolveVersionForShot(Shot, EvaluatedConfig);
 	}
diff --git a/Engine/Plugins/MovieScene/MovieRenderPipeline/Source/MovieRenderPipelineCore/Public/Graph/MovieGraphCoreTimeStep.h b/Engine/Plugins/MovieScene/MovieRenderPipeline/Source/MovieRenderPipelineCore/Public/Graph/MovieGraphCoreTimeStep.h
--- a/Engine/Plugins/MovieScene/MovieRenderPipeline/Source/MovieRenderPipelineCore/Public/Graph/MovieGraphCoreTimeStep.h
+++ b/Engine/Plugins/MovieScene/MovieRenderPipeline/Source/MovieRenderPipelineCore/Public/Graph/MovieGraphCoreTimeStep.h
@@ -27,6 +27,7 @@ public:
 	UMovieGraphCoreTimeStep();
 
 	// UMovieGraphTimeStepBase Interface
+	virtual void Initialize() override;
 	virtual void Shutdown() override;
 	virtual void TickProducingFrames() override;
 	virtual FMovieGraphTimeStepData GetCalculatedTimeData() const override { return CurrentTimeStepData; }
@@ -67,6 +68,7 @@ protected:
 	{
 		FOutputFrameMetrics()
 			: MotionBlurAmount(0.f)
+			, TimeDilation(1.f)
 		{}
 
 		/** The human readable frame rate, adjusted by the config file (24fps, 30fps, etc.) */
@@ -80,6 +82,9 @@ protected:
 
 		/** The 0-1 Motion Blur amount the camera has. */
 		float MotionBlurAmount;
+		
+		/** The time dilation multiplier set by the world, 1.0 if there is no time dilation being applied. */
+		float TimeDilation;
 
 		/** The amount of time (in Tick Resolution) that the shutter is open for. Could be <= FrameTimePerOutputFrame. */
 		FFrameTime FrameTimeWhileShutterOpen;
@@ -151,4 +156,10 @@ protected:
 	* across engine ticks to work towards producing a single frame.
 	*/
 	FCurrentFrameData CurrentFrameData;
+
+	/**
+	* Was the world allowing time dilation before this render started? We'll cache and then restore this value to
+	* avoid changing any user settings.
+	*/
+	bool bPrevWorldAllowTimeDilation = true;
 };
diff --git a/Engine/Plugins/MovieScene/MovieRenderPipeline/Source/MovieRenderPipelineCore/Public/Graph/MovieGraphDataTypes.h b/Engine/Plugins/MovieScene/MovieRenderPipeline/Source/MovieRenderPipelineCore/Public/Graph/MovieGraphDataTypes.h
--- a/Engine/Plugins/MovieScene/MovieRenderPipeline/Source/MovieRenderPipelineCore/Public/Graph/MovieGraphDataTypes.h
+++ b/Engine/Plugins/MovieScene/MovieRenderPipeline/Source/MovieRenderPipelineCore/Public/Graph/MovieGraphDataTypes.h
@@ -150,6 +150,7 @@ public:
 		FSharedData()
 			: OutputFrameNumber(0)
 			, RenderedFrameNumber(0)
+			, bHasTimeDilationBeenUsed(false)
 		{}
 		
 		
@@ -161,6 +162,12 @@ public:
 		* Used internally by the rendering engine to keep track of which frames need to be read back.
 		*/
 		int32 RenderedFrameNumber;
+
+		/**
+		* If true, a frame rendered for this job used time dilation, and all subsequent frame numbers must be their relative versions
+		* to avoid any frame overwriting an already produced one.
+		*/
+		bool bHasTimeDilationBeenUsed;
 	};
 
 public:
diff --git a/Engine/Plugins/MovieScene/MovieRenderPipeline/Source/MovieRenderPipelineCore/Public/Graph/MovieGraphFilenameResolveParams.h b/Engine/Plugins/MovieScene/MovieRenderPipeline/Source/MovieRenderPipelineCore/Public/Graph/MovieGraphFilenameResolveParams.h
--- a/Engine/Plugins/MovieScene/MovieRenderPipeline/Source/MovieRenderPipelineCore/Public/Graph/MovieGraphFilenameResolveParams.h
+++ b/Engine/Plugins/MovieScene/MovieRenderPipeline/Source/MovieRenderPipelineCore/Public/Graph/MovieGraphFilenameResolveParams.h
@@ -83,28 +83,28 @@ struct MOVIERENDERPIPELINECORE_API FMovieGraphFilenameResolveParams
 	FMovieGraphRenderDataIdentifier RenderDataIdentifier;
 	
 	/**
-	* Frame Number for the Root (matching what you see in the Sequencer timeline. ie: If the Sequence PlaybackRange starts on 50, this value would be 50 on the first frame.
+	* Frame Number for the Root (matching what you see in the Sequencer timeline. ie: If the Sequence PlaybackRange starts on 50, this value would be 50 on the first frame.)
 	* Used to fill out the {frame_number} token.
 	*/
 	UPROPERTY(BlueprintReadWrite, EditAnywhere, Category = "Movie Graph")
 	int32 RootFrameNumber;
 	
 	/** 
-	* Frame Number for the Shot (matching what you would see in Sequencer at the sub-sequence level. 
+	* Frame Number for the Shot (matching what you would see in Sequencer at the sub-sequence level.) 
 	* Used to fill out the {frame_number_shot} token.
 	*/
 	UPROPERTY(BlueprintReadWrite, EditAnywhere, Category = "Movie Graph")
 	int32 ShotFrameNumber;
 	
 	/** 
-	* Frame Number for the Root (relative to 0, not what you would see in the Sequencer timeline. ie: If sequence PlaybackRange starts on 50, this value would be 0 on the first frame.)
+	* Frame Number for the Root (relative to the first RootFrameNumber for this render, ie: if your sequence started at 35, this would be 35 as well, but supports time dilation because it will count up based on the total number of frame rendered, not the closest Sequencer frame (like RootFrameNumber))
 	* Used to fill out the {frame_number_rel} token.
 	*/
 	UPROPERTY(BlueprintReadWrite, EditAnywhere, Category = "Movie Graph")
 	int32 RootFrameNumberRel;
 	
 	/** 
-	* Frame Number for the Shot (relative to 0, not what you would see in the Sequencer timeline.) 
+	* Frame Number for the Shot (relative to the first ShotFrameNumber for this shot, ie: if your shot started at 35, this would be 35 as well, but supports time dilation because it will count up based on the total number of frame rendered, not the closest Sequencer frame (like ShotFrameNumber))
 	* Used to fill out the {frame_number_shot_rel} token.
 	*/
 	UPROPERTY(BlueprintReadWrite, EditAnywhere, Category = "Movie Graph")
diff --git a/Engine/Plugins/MovieScene/MovieRenderPipeline/Source/MovieRenderPipelineCore/Public/Graph/MovieGraphTimeStepData.h b/Engine/Plugins/MovieScene/MovieRenderPipeline/Source/MovieRenderPipelineCore/Public/Graph/MovieGraphTimeStepData.h
--- a/Engine/Plugins/MovieScene/MovieRenderPipeline/Source/MovieRenderPipelineCore/Public/Graph/MovieGraphTimeStepData.h
+++ b/Engine/Plugins/MovieScene/MovieRenderPipeline/Source/MovieRenderPipelineCore/Public/Graph/MovieGraphTimeStepData.h
@@ -57,6 +57,10 @@ public:
 	
 	UPROPERTY(BlueprintReadWrite, EditAnywhere, Category = "Movie Graph")
 	float WorldTimeDilation;
+
+	/** True if time dilation was used on this frame, or any in the past, even if WorldTimeDilation has returned to 1.0 */
+	UPROPERTY(BlueprintReadWrite, EditAnywhere, Category = "Movie Graph")
+	bool bHasRelativeTimeBeenUsed;
 	
 	UPROPERTY(BlueprintReadWrite, EditAnywhere, Category = "Movie Graph")
 	float WorldSeconds;
diff --git a/Engine/Plugins/MovieScene/MovieRenderPipeline/Source/MovieRenderPipelineCore/Public/MovieRenderPipelineDataTypes.h b/Engine/Plugins/MovieScene/MovieRenderPipeline/Source/MovieRenderPipelineCore/Public/MovieRenderPipelineDataTypes.h
--- a/Engine/Plugins/MovieScene/MovieRenderPipeline/Source/MovieRenderPipelineCore/Public/MovieRenderPipelineDataTypes.h
+++ b/Engine/Plugins/MovieScene/MovieRenderPipeline/Source/MovieRenderPipelineCore/Public/MovieRenderPipelineDataTypes.h
@@ -574,6 +574,11 @@ public:
 	/** The current time of this shot that we're on in root space (for the Graph version. Tracks sub-frames too.) */
 	FFrameTime CurrentTimeInRoot;
 
+	/** The initial time in the root space for this shot when we started evaluating. Only assigned in the Graph version of MRQ. Effectively the same as CurrentTimeInRoot, but is immutable during the shot. */
+	FFrameTime InitialTimeInRoot;
+	/** The initial time in the shot time space for this shot when we started evaluating. Only assigned in the Graph version of MRQ. Effectively the same as CurrentTimeInRoot * OuterToInnerTransform, but is immutable during the shot. */
+	FFrameTime InitialTimeInShot;
+
 	/** Converts from the outermost space into the innermost space. Only works with linear transforms. */
 	FMovieSceneSequenceTransform OuterToInnerTransform;
 	
