MoviePipeline: Implemented the "Selected Camera(s)" mode in Quick Render.

Also added the ability to toggle on/off editor visibility for all modes.

diff --git a/Engine/Plugins/MovieScene/MovieRenderPipeline/Source/MovieRenderPipelineCore/Public/Graph/MovieGraphQuickRenderSettings.h b/Engine/Plugins/MovieScene/MovieRenderPipeline/Source/MovieRenderPipelineCore/Public/Graph/MovieGraphQuickRenderSettings.h
--- a/Engine/Plugins/MovieScene/MovieRenderPipeline/Source/MovieRenderPipelineCore/Public/Graph/MovieGraphQuickRenderSettings.h
+++ b/Engine/Plugins/MovieScene/MovieRenderPipeline/Source/MovieRenderPipelineCore/Public/Graph/MovieGraphQuickRenderSettings.h
@@ -66,6 +66,7 @@ enum class EMovieGraphQuickRenderViewportLookFlags : uint8
 	Ocio		= 1 << 0 UMETA(DisplayName = "OCIO"),
 	ShowFlags	= 1 << 1,
 	ViewMode	= 1 << 2,
+	Visibility	= 1 << 3,
 };
 ENUM_CLASS_FLAGS(EMovieGraphQuickRenderViewportLookFlags)
 
diff --git a/Engine/Plugins/MovieScene/MovieRenderPipeline/Source/MovieRenderPipelineCore/Public/Graph/Nodes/MovieGraphApplyViewportLookNode.h b/Engine/Plugins/MovieScene/MovieRenderPipeline/Source/MovieRenderPipelineCore/Public/Graph/Nodes/MovieGraphApplyViewportLookNode.h
--- a/Engine/Plugins/MovieScene/MovieRenderPipeline/Source/MovieRenderPipelineCore/Public/Graph/Nodes/MovieGraphApplyViewportLookNode.h
+++ b/Engine/Plugins/MovieScene/MovieRenderPipeline/Source/MovieRenderPipelineCore/Public/Graph/Nodes/MovieGraphApplyViewportLookNode.h
@@ -54,6 +54,9 @@ public:
 	UPROPERTY(EditAnywhere, BlueprintReadWrite, Category = Overrides, meta = (InlineEditConditionToggle))
 	uint8 bOverride_bViewMode : 1;
 	
+	UPROPERTY(EditAnywhere, BlueprintReadWrite, Category = Overrides, meta = (InlineEditConditionToggle))
+	uint8 bOverride_bVisibility : 1;
+	
 	/** Set to true to apply the viewport's OCIO settings to the render. Only render nodes that have "Allow OCIO" turned on will be impacted. */
 	UPROPERTY(EditAnywhere, BlueprintReadWrite, Category="Viewport Look", DisplayName="OCIO", meta=(EditCondition="bOverride_bOcio"))
 	bool bOcio;
@@ -65,4 +68,8 @@ public:
 	/** Set to true to apply the viewport's view mode to the render. */
 	UPROPERTY(EditAnywhere, BlueprintReadWrite, Category="Viewport Look", meta=(EditCondition="bOverride_bViewMode"))
 	bool bViewMode;
+
+	/** Set to true to apply editor visibility to actors in the render. */
+	UPROPERTY(EditAnywhere, BlueprintReadWrite, Category="Viewport Look", meta=(EditCondition="bOverride_bVisibility"))
+	bool bVisibility;
 };
\ No newline at end of file
diff --git a/Engine/Plugins/MovieScene/MovieRenderPipeline/Source/MovieRenderPipelineEditor/Private/Graph/MovieGraphQuickRender.cpp b/Engine/Plugins/MovieScene/MovieRenderPipeline/Source/MovieRenderPipelineEditor/Private/Graph/MovieGraphQuickRender.cpp
--- a/Engine/Plugins/MovieScene/MovieRenderPipeline/Source/MovieRenderPipelineEditor/Private/Graph/MovieGraphQuickRender.cpp
+++ b/Engine/Plugins/MovieScene/MovieRenderPipeline/Source/MovieRenderPipelineEditor/Private/Graph/MovieGraphQuickRender.cpp
@@ -32,6 +32,7 @@
 
 #if WITH_EDITOR
 #include "EditorViewportClient.h"
+#include "Selection.h"
 #endif
 
 #if WITH_OCIO
@@ -68,6 +69,12 @@ void FMovieGraphQuickRender::BeginQuickRender(const UMovieGraphQuickRenderSettin
 		return;
 	}
 
+	// Don't do additional work if the selected mode doesn't have what it needs to do a render. AreModePrerequisitesMet() will output the appropriate warnings.
+	if (!AreModePrerequisitesMet(InQuickRenderSettings.QuickRenderMode))
+	{
+		return;
+	}
+
 	// Find the editor world
 	UWorld* EditorWorld = GetWorldOfType(EWorldType::Editor);
 	if (!EditorWorld)
@@ -131,16 +138,8 @@ void FMovieGraphQuickRender::BeginQuickRender(const UMovieGraphQuickRenderSettin
 	// should be fully initialized at this point.
 	ApplyQuickRenderUpdatesToDuplicateGraph(InQuickRenderSettings, EditorWorld, NewJob);
 
-	// Prevent all existing level sequences in the world from auto-playing. We need the level sequence from Quick Render to be in control. 
-	for (auto It = TActorIterator<ALevelSequenceActor>(EditorWorld); It; ++It)
-	{
-		const ALevelSequenceActor* LevelSequenceActor = *It;
-		if (LevelSequenceActor->PlaybackSettings.bAutoPlay)
-		{
-			(*It)->PlaybackSettings.bAutoPlay = false;
-			ModifiedLevelSequenceActors.Add(MakeWeakObjectPtr(*It));
-		}
-	}
+	// Cache out anything necessary before PIE starts
+	CachePrePieData(EditorWorld);
 
 	// Do the render
 	UMoviePipelinePIEExecutor* Executor = Cast<UMoviePipelinePIEExecutor>(Subsystem->RenderQueueInstanceWithExecutor(TemporaryQueue.Get(), ExecutorClass));
@@ -170,14 +169,7 @@ void FMovieGraphQuickRender::BeginQuickRender(const UMovieGraphQuickRenderSettin
 		TemporaryGraph.Reset();
 		TemporaryLevelSequence.Reset();
 
-		// Restore the modified level sequences to the way they were before the render
-		for (const TWeakObjectPtr<ALevelSequenceActor>& LevelSequenceActor : ModifiedLevelSequenceActors)
-		{
-			if (LevelSequenceActor.IsValid())
-			{
-				LevelSequenceActor->PlaybackSettings.bAutoPlay = true;
-			}
-		}
+		RestorePrePieData();
 	});
 }
 
@@ -261,11 +253,15 @@ PRAGMA_ENABLE_DEPRECATION_WARNINGS
 		APPLY_VIEWPORT_LOOK_FLAG(EMovieGraphQuickRenderViewportLookFlags::Ocio, Ocio);
 		APPLY_VIEWPORT_LOOK_FLAG(EMovieGraphQuickRenderViewportLookFlags::ShowFlags, ShowFlags);
 		APPLY_VIEWPORT_LOOK_FLAG(EMovieGraphQuickRenderViewportLookFlags::ViewMode, ViewMode);
+		APPLY_VIEWPORT_LOOK_FLAG(EMovieGraphQuickRenderViewportLookFlags::Visibility, Visibility);
 #undef APPLY_VIEWPORT_LOOK_FLAG
 	}
 
 	// Apply viewport/editor actor visibility
-	ApplyQuickRenderUpdatesToDuplicateGraph_ApplyEditorVisibility(InEditorWorld);
+	if (EnumHasAnyFlags(LookFlags, EMovieGraphQuickRenderViewportLookFlags::Visibility))
+	{
+		ApplyQuickRenderUpdatesToDuplicateGraph_ApplyEditorVisibility(InEditorWorld);
+	}
 
 	// Apply OCIO if it is activated on the viewport and enabled in the viewport look flags
 	if (InQuickRenderSettings.bOverride_ViewportLookFlags && EnumHasAnyFlags(LookFlags, EMovieGraphQuickRenderViewportLookFlags::Ocio))
@@ -471,9 +467,64 @@ UMovieGraphNode* FMovieGraphQuickRender::InjectNodeIntoBranch(const TSubclassOf<
 	return NewNode;
 }
 
+bool FMovieGraphQuickRender::AreModePrerequisitesMet(const EMovieGraphQuickRenderMode QuickRenderMode)
+{
+	if (QuickRenderMode == EMovieGraphQuickRenderMode::SelectedCameras)
+	{
+		TArray<ACameraActor*> SelectedCameras;
+		GEditor->GetSelectedActors()->GetSelectedObjects<ACameraActor>(SelectedCameras);
+
+		if (SelectedCameras.IsEmpty())
+		{
+			FMessageDialog::Open(
+				EAppMsgType::Ok,
+				LOCTEXT("NoCamerasSelectedWarning", "The 'Selected Camera(s)' mode needs at least one camera selected. Select some cameras in the Outliner, then try again."));
+			
+			return false;
+		}
+	}
+
+	return true;
+}
+
+void FMovieGraphQuickRender::CachePrePieData(UWorld* InEditorWorld)
+{
+	// Prevent all existing level sequences in the world from auto-playing. We need the level sequence from Quick Render to be in control. 
+	for (auto It = TActorIterator<ALevelSequenceActor>(InEditorWorld); It; ++It)
+	{
+		const ALevelSequenceActor* LevelSequenceActor = *It;
+		if (LevelSequenceActor->PlaybackSettings.bAutoPlay)
+		{
+			(*It)->PlaybackSettings.bAutoPlay = false;
+			CachedPrePieData.ModifiedLevelSequenceActors.Add(MakeWeakObjectPtr(*It));
+		}
+	}
+
+	// Cache out the selected camera actors before PIE starts. They'll be deselected once the PIE window shows up.
+	TArray<ACameraActor*> CameraActors;
+	GEditor->GetSelectedActors()->GetSelectedObjects<ACameraActor>(CameraActors);
+	CopyToWeakArray(CachedPrePieData.SelectedCameras, CameraActors);
+}
+
+void FMovieGraphQuickRender::RestorePrePieData()
+{
+	// Restore the modified level sequences to the way they were before the render.
+	for (const TWeakObjectPtr<ALevelSequenceActor>& LevelSequenceActor : CachedPrePieData.ModifiedLevelSequenceActors)
+	{
+		if (LevelSequenceActor.IsValid())
+		{
+			LevelSequenceActor->PlaybackSettings.bAutoPlay = true;
+		}
+	}
+	
+	CachedPrePieData.SelectedCameras.Empty();
+}
+
 void FMovieGraphQuickRender::SetUpLevelSequenceForRender(const UMovieGraphQuickRenderSettings& InQuickRenderSettings)
 {
-	if (InQuickRenderSettings.QuickRenderMode == EMovieGraphQuickRenderMode::CurrentViewport)
+	const EMovieGraphQuickRenderMode QuickRenderMode = InQuickRenderSettings.QuickRenderMode;
+	
+	if ((QuickRenderMode == EMovieGraphQuickRenderMode::CurrentViewport) || (QuickRenderMode == EMovieGraphQuickRenderMode::SelectedCameras))
 	{
 		// Create the new, temporary level sequence
 		ULevelSequence* QuickRenderSequence = NewObject<ULevelSequence>(GetTransientPackage(), TEXT("QuickRenderSequence"), RF_Transient);
@@ -511,6 +562,49 @@ UWorld* FMovieGraphQuickRender::GetWorldOfType(EWorldType::Type WorldType)
 
 void FMovieGraphQuickRender::PerformPreRenderSetup(const UMovieGraphQuickRenderSettings& InQuickRenderSettings)
 {
+	// Adds the specified camera to the Camera Cut track in the level sequence that will be used for Quick Render.
+	// If there's already a camera cut section within the Camera Cut track, a new one will be added immediately after.
+	// Each section is only one frame long.
+	auto AddCameraToLevelSequence = [](ACameraActor* InCameraToAdd)
+	{
+		UMovieScene* MovieScene = TemporaryLevelSequence->GetMovieScene();
+
+		// Add a Camera Cut track if one hasn't already been added
+		UMovieSceneCameraCutTrack* CameraCutTrack = Cast<UMovieSceneCameraCutTrack>(MovieScene->GetCameraCutTrack());
+		if (!CameraCutTrack)
+		{
+			CameraCutTrack = CastChecked<UMovieSceneCameraCutTrack>(MovieScene->AddCameraCutTrack(UMovieSceneCameraCutTrack::StaticClass()));
+		}
+
+		const int32 NumCameraCutSections = CameraCutTrack->GetAllSections().Num();
+
+		// Determine the playback range for the new camera cut. Each camera cut section has a playback range of one frame.
+		const FFrameNumber StartFrame = MovieScene->GetPlaybackRange().GetLowerBound().GetValue();
+		const FFrameNumber EndFrame = FFrameRate::TransformTime(FFrameTime(NumCameraCutSections + 1), MovieScene->GetDisplayRate(), MovieScene->GetTickResolution()).GetFrame();
+
+		// The entire playback range for the movie scene always starts a 0 and ends at the last camera cut section's end frame.
+		MovieScene->SetPlaybackRange(TRange<FFrameNumber>::Inclusive(FFrameNumber(0), EndFrame));
+		
+		UE::Sequencer::FCreateBindingParams CreateBindingParams;
+		CreateBindingParams.bSpawnable = true;
+		CreateBindingParams.bAllowCustomBinding = true;
+		CreateBindingParams.BindingNameOverride = InCameraToAdd->GetActorLabel();
+
+		// Create the binding for the camera. This is wrapped in a dummy transaction so we can cancel it later on to prevent the binding creation from
+		// making an entry on the undo stack.
+		#if WITH_EDITOR
+		const int32 TransactionId = GEditor->BeginTransaction(FText::GetEmpty());
+#endif
+		const FGuid CameraBinding = FSequencerUtilities::CreateOrReplaceBinding(nullptr, MovieScene->GetTypedOuter<UMovieSceneSequence>(), InCameraToAdd, CreateBindingParams);
+#if WITH_EDITOR
+		GEditor->CancelTransaction(TransactionId);
+#endif
+
+		// Add the new camera to the camera cut track.
+		UMovieSceneCameraCutSection* CameraCutSection = CameraCutTrack->AddNewCameraCut(UE::MovieScene::FRelativeObjectBindingID(CameraBinding), StartFrame);
+		CameraCutSection->SetRange(TRange<FFrameNumber>::Inclusive(StartFrame, EndFrame));
+	};
+	
 	if (InQuickRenderSettings.QuickRenderMode == EMovieGraphQuickRenderMode::CurrentViewport)
 	{
 		// Get the PIE world, this is where we need to target the changes (not the editor world). Otherwise we'll have to deal with deleting actors
@@ -523,38 +617,25 @@ void FMovieGraphQuickRender::PerformPreRenderSetup(const UMovieGraphQuickRenderS
 		CameraSpawnParams.OverridePackage = GetTransientPackage();
 		CameraSpawnParams.bCreateActorPackage = false;
 		CameraSpawnParams.OverrideLevel = PieWorld->GetCurrentLevel();
-		const FEditorViewportClient* ViewportClient = UMovieGraphApplyViewportLookNode::GetViewportClient();
 		ACameraActor* NewCamera = Cast<ACameraActor>(PieWorld->SpawnActor<ACameraActor>(CameraSpawnParams));
 
 		// Update the camera to look like the viewport's camera
+		const FEditorViewportClient* ViewportClient = UMovieGraphApplyViewportLookNode::GetViewportClient();
 		NewCamera->SetActorLocation(ViewportClient->GetViewLocation(), false);
 		NewCamera->SetActorRotation(ViewportClient->GetViewRotation());
 		NewCamera->GetCameraComponent()->SetFieldOfView(ViewportClient->ViewFOV);
 
-		// Add a new camera cut track. We only need 1 frame in Viewport mode, so use a range of [0, 1]
-		UMovieScene* MovieScene = TemporaryLevelSequence->GetMovieScene();
-		const FFrameTime EndOfRange = FFrameRate::TransformTime(FFrameTime(1), MovieScene->GetDisplayRate(), MovieScene->GetTickResolution());
-		MovieScene->SetPlaybackRange(TRange<FFrameNumber>::Inclusive(FFrameNumber(0), EndOfRange.GetFrame()));
-		UMovieSceneCameraCutTrack* CameraCutTrack = CastChecked<UMovieSceneCameraCutTrack>(MovieScene->AddCameraCutTrack(UMovieSceneCameraCutTrack::StaticClass()));
-		
-		UE::Sequencer::FCreateBindingParams CreateBindingParams;
-		CreateBindingParams.bSpawnable = true;
-		CreateBindingParams.bAllowCustomBinding = true;
-		CreateBindingParams.BindingNameOverride = NewCamera->GetActorLabel();
-
-		// Create the binding for the camera. This is wrapped in a dummy transaction so we can cancel it later on to prevent the binding creation from
-		// making an entry on the undo stack.
-#if WITH_EDITOR
-		const int32 TransactionId = GEditor->BeginTransaction(FText::GetEmpty());
-#endif
-		const FGuid CameraBinding = FSequencerUtilities::CreateOrReplaceBinding(nullptr, MovieScene->GetTypedOuter<UMovieSceneSequence>(), NewCamera, CreateBindingParams);
-#if WITH_EDITOR
-		GEditor->CancelTransaction(TransactionId);
-#endif
-
-		// Add the new camera to the camera cut track.
-		UMovieSceneCameraCutSection* CameraCutSection = CameraCutTrack->AddNewCameraCut(UE::MovieScene::FRelativeObjectBindingID(CameraBinding), FFrameNumber(0));
-		CameraCutSection->SetRange(MovieScene->GetPlaybackRange());
+		AddCameraToLevelSequence(NewCamera);
+	}
+	else if (InQuickRenderSettings.QuickRenderMode == EMovieGraphQuickRenderMode::SelectedCameras)
+	{
+		for (const TWeakObjectPtr<ACameraActor>& SelectedCamera : CachedPrePieData.SelectedCameras)
+		{
+			if (SelectedCamera.IsValid())
+			{
+				AddCameraToLevelSequence(SelectedCamera.Get());
+			}
+		}
 	}
 }
 
diff --git a/Engine/Plugins/MovieScene/MovieRenderPipeline/Source/MovieRenderPipelineEditor/Public/Graph/MovieGraphQuickRender.h b/Engine/Plugins/MovieScene/MovieRenderPipeline/Source/MovieRenderPipelineEditor/Public/Graph/MovieGraphQuickRender.h
--- a/Engine/Plugins/MovieScene/MovieRenderPipeline/Source/MovieRenderPipelineEditor/Public/Graph/MovieGraphQuickRender.h
+++ b/Engine/Plugins/MovieScene/MovieRenderPipeline/Source/MovieRenderPipelineEditor/Public/Graph/MovieGraphQuickRender.h
@@ -24,6 +24,19 @@ public:
 	static void BeginQuickRender(const UMovieGraphQuickRenderSettings& InQuickRenderSettings);
 
 private:
+	/** Data that was cached prior to PIE starting. */
+	struct FCachedPrePieData
+	{
+		/**
+		 * Level sequence actors which had their AutoPlay setting turned ON prior to PIE (this setting needs to be turned OFF during PIE). This
+		 * setting will be restored to ON after PIE ends.
+		 */
+		TArray<TWeakObjectPtr<class ALevelSequenceActor>> ModifiedLevelSequenceActors;
+		
+		/** The cameras that were selected. When PIE starts, the Outliner selection is cleared out, so it needs to be cached prior to PIE. */
+		TArray<TWeakObjectPtr<ACameraActor>> SelectedCameras;
+	};
+	
 	/**
 	 * Gets the graph that should be used for the quick render. This will be a duplicate of either the user-specified graph or the default quick render graph.
 	 * OutOriginalToDupeMap will provide a mapping of the pre-duplicate -> duplicate graphs (may contain multiple graphs if subgraphs are involved).
@@ -42,6 +55,15 @@ private:
 	/** Injects a new node in the furthest-downstream position within the specified branch in the duplicated graph. */
 	static UMovieGraphNode* InjectNodeIntoBranch(const TSubclassOf<UMovieGraphNode> NodeType, const FName& InBranchName);
 
+	/** Determines if the prerequisites for the given mode are met. */
+	static bool AreModePrerequisitesMet(const EMovieGraphQuickRenderMode QuickRenderMode);
+
+	/** Caches any necessary data before PIE starts. Stored in the CachedPrePieData member variable. */
+	static void CachePrePieData(UWorld* InEditorWorld);
+
+	/** Handles restoring data after PIE ends. */
+	static void RestorePrePieData();
+
 	/**
 	 * Populates TemporaryLevelSequence with the current level sequence, or sets up a new level sequence for rendering (depending on the mode that
 	 * Quick Render is using).
@@ -64,6 +86,6 @@ private:
 	/** The level sequence that is being used by Quick Render. In some cases this may be made specifically for a quick render, in other cases it may be a level sequence that existed before. */
 	inline static TStrongObjectPtr<ULevelSequence> TemporaryLevelSequence = nullptr;
 
-	/** Level sequence actors which had their AutoPlay settings turned off prior to render, and need to be turned back on again after the render finishes. */
-	inline static TArray<TWeakObjectPtr<class ALevelSequenceActor>> ModifiedLevelSequenceActors;
+	/** Any data that needed to be cached prior to starting PIE. */
+	inline static FCachedPrePieData CachedPrePieData;
 };
