CIS: Fixed uninitialized struct property.

diff --git a/Engine/Plugins/MovieScene/MovieRenderPipeline/Source/MovieRenderPipelineCore/Public/Graph/MovieGraphTimeStepData.h b/Engine/Plugins/MovieScene/MovieRenderPipeline/Source/MovieRenderPipelineCore/Public/Graph/MovieGraphTimeStepData.h
--- a/Engine/Plugins/MovieScene/MovieRenderPipeline/Source/MovieRenderPipelineCore/Public/Graph/MovieGraphTimeStepData.h
+++ b/Engine/Plugins/MovieScene/MovieRenderPipeline/Source/MovieRenderPipelineCore/Public/Graph/MovieGraphTimeStepData.h
@@ -24,6 +24,7 @@ public:
 		, RenderedFrameNumber(0)
 		, FrameDeltaTime(0.f)
 		, WorldTimeDilation(0.f)
+		, bHasRelativeTimeBeenUsed(false)
 		, WorldSeconds(0.f)
 		, MotionBlurFraction(0.f)
 		, FrameRate(0, 0)
