MoviePipeline: Implemented the "Use Viewport Camera in Sequence" mode in Quick Render.

diff --git a/Engine/Plugins/MovieScene/MovieRenderPipeline/Source/MovieRenderPipelineEditor/Private/Graph/MovieGraphQuickRender.cpp b/Engine/Plugins/MovieScene/MovieRenderPipeline/Source/MovieRenderPipelineEditor/Private/Graph/MovieGraphQuickRender.cpp
--- a/Engine/Plugins/MovieScene/MovieRenderPipeline/Source/MovieRenderPipelineEditor/Private/Graph/MovieGraphQuickRender.cpp
+++ b/Engine/Plugins/MovieScene/MovieRenderPipeline/Source/MovieRenderPipelineEditor/Private/Graph/MovieGraphQuickRender.cpp
@@ -20,6 +20,7 @@
 #include "Sections/MovieSceneCameraCutSection.h"
 #include "SequencerUtilities.h"
 #include "Tracks/MovieSceneCameraCutTrack.h"
+#include "Tracks/MovieSceneSubTrack.h"
 
 #include "Graph/MovieGraphRenderLayerSubsystem.h"
 #include "Graph/MovieRenderGraphEditorSettings.h"
@@ -62,13 +63,16 @@ void FMovieGraphQuickRender::BeginQuickRender(const UMovieGraphQuickRenderSettin
 	// Always use the PIE executor for quick render, rather than the executor specified in the project settings for normal MRG.
 	const TSubclassOf<UMoviePipelineExecutorBase> ExecutorClass = UMoviePipelinePIEExecutor::StaticClass();
 
-	// Get/create the level sequence that's appropriate for the current Quick Render mode
-	SetUpLevelSequenceForRender(InQuickRenderSettings);
-	if (!TemporaryLevelSequence.IsValid())
+	// Get/create the level sequence that should be rendered from (varies depending on the Quick Render mode).
+	SetUpRenderingLevelSequence(InQuickRenderSettings);
+	if (!RenderingLevelSequence.IsValid())
 	{
 		return;
 	}
 
+	// Also set up the utility level sequence. This is not expected to be valid for most modes, so we don't need to check that this was successfully created.
+	SetUpUtilityLevelSequence(InQuickRenderSettings);
+
 	// Don't do additional work if the selected mode doesn't have what it needs to do a render. AreModePrerequisitesMet() will output the appropriate warnings.
 	if (!AreModePrerequisitesMet(InQuickRenderSettings.QuickRenderMode))
 	{
@@ -88,7 +92,7 @@ void FMovieGraphQuickRender::BeginQuickRender(const UMovieGraphQuickRenderSettin
 
 	// Add a temp job to the queue, and give it the editor's level and the current level sequence being edited in Sequencer
 	UMoviePipelineExecutorJob* NewJob = TemporaryQueue->AllocateNewJob(UMoviePipelineExecutorJob::StaticClass());
-	NewJob->SetSequence(FSoftObjectPath(TemporaryLevelSequence.Get()));
+	NewJob->SetSequence(FSoftObjectPath(RenderingLevelSequence.Get()));
 	NewJob->Map = FSoftObjectPath(EditorWorld);
 	NewJob->JobName = FString(TEXT("QuickRender"));
 
@@ -162,14 +166,15 @@ void FMovieGraphQuickRender::BeginQuickRender(const UMovieGraphQuickRenderSettin
 	});
 
 	// Clean up
-	Executor->OnExecutorFinished().AddLambda([](UMoviePipelineExecutorBase* InExecutor, bool bSuccess)
+	Executor->OnExecutorFinished().AddLambda([&InQuickRenderSettings](UMoviePipelineExecutorBase* InExecutor, bool bSuccess)
 	{
+		RestorePreRenderState(InQuickRenderSettings);
+
 		// The temp objects are part of the transient package, so manually release them when the render is finished
 		TemporaryQueue.Reset();
 		TemporaryGraph.Reset();
-		TemporaryLevelSequence.Reset();
-
-		RestorePrePieData();
+		RenderingLevelSequence.Reset();
+		UtilityLevelSequence.Reset();
 	});
 }
 
@@ -506,7 +511,7 @@ void FMovieGraphQuickRender::CachePrePieData(UWorld* InEditorWorld)
 	CopyToWeakArray(CachedPrePieData.SelectedCameras, CameraActors);
 }
 
-void FMovieGraphQuickRender::RestorePrePieData()
+void FMovieGraphQuickRender::RestorePreRenderState(const UMovieGraphQuickRenderSettings& InQuickRenderSettings)
 {
 	// Restore the modified level sequences to the way they were before the render.
 	for (const TWeakObjectPtr<ALevelSequenceActor>& LevelSequenceActor : CachedPrePieData.ModifiedLevelSequenceActors)
@@ -518,31 +523,86 @@ void FMovieGraphQuickRender::RestorePrePieData()
 	}
 	
 	CachedPrePieData.SelectedCameras.Empty();
+
+	// If we're rendering with the Viewport Camera in Current Sequence mode, remove the subsequence track that was added
+	if (InQuickRenderSettings.QuickRenderMode == EMovieGraphQuickRenderMode::UseViewportCameraInSequence)
+	{
+		const ULevelSequence* CurrentLevelSequence = ULevelSequenceEditorBlueprintLibrary::GetCurrentLevelSequence();
+		check(CurrentLevelSequence);
+
+		UMovieSceneTrack* TrackToRemove = nullptr;
+		for (UMovieSceneSection* Section : CurrentLevelSequence->GetMovieScene()->GetAllSections())
+		{
+			if (const UMovieSceneSubSection* SubsequenceSection = Cast<UMovieSceneSubSection>(Section))
+			{
+				if (SubsequenceSection->GetSequence() == UtilityLevelSequence.Get())
+				{
+					TrackToRemove = Section->GetTypedOuter<UMovieSceneTrack>();
+					break;
+				}
+			}
+		}
+
+		if (TrackToRemove)
+		{
+			CurrentLevelSequence->GetMovieScene()->RemoveTrack(*TrackToRemove);
+		}
+	}
 }
 
-void FMovieGraphQuickRender::SetUpLevelSequenceForRender(const UMovieGraphQuickRenderSettings& InQuickRenderSettings)
+void FMovieGraphQuickRender::SetUpRenderingLevelSequence(const UMovieGraphQuickRenderSettings& InQuickRenderSettings)
 {
+	auto GetCurrentLevelSequenceOrWarn = [](EMovieGraphQuickRenderMode ActiveQuickRenderMode) -> ULevelSequence*
+	{
+		ULevelSequence* CurrentLevelSequence = ULevelSequenceEditorBlueprintLibrary::GetCurrentLevelSequence();
+		if (!CurrentLevelSequence)
+		{
+			const FText ModeName = StaticEnum<EMovieGraphQuickRenderMode>()->GetDisplayNameTextByValue(static_cast<int64>(ActiveQuickRenderMode));
+			const FText MessageText = LOCTEXT("NoActiveSequenceWarning", "Quick Render needs a level sequence opened in Sequencer in order to function in the current mode [{0}]. Open a level sequence, then try again.");
+		
+			FMessageDialog::Open(EAppMsgType::Ok, FText::Format(MessageText, ModeName));
+
+			return nullptr;
+		}
+
+		return CurrentLevelSequence;
+	};
+	
 	const EMovieGraphQuickRenderMode QuickRenderMode = InQuickRenderSettings.QuickRenderMode;
 	
+	// For the CurrentViewport and SelectedCameras modes, a temporary level sequence will be used to drive the render.
 	if ((QuickRenderMode == EMovieGraphQuickRenderMode::CurrentViewport) || (QuickRenderMode == EMovieGraphQuickRenderMode::SelectedCameras))
 	{
 		// Create the new, temporary level sequence
 		ULevelSequence* QuickRenderSequence = NewObject<ULevelSequence>(GetTransientPackage(), TEXT("QuickRenderSequence"), RF_Transient);
 		QuickRenderSequence->Initialize();
 
-		TemporaryLevelSequence.Reset(QuickRenderSequence);
+		RenderingLevelSequence.Reset(QuickRenderSequence);
+		
 		return;
 	}
 
-	// Use the level sequence that's currently active
-	ULevelSequence* CurrentLevelSequence = ULevelSequenceEditorBlueprintLibrary::GetCurrentLevelSequence();
+	// Otherwise, the active mode should use the level sequence that's currently active in Sequencer.
+	ULevelSequence* CurrentLevelSequence = GetCurrentLevelSequenceOrWarn(QuickRenderMode);
 	if (!CurrentLevelSequence)
 	{
-		UE_LOG(LogMovieRenderPipeline, Error, TEXT("Unable to Quick Render: No level sequence currently active to render from ('Current Sequence' mode is active)."));
 		return;
 	}
 	
-	TemporaryLevelSequence.Reset(CurrentLevelSequence);
+	RenderingLevelSequence.Reset(CurrentLevelSequence);
+}
+
+void FMovieGraphQuickRender::SetUpUtilityLevelSequence(const UMovieGraphQuickRenderSettings& InQuickRenderSettings)
+{
+	// This mode needs a utility sequence that contains a camera cut with a high hierarchical bias, which will allow it to override all cameras in
+	// the level sequence being rendered.
+	if (InQuickRenderSettings.QuickRenderMode == EMovieGraphQuickRenderMode::UseViewportCameraInSequence)
+	{
+		ULevelSequence* UtilitySequence = NewObject<ULevelSequence>(GetTransientPackage(), TEXT("QuickRenderUtilitySequence"), RF_Transient);
+		UtilitySequence->Initialize();
+
+		UtilityLevelSequence.Reset(UtilitySequence);
+	}
 }
 
 UWorld* FMovieGraphQuickRender::GetWorldOfType(EWorldType::Type WorldType)
@@ -562,12 +622,12 @@ UWorld* FMovieGraphQuickRender::GetWorldOfType(EWorldType::Type WorldType)
 
 void FMovieGraphQuickRender::PerformPreRenderSetup(const UMovieGraphQuickRenderSettings& InQuickRenderSettings)
 {
-	// Adds the specified camera to the Camera Cut track in the level sequence that will be used for Quick Render.
+	// Adds the specified camera to the Camera Cut track in the given level sequence.
 	// If there's already a camera cut section within the Camera Cut track, a new one will be added immediately after.
 	// Each section is only one frame long.
-	auto AddCameraToLevelSequence = [](ACameraActor* InCameraToAdd)
+	auto AddCameraToLevelSequence = [](ACameraActor* InCameraToAdd, const ULevelSequence* InDestLevelSequence, const TRange<FFrameNumber> OptionalFrameRangeOverride = TRange<FFrameNumber>::Empty())
 	{
-		UMovieScene* MovieScene = TemporaryLevelSequence->GetMovieScene();
+		UMovieScene* MovieScene = InDestLevelSequence->GetMovieScene();
 
 		// Add a Camera Cut track if one hasn't already been added
 		UMovieSceneCameraCutTrack* CameraCutTrack = Cast<UMovieSceneCameraCutTrack>(MovieScene->GetCameraCutTrack());
@@ -582,8 +642,12 @@ void FMovieGraphQuickRender::PerformPreRenderSetup(const UMovieGraphQuickRenderS
 		const FFrameNumber StartFrame = MovieScene->GetPlaybackRange().GetLowerBound().GetValue();
 		const FFrameNumber EndFrame = FFrameRate::TransformTime(FFrameTime(NumCameraCutSections + 1), MovieScene->GetDisplayRate(), MovieScene->GetTickResolution()).GetFrame();
 
-		// The entire playback range for the movie scene always starts a 0 and ends at the last camera cut section's end frame.
-		MovieScene->SetPlaybackRange(TRange<FFrameNumber>::Inclusive(FFrameNumber(0), EndFrame));
+		// The entire playback range for the movie scene always starts a 0 and ends at the last camera cut section's end frame, unless there's an override.
+		const bool bHasFrameRangeOverride = !OptionalFrameRangeOverride.IsEmpty();
+		const TRange<FFrameNumber> PlaybackRange = bHasFrameRangeOverride
+			? OptionalFrameRangeOverride
+			: TRange<FFrameNumber>::Inclusive(FFrameNumber(0), EndFrame);
+		MovieScene->SetPlaybackRange(PlaybackRange);
 		
 		UE::Sequencer::FCreateBindingParams CreateBindingParams;
 		CreateBindingParams.bSpawnable = true;
@@ -592,7 +656,7 @@ void FMovieGraphQuickRender::PerformPreRenderSetup(const UMovieGraphQuickRenderS
 
 		// Create the binding for the camera. This is wrapped in a dummy transaction so we can cancel it later on to prevent the binding creation from
 		// making an entry on the undo stack.
-		#if WITH_EDITOR
+#if WITH_EDITOR
 		const int32 TransactionId = GEditor->BeginTransaction(FText::GetEmpty());
 #endif
 		const FGuid CameraBinding = FSequencerUtilities::CreateOrReplaceBinding(nullptr, MovieScene->GetTypedOuter<UMovieSceneSequence>(), InCameraToAdd, CreateBindingParams);
@@ -600,12 +664,17 @@ void FMovieGraphQuickRender::PerformPreRenderSetup(const UMovieGraphQuickRenderS
 		GEditor->CancelTransaction(TransactionId);
 #endif
 
+		const TRange<FFrameNumber> SectionRange = bHasFrameRangeOverride
+			? OptionalFrameRangeOverride
+			: TRange<FFrameNumber>::Inclusive(StartFrame, EndFrame);
+
 		// Add the new camera to the camera cut track.
 		UMovieSceneCameraCutSection* CameraCutSection = CameraCutTrack->AddNewCameraCut(UE::MovieScene::FRelativeObjectBindingID(CameraBinding), StartFrame);
-		CameraCutSection->SetRange(TRange<FFrameNumber>::Inclusive(StartFrame, EndFrame));
+		CameraCutSection->SetRange(SectionRange);
 	};
 	
-	if (InQuickRenderSettings.QuickRenderMode == EMovieGraphQuickRenderMode::CurrentViewport)
+	// Spawn a new camera which mimics the viewport's camera 
+	auto DuplicateViewportCamera = []() -> ACameraActor*
 	{
 		// Get the PIE world, this is where we need to target the changes (not the editor world). Otherwise we'll have to deal with deleting actors
 		// after the render finishes, which isn't ideal (and this will also create entries in the undo stack, also not ideal).
@@ -625,7 +694,32 @@ void FMovieGraphQuickRender::PerformPreRenderSetup(const UMovieGraphQuickRenderS
 		NewCamera->SetActorRotation(ViewportClient->GetViewRotation());
 		NewCamera->GetCameraComponent()->SetFieldOfView(ViewportClient->ViewFOV);
 
-		AddCameraToLevelSequence(NewCamera);
+		return NewCamera;
+	};
+	
+	if (InQuickRenderSettings.QuickRenderMode == EMovieGraphQuickRenderMode::CurrentViewport)
+	{
+		ACameraActor* NewCamera = DuplicateViewportCamera();
+		AddCameraToLevelSequence(NewCamera, RenderingLevelSequence.Get());
+	}
+	else if (InQuickRenderSettings.QuickRenderMode == EMovieGraphQuickRenderMode::UseViewportCameraInSequence)
+	{
+		UMovieScene* CurrentMovieScene = RenderingLevelSequence.Get()->GetMovieScene();
+		const TRange<FFrameNumber> CurrentPlaybackRange = CurrentMovieScene->GetPlaybackRange();
+
+		// Add a subsequence to the currently-active level sequence
+		const FFrameNumber StartFrame = CurrentPlaybackRange.GetLowerBoundValue();
+		const FFrameNumber EndFrame = CurrentPlaybackRange.GetUpperBoundValue();
+		UMovieSceneSubTrack* NewSubsequenceTrack = CastChecked<UMovieSceneSubTrack>(CurrentMovieScene->AddTrack(UMovieSceneSubTrack::StaticClass()));
+		UMovieSceneSubSection* SequenceSubsection = NewSubsequenceTrack->AddSequence(UtilityLevelSequence.Get(), StartFrame, (EndFrame - StartFrame).Value);
+
+		// Set the hierarchical bias of the new subsequence section so its camera takes precedence over all other cameras (eg, those from shots).
+		// Use a very high number; the value used here is arbitrary, but should be high enough to do the job.
+		SequenceSubsection->Parameters.HierarchicalBias = 1000000;
+
+		// Add the viewport camera to the subsequence. This is done through the level sequence that was added to it.
+		ACameraActor* NewCamera = DuplicateViewportCamera();
+		AddCameraToLevelSequence(NewCamera, UtilityLevelSequence.Get(), CurrentPlaybackRange);
 	}
 	else if (InQuickRenderSettings.QuickRenderMode == EMovieGraphQuickRenderMode::SelectedCameras)
 	{
@@ -633,7 +727,7 @@ void FMovieGraphQuickRender::PerformPreRenderSetup(const UMovieGraphQuickRenderS
 		{
 			if (SelectedCamera.IsValid())
 			{
-				AddCameraToLevelSequence(SelectedCamera.Get());
+				AddCameraToLevelSequence(SelectedCamera.Get(), RenderingLevelSequence.Get());
 			}
 		}
 	}
diff --git a/Engine/Plugins/MovieScene/MovieRenderPipeline/Source/MovieRenderPipelineEditor/Public/Graph/MovieGraphQuickRender.h b/Engine/Plugins/MovieScene/MovieRenderPipeline/Source/MovieRenderPipelineEditor/Public/Graph/MovieGraphQuickRender.h
--- a/Engine/Plugins/MovieScene/MovieRenderPipeline/Source/MovieRenderPipelineEditor/Public/Graph/MovieGraphQuickRender.h
+++ b/Engine/Plugins/MovieScene/MovieRenderPipeline/Source/MovieRenderPipelineEditor/Public/Graph/MovieGraphQuickRender.h
@@ -62,13 +62,16 @@ private:
 	static void CachePrePieData(UWorld* InEditorWorld);
 
 	/** Handles restoring data after PIE ends. */
-	static void RestorePrePieData();
+	static void RestorePreRenderState(const UMovieGraphQuickRenderSettings& InQuickRenderSettings);
 
 	/**
-	 * Populates TemporaryLevelSequence with the current level sequence, or sets up a new level sequence for rendering (depending on the mode that
+	 * Populates RenderingLevelSequence with the current level sequence, or sets up a new level sequence for rendering (depending on the mode that
 	 * Quick Render is using).
 	 */
-	static void SetUpLevelSequenceForRender(const UMovieGraphQuickRenderSettings& InQuickRenderSettings);
+	static void SetUpRenderingLevelSequence(const UMovieGraphQuickRenderSettings& InQuickRenderSettings);
+
+	/** Populates UtilityLevelSequence with a level sequence used for utility purposes, or does nothing if no utility sequence is needed for the current mode. */
+	static void SetUpUtilityLevelSequence(const UMovieGraphQuickRenderSettings& InQuickRenderSettings);
 
 	/** Gets the first world found of the specified type. */
 	static UWorld* GetWorldOfType(EWorldType::Type WorldType);
@@ -83,8 +86,17 @@ private:
 	/** The temporary graph that is used by quick render. Usually a duplicate of a graph created outside of quick render, but will be modified by quick render. */
 	inline static TStrongObjectPtr<UMovieGraphConfig> TemporaryGraph = nullptr;
 
-	/** The level sequence that is being used by Quick Render. In some cases this may be made specifically for a quick render, in other cases it may be a level sequence that existed before. */
-	inline static TStrongObjectPtr<ULevelSequence> TemporaryLevelSequence = nullptr;
+	/**
+	 * The level sequence that is used for rendering (ie, provided directly to MRG to render from). Sometimes this will be the current level sequence
+	 * in Sequencer, and other times it may be a temporary level sequence solely for use in Quick Render. It depends on the mode that is being used.
+	 */
+	inline static TStrongObjectPtr<ULevelSequence> RenderingLevelSequence = nullptr;
+	
+	/**
+     * The level sequence that is being used by Quick Render for utility purposes. This is not directly rendered from, but may be included in another
+     * level sequence to manipulate it (for example). For most modes, this will remain unused.
+     */
+    inline static TStrongObjectPtr<ULevelSequence> UtilityLevelSequence = nullptr;
 
 	/** Any data that needed to be cached prior to starting PIE. */
 	inline static FCachedPrePieData CachedPrePieData;
