MovieGraph: Added the ability for the output resolution to automatically be adapted based on the aspect ratio of cameras, allowing the user to leave their settings at a given resolution and know that they'll get an image that is no bigger than the specified image, while preserving the aspect ratio, and removing any letterboxing. See the "Adapt Resolution" setting on the Global Output Settings Node.

Tested against existing Constrain Aspect Ratio tests, verified that it produces expected results in permutations of (wide, tall, equal) camera aspect ratio vs (wide, tall, equal) output resolutions, with automatic, width, and height tests. Verified cameras that do not ask you to constrain aspect ratio don't edit the aspect ratio. Verified overscan and burn-ins.

diff --git a/Engine/Plugins/MovieScene/MovieRenderPipeline/Source/MovieRenderPipelineCore/Private/Graph/MovieGraphBlueprintLibrary.cpp b/Engine/Plugins/MovieScene/MovieRenderPipeline/Source/MovieRenderPipelineCore/Private/Graph/MovieGraphBlueprintLibrary.cpp
--- a/Engine/Plugins/MovieScene/MovieRenderPipeline/Source/MovieRenderPipelineCore/Private/Graph/MovieGraphBlueprintLibrary.cpp
+++ b/Engine/Plugins/MovieScene/MovieRenderPipeline/Source/MovieRenderPipelineCore/Private/Graph/MovieGraphBlueprintLibrary.cpp
@@ -378,7 +378,7 @@ FIntPoint UMovieGraphBlueprintLibrary::GetEffectiveOutputResolution(UMovieGraphE
 	return GetOverscannedResolution(InEvaluatedGraph, DefaultOverscan);
 }
 
-FIntPoint UMovieGraphBlueprintLibrary::GetDesiredOutputResolution(UMovieGraphEvaluatedConfig* InEvaluatedGraph)
+FIntPoint UMovieGraphBlueprintLibrary::GetDesiredOutputResolution(UMovieGraphEvaluatedConfig* InEvaluatedGraph, float CameraAspectRatio)
 {
 	if (!InEvaluatedGraph)
 	{
@@ -405,10 +405,58 @@ FIntPoint UMovieGraphBlueprintLibrary::GetDesiredOutputResolution(UMovieGraphEva
 		NamedResolution = OutputSetting->OutputResolution;
 	}
 
-	return NamedResolution.Resolution;
+	EMovieGraphAspectRatioAdaptBehavior ResizeBehavior = OutputSetting->AdaptResolution;
+	if (ResizeBehavior == EMovieGraphAspectRatioAdaptBehavior::Disabled)
+	{
+		return NamedResolution.Resolution;
+	}
+
+	float TargetAspectRatio = NamedResolution.Resolution.X / (float)NamedResolution.Resolution.Y;
+
+	// For backwards compatibility reasons, if they did not specify an aspect ratio, then we use the aspect ratio of the
+	// output image. This should result in no change for previous callers that do not specify the camera aspect ratio.
+	if (FMath::IsNearlyZero(CameraAspectRatio))
+	{
+		CameraAspectRatio = TargetAspectRatio;
+	}
+
+	// Automatic works by inspecting the aspect ratio and picking the 'correct' behavior that matches
+	// what you would see in the viewport, but without the letterboxing.
+	if (ResizeBehavior == EMovieGraphAspectRatioAdaptBehavior::Automatic)
+	{
+		if (CameraAspectRatio > TargetAspectRatio)
+		{
+			ResizeBehavior = EMovieGraphAspectRatioAdaptBehavior::ScaleToWidth;
+		}
+		else
+		{
+			ResizeBehavior = EMovieGraphAspectRatioAdaptBehavior::ScaleToHeight;
+		}
+	}
+
+	FIntPoint OutputResolution = NamedResolution.Resolution;
+	switch (ResizeBehavior)
+	{
+		case EMovieGraphAspectRatioAdaptBehavior::ScaleToWidth:
+		{
+			OutputResolution.Y = FMath::RoundToInt(OutputResolution.X / CameraAspectRatio);
+			break;
+
+		}
+		case EMovieGraphAspectRatioAdaptBehavior::ScaleToHeight:
+		{
+			OutputResolution.X = FMath::RoundToInt(OutputResolution.Y * CameraAspectRatio);
+			break;
+		}
+
+	default:
+		check(0); // Unsupported
+	}
+
+	return OutputResolution;
 }
 
-FIntPoint UMovieGraphBlueprintLibrary::GetOverscannedResolution(UMovieGraphEvaluatedConfig* InEvaluatedGraph, float DefaultOverscan)
+FIntPoint UMovieGraphBlueprintLibrary::GetOverscannedResolution(UMovieGraphEvaluatedConfig* InEvaluatedGraph, float DefaultOverscan, float CameraAspectRatio)
 {
 	if (!InEvaluatedGraph)
     {
@@ -416,7 +464,7 @@ FIntPoint UMovieGraphBlueprintLibrary::GetOverscannedResolution(UMovieGraphEvalu
     	return FIntPoint();
     }
 
-	FIntPoint DesiredOutputResolution = GetDesiredOutputResolution(InEvaluatedGraph);
+	FIntPoint DesiredOutputResolution = GetDesiredOutputResolution(InEvaluatedGraph, CameraAspectRatio);
 	
     constexpr bool bIncludeCDOs = false;
     float RescaledOverscan = DefaultOverscan;
@@ -432,7 +480,7 @@ FIntPoint UMovieGraphBlueprintLibrary::GetOverscannedResolution(UMovieGraphEvalu
 	return UMoviePipelineBlueprintLibrary::ScaleResolutionByOverscan(RescaledOverscan, DesiredOutputResolution);
 }
 
-FIntPoint UMovieGraphBlueprintLibrary::GetBackbufferResolution(UMovieGraphEvaluatedConfig* InEvaluatedGraph, float DefaultOverscan)
+FIntPoint UMovieGraphBlueprintLibrary::GetBackbufferResolution(UMovieGraphEvaluatedConfig* InEvaluatedGraph, float DefaultOverscan, float CameraAspectRatio)
 {
 	if (!InEvaluatedGraph)
 	{
@@ -440,13 +488,13 @@ FIntPoint UMovieGraphBlueprintLibrary::GetBackbufferResolution(UMovieGraphEvalua
 		return FIntPoint();
 	}
 
-	FIntPoint OverscannedResolution = GetOverscannedResolution(InEvaluatedGraph, DefaultOverscan);
+	FIntPoint OverscannedResolution = GetOverscannedResolution(InEvaluatedGraph, DefaultOverscan, CameraAspectRatio);
 
 	// TODO: When tiling support is added, compute the appropriate backbuffer resolution from OverscannedResolution using tile count and return it
 	return OverscannedResolution;
 }
 
-FIntRect UMovieGraphBlueprintLibrary::GetOverscanCropRectangle(UMovieGraphEvaluatedConfig* InEvaluatedGraph, float DefaultOverscan)
+FIntRect UMovieGraphBlueprintLibrary::GetOverscanCropRectangle(UMovieGraphEvaluatedConfig* InEvaluatedGraph, float DefaultOverscan, float CameraAspectRatio)
 {
 	if (!InEvaluatedGraph)
 	{
@@ -454,8 +502,8 @@ FIntRect UMovieGraphBlueprintLibrary::GetOverscanCropRectangle(UMovieGraphEvalua
 		return FIntRect();
 	}
 
-	const FIntPoint DesiredOutputResolution = GetDesiredOutputResolution(InEvaluatedGraph);
-	const FIntPoint OverscannedResolution = GetOverscannedResolution(InEvaluatedGraph, DefaultOverscan);
+	const FIntPoint DesiredOutputResolution = GetDesiredOutputResolution(InEvaluatedGraph, CameraAspectRatio);
+	const FIntPoint OverscannedResolution = GetOverscannedResolution(InEvaluatedGraph, DefaultOverscan, CameraAspectRatio);
 
 	FIntRect CropRect;
 
diff --git a/Engine/Plugins/MovieScene/MovieRenderPipeline/Source/MovieRenderPipelineCore/Private/Graph/Nodes/MovieGraphWidgetRendererBaseNode.cpp b/Engine/Plugins/MovieScene/MovieRenderPipeline/Source/MovieRenderPipelineCore/Private/Graph/Nodes/MovieGraphWidgetRendererBaseNode.cpp
--- a/Engine/Plugins/MovieScene/MovieRenderPipeline/Source/MovieRenderPipelineCore/Private/Graph/Nodes/MovieGraphWidgetRendererBaseNode.cpp
+++ b/Engine/Plugins/MovieScene/MovieRenderPipeline/Source/MovieRenderPipelineCore/Private/Graph/Nodes/MovieGraphWidgetRendererBaseNode.cpp
@@ -113,9 +113,12 @@ void UMovieGraphWidgetRendererBaseNode::FMovieGraphWidgetPass::Render(const FMov
 		return;
 	}
 
+	UE::MovieGraph::DefaultRenderer::FCameraInfo CameraInfo = Renderer->GetCameraInfo(LayerData.CameraIndex);
+	// 0.f will preserve original behavior (where aspect ratio is assumed to be the output aspect ratio)
+	const float CameraAspectRatio = CameraInfo.bAllowCameraAspectRatio && CameraInfo.ViewInfo.bConstrainAspectRatio ? CameraInfo.ViewInfo.AspectRatio : 0.f;
 	const float CameraOverscan = Renderer->GetCameraOverscan(LayerData.CameraIndex);
-	const FIntPoint OverscannedResolution = UMovieGraphBlueprintLibrary::GetOverscannedResolution(InTimeData.EvaluatedConfig, CameraOverscan);
-	const FIntRect CropRect = UMovieGraphBlueprintLibrary::GetOverscanCropRectangle(InTimeData.EvaluatedConfig, CameraOverscan);
+	const FIntPoint OverscannedResolution = UMovieGraphBlueprintLibrary::GetOverscannedResolution(InTimeData.EvaluatedConfig, CameraOverscan, CameraAspectRatio);
+	const FIntRect CropRect = UMovieGraphBlueprintLibrary::GetOverscanCropRectangle(InTimeData.EvaluatedConfig, CameraOverscan, CameraAspectRatio);
 	
 	// Composited elements should be sized to the original frustum size, as the final image is either cropped to that size, or the composite will be offset
 	// to match the original frustum
diff --git a/Engine/Plugins/MovieScene/MovieRenderPipeline/Source/MovieRenderPipelineCore/Public/Graph/MovieGraphBlueprintLibrary.h b/Engine/Plugins/MovieScene/MovieRenderPipeline/Source/MovieRenderPipelineCore/Public/Graph/MovieGraphBlueprintLibrary.h
--- a/Engine/Plugins/MovieScene/MovieRenderPipeline/Source/MovieRenderPipelineCore/Public/Graph/MovieGraphBlueprintLibrary.h
+++ b/Engine/Plugins/MovieScene/MovieRenderPipeline/Source/MovieRenderPipelineCore/Public/Graph/MovieGraphBlueprintLibrary.h
@@ -76,37 +76,45 @@ public:
 	 * and is the target resolution that the pipeline will generally output (e.g. when cropping overscan in non-EXR formats). Use GetOverscannedResolution or
 	 * GetBackbufferResolution to get resolutions that factor in overscan and tiling, respectively
 	 * @param	InEvaluatedGraph	- The evaluated graph that will provide context for resolving the resolution
+	 * @param	CameraAspectRatio	- The aspect ratio for the camera. Set to zero if you want to adapt to the output resolution's 
+	 *								  aspect ratio (otherwise resolution may be adapted to fit aspect ratio based on config settings).
 	 * @return						- The desired output resolution
 	 */
 	UFUNCTION(BlueprintCallable, Category = "Movie Graph")
-	static FIntPoint GetDesiredOutputResolution(UMovieGraphEvaluatedConfig* InEvaluatedGraph);
+	static FIntPoint GetDesiredOutputResolution(UMovieGraphEvaluatedConfig* InEvaluatedGraph, float CameraAspectRatio = 0.0f);
 
 	/**
 	* Gets the desired output resolution scaled by any configured overscan for the movie render graph
 	* @param	InEvaluatedGraph	- The evaluated graph that will provide context for resolving the resolution
 	* @param	DefaultOverscan		- The default overscan to use if there are no camera settings that provide an overscan override value, from 0.0 to 1.0
+	* @param	CameraAspectRatio	- The aspect ratio for the camera. Set to zero if you want to adapt to the output resolution's
+	*								  aspect ratio (otherwise resolution may be adapted to fit aspect ratio based on config settings).
 	* @return						- The overscan-scaled resolution
 	*/
 	UFUNCTION(BlueprintCallable, Category = "Movie Graph")
-	static FIntPoint GetOverscannedResolution(UMovieGraphEvaluatedConfig* InEvaluatedGraph, float DefaultOverscan = 0.0f);
+	static FIntPoint GetOverscannedResolution(UMovieGraphEvaluatedConfig* InEvaluatedGraph, float DefaultOverscan = 0.0f, float CameraAspectRatio = 0.0f);
 
 	/**
 	 * Gets the resolution that frames will actually be rendered at in MRG, which includes factors such as overscan and tiling
 	 * @param	InEvaluatedGraph	- The evaluated graph that will provide context for resolving the resolution
 	 * @param	DefaultOverscan		- The default overscan to use if there are no camera settings that provide an overscan override value, from 0.0 to 1.0
+	 * @param	CameraAspectRatio	- The aspect ratio for the camera. Set to zero if you want to adapt to the output resolution's
+	 *								  aspect ratio (otherwise resolution may be adapted to fit aspect ratio based on config settings).
 	 * @return						- The resolution of a rendered frame
 	 */
 	UFUNCTION(BlueprintCallable, Category = "Movie Graph")
-	static FIntPoint GetBackbufferResolution(UMovieGraphEvaluatedConfig* InEvaluatedGraph, float DefaultOverscan = 0.0f);
+	static FIntPoint GetBackbufferResolution(UMovieGraphEvaluatedConfig* InEvaluatedGraph, float DefaultOverscan = 0.0f, float CameraAspectRatio = 0.0f);
 	
 	/**
 	 * Gets a rectangle that will crop out any overscan applied to the scene. If there is no overscan, the rectangle returned will match the backbuffer rectangle
 	 * @param InEvaluatedGraph		- The evaluated graph that will provide context for resolving the resolution
 	 * @param DefaultOverscan		- The default overscan to use if there are no camera settings that provide an overscan override value, from 0.0 to 1.0
+	 * @param	CameraAspectRatio	- The aspect ratio for the camera. Set to zero if you want to adapt to the output resolution's
+	 *								  aspect ratio (otherwise resolution may be adapted to fit aspect ratio based on config settings).
 	 * @return						- The crop rectangle, which will have a resolution matching the requested output resolution
 	 */
 	UFUNCTION(BlueprintCallable, Category = "Movie Graph")
-	static FIntRect GetOverscanCropRectangle(UMovieGraphEvaluatedConfig* InEvaluatedGraph, float DefaultOverscan = 0.0f);
+	static FIntRect GetOverscanCropRectangle(UMovieGraphEvaluatedConfig* InEvaluatedGraph, float DefaultOverscan = 0.0f, float CameraAspectRatio = 0.0f);
 	
 	/**
 	* Gets the name of the current job.
diff --git a/Engine/Plugins/MovieScene/MovieRenderPipeline/Source/MovieRenderPipelineCore/Public/Graph/Nodes/MovieGraphGlobalOutputSettingNode.h b/Engine/Plugins/MovieScene/MovieRenderPipeline/Source/MovieRenderPipelineCore/Public/Graph/Nodes/MovieGraphGlobalOutputSettingNode.h
--- a/Engine/Plugins/MovieScene/MovieRenderPipeline/Source/MovieRenderPipelineCore/Public/Graph/Nodes/MovieGraphGlobalOutputSettingNode.h
+++ b/Engine/Plugins/MovieScene/MovieRenderPipeline/Source/MovieRenderPipelineCore/Public/Graph/Nodes/MovieGraphGlobalOutputSettingNode.h
@@ -65,6 +65,49 @@ public:
 	int32 Value;
 };
 
+UENUM(BlueprintType)
+enum class EMovieGraphAspectRatioAdaptBehavior : uint8
+{
+	/** 
+	* Output resolution will not be modified. If the camera has Constrain Aspect Ratio enabled, then 
+	* letterboxing may be shown to preserve that aspect ratio within the specified Output Resolution.
+	*/
+	Disabled,
+
+	/**
+	* If the camera has Constrain Aspect Ratio enabled, then this should match Disabled, except the letterboxing
+	* will be cropped off. No behavior change if Constrain Aspect Ratio is not enabled.
+	*
+	* The output resolution will be resized to respect the camera aspect ratio, matching either the specified Width,
+	* or Height. Which dimension is picked will depend on the aspect ratio (ie: a tall camera aspect ratio will 
+	* preserve the height of the output, and crop the width, while a wide camera aspect ratio will preserve 
+	* the width of the output and crop the height.)
+	*/
+	Automatic,
+
+	/**
+	* If the camera has Constrain Aspect Ratio enabled, then the width of the output resolution will be preserved,
+	* and the height will automatically be adjusted based on the camera aspect ratio.
+	* 
+	* ie: If you have a 1.77 (16:9) Camera Aspect Ratio and target a 1.0 (1:1) output image, then the output image
+	* will have its height adjusted to preserve the 1.77 aspect ratio of the camera, without showing letterboxing.
+	* For example, a 1.77 Camera Aspect Ratio and a 1024x1024 output resolution will produce a render that is
+	* 1024x576, which preserves the existing aspect ratio and the given width. No behavior change if Constrain Aspect Ratio is not enabled.
+	*/
+	ScaleToWidth,
+
+	/**
+	* If the camera has Constrain Aspect Ratio enabled, then the height of the output resolution will be preserved,
+	* and the width will automatically be adjusted based on the camera aspect ratio.
+	*
+	* ie: If you have a 1.77 (16:9) Camera Aspect Ratio and target a 1.0 (1:1) output image, then the output image
+	* will have its width adjusted to preserve the 1.77 aspect ratio of the camera, without showing letterboxing.
+	* For example, a 1.77 Camera Aspect Ratio and a 1024x1024 output resolution will produce a render that is
+	* 1820x1024, which preserves the existing aspect ratio and the given height. No behavior change if Constrain Aspect Ratio is not enabled.
+	*/
+	ScaleToHeight
+};
+
 UCLASS()
 class MOVIERENDERPIPELINECORE_API UMovieGraphGlobalOutputSettingNode : public UMovieGraphSettingNode
 {
@@ -104,6 +147,9 @@ public:
 
 	UPROPERTY(EditAnywhere, BlueprintReadWrite, Category = Overrides, meta = (InlineEditConditionToggle))
 	uint8 bOverride_OutputResolution : 1;
+
+	UPROPERTY(EditAnywhere, BlueprintReadWrite, Category = Overrides, meta = (InlineEditConditionToggle))
+	uint8 bOverride_AdaptResolution : 1;
 	
 	UPROPERTY(EditAnywhere, BlueprintReadWrite, Category = Overrides, meta = (InlineEditConditionToggle))
 	uint8 bOverride_OutputFrameRate : 1;
@@ -153,6 +199,10 @@ public:
 	/** What resolution should our output files be exported at? */
 	UPROPERTY(EditAnywhere, BlueprintReadWrite, Category = "File Output", meta = (EditCondition = "bOverride_OutputResolution"))
 	FMovieGraphNamedResolution OutputResolution;
+
+	/** Should the output resolution be automatically adjusted to match the aspect ratio on cameras with Constrain Aspect Ratio? */
+	UPROPERTY(EditAnywhere, BlueprintReadWrite, Category = "File Output", meta = (EditCondition = "bOverride_AdaptResolution"))
+	EMovieGraphAspectRatioAdaptBehavior AdaptResolution;
 	
 	/** What frame rate should the output files be exported at? This overrides the Display Rate of the target sequence. If not overwritten, uses the default Sequence Display Rate. */
 	UPROPERTY(EditAnywhere, BlueprintReadWrite, Category = "File Output", meta=(EditCondition="bOverride_OutputFrameRate"))
diff --git a/Engine/Plugins/MovieScene/MovieRenderPipeline/Source/MovieRenderPipelineRenderPasses/Private/Graph/Renderers/MovieGraphDeferredPass.cpp b/Engine/Plugins/MovieScene/MovieRenderPipeline/Source/MovieRenderPipelineRenderPasses/Private/Graph/Renderers/MovieGraphDeferredPass.cpp
--- a/Engine/Plugins/MovieScene/MovieRenderPipeline/Source/MovieRenderPipelineRenderPasses/Private/Graph/Renderers/MovieGraphDeferredPass.cpp
+++ b/Engine/Plugins/MovieScene/MovieRenderPipeline/Source/MovieRenderPipelineRenderPasses/Private/Graph/Renderers/MovieGraphDeferredPass.cpp
@@ -178,18 +178,24 @@ void FMovieGraphDeferredPass::Render(const FMovieGraphTraversalContext& InFrameT
 		bOverrideCameraOverscan = CameraNode->bOverride_OverscanPercentage;
 	}
 
+	// We need the aspect ratio from the camera, but below it ends up getting modified and we pass references around
+	// so they need to stay unique instances below.
+	UE::MovieGraph::DefaultRenderer::FCameraInfo GlobalCameraInfo = GetRenderer()->GetCameraInfo(LayerData.CameraIndex);
+
 	// ToDo: When tiling is used, this should be the size of the per-tile backbuffer
 	const float CameraOverscan = GetRenderer()->GetCameraOverscan(LayerData.CameraIndex);
-	FIntPoint OverscannedResolution = UMovieGraphBlueprintLibrary::GetOverscannedResolution(InTimeData.EvaluatedConfig, CameraOverscan);
+	// 0.f will preserve original behavior (where aspect ratio is assumed to be the output aspect ratio)
+	const float CameraAspectRatio = GlobalCameraInfo.bAllowCameraAspectRatio && GlobalCameraInfo.ViewInfo.bConstrainAspectRatio ? GlobalCameraInfo.ViewInfo.AspectRatio : 0.f;
+	FIntPoint OverscannedResolution = UMovieGraphBlueprintLibrary::GetOverscannedResolution(InTimeData.EvaluatedConfig, CameraOverscan, CameraAspectRatio);
 	FIntPoint AccumulatorResolution = OverscannedResolution;
-	FIntPoint BackbufferResolution = UMovieGraphBlueprintLibrary::GetBackbufferResolution(InTimeData.EvaluatedConfig, CameraOverscan);
+	FIntPoint BackbufferResolution = UMovieGraphBlueprintLibrary::GetBackbufferResolution(InTimeData.EvaluatedConfig, CameraOverscan, CameraAspectRatio);
 	
 	DefaultRenderer::FRenderTargetInitParams RenderTargetInitParams = GetRenderTargetInitParams(InTimeData, AccumulatorResolution);
 	UTextureRenderTarget2D* RenderTarget = GraphRenderer->GetOrCreateViewRenderTarget(RenderTargetInitParams, RenderDataIdentifier);
 	FRenderTarget* RenderTargetResource = RenderTarget->GameThread_GetRenderTargetResource();
 	check(RenderTargetResource);
 
-	const FIntRect CropRect = UMovieGraphBlueprintLibrary::GetOverscanCropRectangle(InTimeData.EvaluatedConfig, CameraOverscan);
+	const FIntRect CropRect = UMovieGraphBlueprintLibrary::GetOverscanCropRectangle(InTimeData.EvaluatedConfig, CameraOverscan, CameraAspectRatio);
 	
 	// Now we can actually construct our ViewFamily, SceneView, and submit it for Rendering + Readback
 	for(int32 SpatialIndex = 0; SpatialIndex < NumSpatialSamples; SpatialIndex++)
