MoviePipeline: Misc polish and bug fixes for Quick Render.
* External viewers that support frame ranges will only display the _rendered_ frame range now, rather than the entire frame range that's available on disk.
* Fixed issue that resulted in an incorrect rendered frame range if only the start or end frame was marked as the selection in Sequencer.
* Fixed issue that could cause the same image sequence to be passed to an external custom viewer more than once.
* Fixed issue that caused the Custom Start/End Frame properties to not appear in the Settings dialog.
* The command used to open an external viewer is now logged.
* Updated the default frame range notation to "# with Start/End Frame". This is a new mode, but makes sense for most external viewers.
* Added a new icon for Quick Render in the toolbar (this will still most likely be temporary, but is better than what was there previously).
* Hid the "Current Shot at Playhead" mode. It was deemed not important enough for the first version of Quick Render, and may be brought back at a later time.
* Changing the "Playback Method" in the post-render settings now updates associated properties (like job/layer playback) to values that will most likely be used for the chosen playback method.
* Removed the requirement to use a cvar to see the Quick Render icon in the toolbar.
* Updated the number of warmup frames in the default Quick Render graph to 64 (matching the default graph for MRG).

#virtualized

diff --git a/Engine/Plugins/MovieScene/MovieRenderPipeline/Source/MovieRenderPipelineCore/Private/Graph/MovieRenderGraphEditorSettings.cpp b/Engine/Plugins/MovieScene/MovieRenderPipeline/Source/MovieRenderPipelineCore/Private/Graph/MovieRenderGraphEditorSettings.cpp
--- a/Engine/Plugins/MovieScene/MovieRenderPipeline/Source/MovieRenderPipelineCore/Private/Graph/MovieRenderGraphEditorSettings.cpp
+++ b/Engine/Plugins/MovieScene/MovieRenderPipeline/Source/MovieRenderPipelineCore/Private/Graph/MovieRenderGraphEditorSettings.cpp
@@ -6,4 +6,60 @@ UMovieRenderGraphEditorSettings::UMovieRenderGraphEditorSettings()
 	: PostRenderSettings(FMovieGraphPostRenderSettings())
 {
 	
-}
\ No newline at end of file
+}
+
+#if WITH_EDITOR
+void UMovieRenderGraphEditorSettings::PostEditChangeChainProperty(FPropertyChangedChainEvent& PropertyChangedEvent)
+{
+	static FProperty* ImageSequenceOptions =
+		FMovieGraphPostRenderSettings::StaticStruct()->FindPropertyByName(
+			GET_MEMBER_NAME_STRING_CHECKED(FMovieGraphPostRenderSettings, PostRenderImageSequencePlayOptions));
+
+	static FProperty* ProResOptions =
+		FMovieGraphPostRenderSettings::StaticStruct()->FindPropertyByName(
+			GET_MEMBER_NAME_STRING_CHECKED(FMovieGraphPostRenderSettings, PostRenderAppleProResPlayOptions));
+
+	static FProperty* AvidOptions =
+		FMovieGraphPostRenderSettings::StaticStruct()->FindPropertyByName(
+			GET_MEMBER_NAME_STRING_CHECKED(FMovieGraphPostRenderSettings, PostRenderAvidDNxHRPlayOptions));
+
+	auto UpdatePlayOptions = [](FMovieGraphPostRenderVideoPlayOptions* InPlayOptions)
+	{
+		if (InPlayOptions->PlaybackMethod == EMovieGraphPlaybackMethod::CustomViewer)
+		{
+			InPlayOptions->JobPlayback = EMovieGraphJobPlaybackRange::AllJobs;
+			InPlayOptions->RenderLayerPlayback = EMovieGraphRenderLayerPlaybackRange::AllRenderLayers;
+		}
+		else
+		{
+			InPlayOptions->JobPlayback = EMovieGraphJobPlaybackRange::FirstJobOnly;
+			InPlayOptions->RenderLayerPlayback = EMovieGraphRenderLayerPlaybackRange::FirstRenderLayerOnly;
+		}
+	};
+
+	// If the Playback Method changes, update some other properties to reflect the most likely use case of this playback method.
+	if ((PropertyChangedEvent.Property != nullptr) && (PropertyChangedEvent.Property->GetFName() == GET_MEMBER_NAME_CHECKED(FMovieGraphPostRenderVideoPlayOptions, PlaybackMethod)))
+	{
+		if (PropertyChangedEvent.PropertyChain.Contains(ImageSequenceOptions))
+		{
+			UpdatePlayOptions(&PostRenderSettings.PostRenderImageSequencePlayOptions);
+
+			// Image sequences additionally have a Playback Range option that needs to be updated
+			const EMovieGraphPlaybackMethod PlaybackMethod = PostRenderSettings.PostRenderImageSequencePlayOptions.PlaybackMethod;
+			PostRenderSettings.PostRenderImageSequencePlayOptions.PlaybackRange = (PlaybackMethod == EMovieGraphPlaybackMethod::CustomViewer)
+				? EMovieGraphImageSequencePlaybackRange::FullRange
+				: EMovieGraphImageSequencePlaybackRange::FirstFrameOnly;
+		}
+		else if (PropertyChangedEvent.PropertyChain.Contains(ProResOptions))
+		{
+			UpdatePlayOptions(&PostRenderSettings.PostRenderAppleProResPlayOptions);
+		}
+		else if (PropertyChangedEvent.PropertyChain.Contains(AvidOptions))
+		{
+			UpdatePlayOptions(&PostRenderSettings.PostRenderAvidDNxHRPlayOptions);
+		}
+	}
+	
+	Super::PostEditChangeChainProperty(PropertyChangedEvent);
+}
+#endif	// WITH_EDITOR
diff --git a/Engine/Plugins/MovieScene/MovieRenderPipeline/Source/MovieRenderPipelineCore/Private/MoviePipelinePostRenderSettings.cpp b/Engine/Plugins/MovieScene/MovieRenderPipeline/Source/MovieRenderPipelineCore/Private/MoviePipelinePostRenderSettings.cpp
--- a/Engine/Plugins/MovieScene/MovieRenderPipeline/Source/MovieRenderPipelineCore/Private/MoviePipelinePostRenderSettings.cpp
+++ b/Engine/Plugins/MovieScene/MovieRenderPipeline/Source/MovieRenderPipelineCore/Private/MoviePipelinePostRenderSettings.cpp
@@ -11,7 +11,7 @@ FMovieGraphPostRenderVideoPlayOptions::FMovieGraphPostRenderVideoPlayOptions()
 }
 
 FMovieGraphPostRenderImageSequencePlayOptions::FMovieGraphPostRenderImageSequencePlayOptions()
-	: FrameRangeNotation(EMovieGraphFrameRangeNotation::Hash)
+	: FrameRangeNotation(EMovieGraphFrameRangeNotation::HashWithStartEndFrame)
 	, PlaybackRange(EMovieGraphImageSequencePlaybackRange::FirstFrameOnly)
 {
 	
diff --git a/Engine/Plugins/MovieScene/MovieRenderPipeline/Source/MovieRenderPipelineCore/Public/Graph/MovieGraphDataTypes.h b/Engine/Plugins/MovieScene/MovieRenderPipeline/Source/MovieRenderPipelineCore/Public/Graph/MovieGraphDataTypes.h
--- a/Engine/Plugins/MovieScene/MovieRenderPipeline/Source/MovieRenderPipelineCore/Public/Graph/MovieGraphDataTypes.h
+++ b/Engine/Plugins/MovieScene/MovieRenderPipeline/Source/MovieRenderPipelineCore/Public/Graph/MovieGraphDataTypes.h
@@ -361,7 +361,7 @@ namespace UE::MovieGraph
 		/** What was the filepath on disk that was generated by the future? */
 		FString FilePath;
 
-		/** The "frame templated" filepath. This groups together a set of output frames by frame number, eg "Seq.Shot.#.exr" if the frame template was "#". */
+		/** The "frame templated" filepath. This groups together a set of output frames by frame number, eg "Seq.Shot.{frame_placeholder}.exr". */
 		FString FrameTemplatedFilePath;
 
 		/** Which render pass is this for? */
diff --git a/Engine/Plugins/MovieScene/MovieRenderPipeline/Source/MovieRenderPipelineCore/Public/Graph/MovieGraphQuickRenderSettings.h b/Engine/Plugins/MovieScene/MovieRenderPipeline/Source/MovieRenderPipelineCore/Public/Graph/MovieGraphQuickRenderSettings.h
--- a/Engine/Plugins/MovieScene/MovieRenderPipeline/Source/MovieRenderPipelineCore/Public/Graph/MovieGraphQuickRenderSettings.h
+++ b/Engine/Plugins/MovieScene/MovieRenderPipeline/Source/MovieRenderPipelineCore/Public/Graph/MovieGraphQuickRenderSettings.h
@@ -26,10 +26,19 @@ enum class EMovieGraphQuickRenderButtonMode : uint8
 UENUM(BlueprintType)
 enum class EMovieGraphQuickRenderMode : uint8
 {
+	/** Renders the level sequence that is currently active in Sequencer. */
 	CurrentSequence,
-	CurrentShotAtPlayhead,
+
+	/** Renders the shot that Sequencer's playhead is currently at. This mode is currently not implemented. */
+	/* CurrentShotAtPlayhead, */
+
+	/** Renders a single frame using the current viewport's camera. */
 	CurrentViewport,
+
+	/** Renders one frame for each camera actor selected in the Outliner. */
 	SelectedCameras,
+
+	/** Renders the level sequence active in Sequencer, but uses the viewport's camera instead (for the entire duration of the render). */
 	UseViewportCameraInSequence
 };
 
@@ -149,11 +158,11 @@ public:
 	EMovieGraphQuickRenderFrameRangeType FrameRangeType = EMovieGraphQuickRenderFrameRangeType::PlaybackRange;
 
 	/** The frame in the level sequence that the render will begin on. */
-	UPROPERTY(BlueprintReadWrite, Category="Quick Render")
+	UPROPERTY(EditAnywhere, BlueprintReadWrite, Category="Quick Render", meta=(EditCondition="FrameRangeType == EMovieGraphQuickRenderFrameRangeType::Custom", EditConditionHides))
 	int32 CustomStartFrame = 0;
 
 	/** The frame in the level sequence that the render will end on. */
-	UPROPERTY(BlueprintReadWrite, Category="Quick Render")
+	UPROPERTY(EditAnywhere, BlueprintReadWrite, Category="Quick Render", meta=(EditCondition="FrameRangeType == EMovieGraphQuickRenderFrameRangeType::Custom", EditConditionHides))
 	int32 CustomEndFrame = 0;
 
 	/** Values that are set on the variables contained within the graph preset. See note on RefreshVariableAssignments(). */
diff --git a/Engine/Plugins/MovieScene/MovieRenderPipeline/Source/MovieRenderPipelineCore/Public/Graph/MovieGraphRenderDataIdentifier.h b/Engine/Plugins/MovieScene/MovieRenderPipeline/Source/MovieRenderPipelineCore/Public/Graph/MovieGraphRenderDataIdentifier.h
--- a/Engine/Plugins/MovieScene/MovieRenderPipeline/Source/MovieRenderPipelineCore/Public/Graph/MovieGraphRenderDataIdentifier.h
+++ b/Engine/Plugins/MovieScene/MovieRenderPipeline/Source/MovieRenderPipelineCore/Public/Graph/MovieGraphRenderDataIdentifier.h
@@ -112,8 +112,8 @@ public:
 	TArray<FString> FilePaths;
 	
 	/**
-	 * The "frame templated" file paths. This groups together a set of output frames by frame number, eg "Seq.Shot.#.exr" if the frame template was
-	 * "#". If the render layer has more than one output type in it (eg, EXR and JPEG), there may be more than one path here, but there should only
+	 * The "frame templated" file paths. This groups together a set of output frames by frame number, eg "Seq.Shot.{frame_placeholder}.exr".
+	 * If the render layer has more than one output type in it (eg, EXR and JPEG), there may be more than one path here, but there should only
 	 * be one per output type.
 	 */
 	UPROPERTY(VisibleAnywhere, BlueprintReadOnly, Category = "Movie Pipeline")
diff --git a/Engine/Plugins/MovieScene/MovieRenderPipeline/Source/MovieRenderPipelineCore/Public/Graph/MovieRenderGraphEditorSettings.h b/Engine/Plugins/MovieScene/MovieRenderPipeline/Source/MovieRenderPipelineCore/Public/Graph/MovieRenderGraphEditorSettings.h
--- a/Engine/Plugins/MovieScene/MovieRenderPipeline/Source/MovieRenderPipelineCore/Public/Graph/MovieRenderGraphEditorSettings.h
+++ b/Engine/Plugins/MovieScene/MovieRenderPipeline/Source/MovieRenderPipelineCore/Public/Graph/MovieRenderGraphEditorSettings.h
@@ -18,6 +18,10 @@ class MOVIERENDERPIPELINECORE_API UMovieRenderGraphEditorSettings : public UDeve
 public:
 	UMovieRenderGraphEditorSettings();
 
+#if WITH_EDITOR
+	virtual void PostEditChangeChainProperty(FPropertyChangedChainEvent& PropertyChangedEvent) override;
+#endif
+
 	/** If PostRenderBehavior is set to PlayRenderOutput, these settings are used to determine how to play media. */
 	UPROPERTY(Config, EditAnywhere, Category="Play Render Output", meta=(ShowOnlyInnerProperties))
 	FMovieGraphPostRenderSettings PostRenderSettings;
diff --git a/Engine/Plugins/MovieScene/MovieRenderPipeline/Source/MovieRenderPipelineCore/Public/MoviePipelinePostRenderSettings.h b/Engine/Plugins/MovieScene/MovieRenderPipeline/Source/MovieRenderPipelineCore/Public/MoviePipelinePostRenderSettings.h
--- a/Engine/Plugins/MovieScene/MovieRenderPipeline/Source/MovieRenderPipelineCore/Public/MoviePipelinePostRenderSettings.h
+++ b/Engine/Plugins/MovieScene/MovieRenderPipeline/Source/MovieRenderPipelineCore/Public/MoviePipelinePostRenderSettings.h
@@ -54,8 +54,28 @@ enum class EMovieGraphRenderLayerPlaybackRange : uint8
 UENUM(BlueprintType)
 enum class EMovieGraphFrameRangeNotation : uint8
 {
-	Hash UMETA(DisplayName = "#"),
+	/**
+	 * Hash symbol. Includes the rendered frame range after the file path(s) are given on the command line.
+	 * For example, "Layer1.#.exr 100-150" for an image sequence that included the frame number in the file name, and a frame range of 100-150 was rendered. 
+	 */
+	HashWithStartEndFrame UMETA(DisplayName = "# with Start/End Frame"),
+
+	/**
+	 * Hash symbol. Does not include the rendered frame range on the command line.
+	 * For example, "Layer1.#.exr" for an image sequence that included the frame number in the file name. 
+	 */
+	Hash UMETA(DisplayName = "# without Start/End Frame"),
+
+	/**
+	 * $F. Includes the rendered frame range before the file paths(s) are given on the command line.
+	 * For example, "-f 100 150 1 Layer1.$F.exr" for an image sequence that included the frame number in the file name, and a frame range of 100-150 was rendered. 
+	 */
 	DollarF UMETA(DisplayName = "$F"),
+
+	/**
+	 * Use the path for the first rendered frame only.
+	 * For example, "Layer1.100.exr", if frame 100 was the first frame rendered. 
+	 */
 	StartFrame
 };
 
diff --git a/Engine/Plugins/MovieScene/MovieRenderPipeline/Source/MovieRenderPipelineEditor/Private/Graph/MovieGraphQuickRender.cpp b/Engine/Plugins/MovieScene/MovieRenderPipeline/Source/MovieRenderPipelineEditor/Private/Graph/MovieGraphQuickRender.cpp
--- a/Engine/Plugins/MovieScene/MovieRenderPipeline/Source/MovieRenderPipelineEditor/Private/Graph/MovieGraphQuickRender.cpp
+++ b/Engine/Plugins/MovieScene/MovieRenderPipeline/Source/MovieRenderPipelineEditor/Private/Graph/MovieGraphQuickRender.cpp
@@ -43,12 +43,6 @@
 
 #define LOCTEXT_NAMESPACE "FMovieGraphQuickRender"
 
-// Quick Render should be hidden in the UI while under development.
-static TAutoConsoleVariable<bool> CVarMoviePipelineEnableQuickRender(
-	TEXT("MoviePipeline.EnableQuickRender"),
-	false,
-	TEXT("Enables or disables whether the Quick Render feature is exposed in the UI (defaults to off)."));
-
 void FMovieGraphQuickRender::BeginQuickRender(const UMovieGraphQuickRenderSettings& InQuickRenderSettings)
 {
 	UMoviePipelineQueueSubsystem* Subsystem = GEditor->GetEditorSubsystem<UMoviePipelineQueueSubsystem>();
@@ -216,27 +210,27 @@ void FMovieGraphQuickRender::ApplyQuickRenderUpdatesToDuplicateGraph(const UMovi
 {
 	check(TemporaryGraph.IsValid());
 
-	// Global Output Settings
+	// Apply changes to Global Output Settings
 	{
 		UMovieGraphGlobalOutputSettingNode* GlobalOutputSettings = Cast<UMovieGraphGlobalOutputSettingNode>(InjectNodeIntoBranch(UMovieGraphGlobalOutputSettingNode::StaticClass(), UMovieGraphNode::GlobalsPinName));
 
-		// If the frame range type is StartEnd, the frame range doesn't need to be modified
-		if ((InQuickRenderSettings.FrameRangeType == EMovieGraphQuickRenderFrameRangeType::SelectionRange) || (InQuickRenderSettings.FrameRangeType == EMovieGraphQuickRenderFrameRangeType::Custom))
+		auto SetStartFrame = [&GlobalOutputSettings](const int32 StartFrame) -> void
 		{
-			const int32 StartFrame = (InQuickRenderSettings.FrameRangeType == EMovieGraphQuickRenderFrameRangeType::SelectionRange)
-				? ULevelSequenceEditorBlueprintLibrary::GetSelectionRangeStart()
-				: InQuickRenderSettings.CustomStartFrame;
-			const int32 EndFrame = (InQuickRenderSettings.FrameRangeType == EMovieGraphQuickRenderFrameRangeType::SelectionRange)
-				? ULevelSequenceEditorBlueprintLibrary::GetSelectionRangeEnd()
-				: InQuickRenderSettings.CustomEndFrame;
-PRAGMA_DISABLE_DEPRECATION_WARNINGS
-			GlobalOutputSettings->bOverride_CustomPlaybackRangeStartFrame = true;
-			GlobalOutputSettings->CustomPlaybackRangeStartFrame = StartFrame;
-
-			GlobalOutputSettings->bOverride_CustomPlaybackRangeEndFrame = true;
-			GlobalOutputSettings->CustomPlaybackRangeEndFrame = EndFrame;
-PRAGMA_ENABLE_DEPRECATION_WARNINGS
-		}
+			GlobalOutputSettings->bOverride_CustomPlaybackRangeStart = true;
+			GlobalOutputSettings->CustomPlaybackRangeStart.Type = EMovieGraphSequenceRangeType::Custom;
+			GlobalOutputSettings->CustomPlaybackRangeStart.Value = StartFrame;
+		};
+
+		auto SetEndFrame = [&GlobalOutputSettings](const int32 EndFrame) -> void
+		{
+			GlobalOutputSettings->bOverride_CustomPlaybackRangeEnd = true;
+			GlobalOutputSettings->CustomPlaybackRangeEnd.Type = EMovieGraphSequenceRangeType::Custom;
+			GlobalOutputSettings->CustomPlaybackRangeEnd.Value = EndFrame;
+		};
+
+		const TRange<FFrameNumber> PlaybackRange = GetPlaybackRange(InQuickRenderSettings);
+		SetStartFrame(PlaybackRange.GetLowerBoundValue().Value);
+		SetEndFrame(PlaybackRange.GetUpperBoundValue().Value);
 	}
 
 	// Inject an Apply Viewport Look node if any viewport look flags were specified. Logic outside of Quick Render will for look this node in order to
@@ -620,6 +614,38 @@ UWorld* FMovieGraphQuickRender::GetWorldOfType(EWorldType::Type WorldType)
 	return World;
 }
 
+TRange<FFrameNumber> FMovieGraphQuickRender::GetPlaybackRange(const UMovieGraphQuickRenderSettings& InQuickRenderSettings)
+{
+	const UMovieScene* MovieScene = RenderingLevelSequence->GetMovieScene();
+
+	const TRange<FFrameNumber> PlaybackRange = MovieScene->GetPlaybackRange();
+	FFrameNumber StartFrame = FFrameRate::TransformTime(PlaybackRange.GetLowerBoundValue(), MovieScene->GetTickResolution(), MovieScene->GetDisplayRate()).GetFrame();
+	FFrameNumber EndFrame = FFrameRate::TransformTime(PlaybackRange.GetUpperBoundValue(), MovieScene->GetTickResolution(), MovieScene->GetDisplayRate()).GetFrame();
+
+	if (InQuickRenderSettings.FrameRangeType == EMovieGraphQuickRenderFrameRangeType::SelectionRange)
+	{
+		const int32 SelectionStartFrame = ULevelSequenceEditorBlueprintLibrary::GetSelectionRangeStart();
+		const int32 SelectionEndFrame = ULevelSequenceEditorBlueprintLibrary::GetSelectionRangeEnd();
+
+		if (SelectionStartFrame != 0)
+		{
+			StartFrame = FFrameNumber(SelectionStartFrame);
+		}
+
+		if (SelectionEndFrame != 0)
+		{
+			EndFrame = FFrameNumber(SelectionEndFrame);
+		}
+	}
+	else if (InQuickRenderSettings.FrameRangeType == EMovieGraphQuickRenderFrameRangeType::Custom)
+	{
+		StartFrame = FFrameNumber(InQuickRenderSettings.CustomStartFrame);
+		EndFrame = FFrameNumber(InQuickRenderSettings.CustomEndFrame);
+	}
+
+	return TRange<FFrameNumber>::Inclusive(StartFrame, EndFrame);
+}
+
 void FMovieGraphQuickRender::PerformPreRenderSetup(const UMovieGraphQuickRenderSettings& InQuickRenderSettings)
 {
 	// Adds the specified camera to the Camera Cut track in the given level sequence.
diff --git a/Engine/Plugins/MovieScene/MovieRenderPipeline/Source/MovieRenderPipelineEditor/Private/MoviePipelinePostRenderFileDisplayProcessor.cpp b/Engine/Plugins/MovieScene/MovieRenderPipeline/Source/MovieRenderPipelineEditor/Private/MoviePipelinePostRenderFileDisplayProcessor.cpp
--- a/Engine/Plugins/MovieScene/MovieRenderPipeline/Source/MovieRenderPipelineEditor/Private/MoviePipelinePostRenderFileDisplayProcessor.cpp
+++ b/Engine/Plugins/MovieScene/MovieRenderPipeline/Source/MovieRenderPipelineEditor/Private/MoviePipelinePostRenderFileDisplayProcessor.cpp
@@ -2,7 +2,10 @@
 
 #include "MoviePipelinePostRenderFileDisplayProcessor.h"
 
+#include "Graph/MovieRenderGraphEditorSettings.h"
+#include "Internationalization/Regex.h"
 #include "MovieGraphImageSequenceOutputNode.h"
+#include "MoviePipelineUtils.h"
 #include "MovieRenderPipelineCoreModule.h"
 
 FMoviePipelinePostRenderFileDisplayProcessor::FMoviePipelinePostRenderFileDisplayProcessor(const FMovieGraphPostRenderSettings& InPostRenderSettings)
@@ -20,6 +23,11 @@ void FMoviePipelinePostRenderFileDisplayProcessor::OpenFiles() const
 {
 	TArray<FString> FinalListOfFilesToOpen;
 	const FMovieGraphPostRenderVideoPlayOptions* FinalPlayOptions = nullptr;
+
+	// The start and end frames to use for opening files. Currently only one set of start and end frames is supported, mostly due to restrictions on
+	// how they can be passed to third-party image viewers.
+	FString StartFrame;
+	FString EndFrame;
 	
 	for (FMoviePipelineOutputData& OutputData : PipelineOutputData)
 	{
@@ -29,6 +37,12 @@ void FMoviePipelinePostRenderFileDisplayProcessor::OpenFiles() const
 		// contain the application executable to use for opening the files.
 		GetFilteredFilesToOpen(FilesToOpen, FinalListOfFilesToOpen, FinalPlayOptions);
 
+		if (StartFrame.IsEmpty() && EndFrame.IsEmpty())
+		{
+			StartFrame = FilesToOpen.StartFrameString;
+			EndFrame = FilesToOpen.EndFrameString;
+		}
+
 		// Settings may dictate that only one job's worth of media should be opened
 		if (PostRenderSettings.PostRenderImageSequencePlayOptions.JobPlayback != EMovieGraphJobPlaybackRange::AllJobs)
 		{
@@ -39,7 +53,7 @@ void FMoviePipelinePostRenderFileDisplayProcessor::OpenFiles() const
 	// Open the application with the filtered list of files
 	if (!FinalListOfFilesToOpen.IsEmpty() && FinalPlayOptions)
 	{
-		LaunchFilesWithSettings(FinalListOfFilesToOpen, FinalPlayOptions);
+		LaunchFilesWithSettings(FinalListOfFilesToOpen, FinalPlayOptions, TTuple<FString, FString>(StartFrame, EndFrame));
 	}
 }
 
@@ -85,9 +99,19 @@ FMoviePipelinePostRenderFileDisplayProcessor::FFilesToOpen FMoviePipelinePostRen
 		{
 			for (const FString& ImagePath : InImagePaths)
 			{
-				if (FrameTemplateToFilePaths.Value.Contains(ImagePath))
+				const TArray<FString>& ImagesAssociatedWithFrameTemplate = FrameTemplateToFilePaths.Value;
+				const FString& FrameTemplatedPath = FrameTemplateToFilePaths.Key;
+				
+				// Get the frame range of this image sequence. This will merge this image sequence's frame range with other image sequences found.
+				// For example, if shot1 has a frame range of 0-100, and shot2 has a frame range of 101-200, and this method is called on both shots,
+				// FFilesToOpen will have a frame range of 0-200.
+				GetStartAndEndFrames(FrameTemplatedPath, InImagePaths, FilesToOpen);
+				
+				// Add the frame-templated path if the templated path is associated with this non-templated path AND the templated path has not already
+				// been added to the list of files to open
+				if (ImagesAssociatedWithFrameTemplate.Contains(ImagePath) && !FilesToOpen.Images.Contains(FrameTemplatedPath))
 				{
-					FilesToOpen.Images.Add(FrameTemplateToFilePaths.Key);
+					FilesToOpen.Images.Add(FrameTemplatedPath);
 
 					// Found an image that matched this frame templated path; move on to the next frame template and see if a path matches it
 					break;
@@ -395,8 +419,67 @@ void FMoviePipelinePostRenderFileDisplayProcessor::GetHighestPriorityExtensions(
 	}
 }
 
-void FMoviePipelinePostRenderFileDisplayProcessor::LaunchFilesWithSettings(const TArray<FString>& InFilesToOpen, const FMovieGraphPostRenderVideoPlayOptions* InPlayOptions) const
+void FMoviePipelinePostRenderFileDisplayProcessor::LaunchFilesWithSettings(const TArray<FString>& InFilesToOpen, const FMovieGraphPostRenderVideoPlayOptions* InPlayOptions, const TTuple<FString, FString>& InFrameRangeToOpen) const
 {
+	const FString StartFrame = InFrameRangeToOpen.Key;
+	const FString EndFrame = InFrameRangeToOpen.Value;
+	
+	auto ReplaceFramePlaceholderInFiles = [&StartFrame, &InFilesToOpen]()
+	{
+		FString FramePlaceholder;
+		
+		const EMovieGraphFrameRangeNotation FrameRangeNotationType = GetDefault<UMovieRenderGraphEditorSettings>()->PostRenderSettings.PostRenderImageSequencePlayOptions.FrameRangeNotation;
+		switch (FrameRangeNotationType)
+		{
+		case EMovieGraphFrameRangeNotation::Hash:
+		case EMovieGraphFrameRangeNotation::HashWithStartEndFrame:
+			FramePlaceholder = TEXT("#");
+			break;
+		case EMovieGraphFrameRangeNotation::DollarF:
+			FramePlaceholder = TEXT("$F");
+			break;
+		case EMovieGraphFrameRangeNotation::StartFrame:
+			{
+				FramePlaceholder = StartFrame;
+				break;
+			}
+		}
+		
+		// Make sure the paths are in the right format for the platform, and have any frame templating applied
+		TArray<FString> TransformedFilePaths = InFilesToOpen;
+		for (FString& FileToOpen : TransformedFilePaths)
+		{
+			FPaths::MakePlatformFilename(FileToOpen);
+
+			if (!FramePlaceholder.IsEmpty())
+			{
+				FileToOpen.ReplaceInline(TEXT("{frame_placeholder}"), *FramePlaceholder);
+			}
+		}
+
+		return TransformedFilePaths;
+	};
+
+	auto InjectFrameRange = [this, &StartFrame, &EndFrame](FString& InConcatenatedFilePaths)
+	{
+		if (!StartFrame.IsEmpty() && !EndFrame.IsEmpty())
+		{
+			const FString FrameRangeString = FString::Format(TEXT("{0}-{1}"), { StartFrame, EndFrame });
+
+			const EMovieGraphFrameRangeNotation FrameRangeNotation = PostRenderSettings.PostRenderImageSequencePlayOptions.FrameRangeNotation;
+			if (FrameRangeNotation == EMovieGraphFrameRangeNotation::HashWithStartEndFrame)
+			{
+				// If using # with frame ranges, specify the frame range after the file path(s). This is how RV expects it.
+				InConcatenatedFilePaths = FString::Format(TEXT("{0} {1}"), { InConcatenatedFilePaths, FrameRangeString });
+			}
+			else if (FrameRangeNotation == EMovieGraphFrameRangeNotation::DollarF)
+			{
+				// $F is used exclusively by MPlay, so provide the frame range via -f with a 1-frame step (eg, "-f 0 150 1")
+				InConcatenatedFilePaths = FString::Format(TEXT("-f {0} 1 {1}"), { FrameRangeString, InConcatenatedFilePaths });
+			}
+		}
+	};
+	
 	if (InPlayOptions->PlaybackMethod == EMovieGraphPlaybackMethod::OperatingSystem)
 	{
 		// Only one frame will be displayed if the OS is being used to open images
@@ -409,18 +492,21 @@ void FMoviePipelinePostRenderFileDisplayProcessor::LaunchFilesWithSettings(const
 		FString ExecutablePathNoQuotes = ExecutablePath.Replace(TEXT("\""), TEXT(""));
 		FPaths::NormalizeFilename(ExecutablePathNoQuotes);
 		
-		// Make sure the paths are in the right format for the platform
-		TArray<FString> TransformedFilePaths = InFilesToOpen;
-		for (FString& FileToOpen : TransformedFilePaths)
-		{
-			FPaths::MakePlatformFilename(FileToOpen);
-		}
+		// Replace {frame_placeholder}, if present, with the appropriate symbol (like '#')
+		const TArray<FString> TransformedFilePaths = ReplaceFramePlaceholderInFiles();
+
+		// Depending on the frame range notation chosen, the command line arguments may need to be augmented with the frame range that was rendered
+		FString ConcatenatedFilePaths = FString::Join(TransformedFilePaths, TEXT(" "));
+		InjectFrameRange(ConcatenatedFilePaths);
 		
 		const FString CommandLineArguments = FString::Format(TEXT("{0} {1}"), {
 			InPlayOptions->AdditionalCommandLineArguments,
-			FString::Join(TransformedFilePaths, TEXT(" "))
+			ConcatenatedFilePaths
 		});
 		
+		const FString FinalCommandString = FString::Format(TEXT("{0} {1}"), {ExecutablePathNoQuotes, CommandLineArguments});
+		UE_LOG(LogMovieRenderPipeline, Display, TEXT("Quick Render: Opening external viewer with command: %s"), *FinalCommandString);
+		
 		// Open the images in the selected application
 		constexpr bool bLaunchDetached = true;
 		constexpr bool bLaunchHidden = false;
@@ -428,7 +514,6 @@ void FMoviePipelinePostRenderFileDisplayProcessor::LaunchFilesWithSettings(const
 		const FProcHandle ProcHandle = FPlatformProcess::CreateProc(*ExecutablePathNoQuotes, *CommandLineArguments, bLaunchDetached, bLaunchHidden, bLaunchReallyHidden, nullptr, 0, nullptr, nullptr, nullptr);
 		if (!ProcHandle.IsValid())
 		{
-			const FString FinalCommandString = FString::Format(TEXT("{0} {1}"), {ExecutablePathNoQuotes, CommandLineArguments});
 			UE_LOG(LogMovieRenderPipeline, Warning, TEXT("Unable to open the post-render media viewer. Ensure the executable and its command line arguments have been specified correctly in Editor Preferences. The command that was run: [%s]"), *FinalCommandString);
 		}
 	}
@@ -441,4 +526,49 @@ int32 FMoviePipelinePostRenderFileDisplayProcessor::GetExtensionPriorityIndex(co
 	// Extensions that have a lower index in the array have the highest priority, thus if the extension was not found (-1), it should be
 	// assigned a very high priority number so it is not considered a high priority (10000 is arbitrary).
 	return (ExtensionPriority == INDEX_NONE) ? 10000 : ExtensionPriority;
-}
\ No newline at end of file
+}
+
+void FMoviePipelinePostRenderFileDisplayProcessor::GetStartAndEndFrames(const FString& InTemplatedPath, const TArray<FString>& InImagePaths, FFilesToOpen& InOutFilesToOpen) const
+{
+	if (InImagePaths.IsEmpty())
+	{
+		return;
+	}
+
+	// Turn the frame template path into a regex pattern. Eg, C:\SomeFolder\Shot1\Layer1.{frame_placeholder}.exr -> \EC:\SomeFolder\Shot1\Layer1.\Q(\d+)\E.exr\Q
+	// Putting the path parts, excluding the frame number (which has been placed in a capture group), inside an \E\Q group allows us to avoid the need
+	// to escape different characters within the path.
+	// First, surround the entire path with \E\Q, then replace {frame_placeholder} with \Q(\d+)\E to capture the frame number. 
+	const FString ReplacementToken = TEXT(R"(\E(\d+)\Q)");
+	FString RegexString = FString::Format(TEXT(R"(\Q{0}\E)"), {InTemplatedPath});
+	RegexString.ReplaceInline(TEXT("{frame_placeholder}"), *ReplacementToken);
+	const FRegexPattern RegexPattern(RegexString);
+
+	// If this is the first time this method is called, init the Start/End frame numbers with bogus numbers so the loop below knows to immediately
+	// initialize them with frame numbers found within the image paths.
+	if ((InOutFilesToOpen.StartFrame == 0) && (InOutFilesToOpen.EndFrame == 0))
+	{
+		InOutFilesToOpen.StartFrame = 100000;
+		InOutFilesToOpen.EndFrame = -100000;
+	}
+
+	for (const FString& ImagePathTest : InImagePaths)
+	{
+		FRegexMatcher FrameMatcher(RegexPattern, ImagePathTest);
+		if (FrameMatcher.FindNext())
+		{
+			const FString FrameNumberString = FrameMatcher.GetCaptureGroup(1);
+			const int32 FrameNumber = FCString::Atoi(*FrameNumberString);
+			if (FrameNumber < InOutFilesToOpen.StartFrame)
+			{
+				InOutFilesToOpen.StartFrame = FrameNumber;
+				InOutFilesToOpen.StartFrameString = FrameNumberString;
+			}
+			if (FrameNumber > InOutFilesToOpen.EndFrame)
+			{
+				InOutFilesToOpen.EndFrame = FrameNumber;
+				InOutFilesToOpen.EndFrameString = FrameNumberString;
+			}
+		}
+	}
+}
diff --git a/Engine/Plugins/MovieScene/MovieRenderPipeline/Source/MovieRenderPipelineEditor/Private/MoviePipelinePostRenderFileDisplayProcessor.h b/Engine/Plugins/MovieScene/MovieRenderPipeline/Source/MovieRenderPipelineEditor/Private/MoviePipelinePostRenderFileDisplayProcessor.h
--- a/Engine/Plugins/MovieScene/MovieRenderPipeline/Source/MovieRenderPipelineEditor/Private/MoviePipelinePostRenderFileDisplayProcessor.h
+++ b/Engine/Plugins/MovieScene/MovieRenderPipeline/Source/MovieRenderPipelineEditor/Private/MoviePipelinePostRenderFileDisplayProcessor.h
@@ -30,6 +30,16 @@ private:
 		TArray<FString> AppleProResMovies;
 		TArray<FString> AvidDNxHRMovies;
 
+		// The start and end frame of the first image sequence contained in Images. Representing these as strings is a little strange. It's
+		// to preserve the padding found in the source filename so it doesn't have to be re-derived later in the pipeline (where determining what that
+		// padding is may be tricky). For example, instead of storing an int-based start frame of 10, we can store it as "0010", which is how it
+		// appears in the filename.
+		FString StartFrameString;
+		FString EndFrameString;
+
+		int32 StartFrame = 0;
+		int32 EndFrame = 0;
+
 		FString HighestPriorityImageExtension;
 		FString HighestPriorityVideoExtension;
 	};
@@ -53,8 +63,8 @@ private:
 		FString& OutHighestPriorityImageExtension,
 		FString& OutHighestPriorityVideoExtension) const;
 
-	/** Opens the given files according to the playback settings provided. */
-	void LaunchFilesWithSettings(const TArray<FString>& InFilesToOpen, const FMovieGraphPostRenderVideoPlayOptions* InPlayOptions) const;
+	/** Opens the given files according to the playback settings provided. The frame range of files to open must be specified, but can be empty strings (to open all frames). */
+	void LaunchFilesWithSettings(const TArray<FString>& InFilesToOpen, const FMovieGraphPostRenderVideoPlayOptions* InPlayOptions, const TTuple<FString, FString>& InFrameRangeToOpen) const;
 
 	/**
 	 * Gets the extension's priority according to the Output Type Priority Order list, or 10000 if the extension was not found. A lower index indicates
@@ -62,6 +72,9 @@ private:
 	 */
 	int32 GetExtensionPriorityIndex(const FString& InExtension) const;
 
+	/** Gets the start and end frame of the images provided, based on the frame-templated path. The string and int-based frame numbers in InOutFilesToOpen are updated. */
+	void GetStartAndEndFrames(const FString& InTemplatedPath, const TArray<FString>& InImagePaths, FFilesToOpen& InOutFilesToOpen) const;
+
 private:
 	/** The post-render settings the class was initialized with. Usually sourced from the editor preferences. */
 	FMovieGraphPostRenderSettings PostRenderSettings;
diff --git a/Engine/Plugins/MovieScene/MovieRenderPipeline/Source/MovieRenderPipelineEditor/Private/MoviePipelineQuickRenderMenu.cpp b/Engine/Plugins/MovieScene/MovieRenderPipeline/Source/MovieRenderPipelineEditor/Private/MoviePipelineQuickRenderMenu.cpp
--- a/Engine/Plugins/MovieScene/MovieRenderPipeline/Source/MovieRenderPipelineEditor/Private/MoviePipelineQuickRenderMenu.cpp
+++ b/Engine/Plugins/MovieScene/MovieRenderPipeline/Source/MovieRenderPipelineEditor/Private/MoviePipelineQuickRenderMenu.cpp
@@ -14,6 +14,7 @@
 #include "MovieRenderPipelineEditorModule.h"
 #include "MovieRenderPipelineEditorUtils.h"
 #include "MovieRenderPipelineSettings.h"
+#include "MovieRenderPipelineStyle.h"
 #include "PropertyEditorModule.h"
 #include "SActionButton.h"
 #include "Styling/ToolBarStyle.h"
@@ -27,16 +28,6 @@ class UMovieGraphConfig;
 
 #define LOCTEXT_NAMESPACE "FMoviePipelineQuickRenderMenu"
 
-bool FMoviePipelineQuickRenderMenu::IsVisible()
-{
-	if (const IConsoleVariable* EnableQuickRenderCVar = IConsoleManager::Get().FindConsoleVariable(TEXT("MoviePipeline.EnableQuickRender")))
-	{
-		return EnableQuickRenderCVar->GetBool();
-	}
-
-	return false;
-}
-
 void FMoviePipelineQuickRenderMenu::AddQuickRenderButtonToToolMenu(UToolMenu* InMenu)
 {
 	if (!InMenu)
@@ -59,7 +50,7 @@ void FMoviePipelineQuickRenderMenu::AddQuickRenderButtonToToolMenu(UToolMenu* In
 					return LOCTEXT("BeginQuickRender", "Begin Quick Render\n\nQuick Render performs a render without having to manually configure a queue and a graph. A typical quick render uses the current map and level sequence, the level sequence's playback range, and the viewport's look, to generate frames. Several different types of quick renders are available -- see the options drop-down.");
 				})
 			),
-			FSlateIcon(FAppStyle::Get().GetStyleSetName(), "LevelEditor.OpenCinematic")
+			FSlateIcon(FMovieRenderPipelineStyle::Get().GetStyleSetName(), "MovieRenderPipeline.QuickRender.Icon.Toolbar")
 		);
 	
 	Section.AddEntry(QuickRenderLaunchButton);
@@ -337,12 +328,13 @@ void FMoviePipelineQuickRenderMenu::GenerateQuickRenderMenuSection(UToolMenu* In
 		FSlateIcon("LevelSequenceEditorStyle", "LevelSequenceEditor.CreateNewLevelSequenceInLevel"),
 		EMovieGraphQuickRenderMode::CurrentSequence);
 
-	AddQuickRenderModeAction(
+	// The "Current Shot at Playhead" mode is disabled for now, may be added in the future.
+	/*AddQuickRenderModeAction(
 		FName("QuickRender_CurrentShotAtPlayhead"),
 		LOCTEXT("QuickRender_CurrentShotAtPlayheadLabel", "Current Shot at Playhead"),
 		LOCTEXT("QuickRender_CurrentShotAtPlayheadToolTip", "Render the current shot starting at the playhead."),
 		FSlateIcon("LevelSequenceEditorStyle", "LevelSequenceEditor.CreateNewLevelSequenceWithShotsInLevel"),
-		EMovieGraphQuickRenderMode::CurrentShotAtPlayhead);
+		EMovieGraphQuickRenderMode::CurrentShotAtPlayhead);*/
 
 	AddQuickRenderModeAction(
 		FName("QuickRender_CurrentViewport"),
diff --git a/Engine/Plugins/MovieScene/MovieRenderPipeline/Source/MovieRenderPipelineEditor/Private/MoviePipelineQuickRenderMenu.h b/Engine/Plugins/MovieScene/MovieRenderPipeline/Source/MovieRenderPipelineEditor/Private/MoviePipelineQuickRenderMenu.h
--- a/Engine/Plugins/MovieScene/MovieRenderPipeline/Source/MovieRenderPipelineEditor/Private/MoviePipelineQuickRenderMenu.h
+++ b/Engine/Plugins/MovieScene/MovieRenderPipeline/Source/MovieRenderPipelineEditor/Private/MoviePipelineQuickRenderMenu.h
@@ -16,9 +16,6 @@ class FMoviePipelineQuickRenderMenu
 public:
 	FMoviePipelineQuickRenderMenu() = default;
 
-	/** Determines if the Quick Render button should be visible. */
-	static bool IsVisible();
-
 	/** Adds the Quick Render button to the specified tool menu. */
 	static void AddQuickRenderButtonToToolMenu(UToolMenu* InMenu);
 
diff --git a/Engine/Plugins/MovieScene/MovieRenderPipeline/Source/MovieRenderPipelineEditor/Private/MovieRenderPipelineEditorModule.cpp b/Engine/Plugins/MovieScene/MovieRenderPipeline/Source/MovieRenderPipelineEditor/Private/MovieRenderPipelineEditorModule.cpp
--- a/Engine/Plugins/MovieScene/MovieRenderPipeline/Source/MovieRenderPipelineEditor/Private/MovieRenderPipelineEditorModule.cpp
+++ b/Engine/Plugins/MovieScene/MovieRenderPipeline/Source/MovieRenderPipelineEditor/Private/MovieRenderPipelineEditorModule.cpp
@@ -241,10 +241,7 @@ void FMovieRenderPipelineEditorModule::RegisterToolbarItem()
 	
 	UnregisterToolbarItem();
 	
-	if (FMoviePipelineQuickRenderMenu::IsVisible())
-	{
-		FMoviePipelineQuickRenderMenu::AddQuickRenderButtonToToolMenu(Menu);
-	}
+	FMoviePipelineQuickRenderMenu::AddQuickRenderButtonToToolMenu(Menu);
 }
 
 void FMovieRenderPipelineEditorModule::UnregisterToolbarItem()
diff --git a/Engine/Plugins/MovieScene/MovieRenderPipeline/Source/MovieRenderPipelineEditor/Private/MovieRenderPipelineStyle.cpp b/Engine/Plugins/MovieScene/MovieRenderPipeline/Source/MovieRenderPipelineEditor/Private/MovieRenderPipelineStyle.cpp
--- a/Engine/Plugins/MovieScene/MovieRenderPipeline/Source/MovieRenderPipelineEditor/Private/MovieRenderPipelineStyle.cpp
+++ b/Engine/Plugins/MovieScene/MovieRenderPipeline/Source/MovieRenderPipelineEditor/Private/MovieRenderPipelineStyle.cpp
@@ -106,6 +106,9 @@ FMovieRenderPipelineStyle::FMovieRenderPipelineStyle()
 	Set("MovieRenderGraph.CollectionsTree.SmallComboButton", FComboButtonStyle(FAppStyle::Get().GetWidgetStyle<FComboButtonStyle>("SimpleComboButton"))
 		.SetDownArrowPadding(FMargin(0)));
 
+	// Quick Render icons
+	Set("MovieRenderPipeline.QuickRender.Icon.Toolbar", new IMAGE_BRUSH_SVG("QuickRender/QuickRenderToolbar_16", Icon16x16));
+
 	FSlateStyleRegistry::RegisterSlateStyle(*this);
 }
 
diff --git a/Engine/Plugins/MovieScene/MovieRenderPipeline/Source/MovieRenderPipelineEditor/Public/Graph/MovieGraphQuickRender.h b/Engine/Plugins/MovieScene/MovieRenderPipeline/Source/MovieRenderPipelineEditor/Public/Graph/MovieGraphQuickRender.h
--- a/Engine/Plugins/MovieScene/MovieRenderPipeline/Source/MovieRenderPipelineEditor/Public/Graph/MovieGraphQuickRender.h
+++ b/Engine/Plugins/MovieScene/MovieRenderPipeline/Source/MovieRenderPipelineEditor/Public/Graph/MovieGraphQuickRender.h
@@ -76,6 +76,9 @@ private:
 	/** Gets the first world found of the specified type. */
 	static UWorld* GetWorldOfType(EWorldType::Type WorldType);
 
+	/** Gets the expected playback range of the sequence being rendered. Will vary depending on the frame range mode that Quick Render is currently in. */
+	static TRange<FFrameNumber> GetPlaybackRange(const UMovieGraphQuickRenderSettings& InQuickRenderSettings);
+
 	/** Does any setup needed in the PIE world prior to the render starting. */
 	static void PerformPreRenderSetup(const UMovieGraphQuickRenderSettings& InQuickRenderSettings);
 
diff --git a/Engine/Plugins/MovieScene/MovieRenderPipeline/Source/MovieRenderPipelineRenderPasses/Private/MovieGraphImageSequenceOutputNode.cpp b/Engine/Plugins/MovieScene/MovieRenderPipeline/Source/MovieRenderPipelineRenderPasses/Private/MovieGraphImageSequenceOutputNode.cpp
--- a/Engine/Plugins/MovieScene/MovieRenderPipeline/Source/MovieRenderPipelineRenderPasses/Private/MovieGraphImageSequenceOutputNode.cpp
+++ b/Engine/Plugins/MovieScene/MovieRenderPipeline/Source/MovieRenderPipelineRenderPasses/Private/MovieGraphImageSequenceOutputNode.cpp
@@ -89,43 +89,35 @@ FString UMovieGraphImageSequenceOutputNode::CreateFileName(
 	// Take our string path from the Output Setting and resolve it.
 	const FString ResolvedFileName = UMovieGraphBlueprintLibrary::ResolveFilenameFormatArguments(FileNameFormatString, Params, OutMergedFormatArgs);
 
-	OutFrameTemplatedFileName = GetFrameTemplatedFileName(InPipeline, Params, FileNameFormatString, OutMergedFormatArgs);
+	OutFrameTemplatedFileName = GetFrameTemplatedFileName(Params, FileNameFormatString, OutMergedFormatArgs);
 
 	return ResolvedFileName;
 }
 
 FString UMovieGraphImageSequenceOutputNode::GetFrameTemplatedFileName(
-	const UMovieGraphPipeline* InPipeline, const FMovieGraphFilenameResolveParams& InParams, const FString& InFileNameFormatString, FMovieGraphResolveArgs& OutMergedFormatArgs) const
+	const FMovieGraphFilenameResolveParams& InParams, const FString& InFileNameFormatString, FMovieGraphResolveArgs& OutMergedFormatArgs) const
 {
-	const EMovieGraphFrameRangeNotation FrameRangeNotationType = GetDefault<UMovieRenderGraphEditorSettings>()->PostRenderSettings.PostRenderImageSequencePlayOptions.FrameRangeNotation;
-	FString FramePlaceholder;
+	const FString FramePlaceholder = TEXT("{frame_placeholder}");
 
-	switch (FrameRangeNotationType)
+	FString FrameTemplatedFormatString = InFileNameFormatString;
+	if (FrameTemplatedFormatString.Contains(TEXT("{frame_number}")))
 	{
-	case EMovieGraphFrameRangeNotation::Hash:
-		FramePlaceholder = TEXT("#");
-		break;
-	case EMovieGraphFrameRangeNotation::DollarF:
-		FramePlaceholder = TEXT("$F");
-		break;
-	case EMovieGraphFrameRangeNotation::StartFrame:
-		{
-			const int32 RootStartFrameNumber = InPipeline->GetActiveShotList()[InPipeline->GetCurrentShotIndex()]->ShotInfo.TotalOutputRangeRoot.GetLowerBoundValue().Value;
-			FramePlaceholder = UE::MoviePipeline::GetPaddingFormatString(InParams.ZeroPadFrameNumberCount, RootStartFrameNumber);
-			break;
-		}
-	default:
-		UE_LOG(LogMovieRenderPipeline, Warning, TEXT("An invalid value was specified for the Frame Range Notation property. This property will be ignored."));
-		break;
+		FrameTemplatedFormatString = InFileNameFormatString.Replace(TEXT("{frame_number}"), *FramePlaceholder);
 	}
-
-	if (!FramePlaceholder.IsEmpty())
+	else if (FrameTemplatedFormatString.Contains(TEXT("{frame_number_rel}")))
 	{
-		const FString FrameTemplatedFormatString = InFileNameFormatString.Replace(TEXT("{frame_number}"), *FramePlaceholder);
-		return UMovieGraphBlueprintLibrary::ResolveFilenameFormatArguments(FrameTemplatedFormatString, InParams, OutMergedFormatArgs);
+		FrameTemplatedFormatString = InFileNameFormatString.Replace(TEXT("{frame_number_rel}"), *FramePlaceholder);
 	}
-
-	return FString();
+	else if (FrameTemplatedFormatString.Contains(TEXT("{frame_number_shot}")))
+	{
+		FrameTemplatedFormatString = InFileNameFormatString.Replace(TEXT("{frame_number_shot}"), *FramePlaceholder);
+	}
+	else if (FrameTemplatedFormatString.Contains(TEXT("{frame_number_shot_rel}")))
+	{
+		FrameTemplatedFormatString = InFileNameFormatString.Replace(TEXT("{frame_number_shot_rel}"), *FramePlaceholder);
+	}
+	
+	return UMovieGraphBlueprintLibrary::ResolveFilenameFormatArguments(FrameTemplatedFormatString, InParams, OutMergedFormatArgs);
 }
 
 void UMovieGraphImageSequenceOutputNode::OnReceiveImageDataImpl(UMovieGraphPipeline* InPipeline, UE::MovieGraph::FMovieGraphOutputMergerFrame* InRawFrameData, const TSet<FMovieGraphRenderDataIdentifier>& InMask)
@@ -802,7 +794,7 @@ FString UMovieGraphImageSequenceOutputNode_MultiLayerEXR::ResolveOutputFilename(
 
 	FString FinalFilePath = UMovieGraphBlueprintLibrary::ResolveFilenameFormatArguments(FilePathFormatString, Params, OutResolveArgs);
 
-	OutFrameTemplatedFilename = GetFrameTemplatedFileName(InPipeline, Params, FileNameFormatString, OutResolveArgs);
+	OutFrameTemplatedFilename = GetFrameTemplatedFileName(Params, FileNameFormatString, OutResolveArgs);
 
 	if (FPaths::IsRelative(FinalFilePath))
 	{
diff --git a/Engine/Plugins/MovieScene/MovieRenderPipeline/Source/MovieRenderPipelineRenderPasses/Public/MovieGraphImageSequenceOutputNode.h b/Engine/Plugins/MovieScene/MovieRenderPipeline/Source/MovieRenderPipelineRenderPasses/Public/MovieGraphImageSequenceOutputNode.h
--- a/Engine/Plugins/MovieScene/MovieRenderPipeline/Source/MovieRenderPipelineRenderPasses/Public/MovieGraphImageSequenceOutputNode.h
+++ b/Engine/Plugins/MovieScene/MovieRenderPipeline/Source/MovieRenderPipelineRenderPasses/Public/MovieGraphImageSequenceOutputNode.h
@@ -69,9 +69,8 @@ protected:
 		FMovieGraphResolveArgs& OutMergedFormatArgs,
 		FString& OutFrameTemplatedFileName) const;
 
-	/** Gets a "frame-number templated" filename (eg, Seq.Shot.#.exr) where '#' is used in place of the frame number. */
+	/** Gets a "frame-number templated" filename (eg, Seq.Shot.{frame_placeholder}.exr) where '{frame_placeholder}' is used in place of the frame number. */
 	FString GetFrameTemplatedFileName(
-		const UMovieGraphPipeline* InPipeline,
 		const FMovieGraphFilenameResolveParams& InParams,
 		const FString& InFileNameFormatString,
 		FMovieGraphResolveArgs& OutMergedFormatArgs) const;
