MoviePipeline: More polish for Quick Render.
* The Quick Render toolbar button will now open the Movie Render Queue editor if the "Movie Render Queue" mode is active, and the MRQ editor is not open.
* Fixed a regression that caused the Viewport and Selected Cameras mode to not render frames.
* Fixed a bug that caused the Selected Cameras mode to exclude frame 0 when playing in an external player.
* Fixed a crash that would occur if the Quick Render menu was opened with a settings asset that was saved in a version of UE greater than the current version (ie, you're on an older CL). The settings cannot be loaded in this scenario, so a new settings asset is created.

diff --git a/Engine/Plugins/MovieScene/MovieRenderPipeline/Source/MovieRenderPipelineCore/Private/Graph/MovieGraphQuickRenderSettings.cpp b/Engine/Plugins/MovieScene/MovieRenderPipeline/Source/MovieRenderPipelineCore/Private/Graph/MovieGraphQuickRenderSettings.cpp
--- a/Engine/Plugins/MovieScene/MovieRenderPipeline/Source/MovieRenderPipelineCore/Private/Graph/MovieGraphQuickRenderSettings.cpp
+++ b/Engine/Plugins/MovieScene/MovieRenderPipeline/Source/MovieRenderPipelineCore/Private/Graph/MovieGraphQuickRenderSettings.cpp
@@ -26,7 +26,10 @@ UMovieGraphQuickRenderSettings* UMovieGraphQuickRenderSettings::GetSavedQuickRen
 	{
 		QuickRenderSettings = Cast<UMovieGraphQuickRenderSettings>(FindObjectWithOuter(QuickRenderSettingsPackage, StaticClass()));
 	}
-	else
+
+	// The settings asset may not exist on disk yet, or (rarely) the settings object may be from a future version (eg, if the user went back to a
+	// previous build). Both of these situations require a new settings asset to be created.
+	if (!QuickRenderSettings)
 	{
 		UPackage* NewSettingsPackage = CreatePackage(*QuickRenderSettingsPackagePath);
 		QuickRenderSettings = NewObject<UMovieGraphQuickRenderSettings>(NewSettingsPackage);
diff --git a/Engine/Plugins/MovieScene/MovieRenderPipeline/Source/MovieRenderPipelineEditor/Private/Graph/MovieGraphQuickRender.cpp b/Engine/Plugins/MovieScene/MovieRenderPipeline/Source/MovieRenderPipelineEditor/Private/Graph/MovieGraphQuickRender.cpp
--- a/Engine/Plugins/MovieScene/MovieRenderPipeline/Source/MovieRenderPipelineEditor/Private/Graph/MovieGraphQuickRender.cpp
+++ b/Engine/Plugins/MovieScene/MovieRenderPipeline/Source/MovieRenderPipelineEditor/Private/Graph/MovieGraphQuickRender.cpp
@@ -229,8 +229,11 @@ void FMovieGraphQuickRender::ApplyQuickRenderUpdatesToDuplicateGraph(const UMovi
 		};
 
 		const TRange<FFrameNumber> PlaybackRange = GetPlaybackRange(InQuickRenderSettings);
-		SetStartFrame(PlaybackRange.GetLowerBoundValue().Value);
-		SetEndFrame(PlaybackRange.GetUpperBoundValue().Value);
+		if (!PlaybackRange.IsEmpty())
+		{
+			SetStartFrame(PlaybackRange.GetLowerBoundValue().Value);
+			SetEndFrame(PlaybackRange.GetUpperBoundValue().Value);
+		}
 	}
 
 	// Inject an Apply Viewport Look node if any viewport look flags were specified. Logic outside of Quick Render will for look this node in order to
@@ -616,11 +619,8 @@ UWorld* FMovieGraphQuickRender::GetWorldOfType(EWorldType::Type WorldType)
 
 TRange<FFrameNumber> FMovieGraphQuickRender::GetPlaybackRange(const UMovieGraphQuickRenderSettings& InQuickRenderSettings)
 {
-	const UMovieScene* MovieScene = RenderingLevelSequence->GetMovieScene();
-
-	const TRange<FFrameNumber> PlaybackRange = MovieScene->GetPlaybackRange();
-	FFrameNumber StartFrame = FFrameRate::TransformTime(PlaybackRange.GetLowerBoundValue(), MovieScene->GetTickResolution(), MovieScene->GetDisplayRate()).GetFrame();
-	FFrameNumber EndFrame = FFrameRate::TransformTime(PlaybackRange.GetUpperBoundValue(), MovieScene->GetTickResolution(), MovieScene->GetDisplayRate()).GetFrame();
+	FFrameNumber StartFrame;
+	FFrameNumber EndFrame;
 
 	if (InQuickRenderSettings.FrameRangeType == EMovieGraphQuickRenderFrameRangeType::SelectionRange)
 	{
@@ -642,6 +642,12 @@ TRange<FFrameNumber> FMovieGraphQuickRender::GetPlaybackRange(const UMovieGraphQ
 		StartFrame = FFrameNumber(InQuickRenderSettings.CustomStartFrame);
 		EndFrame = FFrameNumber(InQuickRenderSettings.CustomEndFrame);
 	}
+	else
+	{
+		// The frame range has not been overridden, so provide an empty range to indicate that the level sequence's normal
+		// playback range should be used.
+		return TRange<FFrameNumber>::Empty();
+	}
 
 	return TRange<FFrameNumber>::Inclusive(StartFrame, EndFrame);
 }
diff --git a/Engine/Plugins/MovieScene/MovieRenderPipeline/Source/MovieRenderPipelineEditor/Private/MoviePipelinePostRenderFileDisplayProcessor.cpp b/Engine/Plugins/MovieScene/MovieRenderPipeline/Source/MovieRenderPipelineEditor/Private/MoviePipelinePostRenderFileDisplayProcessor.cpp
--- a/Engine/Plugins/MovieScene/MovieRenderPipeline/Source/MovieRenderPipelineEditor/Private/MoviePipelinePostRenderFileDisplayProcessor.cpp
+++ b/Engine/Plugins/MovieScene/MovieRenderPipeline/Source/MovieRenderPipelineEditor/Private/MoviePipelinePostRenderFileDisplayProcessor.cpp
@@ -546,7 +546,7 @@ void FMoviePipelinePostRenderFileDisplayProcessor::GetStartAndEndFrames(const FS
 
 	// If this is the first time this method is called, init the Start/End frame numbers with bogus numbers so the loop below knows to immediately
 	// initialize them with frame numbers found within the image paths.
-	if ((InOutFilesToOpen.StartFrame == 0) && (InOutFilesToOpen.EndFrame == 0))
+	if ((InOutFilesToOpen.StartFrame == TNumericLimits<int32>::Min()) && (InOutFilesToOpen.EndFrame == TNumericLimits<int32>::Min()))
 	{
 		InOutFilesToOpen.StartFrame = 100000;
 		InOutFilesToOpen.EndFrame = -100000;
diff --git a/Engine/Plugins/MovieScene/MovieRenderPipeline/Source/MovieRenderPipelineEditor/Private/MoviePipelinePostRenderFileDisplayProcessor.h b/Engine/Plugins/MovieScene/MovieRenderPipeline/Source/MovieRenderPipelineEditor/Private/MoviePipelinePostRenderFileDisplayProcessor.h
--- a/Engine/Plugins/MovieScene/MovieRenderPipeline/Source/MovieRenderPipelineEditor/Private/MoviePipelinePostRenderFileDisplayProcessor.h
+++ b/Engine/Plugins/MovieScene/MovieRenderPipeline/Source/MovieRenderPipelineEditor/Private/MoviePipelinePostRenderFileDisplayProcessor.h
@@ -37,8 +37,8 @@ private:
 		FString StartFrameString;
 		FString EndFrameString;
 
-		int32 StartFrame = 0;
-		int32 EndFrame = 0;
+		int32 StartFrame = TNumericLimits<int32>::Min();
+		int32 EndFrame = TNumericLimits<int32>::Min();
 
 		FString HighestPriorityImageExtension;
 		FString HighestPriorityVideoExtension;
diff --git a/Engine/Plugins/MovieScene/MovieRenderPipeline/Source/MovieRenderPipelineEditor/Private/MoviePipelineQuickRenderMenu.cpp b/Engine/Plugins/MovieScene/MovieRenderPipeline/Source/MovieRenderPipelineEditor/Private/MoviePipelineQuickRenderMenu.cpp
--- a/Engine/Plugins/MovieScene/MovieRenderPipeline/Source/MovieRenderPipelineEditor/Private/MoviePipelineQuickRenderMenu.cpp
+++ b/Engine/Plugins/MovieScene/MovieRenderPipeline/Source/MovieRenderPipelineEditor/Private/MoviePipelineQuickRenderMenu.cpp
@@ -7,8 +7,7 @@
 #include "Framework/Application/SlateApplication.h"
 #include "Framework/Docking/TabManager.h"
 #include "Graph/MovieGraphQuickRender.h"
-#include "Graph/MovieRenderGraphEditorSettings.h"
-#include "HAL/IConsoleManager.h"
+#include "Misc/MessageDialog.h"
 #include "MoviePipelineUtils.h"
 #include "MovieRenderPipelineCoreModule.h"
 #include "MovieRenderPipelineEditorModule.h"
@@ -104,9 +103,19 @@ void FMoviePipelineQuickRenderMenu::QuickRenderButtonPressed()
 	}
 	else
 	{
-		if (!UE::MovieRenderPipelineEditor::Private::PerformLocalRender())
+		// If there's an existing Movie Render Queue tab open, try to do a render with what's in the queue. Otherwise, just open the MRQ editor tab.
+		if (FGlobalTabmanager::Get()->FindExistingLiveTab(IMovieRenderPipelineEditorModule::MoviePipelineQueueTabName))
 		{
-			UE_LOG(LogMovieRenderPipeline, Warning, TEXT("Unable to start local render. Make sure an executor was specified in Project Settings, a job is present in the queue, and a render is not currently running."));
+			if (!UE::MovieRenderPipelineEditor::Private::PerformLocalRender())
+			{
+				FMessageDialog::Open(
+					EAppMsgType::Ok,
+					LOCTEXT("UnableToStartLocalRender", "Unable to start local render. Make sure a job is present in the queue, a render is not currently running, and an executor is specified in Project Settings."));
+			}
+		}
+		else
+		{
+			FGlobalTabmanager::Get()->TryInvokeTab(IMovieRenderPipelineEditorModule::MoviePipelineQueueTabName);
 		}
 	}
 }
diff --git a/Engine/Plugins/MovieScene/MovieRenderPipeline/Source/MovieRenderPipelineEditor/Public/Graph/MovieGraphQuickRender.h b/Engine/Plugins/MovieScene/MovieRenderPipeline/Source/MovieRenderPipelineEditor/Public/Graph/MovieGraphQuickRender.h
--- a/Engine/Plugins/MovieScene/MovieRenderPipeline/Source/MovieRenderPipelineEditor/Public/Graph/MovieGraphQuickRender.h
+++ b/Engine/Plugins/MovieScene/MovieRenderPipeline/Source/MovieRenderPipelineEditor/Public/Graph/MovieGraphQuickRender.h
@@ -76,7 +76,10 @@ private:
 	/** Gets the first world found of the specified type. */
 	static UWorld* GetWorldOfType(EWorldType::Type WorldType);
 
-	/** Gets the expected playback range of the sequence being rendered. Will vary depending on the frame range mode that Quick Render is currently in. */
+	/**
+	 * Gets the explicitly-set playback range of the level sequence being rendered (set either by a selection in Sequencer, or custom start/end
+	 * frames). If using the unaltered playback range of the level sequence, this will return an empty range.
+	 */
 	static TRange<FFrameNumber> GetPlaybackRange(const UMovieGraphQuickRenderSettings& InQuickRenderSettings);
 
 	/** Does any setup needed in the PIE world prior to the render starting. */
