MoviePipeline: Added the ability to change modes in the Quick Render Settings dialog. The settings in the Settings dialog are unique per mode (because each mode may need a different graph config, etc).

diff --git a/Engine/Plugins/MovieScene/MovieRenderPipeline/Source/MovieRenderPipelineCore/Private/Graph/MovieGraphQuickRenderSettings.cpp b/Engine/Plugins/MovieScene/MovieRenderPipeline/Source/MovieRenderPipelineCore/Private/Graph/MovieGraphQuickRenderSettings.cpp
--- a/Engine/Plugins/MovieScene/MovieRenderPipeline/Source/MovieRenderPipelineCore/Private/Graph/MovieGraphQuickRenderSettings.cpp
+++ b/Engine/Plugins/MovieScene/MovieRenderPipeline/Source/MovieRenderPipelineCore/Private/Graph/MovieGraphQuickRenderSettings.cpp
@@ -1,4 +1,4 @@
-﻿// Copyright Epic Games, Inc. All Rights Reserved.
+// Copyright Epic Games, Inc. All Rights Reserved.
 
 #include "Graph/MovieGraphQuickRenderSettings.h"
 
@@ -12,14 +12,17 @@
 #include "Settings/EditorLoadingSavingSettings.h"
 #endif
 
-UMovieGraphQuickRenderSettings::UMovieGraphQuickRenderSettings()
-	: GraphPreset(FSoftObjectPath(TEXT("/MovieRenderPipeline/DefaultQuickRenderGraph.DefaultQuickRenderGraph")))
+UMovieGraphQuickRenderModeSettings* UMovieGraphQuickRenderSettings::GetSavedQuickRenderModeSettings(const EMovieGraphQuickRenderMode SettingsMode)
 {
+	// Some modes share the same settings, so here we map the mode to an internal settings grouping. Each settings group shares the same settings.
+	static const TMap<EMovieGraphQuickRenderMode, FName> ModeToInternalName =
+	{
+		{ EMovieGraphQuickRenderMode::CurrentSequence, FName(TEXT("Sequence"))},
+		{ EMovieGraphQuickRenderMode::CurrentViewport, FName(TEXT("CurrentViewport"))},
+		{ EMovieGraphQuickRenderMode::SelectedCameras, FName(TEXT("SelectedCameras"))},
+		{ EMovieGraphQuickRenderMode::UseViewportCameraInSequence, FName(TEXT("Sequence"))},
+	};
 	
-}
-
-UMovieGraphQuickRenderSettings* UMovieGraphQuickRenderSettings::GetSavedQuickRenderSettings()
-{
 	UMovieGraphQuickRenderSettings* QuickRenderSettings = nullptr;
 	
 	if (const UPackage* QuickRenderSettingsPackage = LoadPackage(nullptr, *QuickRenderSettingsPackagePath, LOAD_None))
@@ -35,17 +38,18 @@ UMovieGraphQuickRenderSettings* UMovieGraphQuickRenderSettings::GetSavedQuickRen
 		QuickRenderSettings = NewObject<UMovieGraphQuickRenderSettings>(NewSettingsPackage);
 	}
 
-	return QuickRenderSettings;
-}
-
-void UMovieGraphQuickRenderSettings::RefreshVariableAssignments(UMovieGraphQuickRenderSettings* InSettings)
-{
-	if (!InSettings)
+	// Get the settings for the mode specified. Create them if they don't exist yet.
+	const FName& InternalSettingsGroupName = ModeToInternalName.FindChecked(SettingsMode);
+	const TObjectPtr<UMovieGraphQuickRenderModeSettings>* ModeSettings = QuickRenderSettings->ModeSettings.Find(InternalSettingsGroupName);
+	if (!ModeSettings || !IsValid(*ModeSettings))
 	{
-		return;
+		UMovieGraphQuickRenderModeSettings* NewModeSettings = NewObject<UMovieGraphQuickRenderModeSettings>(QuickRenderSettings, NAME_None, RF_Transactional);
+		QuickRenderSettings->ModeSettings.Add(InternalSettingsGroupName, NewModeSettings);
+
+		return NewModeSettings;
 	}
-	
-	MoviePipeline::RefreshVariableAssignments(InSettings->GraphPreset.LoadSynchronous(), InSettings->GraphVariableAssignments, InSettings);
+
+	return ModeSettings->Get();
 }
 
 #if WITH_EDITOR
@@ -75,10 +79,8 @@ void UMovieGraphQuickRenderSettings::SaveSettings(const UMovieGraphQuickRenderSe
 		UE_LOG(LogMovieRenderPipeline, Warning, TEXT("Unable to save Quick Render settings. Could not save to destination file [%s]."), *PackageFileName);
 	}
 }
-#endif	// WITH_EDITOR
 
-#if WITH_EDITOR
-void UMovieGraphQuickRenderSettings::PostEditChangeProperty(FPropertyChangedEvent& PropertyChangedEvent)
+void UMovieGraphQuickRenderSettings::NotifyNeedsSave()
 {
 	if (OnEnginePreExitHandle.IsValid())
 	{
@@ -91,14 +93,40 @@ void UMovieGraphQuickRenderSettings::PostEditChangeProperty(FPropertyChangedEven
 	{
 		SaveSettings(this);
 	});
+}
+#endif	// WITH_EDITOR
+
+UMovieGraphQuickRenderModeSettings::UMovieGraphQuickRenderModeSettings()
+	: GraphPreset(FSoftObjectPath(TEXT("/MovieRenderPipeline/DefaultQuickRenderGraph.DefaultQuickRenderGraph")))
+{
+}
+
+void UMovieGraphQuickRenderModeSettings::RefreshVariableAssignments(UMovieGraphQuickRenderModeSettings* InSettings)
+{
+	if (!InSettings)
+	{
+		return;
+	}
 	
-	if (PropertyChangedEvent.GetPropertyName() == GET_MEMBER_NAME_CHECKED(UMovieGraphQuickRenderSettings, GraphPreset))
+	MoviePipeline::RefreshVariableAssignments(InSettings->GraphPreset.LoadSynchronous(), InSettings->GraphVariableAssignments, InSettings);
+}
+
+#if WITH_EDITOR
+void UMovieGraphQuickRenderModeSettings::PostEditChangeProperty(FPropertyChangedEvent& PropertyChangedEvent)
+{
+	if (PropertyChangedEvent.GetPropertyName() == GET_MEMBER_NAME_CHECKED(UMovieGraphQuickRenderModeSettings, GraphPreset))
 	{
 		// Refresh the variable assignments if the graph preset changes
 		RefreshVariableAssignments(this);
 		
 		OnGraphChangedDelegate.Broadcast();
 	}
+
+	// Let the parent settings object know that it is dirty and should queue up a save
+	if (UMovieGraphQuickRenderSettings* QuickRenderSettings = GetTypedOuter<UMovieGraphQuickRenderSettings>())
+	{
+		QuickRenderSettings->NotifyNeedsSave();
+	}
 	
 	Super::PostEditChangeProperty(PropertyChangedEvent);
 }
diff --git a/Engine/Plugins/MovieScene/MovieRenderPipeline/Source/MovieRenderPipelineCore/Public/Graph/MovieGraphQuickRenderSettings.h b/Engine/Plugins/MovieScene/MovieRenderPipeline/Source/MovieRenderPipelineCore/Public/Graph/MovieGraphQuickRenderSettings.h
--- a/Engine/Plugins/MovieScene/MovieRenderPipeline/Source/MovieRenderPipelineCore/Public/Graph/MovieGraphQuickRenderSettings.h
+++ b/Engine/Plugins/MovieScene/MovieRenderPipeline/Source/MovieRenderPipelineCore/Public/Graph/MovieGraphQuickRenderSettings.h
@@ -4,8 +4,6 @@
 
 #include "CoreMinimal.h"
 
-#include "MovieGraphNamedResolution.h"
-
 #include "MovieGraphQuickRenderSettings.generated.h"
 
 class UMovieGraphConfig;
@@ -79,57 +77,62 @@ enum class EMovieGraphQuickRenderViewportLookFlags : uint8
 };
 ENUM_CLASS_FLAGS(EMovieGraphQuickRenderViewportLookFlags)
 
-/** TODO: Docs */
-UENUM(BlueprintType)
-enum class EMovieGraphQuickRenderSettingsCategory : uint8
-{
-	Camera,
-	Sequence,
-	Viewport
-};
-
 #if WITH_EDITOR
 DECLARE_MULTICAST_DELEGATE(FOnMovieGraphQuickRenderGraphChanged);
 #endif
 
-/** Data that specifies user-configurable aspects of a quick render. */
+/** Data that specifies user-configurable aspects of a quick render. Settings are stored for each mode separately. */
 UCLASS(BlueprintType)
 class MOVIERENDERPIPELINECORE_API UMovieGraphQuickRenderSettings : public UObject
 {
 	GENERATED_BODY()
 
 public:
-	UMovieGraphQuickRenderSettings();
+	UMovieGraphQuickRenderSettings() = default;
 	
-	/** Gets the Quick Render settings that have been persisted to disk (Saved/MovieRenderPipeline/QuickRenderSettings.uasset). */
+	/** Gets the Quick Render settings (for a specific mode) that have been persisted to disk (Saved/MovieRenderPipeline/QuickRenderSettings.uasset). */
 	UFUNCTION(BlueprintCallable, Category="Quick Render")
-	static UMovieGraphQuickRenderSettings* GetSavedQuickRenderSettings();
-
-	/**
-	 * Within the settings object provided, refreshes the variable assignments in the settings to match the setting's GraphPreset. Note that
-	 * changing the graph preset within the editor will do this automatically, but if changing the graph outside of the editor, this may need to be
-	 * run manually.
-	 */
-	UFUNCTION(BlueprintCallable, Category="Quick Render")
-	static void RefreshVariableAssignments(UMovieGraphQuickRenderSettings* InSettings);
+	static UMovieGraphQuickRenderModeSettings* GetSavedQuickRenderModeSettings(EMovieGraphQuickRenderMode SettingsMode);
 
 private:
 	inline static const FString QuickRenderSettingsPackagePath = TEXT("/Temp/MovieRenderPipeline/QuickRenderSettings");
 	inline static FDelegateHandle OnEnginePreExitHandle;
 
+	/**
+	 * Individual settings for each mode. Some modes have shared settings, hence why this is keyed by FName (an internal identifier) instead of
+	 * the mode enum.
+	 */
+	UPROPERTY()
+	TMap<FName, TObjectPtr<UMovieGraphQuickRenderModeSettings>> ModeSettings;
+
 private:
 #if WITH_EDITOR
 	/** Saves the settings to a uasset file (Saved/MovieRenderPipeline/QuickRenderSettings.uasset). The file will be created if it doesn't exist. */
 	static void SaveSettings(const UMovieGraphQuickRenderSettings* InSettings);
-#endif
 
-#if WITH_EDITOR
-	//~ Begin UObject Interface
-	virtual void PostEditChangeProperty(FPropertyChangedEvent& PropertyChangedEvent) override;
-	//~ End UObject Interface
-#endif
+	// Individual mode settings objects need to notify this parent settings container that it should save.
+	friend class UMovieGraphQuickRenderModeSettings;
+	void NotifyNeedsSave();
+#endif	// WITH_EDITOR
+};
+
+/** Settings for a specific mode within Quick Render. Note that some modes share settings. */
+UCLASS(BlueprintType)
+class MOVIERENDERPIPELINECORE_API UMovieGraphQuickRenderModeSettings : public UObject
+{
+	GENERATED_BODY()
 
 public:
+	UMovieGraphQuickRenderModeSettings();
+
+	/**
+	 * Within the settings object provided, refreshes the variable assignments in the settings to match the setting's GraphPreset. Note that
+	 * changing the graph preset within the editor will do this automatically, but if changing the graph outside of the editor, this may need to be
+	 * run manually.
+	 */
+	UFUNCTION(BlueprintCallable, Category="Quick Render")
+	static void RefreshVariableAssignments(UMovieGraphQuickRenderModeSettings* InSettings);
+
 #if WITH_EDITOR
 	FOnMovieGraphQuickRenderGraphChanged OnGraphChangedDelegate;
 #endif
@@ -169,23 +172,10 @@ public:
 	UPROPERTY(EditAnywhere, BlueprintReadWrite, Category="Primary Graph Variables")
 	TArray<TObjectPtr<UMovieJobVariableAssignmentContainer>> GraphVariableAssignments;
 
-	// 
-	// ---- The settings below should generally not need to be controlled outside of the UI. ----
-	//
-
-	/** TODO: Docs */
-	UPROPERTY()
-	EMovieGraphQuickRenderSettingsCategory SettingsCategory;
-
-	/** What mode the quick render mode button is using. */
-	UPROPERTY()
-	EMovieGraphQuickRenderButtonMode QuickAccessButtonMode = EMovieGraphQuickRenderButtonMode::NormalMovieRenderQueue;
-
-	/** The mode that Quick Render will use when a render begins. Generally dictates the level sequence and camera(s) that will be used. */
-	UPROPERTY()
-	EMovieGraphQuickRenderMode QuickRenderMode = EMovieGraphQuickRenderMode::CurrentSequence;
-
-	/** Whether the settings dialog should be shown before a quick render is started. */
-	UPROPERTY()
-	bool bShowSettingsBeforeQuickRender = false;
+private:
+#if WITH_EDITOR
+	//~ Begin UObject Interface
+	virtual void PostEditChangeProperty(FPropertyChangedEvent& PropertyChangedEvent) override;
+	//~ End UObject Interface
+#endif
 };
\ No newline at end of file
diff --git a/Engine/Plugins/MovieScene/MovieRenderPipeline/Source/MovieRenderPipelineEditor/Private/Customizations/Graph/MovieGraphQuickRenderSettingsCustomization.h b/Engine/Plugins/MovieScene/MovieRenderPipeline/Source/MovieRenderPipelineEditor/Private/Customizations/Graph/MovieGraphQuickRenderSettingsCustomization.h
--- a/Engine/Plugins/MovieScene/MovieRenderPipeline/Source/MovieRenderPipelineEditor/Private/Customizations/Graph/MovieGraphQuickRenderSettingsCustomization.h
+++ b/Engine/Plugins/MovieScene/MovieRenderPipeline/Source/MovieRenderPipelineEditor/Private/Customizations/Graph/MovieGraphQuickRenderSettingsCustomization.h
@@ -35,9 +35,9 @@ public:
 		
 		UPackage::PackageSavedWithContextEvent.RemoveAll(this);
 
-		if (QuickRenderSettings.IsValid())
+		if (QuickRenderModeSettings.IsValid())
 		{
-			QuickRenderSettings->OnGraphChangedDelegate.RemoveAll(this);
+			QuickRenderModeSettings->OnGraphChangedDelegate.RemoveAll(this);
 		}
 	}
 
@@ -54,7 +54,7 @@ protected:
 		DetailBuilder->ForceRefreshDetails();
 
 		// Refresh the variable assignments (since the variables within the graph may have changed, thus the variable assignments need to be updated)
-		MoviePipeline::RefreshVariableAssignments(QuickRenderSettings->GraphPreset.LoadSynchronous(), QuickRenderSettings->GraphVariableAssignments, QuickRenderSettings.Get());
+		MoviePipeline::RefreshVariableAssignments(QuickRenderModeSettings->GraphPreset.LoadSynchronous(), QuickRenderModeSettings->GraphVariableAssignments, QuickRenderModeSettings.Get());
 	}
 
 	void RefreshLayout() const
@@ -77,19 +77,19 @@ protected:
 			return;
 		}
 
-		QuickRenderSettings = Cast<UMovieGraphQuickRenderSettings>(ObjectsBeingCustomized[0]);
-		if (!QuickRenderSettings.IsValid())
+		QuickRenderModeSettings = Cast<UMovieGraphQuickRenderModeSettings>(ObjectsBeingCustomized[0]);
+		if (!QuickRenderModeSettings.IsValid())
 		{
 			return;
 		}
 
 		// Hide the original assignments property (since it presents an asset picker).
 		const TSharedRef<IPropertyHandle> GraphVariableAssignmentsProperty =
-			InDetailBuilder.GetProperty(TEXT("GraphVariableAssignments"), UMovieGraphQuickRenderSettings::StaticClass());
+			InDetailBuilder.GetProperty(TEXT("GraphVariableAssignments"), UMovieGraphQuickRenderModeSettings::StaticClass());
 		InDetailBuilder.HideProperty(GraphVariableAssignmentsProperty);
 
 		// Refresh the UI if the graph preset changes (so the new variable assignments are displayed).
-		QuickRenderSettings->OnGraphChangedDelegate.AddSP(this, &FMovieGraphQuickRenderSettingsCustomization::RefreshLayout);
+		QuickRenderModeSettings->OnGraphChangedDelegate.AddSP(this, &FMovieGraphQuickRenderSettingsCustomization::RefreshLayout);
 
 		// Set up the category for variable assignments.
 		IDetailCategoryBuilder& PrimaryGraphVariablesCategory = InDetailBuilder.EditCategory(
@@ -101,7 +101,7 @@ protected:
 
 		// Set the variable assignments category as hidden by default, and add the variable assignments. Individual categories will be made visible if variables are added under them.
 		PrimaryGraphVariablesCategory.SetCategoryVisibility(false);
-		UE::MovieRenderPipelineEditor::Private::AddVariableAssignments(QuickRenderSettings->GraphVariableAssignments, PrimaryGraphVariablesCategory, DetailBuilder);
+		UE::MovieRenderPipelineEditor::Private::AddVariableAssignments(QuickRenderModeSettings->GraphVariableAssignments, PrimaryGraphVariablesCategory, DetailBuilder);
 
 		// Give the categories a specific ordering.
 		int32 SortOrder = 0;
@@ -115,8 +115,8 @@ private:
 	/** The details builder associated with the customization. */
 	IDetailLayoutBuilder* DetailBuilder = nullptr;
 
-	/** The quick render settings being displayed. */
-	TWeakObjectPtr<UMovieGraphQuickRenderSettings> QuickRenderSettings = nullptr;
+	/** The quick render mode settings being displayed. */
+	TWeakObjectPtr<UMovieGraphQuickRenderModeSettings> QuickRenderModeSettings = nullptr;
 };
 
 #undef LOCTEXT_NAMESPACE
\ No newline at end of file
diff --git a/Engine/Plugins/MovieScene/MovieRenderPipeline/Source/MovieRenderPipelineEditor/Private/Graph/MovieGraphQuickRender.cpp b/Engine/Plugins/MovieScene/MovieRenderPipeline/Source/MovieRenderPipelineEditor/Private/Graph/MovieGraphQuickRender.cpp
--- a/Engine/Plugins/MovieScene/MovieRenderPipeline/Source/MovieRenderPipelineEditor/Private/Graph/MovieGraphQuickRender.cpp
+++ b/Engine/Plugins/MovieScene/MovieRenderPipeline/Source/MovieRenderPipelineEditor/Private/Graph/MovieGraphQuickRender.cpp
@@ -43,7 +43,7 @@
 
 #define LOCTEXT_NAMESPACE "FMovieGraphQuickRender"
 
-void FMovieGraphQuickRender::BeginQuickRender(const UMovieGraphQuickRenderSettings& InQuickRenderSettings)
+void FMovieGraphQuickRender::BeginQuickRender(const EMovieGraphQuickRenderMode QuickRenderMode, const UMovieGraphQuickRenderModeSettings& InQuickRenderSettings)
 {
 	UMoviePipelineQueueSubsystem* Subsystem = GEditor->GetEditorSubsystem<UMoviePipelineQueueSubsystem>();
 	check(Subsystem);
@@ -58,17 +58,17 @@ void FMovieGraphQuickRender::BeginQuickRender(const UMovieGraphQuickRenderSettin
 	const TSubclassOf<UMoviePipelineExecutorBase> ExecutorClass = UMoviePipelinePIEExecutor::StaticClass();
 
 	// Get/create the level sequence that should be rendered from (varies depending on the Quick Render mode).
-	SetUpRenderingLevelSequence(InQuickRenderSettings);
+	SetUpRenderingLevelSequence(QuickRenderMode);
 	if (!RenderingLevelSequence.IsValid())
 	{
 		return;
 	}
 
 	// Also set up the utility level sequence. This is not expected to be valid for most modes, so we don't need to check that this was successfully created.
-	SetUpUtilityLevelSequence(InQuickRenderSettings);
+	SetUpUtilityLevelSequence(QuickRenderMode);
 
 	// Don't do additional work if the selected mode doesn't have what it needs to do a render. AreModePrerequisitesMet() will output the appropriate warnings.
-	if (!AreModePrerequisitesMet(InQuickRenderSettings.QuickRenderMode))
+	if (!AreModePrerequisitesMet(QuickRenderMode))
 	{
 		return;
 	}
@@ -143,9 +143,9 @@ void FMovieGraphQuickRender::BeginQuickRender(const UMovieGraphQuickRenderSettin
 	UMoviePipelinePIEExecutor* Executor = Cast<UMoviePipelinePIEExecutor>(Subsystem->RenderQueueInstanceWithExecutor(TemporaryQueue.Get(), ExecutorClass));
 
 	// Do any needed setup in the PIE world before rendering starts
-	Executor->OnIndividualJobStarted().AddLambda([&InQuickRenderSettings](UMoviePipelineExecutorJob* InJob)
+	Executor->OnIndividualJobStarted().AddLambda([QuickRenderMode](UMoviePipelineExecutorJob* InJob)
 	{
-		PerformPreRenderSetup(InQuickRenderSettings);
+		PerformPreRenderSetup(QuickRenderMode);
 	});
 
 	// View the job's rendered frame(s)/video(s) if requested
@@ -160,9 +160,9 @@ void FMovieGraphQuickRender::BeginQuickRender(const UMovieGraphQuickRenderSettin
 	});
 
 	// Clean up
-	Executor->OnExecutorFinished().AddLambda([&InQuickRenderSettings](UMoviePipelineExecutorBase* InExecutor, bool bSuccess)
+	Executor->OnExecutorFinished().AddLambda([QuickRenderMode](UMoviePipelineExecutorBase* InExecutor, bool bSuccess)
 	{
-		RestorePreRenderState(InQuickRenderSettings);
+		RestorePreRenderState(QuickRenderMode);
 
 		// The temp objects are part of the transient package, so manually release them when the render is finished
 		TemporaryQueue.Reset();
@@ -206,7 +206,7 @@ UMovieGraphConfig* FMovieGraphQuickRender::GetQuickRenderGraph(UMovieGraphConfig
 	return TemporaryGraph.Get();
 }
 
-void FMovieGraphQuickRender::ApplyQuickRenderUpdatesToDuplicateGraph(const UMovieGraphQuickRenderSettings& InQuickRenderSettings, UWorld* InEditorWorld, UMoviePipelineExecutorJob* InJob)
+void FMovieGraphQuickRender::ApplyQuickRenderUpdatesToDuplicateGraph(const UMovieGraphQuickRenderModeSettings& InQuickRenderSettings, UWorld* InEditorWorld, UMoviePipelineExecutorJob* InJob)
 {
 	check(TemporaryGraph.IsValid());
 
@@ -508,7 +508,7 @@ void FMovieGraphQuickRender::CachePrePieData(UWorld* InEditorWorld)
 	CopyToWeakArray(CachedPrePieData.SelectedCameras, CameraActors);
 }
 
-void FMovieGraphQuickRender::RestorePreRenderState(const UMovieGraphQuickRenderSettings& InQuickRenderSettings)
+void FMovieGraphQuickRender::RestorePreRenderState(const EMovieGraphQuickRenderMode QuickRenderMode)
 {
 	// Restore the modified level sequences to the way they were before the render.
 	for (const TWeakObjectPtr<ALevelSequenceActor>& LevelSequenceActor : CachedPrePieData.ModifiedLevelSequenceActors)
@@ -522,7 +522,7 @@ void FMovieGraphQuickRender::RestorePreRenderState(const UMovieGraphQuickRenderS
 	CachedPrePieData.SelectedCameras.Empty();
 
 	// If we're rendering with the Viewport Camera in Current Sequence mode, remove the subsequence track that was added
-	if (InQuickRenderSettings.QuickRenderMode == EMovieGraphQuickRenderMode::UseViewportCameraInSequence)
+	if (QuickRenderMode == EMovieGraphQuickRenderMode::UseViewportCameraInSequence)
 	{
 		const ULevelSequence* CurrentLevelSequence = ULevelSequenceEditorBlueprintLibrary::GetCurrentLevelSequence();
 		check(CurrentLevelSequence);
@@ -547,7 +547,7 @@ void FMovieGraphQuickRender::RestorePreRenderState(const UMovieGraphQuickRenderS
 	}
 }
 
-void FMovieGraphQuickRender::SetUpRenderingLevelSequence(const UMovieGraphQuickRenderSettings& InQuickRenderSettings)
+void FMovieGraphQuickRender::SetUpRenderingLevelSequence(const EMovieGraphQuickRenderMode QuickRenderMode)
 {
 	auto GetCurrentLevelSequenceOrWarn = [](EMovieGraphQuickRenderMode ActiveQuickRenderMode) -> ULevelSequence*
 	{
@@ -565,8 +565,6 @@ void FMovieGraphQuickRender::SetUpRenderingLevelSequence(const UMovieGraphQuickR
 		return CurrentLevelSequence;
 	};
 	
-	const EMovieGraphQuickRenderMode QuickRenderMode = InQuickRenderSettings.QuickRenderMode;
-	
 	// For the CurrentViewport and SelectedCameras modes, a temporary level sequence will be used to drive the render.
 	if ((QuickRenderMode == EMovieGraphQuickRenderMode::CurrentViewport) || (QuickRenderMode == EMovieGraphQuickRenderMode::SelectedCameras))
 	{
@@ -589,11 +587,11 @@ void FMovieGraphQuickRender::SetUpRenderingLevelSequence(const UMovieGraphQuickR
 	RenderingLevelSequence.Reset(CurrentLevelSequence);
 }
 
-void FMovieGraphQuickRender::SetUpUtilityLevelSequence(const UMovieGraphQuickRenderSettings& InQuickRenderSettings)
+void FMovieGraphQuickRender::SetUpUtilityLevelSequence(const EMovieGraphQuickRenderMode QuickRenderMode)
 {
 	// This mode needs a utility sequence that contains a camera cut with a high hierarchical bias, which will allow it to override all cameras in
 	// the level sequence being rendered.
-	if (InQuickRenderSettings.QuickRenderMode == EMovieGraphQuickRenderMode::UseViewportCameraInSequence)
+	if (QuickRenderMode == EMovieGraphQuickRenderMode::UseViewportCameraInSequence)
 	{
 		ULevelSequence* UtilitySequence = NewObject<ULevelSequence>(GetTransientPackage(), TEXT("QuickRenderUtilitySequence"), RF_Transient);
 		UtilitySequence->Initialize();
@@ -617,7 +615,7 @@ UWorld* FMovieGraphQuickRender::GetWorldOfType(EWorldType::Type WorldType)
 	return World;
 }
 
-TRange<FFrameNumber> FMovieGraphQuickRender::GetPlaybackRange(const UMovieGraphQuickRenderSettings& InQuickRenderSettings)
+TRange<FFrameNumber> FMovieGraphQuickRender::GetPlaybackRange(const UMovieGraphQuickRenderModeSettings& InQuickRenderSettings)
 {
 	FFrameNumber StartFrame;
 	FFrameNumber EndFrame;
@@ -652,7 +650,7 @@ TRange<FFrameNumber> FMovieGraphQuickRender::GetPlaybackRange(const UMovieGraphQ
 	return TRange<FFrameNumber>::Inclusive(StartFrame, EndFrame);
 }
 
-void FMovieGraphQuickRender::PerformPreRenderSetup(const UMovieGraphQuickRenderSettings& InQuickRenderSettings)
+void FMovieGraphQuickRender::PerformPreRenderSetup(const EMovieGraphQuickRenderMode QuickRenderMode)
 {
 	// Adds the specified camera to the Camera Cut track in the given level sequence.
 	// If there's already a camera cut section within the Camera Cut track, a new one will be added immediately after.
@@ -729,12 +727,12 @@ void FMovieGraphQuickRender::PerformPreRenderSetup(const UMovieGraphQuickRenderS
 		return NewCamera;
 	};
 	
-	if (InQuickRenderSettings.QuickRenderMode == EMovieGraphQuickRenderMode::CurrentViewport)
+	if (QuickRenderMode == EMovieGraphQuickRenderMode::CurrentViewport)
 	{
 		ACameraActor* NewCamera = DuplicateViewportCamera();
 		AddCameraToLevelSequence(NewCamera, RenderingLevelSequence.Get());
 	}
-	else if (InQuickRenderSettings.QuickRenderMode == EMovieGraphQuickRenderMode::UseViewportCameraInSequence)
+	else if (QuickRenderMode == EMovieGraphQuickRenderMode::UseViewportCameraInSequence)
 	{
 		UMovieScene* CurrentMovieScene = RenderingLevelSequence.Get()->GetMovieScene();
 		const TRange<FFrameNumber> CurrentPlaybackRange = CurrentMovieScene->GetPlaybackRange();
@@ -753,7 +751,7 @@ void FMovieGraphQuickRender::PerformPreRenderSetup(const UMovieGraphQuickRenderS
 		ACameraActor* NewCamera = DuplicateViewportCamera();
 		AddCameraToLevelSequence(NewCamera, UtilityLevelSequence.Get(), CurrentPlaybackRange);
 	}
-	else if (InQuickRenderSettings.QuickRenderMode == EMovieGraphQuickRenderMode::SelectedCameras)
+	else if (QuickRenderMode == EMovieGraphQuickRenderMode::SelectedCameras)
 	{
 		for (const TWeakObjectPtr<ACameraActor>& SelectedCamera : CachedPrePieData.SelectedCameras)
 		{
diff --git a/Engine/Plugins/MovieScene/MovieRenderPipeline/Source/MovieRenderPipelineEditor/Private/MoviePipelineQuickRenderMenu.cpp b/Engine/Plugins/MovieScene/MovieRenderPipeline/Source/MovieRenderPipelineEditor/Private/MoviePipelineQuickRenderMenu.cpp
--- a/Engine/Plugins/MovieScene/MovieRenderPipeline/Source/MovieRenderPipelineEditor/Private/MoviePipelineQuickRenderMenu.cpp
+++ b/Engine/Plugins/MovieScene/MovieRenderPipeline/Source/MovieRenderPipelineEditor/Private/MoviePipelineQuickRenderMenu.cpp
@@ -7,6 +7,7 @@
 #include "Framework/Application/SlateApplication.h"
 #include "Framework/Docking/TabManager.h"
 #include "Graph/MovieGraphQuickRender.h"
+#include "Misc/ConfigCacheIni.h"
 #include "Misc/MessageDialog.h"
 #include "MoviePipelineUtils.h"
 #include "MovieRenderPipelineCoreModule.h"
@@ -82,23 +83,23 @@ void FMoviePipelineQuickRenderMenu::RemoveQuickRenderButtonToolMenu(UToolMenu* I
 
 void FMoviePipelineQuickRenderMenu::LoadQuickRenderSettings()
 {
-	QuickRenderSettings = TStrongObjectPtr(UMovieGraphQuickRenderSettings::GetSavedQuickRenderSettings());
+	InitQuickRenderModeSettingsFromMode(GetQuickRenderMode());
 }
 
 void FMoviePipelineQuickRenderMenu::QuickRenderButtonPressed()
 {
 	LoadQuickRenderSettings();
 
-	if (QuickRenderSettings->QuickAccessButtonMode == EMovieGraphQuickRenderButtonMode::QuickRender)
+	if (GetQuickRenderButtonMode() == EMovieGraphQuickRenderButtonMode::QuickRender)
 	{
 		// The settings window will have its own Render button, so don't perform the render if showing the settings first
-		if (QuickRenderSettings->bShowSettingsBeforeQuickRender)
+		if (GetShouldShowSettingsBeforeRender())
 		{
 			OpenQuickRenderSettingsWindow(FToolMenuContext());
 		}
 		else
 		{
-			FMovieGraphQuickRender::BeginQuickRender(*QuickRenderSettings);
+			FMovieGraphQuickRender::BeginQuickRender(GetQuickRenderMode(), *QuickRenderModeSettings);
 		}
 	}
 	else
@@ -150,16 +151,19 @@ void FMoviePipelineQuickRenderMenu::OpenQuickRenderSettingsWindow(const FToolMen
 		return;
 	}
 
-	// Init the settings categories
-	QuickRenderSettingsCategories.Empty();
-	const UEnum* SettingsCategoryEnum = StaticEnum<EMovieGraphQuickRenderSettingsCategory>();
-	for (int32 EnumIndex = 0; EnumIndex < SettingsCategoryEnum->NumEnums() - 1; ++EnumIndex)	// -1 to avoid MAX
+	// Before the window opens, sync up the mode that the Settings window is using to be the mode that the toolbar shows
+	SettingsWindowQuickRenderMode = GetQuickRenderMode();
+
+	// Init the modes that are available to switch to
+	QuickRenderModes.Empty();
+	const UEnum* ModesEnum = StaticEnum<EMovieGraphQuickRenderMode>();
+	for (int32 EnumIndex = 0; EnumIndex < ModesEnum->NumEnums() - 1; ++EnumIndex)	// -1 to avoid MAX
 	{
-		QuickRenderSettingsCategories.Add(MakeShared<FName>(SettingsCategoryEnum->GetNameByIndex(EnumIndex)));
+		QuickRenderModes.Add(MakeShared<FName>(ModesEnum->GetNameByIndex(EnumIndex)));
 	}
 	
 	// Update the graph preset's variable assignments in case the graph was updated since the last time the window opened
-	MoviePipeline::RefreshVariableAssignments(QuickRenderSettings->GraphPreset.LoadSynchronous(), QuickRenderSettings->GraphVariableAssignments, QuickRenderSettings.Get());
+	UpdateVariableAssignmentsForCurrentGraph();
 
 	FPropertyEditorModule& PropertyEditorModule = FModuleManager::Get().LoadModuleChecked<FPropertyEditorModule>("PropertyEditor");
 	FDetailsViewArgs DetailsViewArgs;
@@ -173,9 +177,9 @@ void FMoviePipelineQuickRenderMenu::OpenQuickRenderSettingsWindow(const FToolMen
 	// Create the details panel and display the quick render settings
 	const TSharedPtr<IDetailsView> DetailsView = PropertyEditorModule.CreateDetailView(DetailsViewArgs);
 	DetailsView->RegisterInstancedCustomPropertyLayout(
-		UMovieGraphQuickRenderSettings::StaticClass(),
+		UMovieGraphQuickRenderModeSettings::StaticClass(),
 		FOnGetDetailCustomizationInstance::CreateStatic(&FMovieGraphQuickRenderSettingsCustomization::MakeInstance));
-	DetailsView->SetObject(QuickRenderSettings.Get());
+	DetailsView->SetObject(QuickRenderModeSettings.Get());
 	
 	const TSharedRef<SWindow> QuickRenderSettingsWindow = SNew(SWindow)
 		.Title(LOCTEXT("QuickRenderSettingsWindow_Title", "Quick Render"))
@@ -207,52 +211,71 @@ void FMoviePipelineQuickRenderMenu::OpenQuickRenderSettingsWindow(const FToolMen
 		FSlateApplication::Get().AddWindow(QuickRenderSettingsWindow);
 	}
 
+	WeakDetailsPanel = DetailsView;
 	WeakQuickRenderSettingsWindow = QuickRenderSettingsWindow;
 }
 
 TSharedRef<SWidget> FMoviePipelineQuickRenderMenu::GenerateQuickRenderSettingsWindowToolbar()
 {
+	auto GetWidgetForQuickRenderMode = [](TFunction<EMovieGraphQuickRenderMode()>&& GetQuickRenderMode)
+	{
+		return
+			SNew(SHorizontalBox)
+			+ SHorizontalBox::Slot()
+			.Padding(FMargin{0.0f, 0.0f, 5.0f, 0.0f})
+			.AutoWidth()
+			[
+				SNew(SImage)
+				.ColorAndOpacity(FSlateColor::UseForeground())
+				.Image_Lambda([GetQuickRenderMode]()
+				{
+					return GetIconForQuickRenderMode(GetQuickRenderMode()).GetIcon();
+				})
+				.DesiredSizeOverride(FVector2d(16, 16))
+			]
+
+			+ SHorizontalBox::Slot()
+			[
+				SNew(STextBlock)
+				.Text_Lambda([GetQuickRenderMode]()
+				{
+					return StaticEnum<EMovieGraphQuickRenderMode>()->GetDisplayNameTextByValue(static_cast<int64>(GetQuickRenderMode()));
+				})
+			];
+	};
+	
 	FToolBarBuilder ToolbarBuilder(TSharedPtr<const FUICommandList>(), FMultiBoxCustomization::None);
 	ToolbarBuilder.BeginSection("Quick Render");
 	ToolbarBuilder.AddWidget(
 		SNew(SComboBox<TSharedPtr<FName>>)
-		.OptionsSource(&FMoviePipelineQuickRenderMenu::QuickRenderSettingsCategories)
-		.IsEnabled(false)	// TODO: Enable once settings categories are implemented
-		.OnGenerateWidget_Lambda([](TSharedPtr<FName> InItem)
+		.OptionsSource(&FMoviePipelineQuickRenderMenu::QuickRenderModes)
+		.OnGenerateWidget_Lambda([GetWidgetForQuickRenderMode](TSharedPtr<FName> InItem)
 		{
-			return
-				SNew(STextBlock)
-				.Text_Lambda([InItem]()
-				{
-					const int64 SettingsCategory = StaticEnum<EMovieGraphQuickRenderSettingsCategory>()->GetValueByName(*InItem);
-
-					return StaticEnum<EMovieGraphQuickRenderSettingsCategory>()->GetDisplayNameTextByValue(SettingsCategory);
-				});
+			return GetWidgetForQuickRenderMode([InItem]()
+			{
+				const int64 ModeValue = StaticEnum<EMovieGraphQuickRenderMode>()->GetValueByName(*InItem);
+				return static_cast<EMovieGraphQuickRenderMode>(ModeValue);
+			});
 		})
 		.OnSelectionChanged_Lambda([](TSharedPtr<FName> InSelectedItem, ESelectInfo::Type)
 		{
-			QuickRenderSettings->SettingsCategory = static_cast<EMovieGraphQuickRenderSettingsCategory>(
-				StaticEnum<EMovieGraphQuickRenderSettingsCategory>()->GetValueByName(*InSelectedItem));
+			SettingsWindowQuickRenderMode = static_cast<EMovieGraphQuickRenderMode>(
+				StaticEnum<EMovieGraphQuickRenderMode>()->GetValueByName(*InSelectedItem));
+
+			// Also inform the details panel of this change. It needs to display new settings for the chosen mode.
+			if (const TSharedPtr<IDetailsView> DetailsPanel = WeakDetailsPanel.Pin())
+			{
+				InitQuickRenderModeSettingsFromMode(SettingsWindowQuickRenderMode);
+				
+				DetailsPanel->SetObject(QuickRenderModeSettings.Get());
+			}
 		})
 		[
-			SNew(STextBlock)
-			.Text_Lambda([]()
+			GetWidgetForQuickRenderMode([]()
 			{
-				return StaticEnum<EMovieGraphQuickRenderSettingsCategory>()->GetDisplayNameTextByValue(static_cast<int64>(QuickRenderSettings->SettingsCategory));
+				return SettingsWindowQuickRenderMode;
 			})
-		]);
-
-	ToolbarBuilder.AddWidget(
-		SNew(SActionButton)
-		.Text(LOCTEXT("QuickRenderSettingsWindow_ApplyModeButtonText", "Apply"))
-		.Icon(FAppStyle::Get().GetBrush("MaterialEditor.Apply"))
-		.ActionButtonType(EActionButtonType::Simple)
-		.IsEnabled(false)	// TODO: Enable once settings categories are implemented
-		.OnClicked_Lambda([]()
-		{
-			// TODO: Apply the mode setting
-			return FReply::Handled();
-		}));
+		], NAME_None, false, HAlign_Left);
 
 	ToolbarBuilder.AddWidget(
 		SNew(SActionButton)
@@ -261,7 +284,10 @@ TSharedRef<SWidget> FMoviePipelineQuickRenderMenu::GenerateQuickRenderSettingsWi
 		.ActionButtonType(EActionButtonType::Primary)
 		.OnClicked_Lambda([]()
 		{
-			FMovieGraphQuickRender::BeginQuickRender(*QuickRenderSettings);
+			// Sync the Setting window's mode to the toolbar mode. This is only done when a render is performed, not when the window is closed.
+			SetQuickRenderMode(SettingsWindowQuickRenderMode);
+			
+			FMovieGraphQuickRender::BeginQuickRender(GetQuickRenderMode(), *QuickRenderModeSettings);
 			return FReply::Handled();
 		}),
 		NAME_None, false, HAlign_Right);
@@ -280,13 +306,13 @@ void FMoviePipelineQuickRenderMenu::GenerateModesMenuSection(UToolMenu* InMenu)
 		FToolUIAction NewAction;
 		NewAction.ExecuteAction = FToolMenuExecuteAction::CreateLambda([CheckedValue](const FToolMenuContext&)
 		{
-			QuickRenderSettings->QuickAccessButtonMode = CheckedValue;
-			QuickRenderSettings->PostEditChange();
+			SetQuickRenderButtonMode(CheckedValue);
+			QuickRenderModeSettings->PostEditChange();
 		});
 		NewAction.CanExecuteAction = FToolMenuCanExecuteAction::CreateLambda([](const FToolMenuContext&) { return true; });
 		NewAction.GetActionCheckState = FToolMenuGetActionCheckState::CreateLambda([CheckedValue](const FToolMenuContext&)
 		{
-			return (QuickRenderSettings->QuickAccessButtonMode == CheckedValue)
+			return (GetQuickRenderButtonMode() == CheckedValue)
 				? ECheckBoxState::Checked
 				: ECheckBoxState::Unchecked;
 		});
@@ -312,29 +338,28 @@ void FMoviePipelineQuickRenderMenu::GenerateQuickRenderMenuSection(UToolMenu* In
 {
 	FToolMenuSection& Section = InMenu->AddSection("MoviePipelineQuickRenderType", LOCTEXT("QuickRenderType", "Quick Render"));
 
-	auto AddQuickRenderModeAction = [&Section](const FName& InName, const FText& InLabel, const FText& InToolTip, const FSlateIcon& InIcon, EMovieGraphQuickRenderMode CheckedValue)
+	auto AddQuickRenderModeAction = [&Section](const FName& InName, const FText& InLabel, const FText& InToolTip, EMovieGraphQuickRenderMode CheckedValue)
 	{
 		FToolUIAction NewAction;
 		NewAction.ExecuteAction = FToolMenuExecuteAction::CreateLambda([CheckedValue](const FToolMenuContext&)
 		{
-			QuickRenderSettings->QuickRenderMode = CheckedValue;
-			QuickRenderSettings->PostEditChange();
+			SetQuickRenderMode(CheckedValue);
+			QuickRenderModeSettings->PostEditChange();
 		});
 		NewAction.CanExecuteAction = FToolMenuCanExecuteAction::CreateLambda([](const FToolMenuContext&) { return true; });
 		NewAction.GetActionCheckState = FToolMenuGetActionCheckState::CreateLambda([CheckedValue](const FToolMenuContext&)
 		{
-			return (QuickRenderSettings->QuickRenderMode == CheckedValue)
+			return (GetQuickRenderMode() == CheckedValue)
 				? ECheckBoxState::Checked
 				: ECheckBoxState::Unchecked;
 		});
-		Section.AddMenuEntry(InName, InLabel, InToolTip, InIcon, NewAction, EUserInterfaceActionType::RadioButton);
+		Section.AddMenuEntry(InName, InLabel, InToolTip, GetIconForQuickRenderMode(CheckedValue), NewAction, EUserInterfaceActionType::RadioButton);
 	};
 
 	AddQuickRenderModeAction(
 		FName("QuickRender_CurrentSequence"),
 		LOCTEXT("QuickRender_CurrentSequenceLabel", "Current Sequence"),
 		LOCTEXT("QuickRender_CurrentSequenceToolTip", "Render the sequence that is currently active in Sequencer."),
-		FSlateIcon("LevelSequenceEditorStyle", "LevelSequenceEditor.CreateNewLevelSequenceInLevel"),
 		EMovieGraphQuickRenderMode::CurrentSequence);
 
 	// The "Current Shot at Playhead" mode is disabled for now, may be added in the future.
@@ -349,21 +374,18 @@ void FMoviePipelineQuickRenderMenu::GenerateQuickRenderMenuSection(UToolMenu* In
 		FName("QuickRender_CurrentViewport"),
 		LOCTEXT("QuickRender_CurrentViewportLabel", "Current Viewport"),
 		LOCTEXT("QuickRender_CurrentViewportToolTip", "Render the current viewport."),
-		FSlateIcon(FAppStyle::Get().GetStyleSetName(), "LevelEditor.Tabs.Viewports"),
 		EMovieGraphQuickRenderMode::CurrentViewport);
 
 	AddQuickRenderModeAction(
 		FName("QuickRender_SelectedCameras"),
 		LOCTEXT("QuickRender_SelectedCamerasLabel", "Selected Camera(s)"),
 		LOCTEXT("QuickRender_SelectedCamerasToolTip", "Render the selected cameras."),
-		FSlateIcon(FAppStyle::Get().GetStyleSetName(), "ClassIcon.CameraActor"),
 		EMovieGraphQuickRenderMode::SelectedCameras);
 
 	AddQuickRenderModeAction(
 		FName("QuickRender_UseViewportCamera"),
 		LOCTEXT("QuickRender_UseViewportCameraLabel", "Use Viewport Camera in Sequence"),
 		LOCTEXT("QuickRender_UseViewportCameraToolTip", "Use the current viewport camera in the active level sequence."),
-		FSlateIcon(FAppStyle::Get().GetStyleSetName(), "ClassIcon.CameraActor"),
 		EMovieGraphQuickRenderMode::UseViewportCameraInSequence);
 }
 
@@ -381,17 +403,17 @@ void FMoviePipelineQuickRenderMenu::GenerateQuickRenderConfigurationMenuSection(
 			.Text_Lambda([]()
 			{
 				// Generate "Configuration" text
-				const FText GraphName = !QuickRenderSettings->GraphPreset.IsNull()
-					? FText::FromString(QuickRenderSettings->GraphPreset.GetAssetName())
+				const FText GraphName = !QuickRenderModeSettings->GraphPreset.IsNull()
+					? FText::FromString(QuickRenderModeSettings->GraphPreset.GetAssetName())
 					: LOCTEXT("QuickRenderConfigMenu_InvalidGraph", "Invalid");
 				
 				// Generate "After Render" text
-				const FText AfterRender = (QuickRenderSettings->PostRenderBehavior == EMovieGraphQuickRenderPostRenderActionType::PlayRenderOutput)
+				const FText AfterRender = (QuickRenderModeSettings->PostRenderBehavior == EMovieGraphQuickRenderPostRenderActionType::PlayRenderOutput)
 					? LOCTEXT("QuickRenderConfigMenu_PlayRenderOutput", "Play Render Output")
 					: LOCTEXT("QuickRenderConfigMenu_DoNothing", "Do Nothing");
 				
 				// Generate "Viewport Look" text
-				const EMovieGraphQuickRenderViewportLookFlags ViewportLookFlags = static_cast<EMovieGraphQuickRenderViewportLookFlags>(QuickRenderSettings->ViewportLookFlags);
+				const EMovieGraphQuickRenderViewportLookFlags ViewportLookFlags = static_cast<EMovieGraphQuickRenderViewportLookFlags>(QuickRenderModeSettings->ViewportLookFlags);
 				TArray<FText> ViewportFlags;
 				if (EnumHasAnyFlags(ViewportLookFlags, EMovieGraphQuickRenderViewportLookFlags::Ocio))
 				{
@@ -411,7 +433,7 @@ void FMoviePipelineQuickRenderMenu::GenerateQuickRenderConfigurationMenuSection(
 				}
 				
 				// Generate "Sequencer Frame Range" text
-				const FText SequencerFrameRange = StaticEnum<EMovieGraphQuickRenderFrameRangeType>()->GetDisplayNameTextByValue(static_cast<int64>(QuickRenderSettings->FrameRangeType));
+				const FText SequencerFrameRange = StaticEnum<EMovieGraphQuickRenderFrameRangeType>()->GetDisplayNameTextByValue(static_cast<int64>(QuickRenderModeSettings->FrameRangeType));
 
 				// TODO: Add in graph variables summary as well (although this might make a summary that's too long in some cases)
 				FText ConfigurationSummaryText = FText::Format(
@@ -455,13 +477,13 @@ void FMoviePipelineQuickRenderMenu::GenerateSettingsMenuSection(UToolMenu* InMen
 	FToolUIAction ShowSettingsAction;
         ShowSettingsAction.ExecuteAction = FToolMenuExecuteAction::CreateLambda([](const FToolMenuContext&)
         {
-			QuickRenderSettings->bShowSettingsBeforeQuickRender = !QuickRenderSettings->bShowSettingsBeforeQuickRender;
-        	QuickRenderSettings->PostEditChange();
+			SetShouldShowSettingsBeforeRender(!GetShouldShowSettingsBeforeRender());
+        	QuickRenderModeSettings->PostEditChange();
         });
         ShowSettingsAction.CanExecuteAction = FToolMenuCanExecuteAction::CreateLambda([](const FToolMenuContext&) { return true; });
         ShowSettingsAction.GetActionCheckState = FToolMenuGetActionCheckState::CreateLambda([](const FToolMenuContext&)
         {
-            return QuickRenderSettings->bShowSettingsBeforeQuickRender
+            return GetShouldShowSettingsBeforeRender()
 				? ECheckBoxState::Checked
 				: ECheckBoxState::Unchecked;
         });
@@ -484,4 +506,70 @@ void FMoviePipelineQuickRenderMenu::GenerateSettingsMenuSection(UToolMenu* InMen
 		OpenSettingsAction);
 }
 
+FSlateIcon FMoviePipelineQuickRenderMenu::GetIconForQuickRenderMode(const EMovieGraphQuickRenderMode QuickRenderMode)
+{
+	switch(QuickRenderMode)
+	{
+	case EMovieGraphQuickRenderMode::CurrentSequence:
+		return FSlateIcon("LevelSequenceEditorStyle", "LevelSequenceEditor.CreateNewLevelSequenceInLevel");
+	case EMovieGraphQuickRenderMode::CurrentViewport:
+		return FSlateIcon(FAppStyle::Get().GetStyleSetName(), "LevelEditor.Tabs.Viewports");
+	case EMovieGraphQuickRenderMode::SelectedCameras:
+		return FSlateIcon(FAppStyle::Get().GetStyleSetName(), "ClassIcon.CameraActor");
+	case EMovieGraphQuickRenderMode::UseViewportCameraInSequence:
+		return FSlateIcon(FAppStyle::Get().GetStyleSetName(), "ClassIcon.CameraActor");
+	default:
+		return FSlateIcon();
+	}
+}
+
+void FMoviePipelineQuickRenderMenu::InitQuickRenderModeSettingsFromMode(const EMovieGraphQuickRenderMode QuickRenderMode)
+{
+	QuickRenderModeSettings = TStrongObjectPtr(UMovieGraphQuickRenderSettings::GetSavedQuickRenderModeSettings(QuickRenderMode));
+	UpdateVariableAssignmentsForCurrentGraph();
+}
+
+void FMoviePipelineQuickRenderMenu::UpdateVariableAssignmentsForCurrentGraph()
+{
+	MoviePipeline::RefreshVariableAssignments(QuickRenderModeSettings->GraphPreset.LoadSynchronous(), QuickRenderModeSettings->GraphVariableAssignments, QuickRenderModeSettings.Get());
+}
+
+EMovieGraphQuickRenderButtonMode FMoviePipelineQuickRenderMenu::GetQuickRenderButtonMode()
+{
+	constexpr EMovieGraphQuickRenderButtonMode DefaultButtonMode = EMovieGraphQuickRenderButtonMode::QuickRender;
+	
+	return static_cast<EMovieGraphQuickRenderButtonMode>(
+		GConfig->GetIntOrDefault(QuickRenderIniSection, QuickRenderIniSettingName_ButtonMode, static_cast<int32>(DefaultButtonMode), GEditorPerProjectIni));
+}
+
+void FMoviePipelineQuickRenderMenu::SetQuickRenderButtonMode(const EMovieGraphQuickRenderButtonMode NewButtonMode)
+{
+	GConfig->SetInt(QuickRenderIniSection, QuickRenderIniSettingName_ButtonMode, static_cast<int32>(NewButtonMode), GEditorPerProjectIni);
+}
+
+EMovieGraphQuickRenderMode FMoviePipelineQuickRenderMenu::GetQuickRenderMode()
+{
+	constexpr EMovieGraphQuickRenderMode DefaultMode = EMovieGraphQuickRenderMode::CurrentSequence;
+	
+	return static_cast<EMovieGraphQuickRenderMode>(
+		GConfig->GetIntOrDefault(QuickRenderIniSection, QuickRenderIniSettingName_Mode, static_cast<int32>(DefaultMode), GEditorPerProjectIni));
+}
+
+void FMoviePipelineQuickRenderMenu::SetQuickRenderMode(const EMovieGraphQuickRenderMode NewMode)
+{
+	GConfig->SetInt(QuickRenderIniSection, QuickRenderIniSettingName_Mode, static_cast<int32>(NewMode), GEditorPerProjectIni);
+}
+
+bool FMoviePipelineQuickRenderMenu::GetShouldShowSettingsBeforeRender()
+{
+	constexpr bool bDefaultShowSettingsBeforeRender = false;
+	
+	return GConfig->GetBoolOrDefault(QuickRenderIniSection, QuickRenderIniSettingName_ShowSettingsBeforeQuickRender, bDefaultShowSettingsBeforeRender, GEditorPerProjectIni);
+}
+
+void FMoviePipelineQuickRenderMenu::SetShouldShowSettingsBeforeRender(const bool bNewValue)
+{
+	GConfig->SetInt(QuickRenderIniSection, QuickRenderIniSettingName_ShowSettingsBeforeQuickRender, bNewValue, GEditorPerProjectIni);
+}
+
 #undef LOCTEXT_NAMESPACE
\ No newline at end of file
diff --git a/Engine/Plugins/MovieScene/MovieRenderPipeline/Source/MovieRenderPipelineEditor/Private/MoviePipelineQuickRenderMenu.h b/Engine/Plugins/MovieScene/MovieRenderPipeline/Source/MovieRenderPipelineEditor/Private/MoviePipelineQuickRenderMenu.h
--- a/Engine/Plugins/MovieScene/MovieRenderPipeline/Source/MovieRenderPipelineEditor/Private/MoviePipelineQuickRenderMenu.h
+++ b/Engine/Plugins/MovieScene/MovieRenderPipeline/Source/MovieRenderPipelineEditor/Private/MoviePipelineQuickRenderMenu.h
@@ -4,6 +4,7 @@
 
 #include "CoreMinimal.h"
 #include "Graph/MovieGraphQuickRenderSettings.h"
+#include "Textures/SlateIcon.h"
 
 struct FToolMenuContext;
 class SWidget;
@@ -37,11 +38,64 @@ private:
 	static void GenerateCreateMenuSection(UToolMenu* InMenu);
 	static void GenerateSettingsMenuSection(UToolMenu* InMenu);
 
+	static FSlateIcon GetIconForQuickRenderMode(EMovieGraphQuickRenderMode QuickRenderMode);
+
+	/** Populates the QuickRenderModeSettings member with the settings for the provided mode. */
+	static void InitQuickRenderModeSettingsFromMode(const EMovieGraphQuickRenderMode QuickRenderMode);
+
+	/** Updates the variable assignments for the graph preset in the currently active mode. */
+	static void UpdateVariableAssignmentsForCurrentGraph();
+
+	// ---- UI state getters and setters ----------------------------------------------------------------------
+
+	/** Gets the mode that the Quick Render button is currently in. */
+	static EMovieGraphQuickRenderButtonMode GetQuickRenderButtonMode();
+
+	/** Sets the mode that the Quick Render button is currently using. */
+	static void SetQuickRenderButtonMode(const EMovieGraphQuickRenderButtonMode NewButtonMode);
+
+	/** Gets the mode that Quick Render will use when starting a new render. */
+	static EMovieGraphQuickRenderMode GetQuickRenderMode();
+
+	/** Sets the mode that Quick Render will use when starting a new render. */
+	static void SetQuickRenderMode(const EMovieGraphQuickRenderMode NewMode);
+
+	/** Gets whether Quick Render should show the Settings dialog before a render starts. */
+	static bool GetShouldShowSettingsBeforeRender();
+
+	/** Sets whether Quick Render should show the Settings dialog before a render starts. */
+	static void SetShouldShowSettingsBeforeRender(const bool bNewValue);
+
+	// --------------------------------------------------------------------------------------------------------
+
 private:
+	/** The name of the section that Quick Render settings are stored in within the target ini file. */
+	inline static const TCHAR* QuickRenderIniSection = TEXT("MovieRenderPipeline.QuickRender");
+
+	/** Ini setting name for: What mode the Quick Render button is using. */
+	inline static const TCHAR* QuickRenderIniSettingName_ButtonMode = TEXT("ButtonMode");
+
+	/** Ini setting name for: The mode that Quick Render will use when a render begins. Generally dictates the level sequence and camera(s) that will be used. */
+	inline static const TCHAR* QuickRenderIniSettingName_Mode = TEXT("Mode");
+
+	/** Ini setting name for: Whether the settings dialog should be shown before a Quick Render is started. */
+	inline static const TCHAR* QuickRenderIniSettingName_ShowSettingsBeforeQuickRender = TEXT("bShowSettingsBeforeQuickRender");
+
+	/** The ID of the Quick Render toolbar that's added to the viewport toolbar. */
 	inline static const FName QuickRenderToolbarId = TEXT("MovieRenderGraphQuickRender");
+
+	/** Weak pointer to the Settings window widget. */
 	inline static TWeakPtr<SWindow> WeakQuickRenderSettingsWindow;
 
-	inline static TStrongObjectPtr<UMovieGraphQuickRenderSettings> QuickRenderSettings = nullptr;
+	/** Weak pointer to the details panel within the Settings window. */
+	inline static TWeakPtr<class IDetailsView> WeakDetailsPanel;
+
+	/** The settings for the currently-active mode within Quick Render. */
+	inline static TStrongObjectPtr<UMovieGraphQuickRenderModeSettings> QuickRenderModeSettings = nullptr;
+
+	/** All of the different mode names that are available to be used. */
+	inline static TArray<TSharedPtr<FName>> QuickRenderModes;
 
-	inline static TArray<TSharedPtr<FName>> QuickRenderSettingsCategories;
+	/** The Quick Render mode that the Settings window is currently using (vs. the mode that's active in the toolbar). */
+	inline static EMovieGraphQuickRenderMode SettingsWindowQuickRenderMode;
 };
diff --git a/Engine/Plugins/MovieScene/MovieRenderPipeline/Source/MovieRenderPipelineEditor/Public/Graph/MovieGraphQuickRender.h b/Engine/Plugins/MovieScene/MovieRenderPipeline/Source/MovieRenderPipelineEditor/Public/Graph/MovieGraphQuickRender.h
--- a/Engine/Plugins/MovieScene/MovieRenderPipeline/Source/MovieRenderPipelineEditor/Public/Graph/MovieGraphQuickRender.h
+++ b/Engine/Plugins/MovieScene/MovieRenderPipeline/Source/MovieRenderPipelineEditor/Public/Graph/MovieGraphQuickRender.h
@@ -5,7 +5,6 @@
 #include "CoreMinimal.h"
 #include "UObject/StrongObjectPtr.h"
 
-#include "Graph/MovieGraphDataTypes.h"
 #include "Graph/MovieGraphQuickRenderSettings.h"
 #include "LevelSequence.h"
 #include "MoviePipelineQueue.h"
@@ -20,8 +19,8 @@ class FEditorViewportClient;
 class MOVIERENDERPIPELINEEDITOR_API FMovieGraphQuickRender
 {
 public:
-	/* Begins a quick render using the provided settings. */
-	static void BeginQuickRender(const UMovieGraphQuickRenderSettings& InQuickRenderSettings);
+	/* Begins a quick render using the provided mode and settings. */
+	static void BeginQuickRender(const EMovieGraphQuickRenderMode QuickRenderMode, const UMovieGraphQuickRenderModeSettings& InQuickRenderSettings);
 
 private:
 	/** Data that was cached prior to PIE starting. */
@@ -44,7 +43,7 @@ private:
 	static UMovieGraphConfig* GetQuickRenderGraph(UMovieGraphConfig* InUserSpecifiedGraph, TMap<TObjectPtr<UMovieGraphConfig>, TObjectPtr<UMovieGraphConfig>>& OutOriginalToDupeMap);
 
 	/** Applies all necessary updates to the duplicated graph to prepare it for a quick render (adjusting start/end frames, etc). */
-	static void ApplyQuickRenderUpdatesToDuplicateGraph(const UMovieGraphQuickRenderSettings& InQuickRenderSettings, UWorld* InEditorWorld, UMoviePipelineExecutorJob* InJob);
+	static void ApplyQuickRenderUpdatesToDuplicateGraph(const UMovieGraphQuickRenderModeSettings& InQuickRenderSettings, UWorld* InEditorWorld, UMoviePipelineExecutorJob* InJob);
 
 	/** Helper for ApplyQuickRenderUpdatesToDuplicateGraph() that applies actor editor visibility (by injecting a visibility modifier into the graph). */
 	static void ApplyQuickRenderUpdatesToDuplicateGraph_ApplyEditorVisibility(const UWorld* InEditorWorld);
@@ -62,16 +61,16 @@ private:
 	static void CachePrePieData(UWorld* InEditorWorld);
 
 	/** Handles restoring data after PIE ends. */
-	static void RestorePreRenderState(const UMovieGraphQuickRenderSettings& InQuickRenderSettings);
+	static void RestorePreRenderState(const EMovieGraphQuickRenderMode QuickRenderMode);
 
 	/**
 	 * Populates RenderingLevelSequence with the current level sequence, or sets up a new level sequence for rendering (depending on the mode that
 	 * Quick Render is using).
 	 */
-	static void SetUpRenderingLevelSequence(const UMovieGraphQuickRenderSettings& InQuickRenderSettings);
+	static void SetUpRenderingLevelSequence(const EMovieGraphQuickRenderMode QuickRenderMode);
 
 	/** Populates UtilityLevelSequence with a level sequence used for utility purposes, or does nothing if no utility sequence is needed for the current mode. */
-	static void SetUpUtilityLevelSequence(const UMovieGraphQuickRenderSettings& InQuickRenderSettings);
+	static void SetUpUtilityLevelSequence(const EMovieGraphQuickRenderMode QuickRenderMode);
 
 	/** Gets the first world found of the specified type. */
 	static UWorld* GetWorldOfType(EWorldType::Type WorldType);
@@ -80,10 +79,10 @@ private:
 	 * Gets the explicitly-set playback range of the level sequence being rendered (set either by a selection in Sequencer, or custom start/end
 	 * frames). If using the unaltered playback range of the level sequence, this will return an empty range.
 	 */
-	static TRange<FFrameNumber> GetPlaybackRange(const UMovieGraphQuickRenderSettings& InQuickRenderSettings);
+	static TRange<FFrameNumber> GetPlaybackRange(const UMovieGraphQuickRenderModeSettings& InQuickRenderSettings);
 
 	/** Does any setup needed in the PIE world prior to the render starting. */
-	static void PerformPreRenderSetup(const UMovieGraphQuickRenderSettings& InQuickRenderSettings);
+	static void PerformPreRenderSetup(const EMovieGraphQuickRenderMode QuickRenderMode);
 
 private:
 	/** The temporary queue that is used by quick render. */
