MoviePipeline - Fixed a missing reference passing the ClothSimCache. This should fix any remaining issues with the Cloth Asset temporal subsampling.

[RN] MoviePipeline - Fixed Cloth Asset simulation when using temporal subsampling.

diff --git a/Engine/Plugins/MovieScene/MovieRenderPipeline/Source/MovieRenderPipelineCore/Private/MoviePipelineUtils.cpp b/Engine/Plugins/MovieScene/MovieRenderPipeline/Source/MovieRenderPipelineCore/Private/MoviePipelineUtils.cpp
--- a/Engine/Plugins/MovieScene/MovieRenderPipeline/Source/MovieRenderPipelineCore/Private/MoviePipelineUtils.cpp
+++ b/Engine/Plugins/MovieScene/MovieRenderPipeline/Source/MovieRenderPipelineCore/Private/MoviePipelineUtils.cpp
@@ -1229,7 +1229,7 @@ namespace UE
 			InOutView->OverridePostProcessSettings(InViewInfo.PostProcessSettings, InViewInfo.PostProcessBlendWeight);
 		}
 
-		void SetSkeletalMeshClothSubSteps(const int32 InSubdivisionCount, UWorld* InWorld, TMap<TWeakObjectPtr<UObject>, TArray<::MoviePipeline::FClothSimSettingsCache>> InClothSimCache)
+		void SetSkeletalMeshClothSubSteps(const int32 InSubdivisionCount, UWorld* InWorld, TMap<TWeakObjectPtr<UObject>, TArray<::MoviePipeline::FClothSimSettingsCache>>& InClothSimCache)
 		{
 			SCOPE_CYCLE_COUNTER(STAT_ClothSubstepAdjust);
 			for (TActorIterator<AActor> ActorIt(InWorld); ActorIt; ++ActorIt)
@@ -1304,7 +1304,7 @@ namespace UE
 			}
 		}
 
-		void RestoreSkeletalMeshClothSubSteps(TMap<TWeakObjectPtr<UObject>, TArray<::MoviePipeline::FClothSimSettingsCache>> InClothSimCache)
+		void RestoreSkeletalMeshClothSubSteps(const TMap<TWeakObjectPtr<UObject>, TArray<::MoviePipeline::FClothSimSettingsCache>>& InClothSimCache)
 		{
 			for (const TPair<TWeakObjectPtr<UObject>, TArray<::MoviePipeline::FClothSimSettingsCache>>& Pair : InClothSimCache)
 			{
diff --git a/Engine/Plugins/MovieScene/MovieRenderPipeline/Source/MovieRenderPipelineCore/Public/Graph/MovieGraphPipeline.h b/Engine/Plugins/MovieScene/MovieRenderPipeline/Source/MovieRenderPipelineCore/Public/Graph/MovieGraphPipeline.h
--- a/Engine/Plugins/MovieScene/MovieRenderPipeline/Source/MovieRenderPipelineCore/Public/Graph/MovieGraphPipeline.h
+++ b/Engine/Plugins/MovieScene/MovieRenderPipeline/Source/MovieRenderPipelineCore/Public/Graph/MovieGraphPipeline.h
@@ -112,8 +112,11 @@ public:
 	TSharedPtr<UE::MovieGraph::IMovieGraphOutputMerger> GetOutputMerger() const { return OutputMerger; }
 
 	/** Gets cached data for cloth simulations, needed to save/restore cloth sub-stepping. */
-	TMap<TWeakObjectPtr<UObject>, TArray<MoviePipeline::FClothSimSettingsCache>> GetClothSimCache() const { return ClothSimCache; }
-	
+	const TMap<TWeakObjectPtr<UObject>, TArray<MoviePipeline::FClothSimSettingsCache>>& GetClothSimCache() const { return ClothSimCache; }
+
+	/** Gets cached data for cloth simulations, needed to save/restore cloth sub-stepping. */
+	TMap<TWeakObjectPtr<UObject>, TArray<MoviePipeline::FClothSimSettingsCache>>& GetClothSimCache() { return ClothSimCache; }
+
 	/** Writing images to disk is an async process. When you start writing, declare a future with the filename you will eventually write to, and complete the future once it is on disk. */
 	void AddOutputFuture(TFuture<bool>&& InOutputFuture, const UE::MovieGraph::FMovieGraphOutputFutureData& InData);
 	
diff --git a/Engine/Plugins/MovieScene/MovieRenderPipeline/Source/MovieRenderPipelineCore/Public/MoviePipelineUtils.h b/Engine/Plugins/MovieScene/MovieRenderPipeline/Source/MovieRenderPipelineCore/Public/MoviePipelineUtils.h
--- a/Engine/Plugins/MovieScene/MovieRenderPipeline/Source/MovieRenderPipelineCore/Public/MoviePipelineUtils.h
+++ b/Engine/Plugins/MovieScene/MovieRenderPipeline/Source/MovieRenderPipelineCore/Public/MoviePipelineUtils.h
@@ -154,8 +154,8 @@ namespace UE
 		MOVIERENDERPIPELINECORE_API FString GetPaddingFormatString(int32 InZeroPadCount, const int32 InFrameNumber);
 		MOVIERENDERPIPELINECORE_API void DoPostProcessBlend(const FVector& InViewLocation, const class UWorld* InWorld, const struct FMinimalViewInfo& InViewInfo, class FSceneView* InOutView);
 
-		MOVIERENDERPIPELINECORE_API void SetSkeletalMeshClothSubSteps(const int32 InSubdivisionCount, UWorld* InWorld, TMap<TWeakObjectPtr<UObject>, TArray<::MoviePipeline::FClothSimSettingsCache>> InClothSimCache);
-		MOVIERENDERPIPELINECORE_API void RestoreSkeletalMeshClothSubSteps(TMap<TWeakObjectPtr<UObject>, TArray<::MoviePipeline::FClothSimSettingsCache>> InClothSimCache);
+		MOVIERENDERPIPELINECORE_API void SetSkeletalMeshClothSubSteps(const int32 InSubdivisionCount, UWorld* InWorld, TMap<TWeakObjectPtr<UObject>, TArray<::MoviePipeline::FClothSimSettingsCache>>& InClothSimCache);
+		MOVIERENDERPIPELINECORE_API void RestoreSkeletalMeshClothSubSteps(const TMap<TWeakObjectPtr<UObject>, TArray<::MoviePipeline::FClothSimSettingsCache>>& InClothSimCache);
 
 
 		/** When using spatial/temporal samples without anti-aliasing, get the sub-pixel jitter for the given frame index. FrameIndex is modded by InSamplesPerFrame so that the aa jitter pattern repeats every output frame. */
